nDisplay: Add StageActor abstract component which is extended by LightCardStageActorComponent and ChromakeyCardStageActorComponent. These components are used to determine root actor ownship of light card and chromakey card actors. This changes the ICVFX panel default behavior to use components instead of layers to determine ownership and prevents having to iterate over every actor to find the correct primitives in a layer, improving per-frame performance.

- Root actor component ownership can be changed through the details panel of light cards, including in bulk. Since we are supporting the Actors and Layers array of the ShowOnlyList, it is possible to have conflicting ownership if an LC is assigned to both this list, and has a root actor pointer set under its component. An editor toast message helps catch this case if conflicting ownership is detected.
- Chromakey Cards are supported in MU/WP -game modes with these changes.
- Remove DataLayer usage completely.
- Fix adding a chromakey card not respecting the current label display values.
- Initial light card creation is transacted properly in MU now since layers aren't used. Previously layer creation wasn't transacted and required a map save.
- Light card positioning has been refactored to be performed in each light card rather than the DCRA tick. This also addresses an issue from 24434309 where lightcards wouldn't stay in place to their orbit origin in all conditions.
- Fix flags having an incorrect icon.
- Fix crash that could occur when removing a layer from a DCRA that a UV light card was associated with.

diff --git a/Engine/Plugins/Runtime/nDisplay/Source/DisplayCluster/DisplayCluster.Build.cs b/Engine/Plugins/Runtime/nDisplay/Source/DisplayCluster/DisplayCluster.Build.cs
--- a/Engine/Plugins/Runtime/nDisplay/Source/DisplayCluster/DisplayCluster.Build.cs
+++ b/Engine/Plugins/Runtime/nDisplay/Source/DisplayCluster/DisplayCluster.Build.cs
@@ -59,8 +59,7 @@ public class DisplayCluster : ModuleRules
 			PrivateDependencyModuleNames.AddRange(
 				new string[]
 				{
-					"ConcertSyncClient",
-					"DataLayerEditor"
+					"ConcertSyncClient"
 				});
 		}
 
diff --git a/Engine/Plugins/Runtime/nDisplay/Source/DisplayCluster/Private/Blueprints/DisplayClusterBlueprintLib.cpp b/Engine/Plugins/Runtime/nDisplay/Source/DisplayCluster/Private/Blueprints/DisplayClusterBlueprintLib.cpp
--- a/Engine/Plugins/Runtime/nDisplay/Source/DisplayCluster/Private/Blueprints/DisplayClusterBlueprintLib.cpp
+++ b/Engine/Plugins/Runtime/nDisplay/Source/DisplayCluster/Private/Blueprints/DisplayClusterBlueprintLib.cpp
@@ -66,7 +66,7 @@ ADisplayClusterLightCardActor* UDisplayClusterBlueprintLib::CreateLightCard(ADis
 	NewActor->AttachToActor(RootActor, AttachmentRules);
 
 	// Add it to the root actor
-	NewActor->AddToLightCardLayer(RootActor);
+	NewActor->AddToRootActor(RootActor);
 
 #if WITH_EDITOR
 	// Required so operator panel updates
@@ -90,7 +90,7 @@ void UDisplayClusterBlueprintLib::DuplicateLightCards(TArray<ADisplayClusterLigh
 			continue;
 		}
 
-		ADisplayClusterRootActor* RootActor = OriginalLightcard->GetRootActorOwner().Get();
+		ADisplayClusterRootActor* RootActor = OriginalLightcard->GetRootActorOwner();
 		if (!RootActor)
 		{
 			continue;
@@ -123,7 +123,7 @@ void UDisplayClusterBlueprintLib::DuplicateLightCards(TArray<ADisplayClusterLigh
 		NewLightCard->AttachToActor(RootActor, AttachmentRules);
 
 		// Add it to the root actor
-		NewLightCard->AddToLightCardLayer(RootActor);
+		NewLightCard->AddToRootActor(RootActor);
 
 #if WITH_EDITOR
 		// Required so operator panel updates
@@ -160,15 +160,7 @@ void UDisplayClusterBlueprintLib::FindLightCardsForRootActor(const ADisplayClust
 		OutLightCards.Add(Cast<ADisplayClusterLightCardActor>(LightCardActor.Get()));
 	}
 
-	// If there are any layers that are specified as light card layers, iterate over all actors in the world and 
-	// add any that are members of any of the light card layers to the list. Only add an actor once, even if it is
-	// in multiple layers
-	if (RootActorLightCards.ActorLayers.IsEmpty())
-	{
-		return;
-	}
-
-	if (UWorld* World = RootActor->GetWorld())
+	if (const UWorld* World = RootActor->GetWorld())
 	{
 		for (TActorIterator<ADisplayClusterLightCardActor> ActorIt(World); ActorIt; ++ActorIt)
 		{
@@ -177,12 +169,22 @@ void UDisplayClusterBlueprintLib::FindLightCardsForRootActor(const ADisplayClust
 				continue;
 			}
 
-			for (const FActorLayer& ActorLayer : RootActorLightCards.ActorLayers)
+			if (ActorIt->GetRootActorOwner() == RootActor)
 			{
-				if (ActorIt->Layers.Contains(ActorLayer.Name))
+				OutLightCards.Add(*ActorIt);
+			}
+			else
+			{
+				// If there are any layers that are specified as light card layers, iterate over all actors in the world and 
+				// add any that are members of any of the light card layers to the list. Only add an actor once, even if it is
+				// in multiple layers
+				for (const FActorLayer& ActorLayer : RootActorLightCards.ActorLayers)
 				{
-					OutLightCards.Add(*ActorIt);
-					break;
+					if (ActorIt->Layers.Contains(ActorLayer.Name))
+					{
+						OutLightCards.Add(*ActorIt);
+						break;
+					}
 				}
 			}
 		}
@@ -214,14 +216,6 @@ void UDisplayClusterBlueprintLib::FindChromakeyCardsForRootActor(const ADisplayC
 			OutChromakeyCards.Add(Cast<ADisplayClusterChromakeyCardActor>(ChromakeyCardActor.Get()));
 		}
 
-		// If there are any layers that are specified as light card layers, iterate over all actors in the world and 
-		// add any that are members of any of the light card layers to the list. Only add an actor once, even if it is
-		// in multiple layers
-		if (RootActorChromakeyCards.ActorLayers.IsEmpty())
-		{
-			continue;
-		}
-
 		if (const UWorld* World = RootActor->GetWorld())
 		{
 			for (TActorIterator<ADisplayClusterChromakeyCardActor> ActorIt(World); ActorIt; ++ActorIt)
@@ -231,12 +225,19 @@ void UDisplayClusterBlueprintLib::FindChromakeyCardsForRootActor(const ADisplayC
 					continue;
 				}
 
-				for (const FActorLayer& ActorLayer : RootActorChromakeyCards.ActorLayers)
+				if (ActorIt->GetRootActorOwner() == RootActor)
 				{
-					if (ActorIt->Layers.Contains(ActorLayer.Name))
+					OutChromakeyCards.Add(*ActorIt);
+				}
+				else
+				{
+					for (const FActorLayer& ActorLayer : RootActorChromakeyCards.ActorLayers)
 					{
-						OutChromakeyCards.Add(*ActorIt);
-						break;
+						if (ActorIt->Layers.Contains(ActorLayer.Name))
+						{
+							OutChromakeyCards.Add(*ActorIt);
+							break;
+						}
 					}
 				}
 			}
diff --git a/Engine/Plugins/Runtime/nDisplay/Source/DisplayCluster/Private/Game/EngineClasses/Scene/DisplayClusterChromakeyCardActor.cpp b/Engine/Plugins/Runtime/nDisplay/Source/DisplayCluster/Private/Game/EngineClasses/Scene/DisplayClusterChromakeyCardActor.cpp
--- a/Engine/Plugins/Runtime/nDisplay/Source/DisplayCluster/Private/Game/EngineClasses/Scene/DisplayClusterChromakeyCardActor.cpp
+++ b/Engine/Plugins/Runtime/nDisplay/Source/DisplayCluster/Private/Game/EngineClasses/Scene/DisplayClusterChromakeyCardActor.cpp
@@ -3,17 +3,14 @@
 #include "DisplayClusterChromakeyCardActor.h"
 
 #include "DisplayClusterRootActor.h"
+#include "Components/DisplayClusterChromakeyCardStageActorComponent.h"
 #include "Components/DisplayClusterICVFXCameraComponent.h"
-#include "Misc/DisplayClusterLog.h"
 
-#if WITH_EDITOR
-#include "Layers/LayersSubsystem.h"
-#include "DataLayer/DataLayerEditorSubsystem.h"
-#endif
 
 ADisplayClusterChromakeyCardActor::ADisplayClusterChromakeyCardActor(const FObjectInitializer& ObjectInitializer)
-	: Super(ObjectInitializer)
+	: Super(ObjectInitializer.DoNotCreateDefaultSubobject(LightCardStageActorComponentName))
 {
+	StageActorComponent = CreateOptionalDefaultSubobject<UDisplayClusterChromakeyCardStageActorComponent>(TEXT("ChromakeyStageActorComponent"));
 }
 
 void ADisplayClusterChromakeyCardActor::Tick(float DeltaSeconds)
@@ -21,115 +18,36 @@ void ADisplayClusterChromakeyCardActor::Tick(float DeltaSeconds)
 	Super::Tick(DeltaSeconds);
 
 #if WITH_EDITOR
-	UpdateChromakeySettings();
+	if (GIsEditor)
+	{
+		UpdateChromakeySettings();
+	}
 #endif
 }
 
-void ADisplayClusterChromakeyCardActor::AddToChromakeyLayer(ADisplayClusterRootActor* InRootActor)
+void ADisplayClusterChromakeyCardActor::AddToRootActor(ADisplayClusterRootActor* InRootActor)
 {
-	check(InRootActor);
-	SetRootActorOwner(InRootActor);
-	
-	TArray<UDisplayClusterICVFXCameraComponent*> ICVFXComponents;
-	InRootActor->GetComponents(ICVFXComponents);
-
-	if (ICVFXComponents.Num() > 0)
-	{
-		static const FString LayerPrefix = TEXT("Chromakey");
-			
-		FName ChromakeyLayerName = NAME_None;
-
-		// Find an existing layer name.
-		for (UDisplayClusterICVFXCameraComponent* ICVFXComponent : ICVFXComponents)
-		{
-			FDisplayClusterConfigurationICVFX_VisibilityList& VisibilityList = ICVFXComponent->CameraSettings.Chromakey.ChromakeyRenderTexture.ShowOnlyList;
-			
-			if (const FActorLayer* ExistingLayer = VisibilityList.ActorLayers.FindByPredicate([](const FActorLayer& Layer)
-			{
-				return Layer.Name.ToString().StartsWith(LayerPrefix);
-			}))
-			{
-				ChromakeyLayerName = ExistingLayer->Name;
-				break;
-			}
-		}
+	Super::AddToRootActor(InRootActor);
+	// @todo uv chromakey cards - determine translucent sort order
+}
 
-		if (ChromakeyLayerName == NAME_None)
-		{
-			// Find a unique name so it's Chromakey_XXXX
-			do
-			{
-				const FGuid LayerGuid = FGuid::NewGuid();
-				const FString GuidStr = LayerGuid.ToString().Left(4);
-				ChromakeyLayerName = *FString::Printf(TEXT("%s_%s"), *LayerPrefix, *GuidStr);
-			}
-			while (FindObject<UObject>(GetWorld(), *ChromakeyLayerName.ToString()) != nullptr);
-		}
+void ADisplayClusterChromakeyCardActor::RemoveFromRootActor()
+{
+	Super::RemoveFromRootActor();
+}
 
-		check(!ChromakeyLayerName.IsNone());
-		bool bLayerAdded = false;
-#if WITH_EDITOR
-		//TArray<AActor*> ActorsInLayer; // @todo uv chromakey cards - this will be used to determine translucent sort order
-	
-		if (SupportsLayers())
-		{
-			if (ULayersSubsystem* LayersSubsystem = GEditor ? GEditor->GetEditorSubsystem<ULayersSubsystem>() : nullptr)
-			{
-				LayersSubsystem->AddActorsToLayer(TArray<AActor*>{ this }, ChromakeyLayerName);
-				bLayerAdded = true;
-			}
-		}
-		else if (SupportsDataLayer() && GetLevel())
-		{
-			if (UDataLayerEditorSubsystem* DataLayersSubsystem = GEditor ? GEditor->GetEditorSubsystem<UDataLayerEditorSubsystem>() : nullptr)
-			{
-				if (UDataLayerAsset* DataLayerAsset = InRootActor->GetOrCreateChromakeyCardDataLayerAsset(ChromakeyLayerName))
-				{
-					UDataLayerInstance* DataLayerInstance = DataLayersSubsystem->GetDataLayerInstance(DataLayerAsset);
-					if (DataLayerInstance == nullptr)
-					{
-						FDataLayerCreationParameters CreationParams;
-						CreationParams.WorldDataLayers = GetLevel()->GetWorldDataLayers();
-						CreationParams.DataLayerAsset = DataLayerAsset;
-						DataLayerInstance = DataLayersSubsystem->CreateDataLayerInstance(MoveTemp(CreationParams));
-					}
-			
-					AddDataLayer(DataLayerInstance);
-				}
-				else
-				{
-					UE_LOG(LogDisplayClusterGame, Error, TEXT("Chromakey Card Actor could not find or create the data layer asset for layer '%s'."),  *ChromakeyLayerName.ToString());
-				}
-			}
-		}
-#endif
-		if (!bLayerAdded)
-		{
-			// Likely -game and we need to add the layer directly
-			Layers.AddUnique(ChromakeyLayerName);
-		}
+bool ADisplayClusterChromakeyCardActor::IsReferencedByICVFXCamera(const UDisplayClusterICVFXCameraComponent* InCamera) const
+{
+	check(InCamera);
 
-		// Make sure all ICVFX components on this root actor have this chromakey layer
-		for (UDisplayClusterICVFXCameraComponent* ICVFXComponent : ICVFXComponents)
+	if (const UDisplayClusterChromakeyCardStageActorComponent* ChromakeyCardComponent = Cast<UDisplayClusterChromakeyCardStageActorComponent>(StageActorComponent))
+	{
+		if (ChromakeyCardComponent->IsReferencedByICVFXCamera(InCamera))
 		{
-			FDisplayClusterConfigurationICVFX_VisibilityList& VisibilityList = ICVFXComponent->CameraSettings.Chromakey.ChromakeyRenderTexture.ShowOnlyList;
-
-			if (!VisibilityList.ActorLayers.FindByPredicate([ChromakeyLayerName](const FActorLayer& Layer)
-			{
-				return Layer.Name == ChromakeyLayerName;
-			}))
-			{
-#if WITH_EDITOR
-				ICVFXComponent->Modify();
-#endif
-				VisibilityList.ActorLayers.AddDefaulted_GetRef().Name = ChromakeyLayerName;
-			}
+			return true;
 		}
 	}
-}
-
-bool ADisplayClusterChromakeyCardActor::IsReferencedByICVFXCamera(UDisplayClusterICVFXCameraComponent* InCamera) const
-{
+	
 	const FDisplayClusterConfigurationICVFX_CameraSettings& CameraSettings = InCamera->GetCameraSettingsICVFX();
 	if (CameraSettings.Chromakey.ChromakeyRenderTexture.ShowOnlyList.Actors.Contains(this))
 	{
@@ -152,14 +70,14 @@ bool ADisplayClusterChromakeyCardActor::IsReferencedByICVFXCamera(UDisplayCluste
 
 void ADisplayClusterChromakeyCardActor::UpdateChromakeySettings()
 {
-	if (RootActorOwner.IsValid())
+	if (const ADisplayClusterRootActor* RootActor = GetRootActorOwner())
 	{
 		int32 Count = 0;
 		FLinearColor ChromaColor(ForceInitToZero);
 		TArray<UDisplayClusterICVFXCameraComponent*> ICVFXComponents;
-		RootActorOwner->GetComponents(ICVFXComponents);
+		RootActor->GetComponents(ICVFXComponents);
 
-		for (UDisplayClusterICVFXCameraComponent* Component : ICVFXComponents)
+		for (const UDisplayClusterICVFXCameraComponent* Component : ICVFXComponents)
 		{
 			if (IsReferencedByICVFXCamera(Component))
 			{
diff --git a/Engine/Plugins/Runtime/nDisplay/Source/DisplayCluster/Private/Game/EngineClasses/Scene/DisplayClusterChromakeyCardStageActorComponent.cpp b/Engine/Plugins/Runtime/nDisplay/Source/DisplayCluster/Private/Game/EngineClasses/Scene/DisplayClusterChromakeyCardStageActorComponent.cpp
new file mode 100644
--- /dev/null
+++ b/Engine/Plugins/Runtime/nDisplay/Source/DisplayCluster/Private/Game/EngineClasses/Scene/DisplayClusterChromakeyCardStageActorComponent.cpp
@@ -0,0 +1,82 @@
+// Copyright Epic Games, Inc. All Rights Reserved.
+
+#include "Components/DisplayClusterChromakeyCardStageActorComponent.h"
+#include "DisplayClusterRootActor.h"
+#include "Components/DisplayClusterICVFXCameraComponent.h"
+
+#if WITH_EDITOR
+void UDisplayClusterChromakeyCardStageActorComponent::PostEditChangeProperty(FPropertyChangedEvent& PropertyChangedEvent)
+{
+	Super::PostEditChangeProperty(PropertyChangedEvent);
+
+	const FName PropertyName = PropertyChangedEvent.GetPropertyName();
+
+	auto RemoveInvalidComponents = [this]()
+	{
+		for (auto ComponentIt = ICVFXCameras.CreateIterator(); ComponentIt; ++ComponentIt)
+		{
+			if (ComponentIt->OtherActor.IsValid() && ComponentIt->OtherActor.Get() != RootActor.Get())
+			{
+				ComponentIt.RemoveCurrent();
+			}
+		}
+	};
+
+	// Make sure component owner and root actor owner are in sync.
+	if (PropertyName == GET_MEMBER_NAME_CHECKED(UDisplayClusterChromakeyCardStageActorComponent, RootActor))
+	{
+		RemoveInvalidComponents();
+		PopulateICVFXOwners();
+	}
+	else if (PropertyName == GET_MEMBER_NAME_CHECKED(UDisplayClusterChromakeyCardStageActorComponent, ICVFXCameras) &&
+		PropertyChangedEvent.ChangeType == EPropertyChangeType::ValueSet)
+	{
+		const int32 ArrayIndex = PropertyChangedEvent.GetArrayIndex(PropertyName.ToString());
+		if (ArrayIndex >= 0 && ArrayIndex < ICVFXCameras.Num())
+		{
+			RootActor = ICVFXCameras[ArrayIndex].OtherActor.Get();
+			RemoveInvalidComponents();
+		}
+	}
+}
+#endif
+
+void UDisplayClusterChromakeyCardStageActorComponent::SetRootActor(ADisplayClusterRootActor* InRootActor)
+{
+	Super::SetRootActor(InRootActor);
+	PopulateICVFXOwners();
+}
+
+bool UDisplayClusterChromakeyCardStageActorComponent::IsReferencedByICVFXCamera(
+	const UDisplayClusterICVFXCameraComponent* InCamera) const
+{
+	for (const FSoftComponentReference& ComponentReference : GetICVFXCameras())
+	{
+		if (ComponentReference.GetComponent(InCamera->GetOwner()) == InCamera)
+		{
+			return true;
+		}
+	}
+
+	return false;
+}
+
+void UDisplayClusterChromakeyCardStageActorComponent::PopulateICVFXOwners()
+{
+	if (RootActor.Get())
+	{
+		TArray<UDisplayClusterICVFXCameraComponent*> ICVFXComponents;
+		RootActor.Get()->GetComponents(ICVFXComponents);
+
+		ICVFXCameras.Reset(ICVFXComponents.Num());
+	
+		for (UDisplayClusterICVFXCameraComponent* ICVFXComponent : ICVFXComponents)
+		{
+			FSoftComponentReference ComponentReference;
+			ComponentReference.OtherActor = RootActor.Get();
+			ComponentReference.ComponentProperty = ICVFXComponent->GetFName();
+		
+			ICVFXCameras.Add(MoveTemp(ComponentReference));
+		}
+	}
+}
diff --git a/Engine/Plugins/Runtime/nDisplay/Source/DisplayCluster/Private/Game/EngineClasses/Scene/DisplayClusterLabelComponent.cpp b/Engine/Plugins/Runtime/nDisplay/Source/DisplayCluster/Private/Game/EngineClasses/Scene/DisplayClusterLabelComponent.cpp
--- a/Engine/Plugins/Runtime/nDisplay/Source/DisplayCluster/Private/Game/EngineClasses/Scene/DisplayClusterLabelComponent.cpp
+++ b/Engine/Plugins/Runtime/nDisplay/Source/DisplayCluster/Private/Game/EngineClasses/Scene/DisplayClusterLabelComponent.cpp
@@ -106,7 +106,7 @@ void UDisplayClusterLabelComponent::UpdateWidgetComponent()
 		{
 			// The light card editor normally will have set this, but if we're running as -game then it will
 			// be set from the DCRA instead
-			SetRootActor(LightCardActor->GetRootActorOwner().Get());
+			SetRootActor(LightCardActor->GetRootActorOwner());
 		}
 
 #if WITH_EDITOR
diff --git a/Engine/Plugins/Runtime/nDisplay/Source/DisplayCluster/Private/Game/EngineClasses/Scene/DisplayClusterLightCardActor.cpp b/Engine/Plugins/Runtime/nDisplay/Source/DisplayCluster/Private/Game/EngineClasses/Scene/DisplayClusterLightCardActor.cpp
--- a/Engine/Plugins/Runtime/nDisplay/Source/DisplayCluster/Private/Game/EngineClasses/Scene/DisplayClusterLightCardActor.cpp
+++ b/Engine/Plugins/Runtime/nDisplay/Source/DisplayCluster/Private/Game/EngineClasses/Scene/DisplayClusterLightCardActor.cpp
@@ -2,6 +2,7 @@
 
 #include "DisplayClusterLightCardActor.h"
 
+#include "Blueprints/DisplayClusterBlueprintLib.h"
 #include "DisplayClusterConfigurationTypes.h"
 #include "DisplayClusterRootActor.h"
 #include "Features/IModularFeatures.h"
@@ -11,6 +12,7 @@
 #include "IDisplayClusterLightCardActorExtender.h"
 
 #include "Components/DisplayClusterLabelComponent.h"
+#include "Components/DisplayClusterLightCardStageActorComponent.h"
 #include "Components/StaticMeshComponent.h"
 #include "Features/IModularFeatures.h"
 #include "GameFramework/SpringArmComponent.h"
@@ -19,11 +21,9 @@
 #include "UObject/ConstructorHelpers.h"
 
 #if WITH_EDITOR
-
-#include "Editor.h"
-#include "Layers/LayersSubsystem.h"
-#include "DataLayer/DataLayerEditorSubsystem.h"
-
+#include "Framework/Application/SlateApplication.h"
+#include "Framework/Notifications/NotificationManager.h"
+#include "Widgets/Notifications/SNotificationList.h"
 #endif
 
 #if WITH_OPENCV
@@ -46,6 +46,7 @@ static FAutoConsoleVariableRef CVarDisplayClusterLightCardPolygonTextureSize(
 
 const float ADisplayClusterLightCardActor::UVPlaneDefaultSize = 200.0f;
 const float ADisplayClusterLightCardActor::UVPlaneDefaultDistance = 100.0f;
+const FName ADisplayClusterLightCardActor::LightCardStageActorComponentName = TEXT("LightCardStageActorComponent");
 
 ADisplayClusterLightCardActor::ADisplayClusterLightCardActor(const FObjectInitializer& ObjectInitializer)
 	: Super(ObjectInitializer)
@@ -103,8 +104,15 @@ ADisplayClusterLightCardActor::ADisplayClusterLightCardActor(const FObjectInitia
 		UVIndicatorComponent->SetRelativeLocation(FVector(0, 0, 1)); // Slightly in front of the parent
 		UVIndicatorComponent->SetVisibility(false);
 	}
+
+	if (!IsTemplate() && GEngine)
+	{
+		GEngine->OnLevelActorDeleted().AddUObject(this, &ADisplayClusterLightCardActor::OnLevelActorDeleted);
+	}
 #endif // WITH_EDITOR
 
+	StageActorComponent = CreateOptionalDefaultSubobject<UDisplayClusterLightCardStageActorComponent>(LightCardStageActorComponentName);
+	
 	UpdateStageActorTransform();
 
 	LabelComponent = CreateOptionalDefaultSubobject<UDisplayClusterLabelComponent>(TEXT("Label"), true);
@@ -113,6 +121,16 @@ ADisplayClusterLightCardActor::ADisplayClusterLightCardActor(const FObjectInitia
 	CreateComponentsForExtenders();
 }
 
+ADisplayClusterLightCardActor::~ADisplayClusterLightCardActor()
+{
+#if WITH_EDITOR
+	if (GEngine)
+	{
+		GEngine->OnLevelActorDeleted().RemoveAll(this);
+	}
+#endif
+}
+
 void ADisplayClusterLightCardActor::PostLoad()
 {
 	Super::PostLoad();
@@ -151,6 +169,7 @@ void ADisplayClusterLightCardActor::Tick(float DeltaSeconds)
 	
 	ClampLatitudeAndLongitude(Latitude, Longitude);
 
+	UpdateLightCardPositionToRootActor();
 	UpdateStageActorTransform();
 	UpdateLightCardMaterialInstance();
 	UpdateLightCardVisibility();
@@ -225,6 +244,11 @@ FName ADisplayClusterLightCardActor::GetCustomIconName() const
 		return TEXT("ClassIcon.DisplayClusterLightCardActor.UVLightCard");
 	}
 
+	if (bIsLightCardFlag)
+	{
+		return TEXT("ClassIcon.DisplayClusterLightCardActor.Flag");
+	}
+	
 	return Super::GetCustomIconName();
 }
 
@@ -463,11 +487,88 @@ void ADisplayClusterLightCardActor::UpdateUVIndicator()
 #endif // WITH_EDITOR
 }
 
+void ADisplayClusterLightCardActor::UpdateLightCardPositionToRootActor()
+{
+	if (!bLockToOwningRootActor || HasAnyFlags(RF_Transient) /* Proxies update from world instance positions instead */)
+	{
+		return;
+	}
+	
+	if (const ADisplayClusterRootActor* RootActor = GetRootActorOwner())
+	{
+		if (const UDisplayClusterConfigurationData* ConfigData = RootActor->GetConfigData())
+		{
+			if (ConfigData->StageSettings.Lightcard.bEnable)
+			{
+				// Find a view origin to use as the transform "anchor" of each light card. At the moment, assume the first view origin component
+				// found is the correct view origin (the same assumption is made in the light card editor). If no view origin is found, use the root component
+				const USceneComponent* ViewOriginComponent = RootActor->GetRootComponent();
+
+				TArray<UDisplayClusterCameraComponent*> ViewOriginComponents;
+				RootActor->GetComponents(ViewOriginComponents);
+
+				if (ViewOriginComponents.Num())
+				{
+					ViewOriginComponent = ViewOriginComponents[0];
+				}
+		
+				const FVector Location = ViewOriginComponent ? ViewOriginComponent->GetComponentLocation() : RootActor->GetActorLocation();
+				const FVector RelativeLocation = RootActor->GetTransform().InverseTransformPosition(Location);
+				const FRotator Rotation = RootActor->GetActorRotation();
+				
+				const bool bRepositionLightCards = RelativeLocation != LastOrbitLocation && !IsUVActor();
+
+				const FTransform OldLightCardTransform = GetStageActorTransform();
+				
+				SetActorLocation(Location);
+				SetActorRotation(Rotation);
+
+				if (bRepositionLightCards)
+				{
+					const FTransform StageActorTransform = GetOrigin();
+					FDisplayClusterPositionalParams PositionalParams = GetPositionalParams();
+					// We want to adjust the positional params so that the light card stays in the same world position.
+					// First, compute the desired relative location based on the new orbit origin and the light card's previous world position
+					// Then convert that to longitude and latitude values
+					const FVector OrbitLocation = Rotation.UnrotateVector(OldLightCardTransform.GetLocation() - Location);
+					PositionalParams.DistanceFromCenter = OrbitLocation.Length();
+					if (PositionalParams.DistanceFromCenter > UE_SMALL_NUMBER)
+					{
+						PositionalParams.Latitude = FMath::RadiansToDegrees(FMath::Asin(OrbitLocation.Z / PositionalParams.DistanceFromCenter));
+						PositionalParams.Longitude = FMath::RadiansToDegrees(FMath::Atan2(OrbitLocation.Y, OrbitLocation.X)) + 180.0;
+					}
+					else
+					{
+						PositionalParams.Latitude = 0.0;
+						PositionalParams.Longitude = 180.0;
+					}
+					// The light card's orientation needs to be adjusted as well, since the orientation parameters are defined in radial space,
+					// which is relative to the orbit origin
+					const FVector WorldNormal = OldLightCardTransform.GetRotation().RotateVector(FVector::ForwardVector);
+					const FVector LocalNormal = FRotationMatrix::MakeFromX(OrbitLocation.GetSafeNormal()).InverseTransformVector(StageActorTransform.InverseTransformVectorNoScale(WorldNormal));
+					const FRotator NormalRotation = FRotationMatrix::MakeFromX(-LocalNormal).Rotator();
+					PositionalParams.Pitch = NormalRotation.Pitch;
+					PositionalParams.Yaw = NormalRotation.Yaw;
+					// The spin also needs to be adjusted. This can be done by transforming the old spin vector from world space into the new normal space,
+					// and then converting to an angle
+					const FVector WorldSpin = OldLightCardTransform.GetRotation().RotateVector(FVector::UpVector);
+					const FVector LocalSpin = FRotationMatrix::MakeFromX(OrbitLocation.GetSafeNormal()).InverseTransformVector(StageActorTransform.InverseTransformVectorNoScale(WorldSpin));
+					const FVector NormalSpin = FRotationMatrix::MakeFromX(LocalNormal).InverseTransformVector(LocalSpin);
+					PositionalParams.Spin = FMath::RadiansToDegrees(FMath::Atan2(NormalSpin.Y, NormalSpin.Z));
+					SetPositionalParams(PositionalParams);
+
+					LastOrbitLocation = RelativeLocation;
+				}
+			}
+		}
+	}
+}
+
 void ADisplayClusterLightCardActor::MakeFlushToWall()
 {
-	if (RootActorOwner.IsValid())
+	if (ADisplayClusterRootActor* RootActor = GetRootActorOwner())
 	{
-		RootActorOwner->MakeStageActorFlushToWall(this);
+		RootActor->MakeStageActorFlushToWall(this);
 	}
 }
 
@@ -493,108 +594,101 @@ void ADisplayClusterLightCardActor::ShowLightCardLabel(bool bValue, float ScaleV
 	LabelComponent->SetWidgetScale(ScaleValue);
 }
 
-void ADisplayClusterLightCardActor::SetRootActorOwner(ADisplayClusterRootActor* InRootActor)
-{
-	RootActorOwner = MakeWeakObjectPtr(InRootActor);
-}
-
-void ADisplayClusterLightCardActor::AddToLightCardLayer(ADisplayClusterRootActor* InRootActor)
+void ADisplayClusterLightCardActor::SetWeakRootActorOwner(ADisplayClusterRootActor* InRootActor)
 {
-	check(InRootActor);
+	// Pre 5.2 light cards won't have a StageActorComponent->RootActor, but still need to know the root actor
+	// owner for certain operations. We shouldn't automatically set the RootActorOwner, because it will cause competing
+	// ownership of the light card. Legacy LCs have a ShowOnlyList on the root actor to determine ownership, whereas
+	// new LCs determine the root actor they belong to.
+	const bool bCanSetWeakActorPtr = !StageActorComponent || !StageActorComponent->GetRootActor().IsValid()
+		|| InRootActor == StageActorComponent->GetRootActor().Get() || InRootActor == nullptr;
+	if (!bCanSetWeakActorPtr && !HasAnyFlags(RF_Transient))
+	{
+#if WITH_EDITOR
+		check(InRootActor && StageActorComponent && StageActorComponent->GetRootActor().IsValid());
+		const FText CurrentClassName = GetClass()->GetDisplayNameText();
+		const FText CurrentActorName = FText::FromString(GetActorLabel());
+		const FText CurrentRootActorName = FText::FromString(StageActorComponent->GetRootActor()->GetActorLabel());
+		const FText NewRootActorName = FText::FromString(InRootActor->GetActorLabel());
+		const FText ErrorMessage = FText::Format(NSLOCTEXT("DisplayClusterLightCard", "InvalidRootActorOwner",
+"The {0} '{1}' belongs to '{2}' but appears in the content for '{3}'. Remove it from the content and only use the {0}'s nDisplay Root Actor property."),
+					CurrentClassName, CurrentActorName, CurrentRootActorName, NewRootActorName);
+		
+		// Only warn once so as not to spam the message log since this function is called per tick.
+		if (!bHadRootActorMismatch && GIsEditor && FSlateApplication::IsInitialized())
+		{
+			FNotificationInfo Info(ErrorMessage);
+			Info.Image = FAppStyle::GetBrush(TEXT("Icons.Error"));
+			Info.bFireAndForget = true;
+			Info.bUseSuccessFailIcons = false;
+			Info.ExpireDuration = 10.0f;
 
-	UDisplayClusterConfigurationData* ConfigData = InRootActor->GetConfigData();
-	check(ConfigData);
-	
-	// Add light cards to an existing layer or a new layer. This helps with restoring snapshots of a DCRA
-	// so they can be done separately from light cards and still have access to all light cards in the layer.
-				
-	static const FString LightCardPrefix = TEXT("LightCards");
-	FName LightCardLayerName = NAME_None;
+			FSlateNotificationManager::Get().AddNotification(Info);
+			UE_LOG(LogDisplayClusterGame, Error, TEXT("%s"), *ErrorMessage.ToString());
 
-	FDisplayClusterConfigurationICVFX_VisibilityList& RootActorLightCards = ConfigData->StageSettings.Lightcard.ShowOnlyList;
-	
-	if (const FActorLayer* ExistingLightCardLayer = RootActorLightCards.ActorLayers.FindByPredicate([](const FActorLayer& Layer)
-	{
-		return Layer.Name.ToString().StartsWith(LightCardPrefix);
-	}))
-	{
-		LightCardLayerName = ExistingLightCardLayer->Name;
+			bHadRootActorMismatch = true;
+		}
+		else
+		{
+			UE_LOG(LogDisplayClusterGame, Verbose, TEXT("%s"), *ErrorMessage.ToString());
+		}
+#endif
 	}
 	else
 	{
-		// Find a unique name so it's LightCards_XXXX
-		do
-		{
-			const FGuid LayerGuid = FGuid::NewGuid();
-			const FString GuidStr = LayerGuid.ToString().Left(4);
-			LightCardLayerName = *FString::Printf(TEXT("%s_%s"), *LightCardPrefix, *GuidStr);
-		}
-		while (FindObject<UObject>(GetWorld(), *LightCardLayerName.ToString()) != nullptr);
-		
-		ConfigData->Modify();
-		RootActorLightCards.ActorLayers.AddDefaulted_GetRef().Name = LightCardLayerName;
+		WeakRootActorOwner = InRootActor;
 	}
-	
-	check(!LightCardLayerName.IsNone());
+}
 
-	bool bLayerAdded = false;
-#if WITH_EDITOR
-	TArray<AActor*> ActorsInLayer;
-	
-	if (SupportsLayers())
+void ADisplayClusterLightCardActor::SetRootActorOwner(ADisplayClusterRootActor* InRootActor)
+{
+	if (StageActorComponent)
 	{
-		if (ULayersSubsystem* LayersSubsystem = GEditor ? GEditor->GetEditorSubsystem<ULayersSubsystem>() : nullptr)
-		{
-			ActorsInLayer = LayersSubsystem->GetActorsFromLayer(LightCardLayerName);
-
-			LayersSubsystem->AddActorsToLayer(TArray<AActor*>{ this }, LightCardLayerName);
-			bLayerAdded = true;
-		}
+		StageActorComponent->SetRootActor(InRootActor);
 	}
-	else if (SupportsDataLayer() && GetLevel())
-	{
-		if (UDataLayerEditorSubsystem* DataLayersSubsystem = GEditor ? GEditor->GetEditorSubsystem<UDataLayerEditorSubsystem>() : nullptr)
-		{
-			if (UDataLayerAsset* DataLayerAsset = InRootActor->GetOrCreateLightCardDataLayerAsset(LightCardLayerName))
-			{
-				UDataLayerInstance* DataLayerInstance = DataLayersSubsystem->GetDataLayerInstance(DataLayerAsset);
-				if (DataLayerInstance == nullptr)
-				{
-					FDataLayerCreationParameters CreationParams;
-					CreationParams.WorldDataLayers = GetLevel()->GetWorldDataLayers();
-					CreationParams.DataLayerAsset = DataLayerAsset;
-					DataLayerInstance = DataLayersSubsystem->CreateDataLayerInstance(MoveTemp(CreationParams));
-				}
+}
 
-				ActorsInLayer = DataLayersSubsystem->GetActorsFromDataLayer(DataLayerInstance);
-			
-				AddDataLayer(DataLayerInstance);
-			}
-			else
-			{
-				UE_LOG(LogDisplayClusterGame, Error, TEXT("Light Card Actor could not find or create the data layer asset for layer '%s'."),  *LightCardLayerName.ToString());
-			}
-		}
-	}
+ADisplayClusterRootActor* ADisplayClusterLightCardActor::GetRootActorOwner() const
+{
+	return StageActorComponent && StageActorComponent->GetRootActor().IsValid() ? StageActorComponent->GetRootActor().Get() : WeakRootActorOwner.Get();
+}
+
+UDisplayClusterStageActorComponent* ADisplayClusterLightCardActor::GetStageActorComponent() const
+{
+	return StageActorComponent;
+}
+
+void ADisplayClusterLightCardActor::AddToLightCardLayer(ADisplayClusterRootActor* InRootActor)
+{
+	AddToRootActor(InRootActor);
+}
+
+void ADisplayClusterLightCardActor::AddToRootActor(ADisplayClusterRootActor* InRootActor)
+{
+	check(InRootActor);
+
+#if WITH_EDITOR
+	Modify();
+#endif
 	
 	// Adjust transparent sort order for UV light cards
-	if (IsUVActor() && ActorsInLayer.Num() > 0)
+	if (IsUVActor())
 	{
+		TSet<ADisplayClusterLightCardActor*> LightCards;
+		UDisplayClusterBlueprintLib::FindLightCardsForRootActor(InRootActor, LightCards);
+		
 		int32 HighestPriority = MIN_int32;
 		bool bExistingPrioritiesFound = false;
 
-		for (AActor* Actor : ActorsInLayer)
+		for (const ADisplayClusterLightCardActor* LightCard : LightCards)
 		{
-			if (const ADisplayClusterLightCardActor* LightCard = Cast<ADisplayClusterLightCardActor>(Actor))
+			if (LightCard->IsUVActor())
 			{
-				if (LightCard->IsUVActor())
+				if (LightCard->LightCardComponent->TranslucencySortPriority > HighestPriority)
 				{
-					if (LightCard->LightCardComponent->TranslucencySortPriority > HighestPriority)
-					{
-						HighestPriority = LightCard->LightCardComponent->TranslucencySortPriority;
-					}
-					bExistingPrioritiesFound = true;
+					HighestPriority = LightCard->LightCardComponent->TranslucencySortPriority;
 				}
+				bExistingPrioritiesFound = true;
 			}
 		}
 
@@ -604,15 +698,20 @@ void ADisplayClusterLightCardActor::AddToLightCardLayer(ADisplayClusterRootActor
 			LightCardComponent->TranslucencySortPriority = HighestPriority < MAX_int32 ? HighestPriority + 1 : HighestPriority;
 		}
 	}
-	
+
+	StageActorComponent->SetRootActor(InRootActor);
+}
+
+void ADisplayClusterLightCardActor::RemoveFromRootActor()
+{
+#if WITH_EDITOR
+	Modify();
 #endif
-	if (!bLayerAdded)
+	if (StageActorComponent)
 	{
-		// Likely -game and we need to add the layer directly
-		Layers.AddUnique(LightCardLayerName);
+		StageActorComponent->SetRootActor(nullptr);
 	}
-
-	SetRootActorOwner(InRootActor);
+	WeakRootActorOwner.Reset();
 }
 
 void ADisplayClusterLightCardActor::UpdateStageActorTransform()
@@ -856,3 +955,13 @@ void ADisplayClusterLightCardActor::CleanUpComponentsForExtenders()
 		}
 	}
 }
+
+#if WITH_EDITOR
+void ADisplayClusterLightCardActor::OnLevelActorDeleted(AActor* DeletedActor)
+{
+	if (DeletedActor && StageActorComponent && Cast<ADisplayClusterRootActor>(DeletedActor) == StageActorComponent->GetRootActor().Get())
+	{
+		RemoveFromRootActor();
+	}
+}
+#endif
diff --git a/Engine/Plugins/Runtime/nDisplay/Source/DisplayCluster/Private/Game/EngineClasses/Scene/DisplayClusterRootActor.cpp b/Engine/Plugins/Runtime/nDisplay/Source/DisplayCluster/Private/Game/EngineClasses/Scene/DisplayClusterRootActor.cpp
--- a/Engine/Plugins/Runtime/nDisplay/Source/DisplayCluster/Private/Game/EngineClasses/Scene/DisplayClusterRootActor.cpp
+++ b/Engine/Plugins/Runtime/nDisplay/Source/DisplayCluster/Private/Game/EngineClasses/Scene/DisplayClusterRootActor.cpp
@@ -44,9 +44,10 @@
 #include "Game/EngineClasses/Scene/DisplayClusterRootActorInitializer.h"
 
 #include "Algo/MaxElement.h"
+#include "Components/DisplayClusterChromakeyCardStageActorComponent.h"
+#include "Components/DisplayClusterStageActorComponent.h"
 #include "Components/DisplayClusterStageGeometryComponent.h"
 #include "UObject/Package.h"
-#include "WorldPartition/DataLayer/DataLayerAsset.h"
 
 
 #if WITH_EDITOR
@@ -55,7 +56,6 @@
 #include "IConcertSyncClient.h"
 
 #include "AssetToolsModule.h"
-#include "DataLayer/DataLayerFactory.h"
 #endif
 
 
@@ -142,53 +142,6 @@ const FDisplayClusterConfigurationRenderFrame& ADisplayClusterRootActor::GetRend
 	return CurrentConfigData->RenderFrameSettings;
 }
 
-UDataLayerAsset* ADisplayClusterRootActor::GetOrCreateLightCardDataLayerAsset(const FName& InName)
-{
-#if WITH_EDITOR
-	if (LightCardDataLayerAsset == nullptr)
-	{
-		Modify();
-		LightCardDataLayerAsset = CreateDataLayerAsset(InName);
-	}
-#endif
-
-	return LightCardDataLayerAsset;
-}
-
-UDataLayerAsset* ADisplayClusterRootActor::GetOrCreateChromakeyCardDataLayerAsset(const FName& InName)
-{
-#if WITH_EDITOR
-	if (ChromakeyCardDataLayerAsset == nullptr)
-	{
-		Modify();
-		ChromakeyCardDataLayerAsset = CreateDataLayerAsset(InName);
-	}
-#endif
-
-	return ChromakeyCardDataLayerAsset;
-}
-
-#if WITH_EDITOR
-UDataLayerAsset* ADisplayClusterRootActor::CreateDataLayerAsset(const FName& InName)
-{
-	IAssetTools& AssetTools = FModuleManager::GetModuleChecked<FAssetToolsModule>("AssetTools").Get();
-
-	const FString PackageName = GetPackage()->GetName();
-	FString AssetName;
-	FString TempPackageName;
-
-	UDataLayerFactory* Factory = NewObject<UDataLayerFactory>();
-
-	AssetTools.CreateUniqueAssetName(PackageName / InName.ToString(), FString(),
-		TempPackageName, AssetName);
-
-	UClass* AssetClass = Factory->GetSupportedClass();
-	
-	return Cast<UDataLayerAsset>(AssetTools.CreateAsset(AssetName,
-		PackageName, AssetClass, Factory));
-}
-#endif
-
 void ADisplayClusterRootActor::InitializeFromConfig(UDisplayClusterConfigurationData* ConfigData)
 {
 	if (ConfigData)
@@ -768,152 +721,106 @@ bool ADisplayClusterRootActor::BuildHierarchy()
 	return true;
 }
 
-void ADisplayClusterRootActor::UpdateLightCardPositions()
+void ADisplayClusterRootActor::SetLightCardOwnership()
 {
-	// Find a view origin to use as the transform "anchor" of each light card. At the moment, assume the first view origin component
-	// found is the correct view origin (the same assumption is made in the light card editor). If no view origin is found, use the root component
-	USceneComponent* ViewOriginComponent = GetRootComponent();
-
-	TArray<UDisplayClusterCameraComponent*> ViewOriginComponents;
-	GetComponents(ViewOriginComponents);
-
-	if (ViewOriginComponents.Num())
-	{
-		ViewOriginComponent = ViewOriginComponents[0];
-	}
-
-	const FVector Location = ViewOriginComponent ? ViewOriginComponent->GetComponentLocation() : GetActorLocation();
-	const FVector RelativeLocation = GetTransform().InverseTransformPosition(Location);
-	const FRotator Rotation = GetActorRotation();
-
-	const bool bRepositionLightCards = RelativeLocation != StageActorOrbitLocation;
-
-	// Iterate over all the light cards referenced in the root actor's config data and update their location and rotation to match the view origin
 	if (UDisplayClusterConfigurationData* CurrentData = GetConfigData())
 	{
-		FDisplayClusterConfigurationICVFX_VisibilityList& LightCards = CurrentData->StageSettings.Lightcard.ShowOnlyList;
+		FDisplayClusterConfigurationICVFX_VisibilityList& LightCardVisibilityList = CurrentData->StageSettings.Lightcard.ShowOnlyList;
+		LightCardVisibilityList.AutoAddedActors.Reset();
+		
+		TArray<UDisplayClusterICVFXCameraComponent*> ICVFXComponents;
+		GetComponents(ICVFXComponents);
 
-		for (const TSoftObjectPtr<AActor>& Actor : LightCards.Actors)
+		// Iterate over the VisibilityList.Actors array, looking for legacy cards and set the owner so the
+		// card can look the root actor up in certain situations, like adjusting labels when running as -game.
 		{
-			ADisplayClusterLightCardActor* LightCardActor = Actor.IsValid() ? Cast<ADisplayClusterLightCardActor>(Actor.Get()) : nullptr;
-			if (LightCardActor)
+			if (CurrentData->StageSettings.Lightcard.bEnable)
 			{
-				// Set the owner so the light card actor can look the root actor up in certain situations, like adjusting labels
-				// when running as -game
-				LightCardActor->SetRootActorOwner(this);
-				if (LightCardActor->bLockToOwningRootActor)
+				for (const TSoftObjectPtr<AActor>& Actor : LightCardVisibilityList.Actors)
 				{
-					const FTransform OldLightCardTransform = LightCardActor->GetStageActorTransform();
-
-					LightCardActor->SetActorLocation(Location);
-					LightCardActor->SetActorRotation(Rotation);
-
-					if (bRepositionLightCards)
+					if (ADisplayClusterLightCardActor* LightCardActor = Actor.IsValid() ? Cast<ADisplayClusterLightCardActor>(Actor.Get()) : nullptr)
 					{
-						const FTransform StageActorTransform = LightCardActor->GetOrigin();
-						FDisplayClusterPositionalParams PositionalParams = LightCardActor->GetPositionalParams();
-
-						// We want to adjust the positional params so that the light card stays in the same world position.
-						// First, compute the desired relative location based on the new orbit origin and the light card's previous world position
-						// Then convert that to longitude and latitude values
-						const FVector OrbitLocation = Rotation.UnrotateVector(OldLightCardTransform.GetLocation() - Location);
-
-						PositionalParams.DistanceFromCenter = OrbitLocation.Length();
-
-						if (PositionalParams.DistanceFromCenter > UE_SMALL_NUMBER)
-						{
-							PositionalParams.Latitude = FMath::RadiansToDegrees(FMath::Asin(OrbitLocation.Z / PositionalParams.DistanceFromCenter));
-							PositionalParams.Longitude = FMath::RadiansToDegrees(FMath::Atan2(OrbitLocation.Y, OrbitLocation.X)) + 180.0;
-						}
-						else
-						{
-							PositionalParams.Latitude = 0.0;
-							PositionalParams.Longitude = 180.0;
-						}
-
-						// The light card's orientation needs to be adjusted as well, since the orientation parameters are defined in radial space,
-						// which is relative to the orbit origin
-						const FVector WorldNormal = OldLightCardTransform.GetRotation().RotateVector(FVector::ForwardVector);
-						const FVector LocalNormal = FRotationMatrix::MakeFromX(OrbitLocation.GetSafeNormal()).InverseTransformVector(StageActorTransform.InverseTransformVectorNoScale(WorldNormal));
-						const FRotator NormalRotation = FRotationMatrix::MakeFromX(-LocalNormal).Rotator();
-
-						PositionalParams.Pitch = NormalRotation.Pitch;
-						PositionalParams.Yaw = NormalRotation.Yaw;
-
-						// The spin also needs to be adjusted. This can be done by transforming the old spin vector from world space into the new normal space,
-						// and then converting to an angle
-						const FVector WorldSpin = OldLightCardTransform.GetRotation().RotateVector(FVector::UpVector);
-						const FVector LocalSpin = FRotationMatrix::MakeFromX(OrbitLocation.GetSafeNormal()).InverseTransformVector(StageActorTransform.InverseTransformVectorNoScale(WorldSpin));
-						const FVector NormalSpin = FRotationMatrix::MakeFromX(LocalNormal).InverseTransformVector(LocalSpin);
-
-						PositionalParams.Spin = FMath::RadiansToDegrees(FMath::Atan2(NormalSpin.Y, NormalSpin.Z));
-
-						LightCardActor->SetPositionalParams(PositionalParams);
+						LightCardActor->SetWeakRootActorOwner(this);
+					}
+				}
+			}
+		
+			for (UDisplayClusterICVFXCameraComponent* Camera : ICVFXComponents)
+			{
+				if (!Camera->CameraSettings.Chromakey.bEnable)
+				{
+					continue;
+				}
+				
+				FDisplayClusterConfigurationICVFX_VisibilityList& ChromakeyCards = Camera->CameraSettings.Chromakey.ChromakeyRenderTexture.ShowOnlyList;
+				ChromakeyCards.AutoAddedActors.Reset();
+				
+				for (const TSoftObjectPtr<AActor>& Actor : ChromakeyCards.Actors)
+				{
+					if (ADisplayClusterChromakeyCardActor* ChromakeyCardActor = Actor.IsValid() ? Cast<ADisplayClusterChromakeyCardActor>(Actor.Get()) : nullptr)
+					{
+						ChromakeyCardActor->SetWeakRootActorOwner(this);
 					}
 				}
 			}
 		}
 
-		if (LightCards.ActorLayers.Num())
+		// Next iterate over all light card world actors looking for 5.2+ light cards that determine the root actor
+		// they belong to, as well as handle legacy layer operations.
+		if (const UWorld* World = GetWorld())
 		{
-			if (UWorld* World = GetWorld())
+			for (TActorIterator<ADisplayClusterLightCardActor> Iter(World); Iter; ++Iter)
 			{
-				for (TActorIterator<ADisplayClusterLightCardActor> Iter(World); Iter; ++Iter)
+				ADisplayClusterLightCardActor* LightCardActor = *Iter;
+				if (ADisplayClusterChromakeyCardActor* ChromakeyCardActor = Cast<ADisplayClusterChromakeyCardActor>(LightCardActor))
 				{
-					ADisplayClusterLightCardActor* LightCardActor = *Iter;
-					for (const FActorLayer& ActorLayer : LightCards.ActorLayers)
+					for (UDisplayClusterICVFXCameraComponent* Camera : ICVFXComponents)
 					{
-						if (LightCardActor->Layers.Contains(ActorLayer.Name))
+						if (!Camera->CameraSettings.Chromakey.bEnable)
+						{
+							continue;
+						}
+
+						FDisplayClusterConfigurationICVFX_VisibilityList& ChromakeyCards = Camera->CameraSettings.Chromakey.ChromakeyRenderTexture.ShowOnlyList;
+
+						const UDisplayClusterChromakeyCardStageActorComponent* ChromakeyStageActor = Cast<UDisplayClusterChromakeyCardStageActorComponent>(ChromakeyCardActor->GetStageActorComponent());
+						if (ChromakeyStageActor && ChromakeyStageActor->IsReferencedByICVFXCamera(Camera))
 						{
-							LightCardActor->SetRootActorOwner(this);
-							if (LightCardActor->bLockToOwningRootActor)
+							// 5.2+ chromakey cards
+							ChromakeyCards.AutoAddedActors.Add(ChromakeyCardActor);
+						}
+						else
+						{
+							// Legacy chromakey layers
+							for (const FActorLayer& ActorLayer : ChromakeyCards.ActorLayers)
 							{
-								LightCardActor->SetActorLocation(Location);
-								LightCardActor->SetActorRotation(Rotation);
+								if (ChromakeyCardActor->Layers.Contains(ActorLayer.Name))
+								{
+									ChromakeyCardActor->SetWeakRootActorOwner(this);
+									break;
+								}
 							}
-							break;
 						}
 					}
 				}
-			}
-		}
-	}
-
-	StageActorOrbitLocation = RelativeLocation;
-}
-
-void ADisplayClusterRootActor::SetChromakeyCardsOwner()
-{
-	TArray<UDisplayClusterICVFXCameraComponent*> ICVFXComponents;
-	GetComponents(ICVFXComponents);
-
-	for (UDisplayClusterICVFXCameraComponent* Camera : ICVFXComponents)
-	{
-		FDisplayClusterConfigurationICVFX_VisibilityList& ChromakeyCards = Camera->CameraSettings.Chromakey.ChromakeyRenderTexture.ShowOnlyList;
-
-		for (const TSoftObjectPtr<AActor>& Actor : ChromakeyCards.Actors)
-		{
-			if (ADisplayClusterChromakeyCardActor* ChromakeyCardActor = Actor.IsValid() ? Cast<ADisplayClusterChromakeyCardActor>(Actor.Get()) : nullptr)
-			{
-				// Set the owner so the light card actor can look the root actor up in certain situations, like adjusting labels
-				// when running as -game
-				ChromakeyCardActor->SetRootActorOwner(this);
-			}
-		}
-
-		if (ChromakeyCards.ActorLayers.Num())
-		{
-			if (const UWorld* World = GetWorld())
-			{
-				for (TActorIterator<ADisplayClusterChromakeyCardActor> Iter(World); Iter; ++Iter)
+				else if (CurrentData->StageSettings.Lightcard.bEnable)
 				{
-					ADisplayClusterChromakeyCardActor* ChromakeyCardActor = *Iter;
-					for (const FActorLayer& ActorLayer : ChromakeyCards.ActorLayers)
+					if (LightCardActor->GetStageActorComponent() &&
+						LightCardActor->GetStageActorComponent()->GetRootActor() == this)
+					{
+						// 5.2+ light cards
+						LightCardVisibilityList.AutoAddedActors.Add(LightCardActor);
+					}
+					else if (LightCardVisibilityList.ActorLayers.Num() > 0)
 					{
-						if (ChromakeyCardActor->Layers.Contains(ActorLayer.Name))
+						// Legacy light card layers
+						for (const FActorLayer& ActorLayer : LightCardVisibilityList.ActorLayers)
 						{
-							ChromakeyCardActor->SetRootActorOwner(this);
-							break;
+							if (LightCardActor->Layers.Contains(ActorLayer.Name))
+							{
+								LightCardActor->SetWeakRootActorOwner(this);
+								break;
+							}
 						}
 					}
 				}
@@ -1049,13 +956,7 @@ void ADisplayClusterRootActor::Tick(float DeltaSeconds)
 	Tick_Editor(DeltaSeconds);
 #endif
 
-	if (UDisplayClusterConfigurationData* CurrentData = GetConfigData())
-	{
-		if (CurrentData->StageSettings.Lightcard.bEnable)
-		{
-			UpdateLightCardPositions();
-		}
-	}
+	SetLightCardOwnership();
 
 	Super::Tick(DeltaSeconds);
 }
diff --git a/Engine/Plugins/Runtime/nDisplay/Source/DisplayCluster/Private/Game/EngineClasses/Scene/DisplayClusterRootActorEditor.cpp b/Engine/Plugins/Runtime/nDisplay/Source/DisplayCluster/Private/Game/EngineClasses/Scene/DisplayClusterRootActorEditor.cpp
--- a/Engine/Plugins/Runtime/nDisplay/Source/DisplayCluster/Private/Game/EngineClasses/Scene/DisplayClusterRootActorEditor.cpp
+++ b/Engine/Plugins/Runtime/nDisplay/Source/DisplayCluster/Private/Game/EngineClasses/Scene/DisplayClusterRootActorEditor.cpp
@@ -176,8 +176,6 @@ void ADisplayClusterRootActor::Tick_Editor(float DeltaSeconds)
 			}
 		}
 	}
-	
-	SetChromakeyCardsOwner();
 }
 
 void ADisplayClusterRootActor::PostActorCreated_Editor()
diff --git a/Engine/Plugins/Runtime/nDisplay/Source/DisplayCluster/Private/Game/EngineClasses/Scene/DisplayClusterStageActorComponent.cpp b/Engine/Plugins/Runtime/nDisplay/Source/DisplayCluster/Private/Game/EngineClasses/Scene/DisplayClusterStageActorComponent.cpp
new file mode 100644
--- /dev/null
+++ b/Engine/Plugins/Runtime/nDisplay/Source/DisplayCluster/Private/Game/EngineClasses/Scene/DisplayClusterStageActorComponent.cpp
@@ -0,0 +1,15 @@
+// Copyright Epic Games, Inc. All Rights Reserved.
+
+#include "Components/DisplayClusterStageActorComponent.h"
+#include "DisplayClusterRootActor.h"
+#include "Components/DisplayClusterICVFXCameraComponent.h"
+
+void UDisplayClusterStageActorComponent::SetRootActor(ADisplayClusterRootActor* InRootActor)
+{
+	RootActor = InRootActor;
+}
+
+const TSoftObjectPtr<ADisplayClusterRootActor>& UDisplayClusterStageActorComponent::GetRootActor() const
+{
+	return RootActor;
+}
\ No newline at end of file
diff --git a/Engine/Plugins/Runtime/nDisplay/Source/DisplayCluster/Private/Render/Viewport/Configuration/DisplayClusterViewportConfigurationHelpers_ICVFX.cpp b/Engine/Plugins/Runtime/nDisplay/Source/DisplayCluster/Private/Render/Viewport/Configuration/DisplayClusterViewportConfigurationHelpers_ICVFX.cpp
--- a/Engine/Plugins/Runtime/nDisplay/Source/DisplayCluster/Private/Render/Viewport/Configuration/DisplayClusterViewportConfigurationHelpers_ICVFX.cpp
+++ b/Engine/Plugins/Runtime/nDisplay/Source/DisplayCluster/Private/Render/Viewport/Configuration/DisplayClusterViewportConfigurationHelpers_ICVFX.cpp
@@ -675,7 +675,7 @@ void FDisplayClusterViewportConfigurationHelpers_ICVFX::UpdateChromakeyViewportS
 		// Update visibility settings only for rendered viewports
 		if (!DstViewport.PostRenderSettings.Replace.IsEnabled())
 		{
-			check(FDisplayClusterViewportConfigurationHelpers_Visibility::IsValid(InRenderSettings.ShowOnlyList));
+			check(FDisplayClusterViewportConfigurationHelpers_Visibility::IsVisibilityListValid(InRenderSettings.ShowOnlyList));
 
 			FDisplayClusterViewportConfigurationHelpers_Visibility::UpdateShowOnlyList(DstViewport, RootActor, InRenderSettings.ShowOnlyList);
 		}
@@ -782,17 +782,17 @@ bool FDisplayClusterViewportConfigurationHelpers_ICVFX::IsShouldUseLightcard(con
 	}
 
 	// Lightcard require layers for render
-	return FDisplayClusterViewportConfigurationHelpers_Visibility::IsValid(InLightcardSettings.ShowOnlyList);
+	return FDisplayClusterViewportConfigurationHelpers_Visibility::IsVisibilityListValid(InLightcardSettings.ShowOnlyList);
 }
 
 bool FDisplayClusterViewportConfigurationHelpers_ICVFX::IsShouldUseUVLightcard(FDisplayClusterViewportManager& InViewportManager, const FDisplayClusterConfigurationICVFX_LightcardSettings& InLightcardSettings)
 {
 	if (InLightcardSettings.bEnable)
 	{
-		TSharedPtr<FDisplayClusterViewportLightCardManager, ESPMode::ThreadSafe> LightCardManager = InViewportManager.GetLightCardManager();;
+		const TSharedPtr<FDisplayClusterViewportLightCardManager, ESPMode::ThreadSafe> LightCardManager = InViewportManager.GetLightCardManager();
 		if (LightCardManager.IsValid() && LightCardManager->IsUVLightCardEnabled())
 		{
-			return true;
+			return FDisplayClusterViewportConfigurationHelpers_Visibility::IsVisibilityListValid(InLightcardSettings.ShowOnlyList);
 		}
 	}
 
@@ -834,7 +834,7 @@ void FDisplayClusterViewportConfigurationHelpers_ICVFX::UpdateLightcardViewportS
 		// Update visibility settings only for rendered viewports
 		if (!DstViewport.PostRenderSettings.Replace.IsEnabled())
 		{
-			check(FDisplayClusterViewportConfigurationHelpers_Visibility::IsValid(LightcardSettings.ShowOnlyList));
+			check(FDisplayClusterViewportConfigurationHelpers_Visibility::IsVisibilityListValid(LightcardSettings.ShowOnlyList));
 
 			FDisplayClusterViewportConfigurationHelpers_Visibility::UpdateShowOnlyList(DstViewport, RootActor, LightcardSettings.ShowOnlyList);
 		}
diff --git a/Engine/Plugins/Runtime/nDisplay/Source/DisplayCluster/Private/Render/Viewport/Configuration/DisplayClusterViewportConfigurationHelpers_Visibility.cpp b/Engine/Plugins/Runtime/nDisplay/Source/DisplayCluster/Private/Render/Viewport/Configuration/DisplayClusterViewportConfigurationHelpers_Visibility.cpp
--- a/Engine/Plugins/Runtime/nDisplay/Source/DisplayCluster/Private/Render/Viewport/Configuration/DisplayClusterViewportConfigurationHelpers_Visibility.cpp
+++ b/Engine/Plugins/Runtime/nDisplay/Source/DisplayCluster/Private/Render/Viewport/Configuration/DisplayClusterViewportConfigurationHelpers_Visibility.cpp
@@ -20,20 +20,27 @@ static void ImplCollectVisibility(ADisplayClusterRootActor& RootActor, const FDi
 		RootActor.FindPrimitivesByName(InVisibilityList.RootActorComponentNames, OutAdditionalComponentsList);
 	}
 
-	// Collect Actors refs
-	for (const TSoftObjectPtr<AActor>& ActorSOPtrIt : InVisibilityList.Actors)
+	auto CollectActorRefs = [&](const TArray<TSoftObjectPtr<AActor>>& InActorRefs)
+ 
 	{
-		if (ActorSOPtrIt.IsValid())
+		for (const TSoftObjectPtr<AActor>& ActorSOPtrIt : InActorRefs)
 		{
-			for (const UActorComponent* Component : ActorSOPtrIt->GetComponents())
+			if (ActorSOPtrIt.IsValid())
 			{
-				if (const UPrimitiveComponent* PrimComp = Cast<UPrimitiveComponent>(Component))
+				for (const UActorComponent* Component : ActorSOPtrIt->GetComponents())
 				{
-					OutAdditionalComponentsList.Add(PrimComp->ComponentId);
+					if (const UPrimitiveComponent* PrimComp = Cast<UPrimitiveComponent>(Component))
+					{
+						OutAdditionalComponentsList.Add(PrimComp->ComponentId);
+					}
 				}
 			}
 		}
-	}
+	};
+	
+	// Collect Actors refs
+	CollectActorRefs(InVisibilityList.Actors);
+	CollectActorRefs(InVisibilityList.AutoAddedActors);
 
 	// Collect ActorLayers
 	for (const FActorLayer& ActorLayerIt : InVisibilityList.ActorLayers)
@@ -45,7 +52,7 @@ static void ImplCollectVisibility(ADisplayClusterRootActor& RootActor, const FDi
 	}
 }
 
-bool FDisplayClusterViewportConfigurationHelpers_Visibility::IsValid(const FDisplayClusterConfigurationICVFX_VisibilityList& InVisibilityList)
+bool FDisplayClusterViewportConfigurationHelpers_Visibility::IsVisibilityListValid(const FDisplayClusterConfigurationICVFX_VisibilityList& InVisibilityList)
 {
 	for (const FString& ComponentNameIt : InVisibilityList.RootActorComponentNames)
 	{
@@ -71,6 +78,14 @@ bool FDisplayClusterViewportConfigurationHelpers_Visibility::IsValid(const FDisp
 		}
 	}
 
+	for (const TSoftObjectPtr<AActor>& AutoAddedActor : InVisibilityList.AutoAddedActors)
+	{
+		if (AutoAddedActor.IsValid())
+		{
+			return true;
+		}
+	}
+	
 	return false;
 }
 
diff --git a/Engine/Plugins/Runtime/nDisplay/Source/DisplayCluster/Private/Render/Viewport/Configuration/DisplayClusterViewportConfigurationHelpers_Visibility.h b/Engine/Plugins/Runtime/nDisplay/Source/DisplayCluster/Private/Render/Viewport/Configuration/DisplayClusterViewportConfigurationHelpers_Visibility.h
--- a/Engine/Plugins/Runtime/nDisplay/Source/DisplayCluster/Private/Render/Viewport/Configuration/DisplayClusterViewportConfigurationHelpers_Visibility.h
+++ b/Engine/Plugins/Runtime/nDisplay/Source/DisplayCluster/Private/Render/Viewport/Configuration/DisplayClusterViewportConfigurationHelpers_Visibility.h
@@ -12,8 +12,8 @@ struct FDisplayClusterConfigurationICVFX_VisibilityList;
 class FDisplayClusterViewportConfigurationHelpers_Visibility
 {
 public:
-	// Retur true, if visibility list not empty
-	static bool IsValid(const FDisplayClusterConfigurationICVFX_VisibilityList& InVisibilityList);
+	// Return true, if visibility list not empty
+	static bool IsVisibilityListValid(const FDisplayClusterConfigurationICVFX_VisibilityList& InVisibilityList);
 
 	// Update ShowOnly list for DstViewport
 	static void UpdateShowOnlyList(FDisplayClusterViewport& DstViewport, ADisplayClusterRootActor& RootActor, const FDisplayClusterConfigurationICVFX_VisibilityList& InVisibilityList);
diff --git a/Engine/Plugins/Runtime/nDisplay/Source/DisplayCluster/Private/Render/Viewport/Configuration/DisplayClusterViewportConfigurationICVFX.cpp b/Engine/Plugins/Runtime/nDisplay/Source/DisplayCluster/Private/Render/Viewport/Configuration/DisplayClusterViewportConfigurationICVFX.cpp
--- a/Engine/Plugins/Runtime/nDisplay/Source/DisplayCluster/Private/Render/Viewport/Configuration/DisplayClusterViewportConfigurationICVFX.cpp
+++ b/Engine/Plugins/Runtime/nDisplay/Source/DisplayCluster/Private/Render/Viewport/Configuration/DisplayClusterViewportConfigurationICVFX.cpp
@@ -232,7 +232,7 @@ protected:
 				return false;
 			}
 
-			bool bIsChromakeyHasAnyRenderComponent = FDisplayClusterViewportConfigurationHelpers_Visibility::IsValid(ChromakeySettings.ChromakeyRenderTexture.ShowOnlyList);
+			const bool bIsChromakeyHasAnyRenderComponent = FDisplayClusterViewportConfigurationHelpers_Visibility::IsVisibilityListValid(ChromakeySettings.ChromakeyRenderTexture.ShowOnlyList);
 			if (ChromakeySettings.ChromakeyRenderTexture.Replace.bAllowReplace == false && !bIsChromakeyHasAnyRenderComponent)
 			{
 				// ChromakeyRenderTexture requires actors for render.
diff --git a/Engine/Plugins/Runtime/nDisplay/Source/DisplayCluster/Public/Components/DisplayClusterChromakeyCardStageActorComponent.h b/Engine/Plugins/Runtime/nDisplay/Source/DisplayCluster/Public/Components/DisplayClusterChromakeyCardStageActorComponent.h
new file mode 100644
--- /dev/null
+++ b/Engine/Plugins/Runtime/nDisplay/Source/DisplayCluster/Public/Components/DisplayClusterChromakeyCardStageActorComponent.h
@@ -0,0 +1,48 @@
+// Copyright Epic Games, Inc. All Rights Reserved.
+
+#pragma once
+
+#include "DisplayClusterStageActorComponent.h"
+
+#include "DisplayClusterChromakeyCardStageActorComponent.generated.h"
+
+class ADisplayClusterRootActor;
+class UDisplayClusterICVFXCameraComponent;
+
+/**
+ * Stage Actor Component to be placed in chromakey card actors
+ */
+UCLASS(MinimalAPI, ClassGroup = (DisplayCluster), meta = (DisplayName = "Chromakey Card Stage Actor"), HideCategories=(Physics, Collision, Lighting, Navigation, Cooking, LOD, MaterialParameters, HLOD, RayTracing, TextureStreaming, Mobile))
+class UDisplayClusterChromakeyCardStageActorComponent final : public UDisplayClusterStageActorComponent
+{
+	GENERATED_BODY()
+
+public:
+#if WITH_EDITOR
+	virtual void PostEditChangeProperty(FPropertyChangedEvent& PropertyChangedEvent) override;
+#endif
+
+	DISPLAYCLUSTER_API virtual void SetRootActor(ADisplayClusterRootActor* InRootActor) override;
+
+	/** Retrieve current ICVFX cameras associated with this stage actor */
+	const TArray<FSoftComponentReference>& GetICVFXCameras() const { return ICVFXCameras; }
+
+	/**
+	* Checks if the given ICVFX camera has chroma key settings supporting this actor.
+	* 
+	* @param InCamera The ICVFX camera component to compare against.
+	*
+	* @return true if the ICVFX camera component is referenced by this component.
+	*/
+	DISPLAYCLUSTER_API bool IsReferencedByICVFXCamera(const UDisplayClusterICVFXCameraComponent* InCamera) const;
+	
+protected:
+	/** Update linked ICVFX cameras based on current settings */
+	void PopulateICVFXOwners();
+	
+protected:
+	/** Associate this Chromakey Card with the nDisplay configuration actor and ICVFX Camera component(s) specified here.
+	 * The content will only appear in the ICVFX Editor when this nDisplay configuration is selected. */
+	UPROPERTY(EditInstanceOnly, Category = "NDisplay Root Actor", DisplayName = "ICVFX Cameras", meta = (UseComponentPicker, AllowedClasses = "/Script/DisplayCluster.DisplayClusterICVFXCameraComponent", AllowAnyActor))
+	TArray<FSoftComponentReference> ICVFXCameras;
+};
\ No newline at end of file
diff --git a/Engine/Plugins/Runtime/nDisplay/Source/DisplayCluster/Public/Components/DisplayClusterLightCardStageActorComponent.h b/Engine/Plugins/Runtime/nDisplay/Source/DisplayCluster/Public/Components/DisplayClusterLightCardStageActorComponent.h
new file mode 100644
--- /dev/null
+++ b/Engine/Plugins/Runtime/nDisplay/Source/DisplayCluster/Public/Components/DisplayClusterLightCardStageActorComponent.h
@@ -0,0 +1,16 @@
+// Copyright Epic Games, Inc. All Rights Reserved.
+
+#pragma once
+
+#include "DisplayClusterStageActorComponent.h"
+
+#include "DisplayClusterLightCardStageActorComponent.generated.h"
+
+/**
+ * Stage Actor Component to be placed in light card actors
+ */
+UCLASS(MinimalAPI, ClassGroup = (DisplayCluster), meta = (DisplayName = "Light Card Stage Actor"), HideCategories=(Physics, Collision, Lighting, Navigation, Cooking, LOD, MaterialParameters, HLOD, RayTracing, TextureStreaming, Mobile))
+class UDisplayClusterLightCardStageActorComponent final : public UDisplayClusterStageActorComponent
+{
+	GENERATED_BODY()
+};
\ No newline at end of file
diff --git a/Engine/Plugins/Runtime/nDisplay/Source/DisplayCluster/Public/Components/DisplayClusterStageActorComponent.h b/Engine/Plugins/Runtime/nDisplay/Source/DisplayCluster/Public/Components/DisplayClusterStageActorComponent.h
new file mode 100644
--- /dev/null
+++ b/Engine/Plugins/Runtime/nDisplay/Source/DisplayCluster/Public/Components/DisplayClusterStageActorComponent.h
@@ -0,0 +1,31 @@
+// Copyright Epic Games, Inc. All Rights Reserved.
+
+#pragma once
+
+#include "Components/ActorComponent.h"
+
+#include "DisplayClusterStageActorComponent.generated.h"
+
+class ADisplayClusterRootActor;
+
+/**
+ * Stage Actor Component used to determine root actor ownership
+ */
+UCLASS(Abstract, ClassGroup = (DisplayCluster), meta = (DisplayName = "Stage Actor"), HideCategories=(Physics, Collision, Lighting, Navigation, Cooking, LOD, MaterialParameters, HLOD, RayTracing, TextureStreaming, Mobile))
+class UDisplayClusterStageActorComponent : public UActorComponent
+{
+	GENERATED_BODY()
+
+public:
+	/** Set the Display Cluster Root Actor owner */
+	DISPLAYCLUSTER_API virtual void SetRootActor(ADisplayClusterRootActor* InRootActor);
+
+	/** Return the root actor owner of the light card */
+	DISPLAYCLUSTER_API const TSoftObjectPtr<ADisplayClusterRootActor>& GetRootActor() const;
+	
+protected:
+	/** Associate this content with the nDisplay configuration actor specified here.
+	 * The content will only appear in the ICVFX Editor when this nDisplay configuration is selected. */
+	UPROPERTY(EditInstanceOnly, DisplayName = "nDisplay Root Actor", Category = "NDisplay Root Actor")
+	TSoftObjectPtr<ADisplayClusterRootActor> RootActor;
+};
\ No newline at end of file
diff --git a/Engine/Plugins/Runtime/nDisplay/Source/DisplayCluster/Public/DisplayClusterChromakeyCardActor.h b/Engine/Plugins/Runtime/nDisplay/Source/DisplayCluster/Public/DisplayClusterChromakeyCardActor.h
--- a/Engine/Plugins/Runtime/nDisplay/Source/DisplayCluster/Public/DisplayClusterChromakeyCardActor.h
+++ b/Engine/Plugins/Runtime/nDisplay/Source/DisplayCluster/Public/DisplayClusterChromakeyCardActor.h
@@ -6,10 +6,9 @@
 
 #include "DisplayClusterChromakeyCardActor.generated.h"
 
-class UDataLayerAsset;
 class UDisplayClusterICVFXCameraComponent;
 
-UCLASS(Blueprintable, NotPlaceable, DisplayName = "Chromakey Card", HideCategories = (Tick, Physics, Collision, Replication, Cooking, Input, Actor))
+UCLASS(Blueprintable, NotPlaceable, DisplayName = "Chromakey Card", HideCategories = (Tick, Physics, Collision, Networking, Replication, Cooking, Input, Actor, HLOD))
 class DISPLAYCLUSTER_API ADisplayClusterChromakeyCardActor : public ADisplayClusterLightCardActor
 {
 	GENERATED_BODY()
@@ -18,12 +17,19 @@ public:
 
 	virtual void Tick(float DeltaSeconds) override;
 
-	/** Setup chroma key for the owning root actor */
-	void AddToChromakeyLayer(ADisplayClusterRootActor* InRootActor);
-	
-	/** Checks if the given ICVFX camera has chroma key settings supporting this actor */
-	bool IsReferencedByICVFXCamera(UDisplayClusterICVFXCameraComponent* InCamera) const;
+	// ~Begin ADisplayClusterLightCardActor interface
+	virtual void AddToRootActor(ADisplayClusterRootActor* InRootActor) override;
+	virtual void RemoveFromRootActor() override;
+	// ~End ADisplayClusterLightCardActor interface
 	
+	/**
+	 * Checks if the given ICVFX camera has chroma key settings supporting this actor.
+	 * 
+	 * @param InCamera The ICVFX camera component to compare against.
+	 *
+	 * @return true if the camera component references this actor.
+	 */
+	bool IsReferencedByICVFXCamera(const UDisplayClusterICVFXCameraComponent* InCamera) const;
 protected:
 	void UpdateChromakeySettings();
 };
diff --git a/Engine/Plugins/Runtime/nDisplay/Source/DisplayCluster/Public/DisplayClusterLightCardActor.h b/Engine/Plugins/Runtime/nDisplay/Source/DisplayCluster/Public/DisplayClusterLightCardActor.h
--- a/Engine/Plugins/Runtime/nDisplay/Source/DisplayCluster/Public/DisplayClusterLightCardActor.h
+++ b/Engine/Plugins/Runtime/nDisplay/Source/DisplayCluster/Public/DisplayClusterLightCardActor.h
@@ -10,8 +10,8 @@
 
 class ADisplayClusterRootActor;
 class UActorComponent;
-class UDataLayerAsset;
 class UDisplayClusterLabelComponent;
+class UDisplayClusterStageActorComponent;
 class UMeshComponent;
 class USceneComponent;
 class USpringArmComponent;
@@ -51,20 +51,24 @@ struct FLightCardAlphaGradientSettings
 	float Angle = 0;
 };
 
-UCLASS(Blueprintable, DisplayName = "Light Card", HideCategories = (Tick, Physics, Collision, Replication, Cooking, Input, Actor))
+UCLASS(Blueprintable, DisplayName = "Light Card", HideCategories = (Tick, Physics, Collision, Networking, Replication, Cooking, Input, Actor, HLOD))
 class DISPLAYCLUSTER_API ADisplayClusterLightCardActor : public AActor, public IDisplayClusterStageActor
 {
 	GENERATED_BODY()
 
 public:
-	/** The default size of the porjection plane UV light cards are rendered to */
+	/** The default size of the projection plane UV light cards are rendered to */
 	static const float UVPlaneDefaultSize;
 
 	/** The default distance from the view of the projection plane UV light cards are rendered to */
 	static const float UVPlaneDefaultDistance;
 
+	/** The name used for the stage actor component */
+	static const FName LightCardStageActorComponentName;
+	
 public:
 	ADisplayClusterLightCardActor(const FObjectInitializer& ObjectInitializer);
+	virtual ~ADisplayClusterLightCardActor() override;
 
 	virtual void PostLoad() override;
 	virtual void OnConstruction(const FTransform& Transform) override;
@@ -100,6 +104,9 @@ public:
 	/** Updates the UV Indicator */
 	void UpdateUVIndicator();
 
+	/** Sync the position to the root actor */
+	void UpdateLightCardPositionToRootActor();
+
 	/** Makes the light card flush to the wall */
 	void MakeFlushToWall();
 
@@ -112,14 +119,27 @@ public:
 	/** Show or hide the light card label  */
 	void ShowLightCardLabel(bool bValue, float ScaleValue, ADisplayClusterRootActor* InRootActor);
 
-	/** Set the current owner of the light card */
+	/** Set a weak root actor owner. This should be used on legacy light cards that don't have StageActorComponent->RootActor set */
+	void SetWeakRootActorOwner(ADisplayClusterRootActor* InRootActor);
+
+	/** Set the root actor of the stage actor component */
 	void SetRootActorOwner(ADisplayClusterRootActor* InRootActor);
 	
-	/** Return the current owner, providing one was set */
-	TWeakObjectPtr<ADisplayClusterRootActor> GetRootActorOwner() const { return RootActorOwner; }
-	
+	/** Return the root actor owner of the light card */
+	ADisplayClusterRootActor* GetRootActorOwner() const;
+
+	/** Retrieve the stage actor component */
+	UDisplayClusterStageActorComponent* GetStageActorComponent() const;
+
 	/** Add this light card to a light card layer on the given root actor */
+	UE_DEPRECATED(5.2, "Layers are no longer used for light cards by default, use 'AddToRootActor' instead")
 	void AddToLightCardLayer(ADisplayClusterRootActor* InRootActor);
+
+	/** Add this light card to a root actor */
+	virtual void AddToRootActor(ADisplayClusterRootActor* InRootActor);
+
+	/** Remove this light card from the root actor */
+	virtual void RemoveFromRootActor();
 	
 	// ~Begin IDisplayClusterStageActor interface
 	/** Updates the Light Card transform based on its positional properties (Lat, Long, etc.) */
@@ -273,6 +293,12 @@ protected:
 	/** Removes components that were added by IDisplayClusterLightCardActorExtender */
 	void CleanUpComponentsForExtenders();
 
+#if WITH_EDITOR
+	/** Called when a level actor is deleted */
+	void OnLevelActorDeleted(AActor* DeletedActor);
+#endif
+	
+protected:
 	UPROPERTY(EditDefaultsOnly, BlueprintReadOnly, Category = "Default")
 	TObjectPtr<USceneComponent> DefaultSceneRootComponent;
 
@@ -292,16 +318,28 @@ protected:
 	UPROPERTY(Transient, BlueprintReadOnly, Category = "Default")
 	TObjectPtr<UDisplayClusterLabelComponent> LabelComponent;
 
+	/** Manages stage actor properties */
+	UPROPERTY(VisibleAnywhere, Category = "Default")
+	TObjectPtr<UDisplayClusterStageActorComponent> StageActorComponent;
+
 #if WITH_EDITORONLY_DATA
 	UPROPERTY(Transient, BlueprintReadOnly, Category = "Default")
 	TObjectPtr<UStaticMeshComponent> UVIndicatorComponent;
 #endif
 
-	/** The current owner of the light card */
-	TWeakObjectPtr<ADisplayClusterRootActor> RootActorOwner;
-
 	/** Indicates this light card should be considered a flag */
 	UPROPERTY()
 	bool bIsLightCardFlag = false;
-	
+
+private:
+	/** Set by DCRA during tick. Meant for legacy light cards that don't have a SoftObjectPtr set */
+	TWeakObjectPtr<ADisplayClusterRootActor> WeakRootActorOwner;
+
+	/** Stores the location relative to the root actor's origin that stage actors like light cards can orbit around */
+	FVector LastOrbitLocation = FVector::ZeroVector;
+
+#if WITH_EDITOR
+	/** Set if there was conflicting root actor ownership */
+	bool bHadRootActorMismatch = false;
+#endif
 };
diff --git a/Engine/Plugins/Runtime/nDisplay/Source/DisplayCluster/Public/DisplayClusterRootActor.h b/Engine/Plugins/Runtime/nDisplay/Source/DisplayCluster/Public/DisplayClusterRootActor.h
--- a/Engine/Plugins/Runtime/nDisplay/Source/DisplayCluster/Public/DisplayClusterRootActor.h
+++ b/Engine/Plugins/Runtime/nDisplay/Source/DisplayCluster/Public/DisplayClusterRootActor.h
@@ -99,13 +99,6 @@ public:
 	const FDisplayClusterConfigurationICVFX_StageSettings& GetStageSettings() const;
 	const FDisplayClusterConfigurationRenderFrame& GetRenderFrameSettings() const;
 
-	UDataLayerAsset* GetOrCreateLightCardDataLayerAsset(const FName& InName);
-	UDataLayerAsset* GetOrCreateChromakeyCardDataLayerAsset(const FName& InName);
-
-private:
-#if WITH_EDITOR
-	UDataLayerAsset* CreateDataLayerAsset(const FName& InName);
-#endif
 protected:
 	//////////////////////////////////////////////////////////////////////////////////////////////
 	// AActor
@@ -127,11 +120,8 @@ protected:
 	// Creates all hierarchy objects declared in a config file
 	bool BuildHierarchy();
 
-	/** Updates the world position and rotation of each light card referenced by this root actor's light card list to match the default view origin */
-	void UpdateLightCardPositions();
-
-	/** Sets self as the owner for any chromakey cards */
-	void SetChromakeyCardsOwner();
+	/** Determine which light card actors are owned by this root actor */
+	void SetLightCardOwnership();
 
 public:
 	UFUNCTION(BlueprintCallable, Category = "NDisplay")
@@ -302,23 +292,13 @@ private:
 	/** Component that stores the stage's geometry map, which is used to make objects flush with the stage's walls and ceilings */
 	UPROPERTY()
 	TObjectPtr<UDisplayClusterStageGeometryComponent> StageGeometryComponent;
-
-	/** Data layer used for light card layers */
-	UPROPERTY()
-	TObjectPtr<UDataLayerAsset> LightCardDataLayerAsset;
-
-	/** Data layer used for chromakey card layers */
-	UPROPERTY()
-	TObjectPtr<UDataLayerAsset> ChromakeyCardDataLayerAsset;
+	
 private:
 	// Current operation mode
 	EDisplayClusterOperationMode OperationMode;
 
 	float LastDeltaSecondsValue = 0.f;
 
-	/** Stores the location relative to the actor's origin that stage actors like light cards or CCWs orbit around when owned by this root actor. */
-	FVector StageActorOrbitLocation = FVector::ZeroVector;
-
 private:
 	template <typename TComp>
 	void GetTypedPrimitives(TSet<FPrimitiveComponentId>& OutPrimitives, const TArray<FString>* InCompNames = nullptr, bool bCollectChildrenVisualizationComponent = true) const;
diff --git a/Engine/Plugins/Runtime/nDisplay/Source/DisplayClusterConfiguration/Public/DisplayClusterConfigurationTypes_ICVFX.h b/Engine/Plugins/Runtime/nDisplay/Source/DisplayClusterConfiguration/Public/DisplayClusterConfigurationTypes_ICVFX.h
--- a/Engine/Plugins/Runtime/nDisplay/Source/DisplayClusterConfiguration/Public/DisplayClusterConfigurationTypes_ICVFX.h
+++ b/Engine/Plugins/Runtime/nDisplay/Source/DisplayClusterConfiguration/Public/DisplayClusterConfigurationTypes_ICVFX.h
@@ -120,6 +120,9 @@ public:
 	/** Reference to RootActor components by names */
 	UPROPERTY(EditDefaultsOnly, BlueprintReadWrite, Category = NDisplay)
 	TArray<FString> RootActorComponentNames; //@todo change link, now by names
+
+	/** Actors added automatically during a session */
+	TArray<TSoftObjectPtr<AActor>> AutoAddedActors;
 };
 
 USTRUCT(Blueprintable)
diff --git a/Engine/Plugins/Runtime/nDisplay/Source/DisplayClusterLightCardEditor/Private/DisplayClusterLightCardEditor.cpp b/Engine/Plugins/Runtime/nDisplay/Source/DisplayClusterLightCardEditor/Private/DisplayClusterLightCardEditor.cpp
--- a/Engine/Plugins/Runtime/nDisplay/Source/DisplayClusterLightCardEditor/Private/DisplayClusterLightCardEditor.cpp
+++ b/Engine/Plugins/Runtime/nDisplay/Source/DisplayClusterLightCardEditor/Private/DisplayClusterLightCardEditor.cpp
@@ -216,25 +216,6 @@ AActor* FDisplayClusterLightCardEditor::SpawnActor(TSubclassOf<AActor> InActorCl
 	}
 
 	const UClass* ClassToUse = InTemplate ? InTemplate->GetClass() : InActorClass.Get(); 
-
-	// Partitioned worlds that are not saved will fail to create the data layer asset. Detect this and display
-	// an error message for the user.
-	if (ActiveRootActor->GetWorld()->IsPartitionedWorld() &&
-		(ClassToUse->IsChildOf(ADisplayClusterLightCardActor::StaticClass()) ||
-		ClassToUse->IsChildOf(UDisplayClusterLightCardTemplate::StaticClass())))
-	{
-		const FString PathName = ActiveRootActor->GetWorld()->GetPathName();
-		if (PathName.StartsWith(TEXT("/Temp")))
-		{
-			FNotificationInfo NotificationInfo(LOCTEXT("CannotSpawnActorInUnsavedWorldPartition",
-	"Cannot spawn an actor of this type in a world partition map that is not saved. Please save the map."));
-			NotificationInfo.bUseThrobber = false;
-			NotificationInfo.ExpireDuration = 4.0f;
-			FSlateNotificationManager::Get().AddNotification(NotificationInfo)->SetCompletionState(SNotificationItem::CS_Fail);
-
-			return nullptr;
-		}
-	}
 	
 	FScopedTransaction Transaction(LOCTEXT("SpawnActorTransactionMessage", "Spawn Actor"));
 
@@ -763,6 +744,8 @@ void FDisplayClusterLightCardEditor::RemoveActors(
 					{
 						return InActor.Get() == LightCard;
 					});
+
+				LightCard->RemoveFromRootActor();
 				
 				if (!bDeleteActors)
 				{
diff --git a/Engine/Plugins/Runtime/nDisplay/Source/DisplayClusterScenePreview/Private/DisplayClusterLightCardEditorHelper.cpp b/Engine/Plugins/Runtime/nDisplay/Source/DisplayClusterScenePreview/Private/DisplayClusterLightCardEditorHelper.cpp
--- a/Engine/Plugins/Runtime/nDisplay/Source/DisplayClusterScenePreview/Private/DisplayClusterLightCardEditorHelper.cpp
+++ b/Engine/Plugins/Runtime/nDisplay/Source/DisplayClusterScenePreview/Private/DisplayClusterLightCardEditorHelper.cpp
@@ -882,10 +882,11 @@ AActor* FDisplayClusterLightCardEditorHelper::SpawnStageActor(const FSpawnActorA
 			}
 		}
 
+		NewLightCard->ShowLightCardLabel(InSpawnArgs.AddLightCardArgs.bShowLabels, InSpawnArgs.AddLightCardArgs.LabelScale, RootActor);
+		
 		if (ADisplayClusterChromakeyCardActor* ChromakeyCardActor = Cast<ADisplayClusterChromakeyCardActor>(NewActor))
 		{
-			// Only add to chroma key layer, not light card layer
-			ChromakeyCardActor->AddToChromakeyLayer(RootActor);
+			ChromakeyCardActor->AddToRootActor(RootActor);
 		}
 		else if (!bIsPreview)
 		{
@@ -929,8 +930,6 @@ void FDisplayClusterLightCardEditorHelper::AddLightCardsToRootActor(
 				return Actor.Get() == LightCard;
 			}))
 		{
-			LightCard->ShowLightCardLabel(AddLightCardArgs.bShowLabels, AddLightCardArgs.LabelScale, RootActor);
-				
 			const TSoftObjectPtr<AActor> LightCardSoftObject(LightCard);
 
 			// Remove any exact paths to this actor. It's possible invalid actors are present if a light card
@@ -940,7 +939,7 @@ void FDisplayClusterLightCardEditorHelper::AddLightCardsToRootActor(
 					return Actor == LightCardSoftObject;
 				});
 
-			LightCard->AddToLightCardLayer(RootActor);
+			LightCard->AddToRootActor(RootActor);
 		}
 
 #if WITH_EDITOR
diff --git a/Engine/Plugins/VirtualProduction/EpicStageApp/Source/EpicStageApp/Private/StageAppRouteHandler.cpp b/Engine/Plugins/VirtualProduction/EpicStageApp/Source/EpicStageApp/Private/StageAppRouteHandler.cpp
--- a/Engine/Plugins/VirtualProduction/EpicStageApp/Source/EpicStageApp/Private/StageAppRouteHandler.cpp
+++ b/Engine/Plugins/VirtualProduction/EpicStageApp/Source/EpicStageApp/Private/StageAppRouteHandler.cpp
@@ -947,7 +947,7 @@ void FStageAppRouteHandler::HandleWebSocketStageAppActorsDuplicate(const FRemote
 			// Add it to the root actor
 			if (ADisplayClusterRootActor* RootActor = StageAppRouteHandlerUtils::GetRootActorForLightCard(Actor))
 			{
-				NewLightCard->AddToLightCardLayer(RootActor);
+				NewLightCard->AddToRootActor(RootActor);
 
 #if WITH_EDITOR
 				// Required so operator panel updates
