SharedMemoryMedia: Fixed sRGB handling.

The SMM player was not labeling the samples with the sRGB according to the format of the incoming shared texture. They were incorrectly always being flagged as sRGB encoded. This issue was exposed after fixing the sRGB handling in the nDisplay media input, which now triggers the correct conversion when the incoming textures are sRGB encoded.

diff --git a/Engine/Plugins/Runtime/nDisplay/Source/DisplayClusterMedia/Private/DisplayClusterMediaHelpers.cpp b/Engine/Plugins/Runtime/nDisplay/Source/DisplayClusterMedia/Private/DisplayClusterMediaHelpers.cpp
--- a/Engine/Plugins/Runtime/nDisplay/Source/DisplayClusterMedia/Private/DisplayClusterMediaHelpers.cpp
+++ b/Engine/Plugins/Runtime/nDisplay/Source/DisplayClusterMedia/Private/DisplayClusterMediaHelpers.cpp
@@ -75,10 +75,13 @@ namespace DisplayClusterMediaHelpers
 	void ResampleTexture_RenderThread(FRHICommandListImmediate& RHICmdList, FRHITexture* SrcTexture, FRHITexture* DstTexture, const FIntRect& SrcRect, const FIntRect& DstRect)
 	{
 		check(SrcTexture);
+		check(DstTexture);
 
 		const bool bSrcSrgb = EnumHasAnyFlags(SrcTexture->GetFlags(), TexCreate_SRGB);
+		const bool bDstSrgb = EnumHasAnyFlags(DstTexture->GetFlags(), TexCreate_SRGB);
 
-		if (bSrcSrgb)
+		// We only need to convert from sRGB encoding to linear if the source is sRGB encoded and the destination is not.
+		if (bSrcSrgb && !bDstSrgb)
 		{
 			ResampleTextureImpl_RenderThread<FScreenPSsRGBSource>(RHICmdList, SrcTexture, DstTexture, SrcRect, DstRect);
 		}
diff --git a/Engine/Plugins/Runtime/nDisplay/Source/DisplayClusterMedia/Private/Input/SharedMemoryMediaPlayer.cpp b/Engine/Plugins/Runtime/nDisplay/Source/DisplayClusterMedia/Private/Input/SharedMemoryMediaPlayer.cpp
--- a/Engine/Plugins/Runtime/nDisplay/Source/DisplayClusterMedia/Private/Input/SharedMemoryMediaPlayer.cpp
+++ b/Engine/Plugins/Runtime/nDisplay/Source/DisplayClusterMedia/Private/Input/SharedMemoryMediaPlayer.cpp
@@ -389,6 +389,8 @@ void FSharedMemoryMediaPlayer::TickFetch(FTimespan DeltaTime, FTimespan Timecode
 			MakeUniqueObjectName(nullptr, UTexture2D::StaticClass(), FName(*FString::Printf(TEXT("SampleCommonTexture_%s"), *UniqueName)))
 		));
 
+		SampleCommonTexture->SRGB = SharedCrossGpuTextureDescriptions[0].bSrgb;
+
 		SampleCommonTexture->UpdateResource();
 	}
 
@@ -417,6 +419,18 @@ void FSharedMemoryMediaPlayer::TickFetch(FTimespan DeltaTime, FTimespan Timecode
 		Samples->CurrentSample->Texture = SampleCommonTexture->GetResource()->GetTextureRHI();
 		Samples->CurrentSample->Dim = FIntPoint(SharedCrossGpuTextureDescriptions[0].Width, SharedCrossGpuTextureDescriptions[0].Height);
 		Samples->CurrentSample->Stride = SharedCrossGpuTextureDescriptions[0].Stride;
+		Samples->CurrentSample->bSrgb = SharedCrossGpuTextureDescriptions[0].bSrgb;
+		
+		// Try to preserve the common input formats when specifying EMediaTextureSampleFormat
+
+		switch (SharedCrossGpuTextureDescriptions[0].Format)
+		{
+		case PF_FloatRGB   : Samples->CurrentSample->Format = EMediaTextureSampleFormat::FloatRGB;    break;
+		case PF_FloatRGBA  : Samples->CurrentSample->Format = EMediaTextureSampleFormat::FloatRGBA;   break;
+		case PF_B8G8R8A8   : Samples->CurrentSample->Format = EMediaTextureSampleFormat::CharBGRA;    break;
+		case PF_A2B10G10R10: Samples->CurrentSample->Format = EMediaTextureSampleFormat::CharBGR10A2; break;
+		default            : Samples->CurrentSample->Format = EMediaTextureSampleFormat::CharBGR10A2; break;
+		}
 
 		// Ticks are "frames" in our player. W/o constraining the sample timing, it would be possible for the
 		// render thread to purge the sample of the next frame, in particular when the game thread stalls.
diff --git a/Engine/Plugins/Runtime/nDisplay/Source/DisplayClusterMedia/Private/Input/SharedMemoryMediaSample.h b/Engine/Plugins/Runtime/nDisplay/Source/DisplayClusterMedia/Private/Input/SharedMemoryMediaSample.h
--- a/Engine/Plugins/Runtime/nDisplay/Source/DisplayClusterMedia/Private/Input/SharedMemoryMediaSample.h
+++ b/Engine/Plugins/Runtime/nDisplay/Source/DisplayClusterMedia/Private/Input/SharedMemoryMediaSample.h
@@ -38,6 +38,12 @@ public:
 	/** Texture dimensions */
 	FIntPoint Dim = FIntPoint::ZeroValue;
 
+	/** Whether the sample is sRGB encoded or not */
+	bool bSrgb = false;
+
+	/** Sample format */
+	EMediaTextureSampleFormat Format = EMediaTextureSampleFormat::CharBGR10A2;
+
 private:
 
 	/** Used in case the Texture is not valid, so we can provide pixels to the Media Framework */
@@ -78,8 +84,7 @@ public:
 
 	EMediaTextureSampleFormat GetFormat() const override
 	{
-		// All output types go through buffer conversion
-		return EMediaTextureSampleFormat::CharBGR10A2;
+		return Format;
 	}
 
 	FIntPoint GetOutputDim() const override
@@ -104,7 +109,7 @@ public:
 
 	bool IsOutputSrgb() const override
 	{
-		return true;
+		return bSrgb;
 	}
 
 public:
diff --git a/Engine/Plugins/Runtime/nDisplay/Source/DisplayClusterMedia/Private/Platforms/Windows/SharedMemoryMediaPlatformWindowsD3D12.cpp b/Engine/Plugins/Runtime/nDisplay/Source/DisplayClusterMedia/Private/Platforms/Windows/SharedMemoryMediaPlatformWindowsD3D12.cpp
--- a/Engine/Plugins/Runtime/nDisplay/Source/DisplayClusterMedia/Private/Platforms/Windows/SharedMemoryMediaPlatformWindowsD3D12.cpp
+++ b/Engine/Plugins/Runtime/nDisplay/Source/DisplayClusterMedia/Private/Platforms/Windows/SharedMemoryMediaPlatformWindowsD3D12.cpp
@@ -220,6 +220,7 @@ FTextureRHIRef FSharedMemoryMediaPlatformWindowsD3D12::OpenSharedCrossGpuTexture
 		OutTextureDescription.Height = SharedGpuTextureDesc.Height;
 		OutTextureDescription.Width = SharedGpuTextureDesc.Width;
 		OutTextureDescription.Format = EPixelFormat::PF_Unknown;
+		OutTextureDescription.bSrgb = false;
 
 		// Find EPixelFormat from platform format
 		for (int32 FormatIdx = 0; FormatIdx < PF_MAX; FormatIdx++)
@@ -235,6 +236,7 @@ FTextureRHIRef FSharedMemoryMediaPlatformWindowsD3D12::OpenSharedCrossGpuTexture
 					OutTextureDescription.Format = EPixelFormat(FormatIdx);
 					OutTextureDescription.BytesPerPixel = GPixelFormats[FormatIdx].BlockBytes;
 					OutTextureDescription.Stride = OutTextureDescription.Width * OutTextureDescription.BytesPerPixel;
+					OutTextureDescription.bSrgb = bSrgb;
 					break;
 				}
 			}
@@ -257,9 +259,16 @@ FTextureRHIRef FSharedMemoryMediaPlatformWindowsD3D12::OpenSharedCrossGpuTexture
 		}
 	}
 
+	ETextureCreateFlags Flags = TexCreate_Dynamic | TexCreate_DisableSRVCreation;
+	
+	if (OutTextureDescription.bSrgb)
+	{
+		Flags |= TexCreate_SRGB;
+	}
+
 	return GetID3D12DynamicRHI()->RHICreateTexture2DFromResource(
 		OutTextureDescription.Format,
-		TexCreate_Dynamic | TexCreate_DisableSRVCreation,
+		Flags,
 		FClearValueBinding::None,
 		SharedCrossGpuTexture
 	);
diff --git a/Engine/Plugins/Runtime/nDisplay/Source/DisplayClusterMedia/Private/SharedMemoryMediaTypes.h b/Engine/Plugins/Runtime/nDisplay/Source/DisplayClusterMedia/Private/SharedMemoryMediaTypes.h
--- a/Engine/Plugins/Runtime/nDisplay/Source/DisplayClusterMedia/Private/SharedMemoryMediaTypes.h
+++ b/Engine/Plugins/Runtime/nDisplay/Source/DisplayClusterMedia/Private/SharedMemoryMediaTypes.h
@@ -50,6 +50,7 @@ struct FSharedMemoryMediaTextureDescription
 	uint32 Stride = 0;
 	EPixelFormat Format = EPixelFormat::PF_Unknown;
 	FGuid Guid = UE::SharedMemoryMedia::ZeroGuid;
+	bool bSrgb = false;
 
 	bool IsEquivalentTo(const FSharedMemoryMediaTextureDescription& Other) const
 	{
@@ -57,7 +58,9 @@ struct FSharedMemoryMediaTextureDescription
 			&& Height == Other.Height
 			&& BytesPerPixel == Other.BytesPerPixel
 			&& Stride == Other.Stride
-			&& Format == Other.Format;
+			&& Format == Other.Format
+			&& bSrgb == Other.bSrgb
+			;
 	}
 };
 
