nDisplay: Fix crash when using UV Light Cards without a separate rendering thread.

diff --git a/Engine/Plugins/Runtime/nDisplay/Source/DisplayCluster/Private/Render/Viewport/LightCard/DisplayClusterViewportLightCardManager.cpp b/Engine/Plugins/Runtime/nDisplay/Source/DisplayCluster/Private/Render/Viewport/LightCard/DisplayClusterViewportLightCardManager.cpp
--- a/Engine/Plugins/Runtime/nDisplay/Source/DisplayCluster/Private/Render/Viewport/LightCard/DisplayClusterViewportLightCardManager.cpp
+++ b/Engine/Plugins/Runtime/nDisplay/Source/DisplayCluster/Private/Render/Viewport/LightCard/DisplayClusterViewportLightCardManager.cpp
@@ -41,7 +41,7 @@ FDisplayClusterViewportLightCardManager::~FDisplayClusterViewportLightCardManage
 void FDisplayClusterViewportLightCardManager::Release()
 {
 	// The destructor is usually called from the rendering thread, so Release() must be called first from the game thread.
-	check(!IsInRenderingThread());
+	check(IsInGameThread());
 
 	// Release UVLightCard
 	ReleaseUVLightCardData();
