[nDisplay] fixed OCIO issue when media input is used

#jria UE-179501

diff --git a/Engine/Plugins/Runtime/nDisplay/Source/DisplayCluster/Private/Render/Viewport/DisplayClusterViewport.cpp b/Engine/Plugins/Runtime/nDisplay/Source/DisplayCluster/Private/Render/Viewport/DisplayClusterViewport.cpp
--- a/Engine/Plugins/Runtime/nDisplay/Source/DisplayCluster/Private/Render/Viewport/DisplayClusterViewport.cpp
+++ b/Engine/Plugins/Runtime/nDisplay/Source/DisplayCluster/Private/Render/Viewport/DisplayClusterViewport.cpp
@@ -23,6 +23,7 @@
 #include "EngineUtils.h"
 #include "SceneManagement.h"
 #include "SceneView.h"
+#include "UnrealClient.h"
 
 #include "DisplayClusterSceneViewExtensions.h"
 
@@ -739,3 +740,25 @@ bool FDisplayClusterViewport::UpdateFrameContexts(const uint32 InStereoViewIndex
 
 	return true;
 }
+
+// Currently, some data like EngineShowFlags and EngineGamma get updated on PostRenderViewFamily.
+// Since the viewports that have media input assigned never gets rendered in normal way, the PostRenderViewFamily
+// callback never gets called, therefore the data mentioned above never gets updated. This workaround initializes
+// those settings for all viewports. The viewports with no media input assigned will override the data
+// in PostRenderViewFamily like it was previously so nothing should be broken.
+void FDisplayClusterViewport::UpdateMediaDependencies(FViewport* InViewport)
+{
+	if (InViewport && InViewport->GetClient())
+	{
+		if (const FEngineShowFlags* const EngineShowFlags = InViewport->GetClient()->GetEngineShowFlags())
+		{
+			const float DisplayGamma = GEngine ? GEngine->DisplayGamma : 2.2f;
+
+			for (FDisplayClusterViewport_Context& Context : Contexts)
+			{
+				Context.RenderThreadData.EngineDisplayGamma = DisplayGamma;
+				Context.RenderThreadData.EngineShowFlags = *EngineShowFlags;
+			}
+		}
+	}
+}
diff --git a/Engine/Plugins/Runtime/nDisplay/Source/DisplayCluster/Private/Render/Viewport/DisplayClusterViewport.h b/Engine/Plugins/Runtime/nDisplay/Source/DisplayCluster/Private/Render/Viewport/DisplayClusterViewport.h
--- a/Engine/Plugins/Runtime/nDisplay/Source/DisplayCluster/Private/Render/Viewport/DisplayClusterViewport.h
+++ b/Engine/Plugins/Runtime/nDisplay/Source/DisplayCluster/Private/Render/Viewport/DisplayClusterViewport.h
@@ -197,6 +197,9 @@ public:
 	 */
 	bool UpdateFrameContexts(const uint32 InStereoViewIndex, const FDisplayClusterRenderFrameSettings& InFrameSettings);
 
+	/* Update media dependent data */
+	void UpdateMediaDependencies(class FViewport* InViewport);
+
 	/** Reset viewport contexts and resources. */
 	void ResetFrameContexts();
 
diff --git a/Engine/Plugins/Runtime/nDisplay/Source/DisplayCluster/Private/Render/Viewport/DisplayClusterViewportManager.cpp b/Engine/Plugins/Runtime/nDisplay/Source/DisplayCluster/Private/Render/Viewport/DisplayClusterViewportManager.cpp
--- a/Engine/Plugins/Runtime/nDisplay/Source/DisplayCluster/Private/Render/Viewport/DisplayClusterViewportManager.cpp
+++ b/Engine/Plugins/Runtime/nDisplay/Source/DisplayCluster/Private/Render/Viewport/DisplayClusterViewportManager.cpp
@@ -532,6 +532,12 @@ bool FDisplayClusterViewportManager::BeginNewFrame(FViewport* InViewport, UWorld
 		}
 	}
 
+	/* Some media dependent data can (and needs to) be initialized in advance */
+	for (FDisplayClusterViewport* const Viewport : ClusterNodeViewports)
+	{
+		Viewport->UpdateMediaDependencies(InViewport);
+	}
+
 	// Handle scene RTT resize
 	UpdateSceneRenderTargetSize();
 
