[nDisplay] Fixed nDisplay crash when viewports get frozen with non-zero global media latency

diff --git a/Engine/Plugins/Runtime/nDisplay/Source/DisplayClusterMedia/Private/Synchronization/DisplayClusterFrameQueue.cpp b/Engine/Plugins/Runtime/nDisplay/Source/DisplayClusterMedia/Private/Synchronization/DisplayClusterFrameQueue.cpp
--- a/Engine/Plugins/Runtime/nDisplay/Source/DisplayClusterMedia/Private/Synchronization/DisplayClusterFrameQueue.cpp
+++ b/Engine/Plugins/Runtime/nDisplay/Source/DisplayClusterMedia/Private/Synchronization/DisplayClusterFrameQueue.cpp
@@ -117,14 +117,17 @@ void FDisplayClusterFrameQueue::HandleProcessLatency_RenderThread(FRHICommandLis
 							// Get viewport texture
 							TArray<FRHITexture*> Textures;
 							ViewportProxyIt->GetResources_RenderThread(EDisplayClusterViewportResourceType::InternalRenderTargetResource, Textures);
-							check(Textures.Num() > 0);
 
-							const FString ViewportId(ViewportProxyIt->GetId());
+							// Save current, and replace with an older
+							if (Textures.Num() > 0)
+							{
+								const FString ViewportId(ViewportProxyIt->GetId());
 
-							// Enqueue this view (head)
-							Frames[IdxHead].SaveView(RHICmdList, ViewportId, Textures[0]);
-							// Dequeue from the past (tail)
-							Frames[IdxTail].LoadView(RHICmdList, ViewportId, Textures[0]);
+								// Enqueue this view (head)
+								Frames[IdxHead].SaveView(RHICmdList, ViewportId, Textures[0]);
+								// Dequeue from the past (tail)
+								Frames[IdxTail].LoadView(RHICmdList, ViewportId, Textures[0]);
+							}
 						}
 					}
 				}
