nDisplay: Added the ability to rename and delete color grading groups from the ICVFX panel color grading drawer through a right click context menu. Unnamed color grading groups now show up as "Unnamed" in the drawer, and hovering over the group name now shows a text edit beam cursor.

diff --git a/Engine/Plugins/Runtime/nDisplay/Source/DisplayClusterColorGrading/Private/DataModelGenerators/DisplayClusterColorGradingGenerator_RootActor.cpp b/Engine/Plugins/Runtime/nDisplay/Source/DisplayClusterColorGrading/Private/DataModelGenerators/DisplayClusterColorGradingGenerator_RootActor.cpp
--- a/Engine/Plugins/Runtime/nDisplay/Source/DisplayClusterColorGrading/Private/DataModelGenerators/DisplayClusterColorGradingGenerator_RootActor.cpp
+++ b/Engine/Plugins/Runtime/nDisplay/Source/DisplayClusterColorGrading/Private/DataModelGenerators/DisplayClusterColorGradingGenerator_RootActor.cpp
@@ -422,6 +422,8 @@ void FDisplayClusterColorGradingGenerator_RootActor::GenerateDataModel(IProperty
 				TSharedRef<IPropertyHandle> PerViewportElementHandle = PerViewportColorGradingHandle->AsArray()->GetElement(Index);
 
 				FDisplayClusterColorGradingDataModel::FColorGradingGroup PerViewportGroup = CreateColorGradingGroup(PerViewportElementHandle);
+				PerViewportGroup.bCanBeDeleted = true;
+				PerViewportGroup.bCanBeRenamed = true;
 				PerViewportGroup.EditConditionPropertyHandle = PerViewportElementHandle->GetChildHandle(GET_MEMBER_NAME_CHECKED(FDisplayClusterConfigurationViewport_PerViewportColorGrading, bIsEnabled));
 
 				PerViewportGroup.DetailsViewCategories.Add(TEXT("DetailView_PerViewport"));
@@ -439,7 +441,23 @@ void FDisplayClusterColorGradingGenerator_RootActor::GenerateDataModel(IProperty
 					.VAlign(VAlign_Center)
 					[
 						SNew(SInlineEditableTextBlock)
-						.Font(FAppStyle::Get().GetFontStyle("NormalFontBold"))
+						.Cursor(EMouseCursor::TextEditBeam)
+						.Font_Lambda([NamePropertyHandle]
+						{
+							FText Name = FText::GetEmpty();
+
+							if (NamePropertyHandle.IsValid() && NamePropertyHandle->IsValidHandle())
+							{
+								NamePropertyHandle->GetValue(Name);
+							}
+
+							if (Name.IsEmpty())
+							{
+								return FAppStyle::Get().GetFontStyle("NormalFontBoldItalic");
+							}
+
+							return FAppStyle::Get().GetFontStyle("NormalFontBold");
+						})
 						.OnTextCommitted_Lambda([NamePropertyHandle](const FText& InText, ETextCommit::Type TextCommitType)
 						{
 							if (NamePropertyHandle.IsValid() && NamePropertyHandle->IsValidHandle())
@@ -457,6 +475,11 @@ void FDisplayClusterColorGradingGenerator_RootActor::GenerateDataModel(IProperty
 								NamePropertyHandle->GetValue(Name);
 							}
 
+							if (Name.IsEmpty())
+							{
+								return LOCTEXT("UnnamedPerViewportLabel", "Unnamed");
+							}
+
 							return Name;
 						})
 					]
@@ -485,6 +508,8 @@ void FDisplayClusterColorGradingGenerator_RootActor::GenerateDataModel(IProperty
 			.ColorAndOpacity(FSlateColor::UseForeground())
 		];
 
+	OutColorGradingDataModel.OnColorGradingGroupDeleted().AddSP(this, &FDisplayClusterColorGradingGenerator_RootActor::DeleteColorGradingGroup);
+	OutColorGradingDataModel.OnColorGradingGroupRenamed().AddSP(this, &FDisplayClusterColorGradingGenerator_RootActor::RenameColorGradingGroup);
 
 	FDisplayClusterColorGradingDataModel::FDetailsSection ViewportsDetailsSection;
 	ViewportsDetailsSection.DisplayName = LOCTEXT("ViewportsDetailsSectionLabel", "Viewports");
@@ -544,6 +569,56 @@ FReply FDisplayClusterColorGradingGenerator_RootActor::AddColorGradingGroup()
 	return FReply::Handled();
 }
 
+void FDisplayClusterColorGradingGenerator_RootActor::DeleteColorGradingGroup(int32 GroupIndex)
+{
+	for (const TWeakObjectPtr<ADisplayClusterRootActor>& RootActor : RootActors)
+	{
+		if (RootActor.IsValid())
+		{
+			if (UDisplayClusterConfigurationData* ConfigData = RootActor->GetConfigData())
+			{
+				// Use GroupIndex - 1 since the entire cluster color grading is group 0
+				if (ConfigData->StageSettings.PerViewportColorGrading.IsValidIndex(GroupIndex - 1))
+				{
+					FScopedTransaction Transaction(LOCTEXT("DeleteViewportColorGradingGroupTransaction", "Delete Viewport Group"));
+
+					RootActor->Modify();
+					ConfigData->Modify();
+
+					ConfigData->StageSettings.PerViewportColorGrading.RemoveAt(GroupIndex - 1);
+
+					IDisplayClusterColorGrading::Get().GetColorGradingDrawerSingleton().RefreshColorGradingDrawers(true);
+				}
+			}
+		}
+	}
+}
+
+void FDisplayClusterColorGradingGenerator_RootActor::RenameColorGradingGroup(int32 GroupIndex, const FText& NewName)
+{
+	for (const TWeakObjectPtr<ADisplayClusterRootActor>& RootActor : RootActors)
+	{
+		if (RootActor.IsValid())
+		{
+			if (UDisplayClusterConfigurationData* ConfigData = RootActor->GetConfigData())
+			{
+				// Use GroupIndex - 1 since the entire cluster color grading is group 0
+				if (ConfigData->StageSettings.PerViewportColorGrading.IsValidIndex(GroupIndex - 1))
+				{
+					FScopedTransaction Transaction(LOCTEXT("RenameViewportColorGradingGroupTransaction", "Rename Viewport Group"));
+
+					RootActor->Modify();
+					ConfigData->Modify();
+
+					ConfigData->StageSettings.PerViewportColorGrading[GroupIndex - 1].Name = NewName;
+
+					IDisplayClusterColorGrading::Get().GetColorGradingDrawerSingleton().RefreshColorGradingDrawers(true);
+				}
+			}
+		}
+	}
+}
+
 TSharedRef<SWidget> FDisplayClusterColorGradingGenerator_RootActor::CreateViewportComboBox(int32 PerViewportColorGradingIndex) const
 {
 	return SNew(SComboButton)
@@ -1002,6 +1077,8 @@ void FDisplayClusterColorGradingGenerator_ICVFXCamera::GenerateDataModel(IProper
 				TSharedRef<IPropertyHandle> PerNodeElementHandle = PerNodeColorGradingHandle->AsArray()->GetElement(Index);
 
 				FDisplayClusterColorGradingDataModel::FColorGradingGroup PerNodeGroup = CreateColorGradingGroup(PerNodeElementHandle);
+				PerNodeGroup.bCanBeDeleted = true;
+				PerNodeGroup.bCanBeRenamed = true;
 				PerNodeGroup.EditConditionPropertyHandle = PerNodeElementHandle->GetChildHandle(GET_MEMBER_NAME_CHECKED(FDisplayClusterConfigurationViewport_PerNodeColorGrading, bIsEnabled));
 
 				PerNodeGroup.DetailsViewCategories.Add(TEXT("DetailView_PerNode"));
@@ -1019,7 +1096,23 @@ void FDisplayClusterColorGradingGenerator_ICVFXCamera::GenerateDataModel(IProper
 					.VAlign(VAlign_Center)
 					[
 						SNew(SInlineEditableTextBlock)
-						.Font(FAppStyle::Get().GetFontStyle("NormalFontBold"))
+						.Cursor(EMouseCursor::TextEditBeam)
+						.Font_Lambda([NamePropertyHandle]
+						{
+							FText Name = FText::GetEmpty();
+
+							if (NamePropertyHandle.IsValid() && NamePropertyHandle->IsValidHandle())
+							{
+								NamePropertyHandle->GetValue(Name);
+							}
+
+							if (Name.IsEmpty())
+							{
+								return FAppStyle::Get().GetFontStyle("NormalFontBoldItalic");
+							}
+
+							return FAppStyle::Get().GetFontStyle("NormalFontBold");
+						})
 						.OnTextCommitted_Lambda([NamePropertyHandle](const FText& InText, ETextCommit::Type TextCommitType)
 						{
 							if (NamePropertyHandle.IsValid() && NamePropertyHandle->IsValidHandle())
@@ -1037,6 +1130,11 @@ void FDisplayClusterColorGradingGenerator_ICVFXCamera::GenerateDataModel(IProper
 								NamePropertyHandle->GetValue(Name);
 							}
 
+							if (Name.IsEmpty())
+							{
+								return LOCTEXT("UnnamedPerNodeLabel", "Unnamed");
+							}
+
 							return Name;
 						})
 					]
@@ -1065,6 +1163,9 @@ void FDisplayClusterColorGradingGenerator_ICVFXCamera::GenerateDataModel(IProper
 			.ColorAndOpacity(FSlateColor::UseForeground())
 		];
 
+	OutColorGradingDataModel.OnColorGradingGroupDeleted().AddSP(this, &FDisplayClusterColorGradingGenerator_ICVFXCamera::DeleteColorGradingGroup);
+	OutColorGradingDataModel.OnColorGradingGroupRenamed().AddSP(this, &FDisplayClusterColorGradingGenerator_ICVFXCamera::RenameColorGradingGroup);
+
 	{
 		FDisplayClusterColorGradingDataModel::FDetailsSection InnerFrustumDetailsSection;
 		InnerFrustumDetailsSection.DisplayName = LOCTEXT("InnerFrustumDetailsSectionLabel", "Inner Frustum");
@@ -1148,6 +1249,46 @@ FReply FDisplayClusterColorGradingGenerator_ICVFXCamera::AddColorGradingGroup()
 	return FReply::Handled();
 }
 
+void FDisplayClusterColorGradingGenerator_ICVFXCamera::DeleteColorGradingGroup(int32 GroupIndex)
+{
+	for (const TWeakObjectPtr<UDisplayClusterICVFXCameraComponent>& CameraComponent : CameraComponents)
+	{
+		if (CameraComponent.IsValid())
+		{
+			// Use GroupIndex - 1 since the entire cluster color grading is group 0
+			if (CameraComponent->CameraSettings.PerNodeColorGrading.IsValidIndex(GroupIndex - 1))
+			{
+				FScopedTransaction Transaction(LOCTEXT("DeleteNodeColorGradingGroupTransaction", "Delete Node Group"));
+				CameraComponent->Modify();
+
+				CameraComponent->CameraSettings.PerNodeColorGrading.RemoveAt(GroupIndex - 1);
+
+				IDisplayClusterColorGrading::Get().GetColorGradingDrawerSingleton().RefreshColorGradingDrawers(true);
+			}
+		}
+	}
+}
+
+void FDisplayClusterColorGradingGenerator_ICVFXCamera::RenameColorGradingGroup(int32 GroupIndex, const FText& NewName)
+{
+	for (const TWeakObjectPtr<UDisplayClusterICVFXCameraComponent>& CameraComponent : CameraComponents)
+	{
+		if (CameraComponent.IsValid())
+		{
+			// Use GroupIndex - 1 since the entire cluster color grading is group 0
+			if (CameraComponent->CameraSettings.PerNodeColorGrading.IsValidIndex(GroupIndex - 1))
+			{
+				FScopedTransaction Transaction(LOCTEXT("RenameNodeColorGradingGroupTransaction", "Rename Node Group"));
+				CameraComponent->Modify();
+
+				CameraComponent->CameraSettings.PerNodeColorGrading[GroupIndex - 1].Name = NewName;
+
+				IDisplayClusterColorGrading::Get().GetColorGradingDrawerSingleton().RefreshColorGradingDrawers(true);
+			}
+		}
+	}
+}
+
 TSharedRef<SWidget> FDisplayClusterColorGradingGenerator_ICVFXCamera::CreateNodeComboBox(int32 PerNodeColorGradingIndex) const
 {
 	return SNew(SComboButton)
@@ -1291,4 +1432,4 @@ TSharedRef<SWidget> FDisplayClusterColorGradingGenerator_ICVFXCamera::GetNodeCom
 	return MenuBuilder.MakeWidget();
 }
 
-#undef LOCTEXT_NAMESPACE
\ No newline at end of file
+#undef LOCTEXT_NAMESPACE
diff --git a/Engine/Plugins/Runtime/nDisplay/Source/DisplayClusterColorGrading/Private/DataModelGenerators/DisplayClusterColorGradingGenerator_RootActor.h b/Engine/Plugins/Runtime/nDisplay/Source/DisplayClusterColorGrading/Private/DataModelGenerators/DisplayClusterColorGradingGenerator_RootActor.h
--- a/Engine/Plugins/Runtime/nDisplay/Source/DisplayClusterColorGrading/Private/DataModelGenerators/DisplayClusterColorGradingGenerator_RootActor.h
+++ b/Engine/Plugins/Runtime/nDisplay/Source/DisplayClusterColorGrading/Private/DataModelGenerators/DisplayClusterColorGradingGenerator_RootActor.h
@@ -48,6 +48,12 @@ private:
 	/** OnClick delegate handler for the "+" button that is added to the color grading group toolbar */
 	FReply AddColorGradingGroup();
 
+	/** Raised when a color grading group is being deleted */
+	void DeleteColorGradingGroup(int32 GroupIndex);
+
+	/** Raised when a color grading group is being renamed */
+	void RenameColorGradingGroup(int32 GroupIndex, const FText& NewName);
+
 	/** Creates the combo box used to display which viewports are part of the specified color grading group */
 	TSharedRef<SWidget> CreateViewportComboBox(int32 PerViewportColorGradingIndex) const;
 
@@ -78,6 +84,12 @@ private:
 	/** OnClick delegate handler for the "+" button that is added to the color grading group toolbar */
 	FReply AddColorGradingGroup();
 
+	/** Raised when a color grading group is being deleted */
+	void DeleteColorGradingGroup(int32 GroupIndex);
+
+	/** Raised when a color grading group is being renamed */
+	void RenameColorGradingGroup(int32 GroupIndex, const FText& NewName);
+
 	/** Creates the combo box used to display which nodes are part of the specified color grading group */
 	TSharedRef<SWidget> CreateNodeComboBox(int32 PerNodeColorGradingIndex) const;
 
diff --git a/Engine/Plugins/Runtime/nDisplay/Source/DisplayClusterColorGrading/Private/DisplayClusterColorGradingDataModel.cpp b/Engine/Plugins/Runtime/nDisplay/Source/DisplayClusterColorGrading/Private/DisplayClusterColorGradingDataModel.cpp
--- a/Engine/Plugins/Runtime/nDisplay/Source/DisplayClusterColorGrading/Private/DisplayClusterColorGradingDataModel.cpp
+++ b/Engine/Plugins/Runtime/nDisplay/Source/DisplayClusterColorGrading/Private/DisplayClusterColorGradingDataModel.cpp
@@ -149,6 +149,9 @@ void FDisplayClusterColorGradingDataModel::Reset()
 		PropertyRowGenerator->UnregisterInstancedCustomPropertyLayout(GeneratorInstance.Key.Get());
 	}
 
+	OnColorGradingGroupDeletedDelegate.Clear();
+	OnColorGradingGroupRenamedDelegate.Clear();
+
 	DataModelGeneratorInstances.Empty();
 	ColorGradingGroups.Empty();
 	DetailsSections.Empty();
@@ -192,6 +195,13 @@ void FDisplayClusterColorGradingDataModel::SetDrawerState(const FDisplayClusterC
 	{
 		PropertyRowGenerator->SetObjects(ObjectsToSelect);
 	}
+
+	// After the data model has been created as part of the SetObjects call, check that the saved selected color grading group is still valid,
+	// and if not, set the selected group to 0
+	if (SelectedColorGradingGroupIndex >= ColorGradingGroups.Num())
+	{
+		SelectedColorGradingGroupIndex = 0;
+	}
 }
 
 FDisplayClusterColorGradingDataModel::FColorGradingGroup* FDisplayClusterColorGradingDataModel::GetSelectedColorGradingGroup()
diff --git a/Engine/Plugins/Runtime/nDisplay/Source/DisplayClusterColorGrading/Private/DisplayClusterColorGradingDataModel.h b/Engine/Plugins/Runtime/nDisplay/Source/DisplayClusterColorGrading/Private/DisplayClusterColorGradingDataModel.h
--- a/Engine/Plugins/Runtime/nDisplay/Source/DisplayClusterColorGrading/Private/DisplayClusterColorGradingDataModel.h
+++ b/Engine/Plugins/Runtime/nDisplay/Source/DisplayClusterColorGrading/Private/DisplayClusterColorGradingDataModel.h
@@ -88,6 +88,12 @@ public:
 
 		/** A list of categories of properties to display in the detail view of the color grading panel */
 		TArray<FName> DetailsViewCategories;
+
+		/** Indicates that the color grading group can be deleted by the user */
+		bool bCanBeDeleted = false;
+
+		/** Indicates that the color grading group can be renamed by the user */
+		bool bCanBeRenamed = false;
 	};
 
 	/** Stores data used to generate a details subsection, which is a selectable element in a details section that modifies which properties are displayed in the section */
@@ -124,6 +130,8 @@ public:
 
 	DECLARE_MULTICAST_DELEGATE(FOnDataModelGenerated);
 	DECLARE_MULTICAST_DELEGATE(FOnColorGradingSelectionChanged);
+	DECLARE_MULTICAST_DELEGATE_OneParam(FOnColorGradingGroupDeleted, int32);
+	DECLARE_MULTICAST_DELEGATE_TwoParams(FOnColorGradingGroupRenamed, int32, const FText&);
 
 public:
 	FDisplayClusterColorGradingDataModel();
@@ -176,6 +184,12 @@ public:
 	/** Gets the delegate that is raised when the selected color grading element has changed */
 	FOnColorGradingSelectionChanged& OnColorGradingElementSelectionChanged() { return OnColorGradingElementSelectionChangedDelegate; }
 
+	/** Gets the delegate that is raised when a color grading group is deleted by the user */
+	FOnColorGradingGroupDeleted& OnColorGradingGroupDeleted() { return OnColorGradingGroupDeletedDelegate; }
+
+	/** Gets the delegate that is raised when a color grading group is renamed by the user */
+	FOnColorGradingGroupRenamed& OnColorGradingGroupRenamed() { return OnColorGradingGroupRenamedDelegate; }
+
 private:
 	/** Initializes any needed data model generators for the specified class */
 	void InitializeDataModelGenerator(UClass* InClass);
@@ -220,4 +234,10 @@ private:
 
 	/** Delegate that is raised when the selected color grading element has been changed */
 	FOnColorGradingSelectionChanged OnColorGradingElementSelectionChangedDelegate;
+
+	/** Delegate that is raised when a color grading group is being deleted */
+	FOnColorGradingGroupDeleted OnColorGradingGroupDeletedDelegate;
+
+	/** Delegate that is raised when a color grading group is being renamed */
+	FOnColorGradingGroupRenamed OnColorGradingGroupRenamedDelegate;
 };
\ No newline at end of file
diff --git a/Engine/Plugins/Runtime/nDisplay/Source/DisplayClusterColorGrading/Private/Drawer/SDisplayClusterColorGradingDrawer.cpp b/Engine/Plugins/Runtime/nDisplay/Source/DisplayClusterColorGrading/Private/Drawer/SDisplayClusterColorGradingDrawer.cpp
--- a/Engine/Plugins/Runtime/nDisplay/Source/DisplayClusterColorGrading/Private/Drawer/SDisplayClusterColorGradingDrawer.cpp
+++ b/Engine/Plugins/Runtime/nDisplay/Source/DisplayClusterColorGrading/Private/Drawer/SDisplayClusterColorGradingDrawer.cpp
@@ -20,8 +20,10 @@
 #include "ColorCorrectRegion.h"
 #include "ColorCorrectWindow.h"
 #include "Engine/PostProcessVolume.h"
+#include "Framework/MultiBox/MultiBoxBuilder.h"
 #include "Modules/ModuleManager.h"
 #include "ScopedTransaction.h"
+#include "Framework/Commands/GenericCommands.h"
 #include "Styling/AppStyle.h"
 #include "Widgets/Input/SButton.h"
 #include "Widgets/Layout/SBorder.h"
@@ -32,6 +34,7 @@
 #include "Widgets/Layout/SSpacer.h"
 #include "Widgets/Layout/SSplitter.h"
 #include "Widgets/SBoxPanel.h"
+#include "Widgets/Text/SInlineEditableTextBlock.h"
 
 #define LOCTEXT_NAMESPACE "DisplayClusterColorGrading"
 
@@ -233,7 +236,7 @@ void SDisplayClusterColorGradingDrawer::Construct(const FArguments& InArgs, bool
 					.BorderImage(bIsInDrawer ? FStyleDefaults::GetNoBrush() : FAppStyle::Get().GetBrush("Brushes.Panel"))
 					[
 						SNew(SBox)
-						.HeightOverride(24.0f)
+						.HeightOverride(28.0f)
 						[
 							SNew(SHorizontalBox)
 
@@ -836,10 +839,12 @@ void SDisplayClusterColorGradingDrawer::FillColorGradingGroupToolBar()
 	if (ColorGradingGroupToolBarBox.IsValid())
 	{
 		ColorGradingGroupToolBarBox->ClearChildren();
+		ColorGradingGroupTextBlocks.Empty(ColorGradingDataModel->ColorGradingGroups.Num());
 
 		for (int32 Index = 0; Index < ColorGradingDataModel->ColorGradingGroups.Num(); ++Index)
 		{
-			const FDisplayClusterColorGradingDataModel::FColorGradingGroup& Group = ColorGradingDataModel->ColorGradingGroups[Index];
+			TSharedPtr<SInlineEditableTextBlock> TextBlock = nullptr;
+
 			ColorGradingGroupToolBarBox->AddSlot()
 			.AutoWidth()
 			.Padding(2.0f, 0.0f, 2.0f, 0.0f)
@@ -848,12 +853,20 @@ void SDisplayClusterColorGradingDrawer::FillColorGradingGroupToolBar()
 				.Style(FAppStyle::Get(), "DetailsView.SectionButton")
 				.OnCheckStateChanged(this, &SDisplayClusterColorGradingDrawer::OnColorGradingGroupCheckedChanged, Index)
 				.IsChecked(this, &SDisplayClusterColorGradingDrawer::IsColorGradingGroupSelected, Index)
+				.OnGetMenuContent(this, &SDisplayClusterColorGradingDrawer::GetColorGradingGroupMenuContent, Index)
 				[
-					SNew(STextBlock)
-					.TextStyle(FAppStyle::Get(), "SmallText")
-					.Text(Group.DisplayName)
+					SNew(SBox)
+					.HeightOverride(20.0)
+					[
+						SAssignNew(TextBlock, SInlineEditableTextBlock)
+						.Text(this, &SDisplayClusterColorGradingDrawer::GetColorGradingGroupDisplayName, Index)
+						.Font(this, &SDisplayClusterColorGradingDrawer::GetColorGradingGroupDisplayNameFont, Index)
+						.OnTextCommitted(this, &SDisplayClusterColorGradingDrawer::OnColorGradingGroupRenamed, Index)
+					]
 				]
 			];
+
+			ColorGradingGroupTextBlocks.Add(TextBlock);
 		}
 
 		if (ColorGradingDataModel->ColorGradingGroupToolBarWidget.IsValid())
@@ -909,6 +922,99 @@ void SDisplayClusterColorGradingDrawer::OnColorGradingGroupCheckedChanged(ECheck
 	}
 }
 
+FText SDisplayClusterColorGradingDrawer::GetColorGradingGroupDisplayName(int32 GroupIndex) const
+{
+	if (ColorGradingDataModel->ColorGradingGroups.IsValidIndex(GroupIndex))
+	{
+		FText DisplayName = ColorGradingDataModel->ColorGradingGroups[GroupIndex].DisplayName;
+		if (DisplayName.IsEmpty())
+		{
+			return LOCTEXT("ColorGradingGroupEmptyNameLabel", "Unnamed");
+		}
+
+		return DisplayName;
+	}
+
+	return FText::GetEmpty();
+}
+
+FSlateFontInfo SDisplayClusterColorGradingDrawer::GetColorGradingGroupDisplayNameFont(int32 GroupIndex) const
+{
+	if (ColorGradingDataModel->ColorGradingGroups.IsValidIndex(GroupIndex))
+	{
+		FText DisplayName = ColorGradingDataModel->ColorGradingGroups[GroupIndex].DisplayName;
+		if (DisplayName.IsEmpty())
+		{
+			return FAppStyle::Get().GetFontStyle("NormalFontItalic");
+		}
+	}
+
+	return FAppStyle::Get().GetFontStyle("NormalFont");
+}
+
+TSharedRef<SWidget> SDisplayClusterColorGradingDrawer::GetColorGradingGroupMenuContent(int32 GroupIndex)
+{
+	if (ColorGradingDataModel->ColorGradingGroups.IsValidIndex(GroupIndex))
+	{
+		const FDisplayClusterColorGradingDataModel::FColorGradingGroup& Group = ColorGradingDataModel->ColorGradingGroups[GroupIndex];
+
+		FMenuBuilder MenuBuilder(true, nullptr);
+
+		const FGenericCommands& GenericCommands = FGenericCommands::Get();
+
+		if (Group.bCanBeRenamed)
+		{
+			MenuBuilder.AddMenuEntry(
+				GenericCommands.Rename->GetLabel(),
+				GenericCommands.Rename->GetDescription(),
+				GenericCommands.Rename->GetIcon(),
+				FUIAction(FExecuteAction::CreateSP(this, &SDisplayClusterColorGradingDrawer::OnColorGradingGroupRequestRename, GroupIndex))
+			);
+		}
+
+		if (Group.bCanBeDeleted)
+		{
+			MenuBuilder.AddMenuEntry(
+				GenericCommands.Delete->GetLabel(),
+				GenericCommands.Delete->GetDescription(),
+				GenericCommands.Delete->GetIcon(),
+				FUIAction(FExecuteAction::CreateSP(this, &SDisplayClusterColorGradingDrawer::OnColorGradingGroupDeleted, GroupIndex))
+			);
+		}
+
+		return MenuBuilder.MakeWidget();
+	}
+
+	return SNullWidget::NullWidget;
+}
+
+void SDisplayClusterColorGradingDrawer::OnColorGradingGroupDeleted(int32 GroupIndex)
+{
+	// If the group being deleted is in front of the currently selected one, we want to make sure
+	// that the same group is selected even after the deletion, so preemptively adjust the
+	// currently selected group index
+	int32 SelectedGroupIndex = ColorGradingDataModel->GetSelectedColorGradingGroupIndex();
+	if (SelectedGroupIndex > GroupIndex)
+	{
+		ColorGradingDataModel->SetSelectedColorGradingGroup(SelectedGroupIndex - 1);
+	}
+
+	ColorGradingDataModel->OnColorGradingGroupDeleted().Broadcast(GroupIndex);
+}
+
+void SDisplayClusterColorGradingDrawer::OnColorGradingGroupRequestRename(int32 GroupIndex)
+{
+	if (ColorGradingGroupTextBlocks.IsValidIndex(GroupIndex) && ColorGradingGroupTextBlocks[GroupIndex].IsValid())
+	{
+		ColorGradingGroupTextBlocks[GroupIndex]->EnterEditingMode();
+	}
+}
+
+void SDisplayClusterColorGradingDrawer::OnColorGradingGroupRenamed(const FText& InText, ETextCommit::Type TextCommitType, int32 GroupIndex)
+{
+	ColorGradingDataModel->OnColorGradingGroupRenamed().Broadcast(GroupIndex, InText);
+}
+
 void SDisplayClusterColorGradingDrawer::OnObjectsReplaced(const TMap<UObject*, UObject*>& OldToNewInstanceMap)
 {
 	bool bNeedsFullRefresh = false;
diff --git a/Engine/Plugins/Runtime/nDisplay/Source/DisplayClusterColorGrading/Private/Drawer/SDisplayClusterColorGradingDrawer.h b/Engine/Plugins/Runtime/nDisplay/Source/DisplayClusterColorGrading/Private/Drawer/SDisplayClusterColorGradingDrawer.h
--- a/Engine/Plugins/Runtime/nDisplay/Source/DisplayClusterColorGrading/Private/Drawer/SDisplayClusterColorGradingDrawer.h
+++ b/Engine/Plugins/Runtime/nDisplay/Source/DisplayClusterColorGrading/Private/Drawer/SDisplayClusterColorGradingDrawer.h
@@ -12,6 +12,7 @@
 #include "SDisplayClusterColorGradingObjectList.h"
 #include "SDisplayClusterColorGradingColorWheel.h"
 
+class SInlineEditableTextBlock;
 class ADisplayClusterRootActor;
 class FDisplayClusterOperatorStatusBarExtender;
 class FDisplayClusterColorGradingDataModel;
@@ -112,6 +113,24 @@ private:
 	/** Raised when the user has selected the specified color grading group */
 	void OnColorGradingGroupCheckedChanged(ECheckBoxState State, int32 GroupIndex);
 
+	/** Gets the display name of the specified color grading group */
+	FText GetColorGradingGroupDisplayName(int32 GroupIndex) const;
+
+	/** Gets the font of the display name label of the specified color grading group */
+	FSlateFontInfo GetColorGradingGroupDisplayNameFont(int32 GroupIndex) const;
+
+	/** Gets the content for the right click menu for the color grading group */
+	TSharedRef<SWidget> GetColorGradingGroupMenuContent(int32 GroupIndex);
+
+	/** Raised when a color grading group has been deleted by the user */
+	void OnColorGradingGroupDeleted(int32 GroupIndex);
+
+	/** Raised when a rename has been requested on a color grading group */
+	void OnColorGradingGroupRequestRename(int32 GroupIndex);
+
+	/** Raised when a rename has been committed on a color grading group */
+	void OnColorGradingGroupRenamed(const FText& InText, ETextCommit::Type TextCommitType, int32 GroupIndex);
+
 	/** Raised when the editor replaces any UObjects with new instantiations, usually when actors have been recompiled from blueprints */
 	void OnObjectsReplaced(const TMap<UObject*, UObject*>& OldToNewInstanceMap);
 
@@ -127,7 +146,7 @@ private:
 	/** Raised when the user has changed the active root actor selected in the nDisplay operator panel */
 	void OnActiveRootActorChanged(ADisplayClusterRootActor* NewRootActor);
 	
-	/** Raised when the obejcts being displayed in the operator's details panel have changed */
+	/** Raised when the objects being displayed in the operator's details panel have changed */
 	void OnDetailObjectsChanged(const TArray<UObject*>& Objects);
 
 	/** Raised when the color grading data model has been generated */
@@ -152,12 +171,13 @@ private:
 	TSharedPtr<SHorizontalBox> ColorGradingGroupToolBarBox;
 	TSharedPtr<SDisplayClusterColorGradingColorWheelPanel> ColorWheelPanel;
 	TSharedPtr<SDisplayClusterColorGradingDetailsPanel> DetailsPanel;
+	TArray<TSharedPtr<SInlineEditableTextBlock>> ColorGradingGroupTextBlocks;
 
 	/** The color grading data model for the currently selected objects */
 	TSharedPtr<FDisplayClusterColorGradingDataModel> ColorGradingDataModel;
 
 	/** The current mode that the drawer is in */
-	EDisplayClusterColorGradingDrawerMode CurrentDrawerMode;
+	EDisplayClusterColorGradingDrawerMode CurrentDrawerMode = EDisplayClusterColorGradingDrawerMode::ColorGrading;
 
 	/** Gets whether this widget is in a drawer or docked in a tab */
 	bool bIsInDrawer = false;
@@ -170,4 +190,4 @@ private:
 
 	/** Delegate handle for the OnActiveRootActorChanged delegate */
 	FDelegateHandle ActiveRootActorChangedHandle;
-};
\ No newline at end of file
+};
