nDisplay: Added Media Plate button event transactions.

diff --git a/Engine/Plugins/Runtime/nDisplay/Source/DisplayClusterMultiUser/DisplayClusterMultiUser.Build.cs b/Engine/Plugins/Runtime/nDisplay/Source/DisplayClusterMultiUser/DisplayClusterMultiUser.Build.cs
--- a/Engine/Plugins/Runtime/nDisplay/Source/DisplayClusterMultiUser/DisplayClusterMultiUser.Build.cs
+++ b/Engine/Plugins/Runtime/nDisplay/Source/DisplayClusterMultiUser/DisplayClusterMultiUser.Build.cs
@@ -14,7 +14,8 @@ public class DisplayClusterMultiUser : ModuleRules
 				"CoreUObject",
 				"ConcertSyncClient",
 				"DisplayClusterConfiguration",
-				"Engine"
+				"Engine",
+				"MediaAssets"
 			});
 	}
 }
diff --git a/Engine/Plugins/Runtime/nDisplay/Source/DisplayClusterMultiUser/Private/DisplayClusterMultiUserModule.cpp b/Engine/Plugins/Runtime/nDisplay/Source/DisplayClusterMultiUser/Private/DisplayClusterMultiUserModule.cpp
--- a/Engine/Plugins/Runtime/nDisplay/Source/DisplayClusterMultiUser/Private/DisplayClusterMultiUserModule.cpp
+++ b/Engine/Plugins/Runtime/nDisplay/Source/DisplayClusterMultiUser/Private/DisplayClusterMultiUserModule.cpp
@@ -1,18 +1,27 @@
 // Copyright Epic Games, Inc. All Rights Reserved.
 
 #include "DisplayClusterMultiUserModule.h"
-
 #include "Modules/ModuleManager.h"
 
+#if WITH_EDITOR
+#include "MediaAssetMultiUserManager.h"
+#endif
+
 
 void FDisplayClusterMultiUserModule::StartupModule()
 {
 	MultiUserManager = MakeUnique<FDisplayClusterMultiUserManager>();
+#if WITH_EDITOR
+	MediaAssetMUManager = MakeUnique<FMediaAssetMultiUserManager>();
+#endif
 }
 
 void FDisplayClusterMultiUserModule::ShutdownModule()
 {
 	MultiUserManager.Reset();
+#if WITH_EDITOR
+	MediaAssetMUManager.Reset();
+#endif
 }
 
 IMPLEMENT_MODULE(FDisplayClusterMultiUserModule, DisplayClusterMultiUser);
diff --git a/Engine/Plugins/Runtime/nDisplay/Source/DisplayClusterMultiUser/Private/DisplayClusterMultiUserModule.h b/Engine/Plugins/Runtime/nDisplay/Source/DisplayClusterMultiUser/Private/DisplayClusterMultiUserModule.h
--- a/Engine/Plugins/Runtime/nDisplay/Source/DisplayClusterMultiUser/Private/DisplayClusterMultiUserModule.h
+++ b/Engine/Plugins/Runtime/nDisplay/Source/DisplayClusterMultiUser/Private/DisplayClusterMultiUserModule.h
@@ -7,6 +7,7 @@
 #include "Modules/ModuleInterface.h"
 #include "Templates/UniquePtr.h"
 
+class FMediaAssetMultiUserManager;
 
 class FDisplayClusterMultiUserModule :
 	public IModuleInterface
@@ -19,4 +20,8 @@ public:
 private:
 	// Manager for multi user connections.
 	TUniquePtr<FDisplayClusterMultiUserManager> MultiUserManager;
+
+#if WITH_EDITOR
+	TUniquePtr<FMediaAssetMultiUserManager> MediaAssetMUManager;
+#endif
 };
diff --git a/Engine/Plugins/Runtime/nDisplay/Source/DisplayClusterMultiUser/Private/MediaAssetMultiUserManager.cpp b/Engine/Plugins/Runtime/nDisplay/Source/DisplayClusterMultiUser/Private/MediaAssetMultiUserManager.cpp
new file mode 100644
--- /dev/null
+++ b/Engine/Plugins/Runtime/nDisplay/Source/DisplayClusterMultiUser/Private/MediaAssetMultiUserManager.cpp
@@ -0,0 +1,78 @@
+// Copyright Epic Games, Inc. All Rights Reserved.
+
+#include "MediaAssetMultiUserManager.h"
+
+#if WITH_EDITOR
+
+#include "IMediaAssetsModule.h"
+#include "IConcertSession.h"
+#include "IConcertSyncClientModule.h"
+#include "IConcertClientModule.h"
+#include "IConcertSyncClient.h"
+
+FMediaAssetMultiUserManager::FMediaAssetMultiUserManager()
+{
+	if (TSharedPtr<IConcertSyncClient> ConcertSyncClient = IConcertSyncClientModule::Get().GetClient(TEXT("MultiUser")))
+	{
+		IConcertClientRef ConcertClient = ConcertSyncClient->GetConcertClient();
+		ConcertClient->OnSessionStartup().AddRaw(this, &FMediaAssetMultiUserManager::Register);
+		ConcertClient->OnSessionShutdown().AddRaw(this, &FMediaAssetMultiUserManager::Unregister);
+
+		if (TSharedPtr<IConcertClientSession> ConcertClientSession = ConcertClient->GetCurrentSession())
+		{
+			Register(ConcertClientSession.ToSharedRef());
+		}
+	}
+
+	IMediaAssetsModule& MediaAssetsModule = FModuleManager::Get().LoadModuleChecked<IMediaAssetsModule>("MediaAssets");
+	OnMediaPlateStateChangedHandle = MediaAssetsModule.RegisterOnMediaStateChangedEvent(IMediaAssetsModule::FMediaStateChangedDelegate::FDelegate::CreateRaw(this, &FMediaAssetMultiUserManager::OnMediaPlateStateChanged));
+}
+
+
+FMediaAssetMultiUserManager::~FMediaAssetMultiUserManager()
+{
+	IMediaAssetsModule& MediaAssetsModule = FModuleManager::Get().GetModuleChecked<IMediaAssetsModule>("MediaAssets");
+	MediaAssetsModule.UnregisterOnMediaStateChangedEvent(OnMediaPlateStateChangedHandle);
+}
+
+void FMediaAssetMultiUserManager::Register(TSharedRef<IConcertClientSession> InSession)
+{
+	ConcertSession = InSession;
+	InSession->RegisterCustomEventHandler<FConcertMediaStateChangedEvent>(this, &FMediaAssetMultiUserManager::OnStateChangedEvent);
+}
+
+void FMediaAssetMultiUserManager::Unregister(TSharedRef<IConcertClientSession> InSession)
+{
+	if (TSharedPtr<IConcertClientSession> ConcertSessionPinned = ConcertSession.Pin())
+	{
+		ConcertSessionPinned->UnregisterCustomEventHandler<FConcertMediaStateChangedEvent>(this);
+	}
+	ConcertSession.Reset();
+}
+
+void FMediaAssetMultiUserManager::OnStateChangedEvent(const FConcertSessionContext& InConcertSessionContext, const FConcertMediaStateChangedEvent& InEvent)
+{
+	IMediaAssetsModule& MediaAssetsModule = FModuleManager::Get().LoadModuleChecked<IMediaAssetsModule>("MediaAssets");
+	MediaAssetsModule.BroadcastOnMediaStateChangedEvent(InEvent.ActorsPathNames, InEvent.State, true /* Broadcast came from a remote endpoint. */);
+}
+
+void FMediaAssetMultiUserManager::OnMediaPlateStateChanged(const TArray<FString>& InActorsPathNames, uint8 InEnumState, bool bRemoteBroadcast)
+{
+	// We don't need to broadcast to other endpoints as this is already handled.
+	if (bRemoteBroadcast)
+	{
+		return;
+	}
+
+	if (TSharedPtr<IConcertClientSession> ConcertSessionPinned = ConcertSession.Pin())
+	{
+		TArray<FGuid> ClientIds = ConcertSessionPinned->GetSessionClientEndpointIds();
+		FConcertMediaStateChangedEvent EventData;
+		EventData.ActorsPathNames = InActorsPathNames;
+		EventData.State = InEnumState;
+		ConcertSessionPinned->SendCustomEvent(EventData, ClientIds, EConcertMessageFlags::ReliableOrdered | EConcertMessageFlags::UniqueId);
+	}
+}
+
+#endif
+
diff --git a/Engine/Plugins/Runtime/nDisplay/Source/DisplayClusterMultiUser/Private/MediaAssetMultiUserManager.h b/Engine/Plugins/Runtime/nDisplay/Source/DisplayClusterMultiUser/Private/MediaAssetMultiUserManager.h
new file mode 100644
--- /dev/null
+++ b/Engine/Plugins/Runtime/nDisplay/Source/DisplayClusterMultiUser/Private/MediaAssetMultiUserManager.h
@@ -0,0 +1,70 @@
+// Copyright Epic Games, Inc. All Rights Reserved.
+
+#pragma once
+#include "Misc/CoreDefines.h"
+#include "MediaAssetMultiUserManager.generated.h"
+
+struct FConcertSessionContext;
+class IConcertSyncClient;
+class IConcertClientSession;
+
+#if WITH_EDITOR
+
+
+/** 
+* The event that is transfered over MU.
+*/
+USTRUCT()
+struct FConcertMediaStateChangedEvent
+{
+public:
+	GENERATED_BODY()
+
+	/** Name used to find the object on the other side of MU. */
+	UPROPERTY()
+	TArray<FString> ActorsPathNames;
+
+	/** State to replicate on the other side. */
+	UPROPERTY()
+	uint8 State;
+};
+
+/**
+ * Reacts to state changes in Media Plates and transfers state change information to remote endpoints.
+ */
+class FMediaAssetMultiUserManager
+{
+public:
+
+	FMediaAssetMultiUserManager();
+
+	virtual ~FMediaAssetMultiUserManager();
+
+	/** 
+	* Subscribes to the relevant Concert Sync events.
+	*/
+	void Register(TSharedRef<IConcertClientSession> InSession);
+
+	/** 
+	* Unsubscribes from concert sync events.
+	*/
+	void Unregister(TSharedRef<IConcertClientSession> InSession);
+
+	/** 
+	* Reacts to changes in Media Plate from remove endpoints and forwards it to local Media Plates.
+	*/
+	void OnStateChangedEvent(const FConcertSessionContext& InConcertSessionContext, const FConcertMediaStateChangedEvent& InEvent);
+
+	/**
+	* Called when Media Plate broadcasts state change events and sends a custom event to remote endpoints.
+	* bRemoteBroadcast identifies if event came from the remote endpoint or broadcasted locally.
+	*/
+	void OnMediaPlateStateChanged(const TArray<FString>& InNameId, uint8 InEnumState, bool bRemoteBroadcast);
+
+private:
+
+	TWeakPtr<IConcertClientSession> ConcertSession;
+	FDelegateHandle OnMediaPlateStateChangedHandle;
+};
+
+#endif // WITH_EDITOR
