[nDisplay] Fixed crash when capturing an nDisplay node backbuffer

diff --git a/Engine/Plugins/Runtime/nDisplay/Source/DisplayCluster/Private/DisplayClusterCallbacks.h b/Engine/Plugins/Runtime/nDisplay/Source/DisplayCluster/Private/DisplayClusterCallbacks.h
--- a/Engine/Plugins/Runtime/nDisplay/Source/DisplayCluster/Private/DisplayClusterCallbacks.h
+++ b/Engine/Plugins/Runtime/nDisplay/Source/DisplayCluster/Private/DisplayClusterCallbacks.h
@@ -122,6 +122,11 @@ public:
 		return DisplayClusterPostBackbufferUpdateEvent;
 	}
 
+	virtual FDisplayClusterPostBackbufferUpdated_RenderThread& OnDisplayClusterPostBackbufferUpdated_RenderThread() override
+	{
+		return DisplayClusterPostBackbufferUpdatedEvent;
+	}
+
 	virtual FDisplayClusterPreProcessIcvfx_RenderThread& OnDisplayClusterPreProcessIcvfx_RenderThread() override
 	{
 		return FDisplayClusterPreProcessIcvfxEvent;
@@ -152,5 +157,6 @@ private:
 	FDisplayClusterProcessLatency_RenderThread       FDisplayClusterProcessLatencyEvent;
 	FDisplayClusterPostFrameRender_RenderThread      DisplayClusterPostFrameRenderEvent;
 	FDisplayClusterPostBackbufferUpdate_RenderThread DisplayClusterPostBackbufferUpdateEvent;
+	FDisplayClusterPostBackbufferUpdated_RenderThread DisplayClusterPostBackbufferUpdatedEvent;
 	FDisplayClusterPreProcessIcvfx_RenderThread      FDisplayClusterPreProcessIcvfxEvent;
 };
diff --git a/Engine/Plugins/Runtime/nDisplay/Source/DisplayCluster/Private/Render/Viewport/DisplayClusterViewportManagerProxy.cpp b/Engine/Plugins/Runtime/nDisplay/Source/DisplayCluster/Private/Render/Viewport/DisplayClusterViewportManagerProxy.cpp
--- a/Engine/Plugins/Runtime/nDisplay/Source/DisplayCluster/Private/Render/Viewport/DisplayClusterViewportManagerProxy.cpp
+++ b/Engine/Plugins/Runtime/nDisplay/Source/DisplayCluster/Private/Render/Viewport/DisplayClusterViewportManagerProxy.cpp
@@ -297,7 +297,10 @@ void FDisplayClusterViewportManagerProxy::ImplRenderFrame(FViewport* InViewport)
 				ViewportManagerProxy->ResolveFrameTargetToBackBuffer_RenderThread(RHICmdList, 0, 0, FrameOutputRTT, FrameOutputRTT->GetSizeXY());
 
 				// Finally, notify about backbuffer update
+				IDisplayCluster::Get().GetCallbacks().OnDisplayClusterPostBackbufferUpdated_RenderThread().Broadcast(RHICmdList, InViewport);
+PRAGMA_DISABLE_DEPRECATION_WARNINGS
 				IDisplayCluster::Get().GetCallbacks().OnDisplayClusterPostBackbufferUpdate_RenderThread().Broadcast(RHICmdList, ViewportManagerProxy, InViewport);
+PRAGMA_ENABLE_DEPRECATION_WARNINGS
 			}
 		}
 	});
diff --git a/Engine/Plugins/Runtime/nDisplay/Source/DisplayCluster/Public/IDisplayClusterCallbacks.h b/Engine/Plugins/Runtime/nDisplay/Source/DisplayCluster/Public/IDisplayClusterCallbacks.h
--- a/Engine/Plugins/Runtime/nDisplay/Source/DisplayCluster/Public/IDisplayClusterCallbacks.h
+++ b/Engine/Plugins/Runtime/nDisplay/Source/DisplayCluster/Public/IDisplayClusterCallbacks.h
@@ -107,10 +107,15 @@ public:
 	DECLARE_EVENT_ThreeParams(IDisplayClusterCallbacks, FDisplayClusterPostFrameRender_RenderThread, FRHICommandListImmediate&, const IDisplayClusterViewportManagerProxy*, FViewport*);
 	virtual FDisplayClusterPostFrameRender_RenderThread& OnDisplayClusterPostFrameRender_RenderThread() = 0;
 
-	/** Called after backbuffer update **/
+	/** (deprecated) Called after backbuffer update **/
 	DECLARE_EVENT_ThreeParams(IDisplayClusterCallbacks, FDisplayClusterPostBackbufferUpdate_RenderThread, FRHICommandListImmediate&, const IDisplayClusterViewportManagerProxy*, FViewport*);
+	UE_DEPRECATED(5.2, "This event getter has been deprecated. Please use OnDisplayClusterPostBackbufferUpdated_RenderThread()")
 	virtual FDisplayClusterPostBackbufferUpdate_RenderThread& OnDisplayClusterPostBackbufferUpdate_RenderThread() = 0;
 
+	/** Called after backbuffer update **/
+	DECLARE_EVENT_TwoParams(IDisplayClusterCallbacks, FDisplayClusterPostBackbufferUpdated_RenderThread, FRHICommandListImmediate&, FViewport*);
+	virtual FDisplayClusterPostBackbufferUpdated_RenderThread& OnDisplayClusterPostBackbufferUpdated_RenderThread() = 0;
+
 	/** Called before applying ICVFX shaders **/
 	DECLARE_EVENT_FourParams(IDisplayClusterCallbacks, FDisplayClusterPreProcessIcvfx_RenderThread, FRHICommandListImmediate&, const IDisplayClusterViewportProxy*, FDisplayClusterShaderParameters_WarpBlend&, FDisplayClusterShaderParameters_ICVFX&);
 	virtual FDisplayClusterPreProcessIcvfx_RenderThread& OnDisplayClusterPreProcessIcvfx_RenderThread() = 0;
diff --git a/Engine/Plugins/Runtime/nDisplay/Source/DisplayClusterConfigurator/Private/DisplayClusterConfiguratorCompiler.cpp b/Engine/Plugins/Runtime/nDisplay/Source/DisplayClusterConfigurator/Private/DisplayClusterConfiguratorCompiler.cpp
--- a/Engine/Plugins/Runtime/nDisplay/Source/DisplayClusterConfigurator/Private/DisplayClusterConfiguratorCompiler.cpp
+++ b/Engine/Plugins/Runtime/nDisplay/Source/DisplayClusterConfigurator/Private/DisplayClusterConfiguratorCompiler.cpp
@@ -259,11 +259,12 @@ void FDisplayClusterConfiguratorKismetCompilerContext::ValidateConfiguration()
 
 	TMap<FString, UDisplayClusterConfigurationViewport*> ViewportsByName;
 	
-	for (const auto& ClusterNodes : BlueprintData->Cluster->Nodes)
+	for (const TPair<FString, TObjectPtr<UDisplayClusterConfigurationClusterNode>>& ClusterNode : BlueprintData->Cluster->Nodes)
 	{
-		if (ClusterNodes.Value->Viewports.Num() > 0)
+		// Validate vieports
+		if (ClusterNode.Value->Viewports.Num() > 0)
 		{
-			for (const auto& Viewport : ClusterNodes.Value->Viewports)
+			for (const TPair<FString, TObjectPtr<UDisplayClusterConfigurationViewport>>& Viewport : ClusterNode.Value->Viewports)
 			{
 				const FString ViewportName = Viewport.Value->GetName();
 
@@ -291,11 +292,17 @@ void FDisplayClusterConfiguratorKismetCompilerContext::ValidateConfiguration()
 		{
 			// It's allowed not to have any viewports when rendering offscreen. For example, we're going to render
 			// ICVFX inner views only which don't require any viewports. In other cases continue throwing a warning.
-			if (!ClusterNodes.Value->bRenderHeadless)
+			if (!ClusterNode.Value->bRenderHeadless)
 			{
-				MessageLog.Warning(*LOCTEXT("NoViewportsError", "Node @@ has no viewports. Please add a viewport.").ToString(), *ClusterNodes.Key);
+				MessageLog.Warning(*LOCTEXT("NoViewportsError", "Node @@ has no viewports. Please add a viewport.").ToString(), *ClusterNode.Key);
 			}
 		}
+
+		// Warn user about fullscreen + headless
+		if (ClusterNode.Value->bRenderHeadless && ClusterNode.Value->bIsFullscreen)
+		{
+			MessageLog.Warning(*LOCTEXT("HeadlessFullscreen", "Node @@ has both 'headless' and 'fullscreen' flags set. It's not forbidden but make sure it's intentional as the backbuffer will have the desktop's size.").ToString(), *ClusterNode.Key);
+		}
 	}
 }
 
diff --git a/Engine/Plugins/Runtime/nDisplay/Source/DisplayClusterMedia/Private/Capture/DisplayClusterMediaCaptureNode.cpp b/Engine/Plugins/Runtime/nDisplay/Source/DisplayClusterMedia/Private/Capture/DisplayClusterMediaCaptureNode.cpp
--- a/Engine/Plugins/Runtime/nDisplay/Source/DisplayClusterMedia/Private/Capture/DisplayClusterMediaCaptureNode.cpp
+++ b/Engine/Plugins/Runtime/nDisplay/Source/DisplayClusterMedia/Private/Capture/DisplayClusterMediaCaptureNode.cpp
@@ -12,6 +12,8 @@
 #include "RHICommandList.h"
 #include "RHIResources.h"
 
+#include "Engine/Engine.h"
+#include "Engine/GameViewportClient.h"
 #include "UnrealClient.h"
 
 
@@ -26,7 +28,7 @@ bool FDisplayClusterMediaCaptureNode::StartCapture()
 	// If capturing initialized and started successfully, subscribe for rendering callbacks
 	if (FDisplayClusterMediaCaptureBase::StartCapture())
 	{
-		IDisplayCluster::Get().GetCallbacks().OnDisplayClusterPostFrameRender_RenderThread().AddRaw(this, &FDisplayClusterMediaCaptureNode::OnPostFrameRender_RenderThread);
+		IDisplayCluster::Get().GetCallbacks().OnDisplayClusterPostBackbufferUpdated_RenderThread().AddRaw(this, &FDisplayClusterMediaCaptureNode::OnPostBackbufferUpdated_RenderThread);
 		return true;
 	}
 
@@ -36,73 +38,33 @@ bool FDisplayClusterMediaCaptureNode::StartCapture()
 void FDisplayClusterMediaCaptureNode::StopCapture()
 {
 	// Stop rendering notifications
-	IDisplayCluster::Get().GetCallbacks().OnDisplayClusterPostFrameRender_RenderThread().RemoveAll(this);
+	IDisplayCluster::Get().GetCallbacks().OnDisplayClusterPostBackbufferUpdated_RenderThread().RemoveAll(this);
 	// Stop capturing
 	FDisplayClusterMediaCaptureBase::StopCapture();
 }
 
 FIntPoint FDisplayClusterMediaCaptureNode::GetCaptureSize() const
 {
-	if (IDisplayCluster::Get().GetConfigMgr() && IDisplayCluster::Get().GetConfigMgr()->GetLocalNode())
+	// Return backbuffer runtime size
+	if (GEngine && GEngine->GameViewport && GEngine->GameViewport->Viewport)
 	{
-		const FDisplayClusterConfigurationRectangle& Window = IDisplayCluster::Get().GetConfigMgr()->GetLocalNode()->WindowRect;
-		return { Window.W, Window.H };
+		return GEngine->GameViewport->Viewport->GetSizeXY();
 	}
 
 	return FIntPoint();
 }
 
-FTextureRHIRef FDisplayClusterMediaCaptureNode::CreateIntermediateTexture_RenderThread(EPixelFormat Format, ETextureCreateFlags Flags, const FIntPoint& Size)
+void FDisplayClusterMediaCaptureNode::OnPostBackbufferUpdated_RenderThread(FRHICommandListImmediate& RHICmdList, FViewport* Viewport)
 {
-	// Prepare description
-	FRHITextureCreateDesc Desc =
-		FRHITextureCreateDesc::Create2D(TEXT("DisplayClusterFrameQueueCacheTexture"), Size.X, Size.Y, Format)
-		.SetClearValue(FClearValueBinding::Black)
-		.SetNumMips(1)
-		.SetFlags(ETextureCreateFlags::Dynamic)
-		.AddFlags(ETextureCreateFlags::MultiGPUGraphIgnore)
-		.SetInitialState(ERHIAccess::SRVMask);
-
-	// Leave original flags, but make sure it's ResolveTargetable but not RenderTargetable
-	Flags &= ~ETextureCreateFlags::RenderTargetable;
-	Flags |= ETextureCreateFlags::ResolveTargetable;
-	Desc.SetFlags(Flags);
-
-	// Create texture
-	return RHICreateTexture(Desc);
-}
-
-void FDisplayClusterMediaCaptureNode::OnPostFrameRender_RenderThread(FRHICommandListImmediate& RHICmdList, const IDisplayClusterViewportManagerProxy* ViewportManagerProxy, FViewport* Viewport)
-{
-	TArray<FRHITexture*> Textures;
-	TArray<FIntPoint>    Regions;
-
-	if (ViewportManagerProxy && ViewportManagerProxy->GetFrameTargets_RenderThread(Textures, Regions))
+	if (Viewport)
 	{
-		if (Textures.Num() > 0 && Regions.Num() > 0 && Textures[0])
+		if (FRHITexture* const BackbufferTexture = Viewport->GetRenderTargetTexture())
 		{
-			// Create interim texture if not already available
-			if (!InterimTexture.IsValid() && Viewport)
-			{
-				if (FRHITexture* const BackbufferTexture = Viewport->GetRenderTargetTexture())
-				{
-					InterimTexture = CreateIntermediateTexture_RenderThread(Textures[0]->GetDesc().Format, Textures[0]->GetDesc().Flags, BackbufferTexture->GetDesc().Extent);
-				}
-			}
-
-			if (InterimTexture.IsValid())
-			{
-				// Copy texture to the intermediate buffer
-				FRHICopyTextureInfo CopyInfo;
-				CopyInfo.Size = Textures[0]->GetDesc().GetSize();
-				RHICmdList.CopyTexture(Textures[0], InterimTexture, CopyInfo);
-
-				// Capture intermediate texture
-				FMediaTextureInfo TextureInfo{ InterimTexture, FIntRect(FIntPoint::ZeroValue, InterimTexture->GetDesc().Extent) };
-				FRDGBuilder Builder(RHICmdList);
-				ExportMediaData(Builder, TextureInfo);
-				Builder.Execute();
-			}
+			// Capture backbuffer
+			FMediaTextureInfo TextureInfo{ BackbufferTexture, FIntRect(FIntPoint::ZeroValue, BackbufferTexture->GetDesc().Extent) };
+			FRDGBuilder Builder(RHICmdList);
+			ExportMediaData(Builder, TextureInfo);
+			Builder.Execute();
 		}
 	}
 }
diff --git a/Engine/Plugins/Runtime/nDisplay/Source/DisplayClusterMedia/Private/Capture/DisplayClusterMediaCaptureNode.h b/Engine/Plugins/Runtime/nDisplay/Source/DisplayClusterMedia/Private/Capture/DisplayClusterMediaCaptureNode.h
--- a/Engine/Plugins/Runtime/nDisplay/Source/DisplayClusterMedia/Private/Capture/DisplayClusterMediaCaptureNode.h
+++ b/Engine/Plugins/Runtime/nDisplay/Source/DisplayClusterMedia/Private/Capture/DisplayClusterMediaCaptureNode.h
@@ -28,13 +28,5 @@ protected:
 	virtual FIntPoint GetCaptureSize() const override;
 
 private:
-	void OnPostFrameRender_RenderThread(FRHICommandListImmediate& RHICmdList, const IDisplayClusterViewportManagerProxy* ViewportManagerProxy, FViewport* Viewport);
-
-	// Auxiliary method to create an intermediate texture
-	FTextureRHIRef CreateIntermediateTexture_RenderThread(EPixelFormat Format, ETextureCreateFlags Flags, const FIntPoint& Size);
-
-private:
-	// We need an intermediate texture with the same size as backbuffer has. This allows to avoid
-	// any texture size related problems caused by optimizations in the nD rendering pipeline.
-	FTextureRHIRef InterimTexture = nullptr;
+	void OnPostBackbufferUpdated_RenderThread(FRHICommandListImmediate& RHICmdList, FViewport* Viewport);
 };
