[nDisplay] Fixed editor crash when removing DCRA instance from scene

diff --git a/Engine/Plugins/Runtime/nDisplay/Source/DisplayCluster/Private/Render/Viewport/DisplayClusterViewport_OpenColorIO.cpp b/Engine/Plugins/Runtime/nDisplay/Source/DisplayCluster/Private/Render/Viewport/DisplayClusterViewport_OpenColorIO.cpp
--- a/Engine/Plugins/Runtime/nDisplay/Source/DisplayCluster/Private/Render/Viewport/DisplayClusterViewport_OpenColorIO.cpp
+++ b/Engine/Plugins/Runtime/nDisplay/Source/DisplayCluster/Private/Render/Viewport/DisplayClusterViewport_OpenColorIO.cpp
@@ -27,14 +27,6 @@ FDisplayClusterViewport_OpenColorIO::FDisplayClusterViewport_OpenColorIO(const F
 FDisplayClusterViewport_OpenColorIO::~FDisplayClusterViewport_OpenColorIO()
 { }
 
-void FDisplayClusterViewport_OpenColorIO::AddReferencedObjects(FReferenceCollector& Collector)
-{
-	if (DisplayConfiguration.ConfigurationSource)
-	{
-		Collector.AddReferencedObject(DisplayConfiguration.ConfigurationSource);
-	}
-}
-
 void FDisplayClusterViewport_OpenColorIO::SetupSceneView(FSceneViewFamily& InOutViewFamily, FSceneView& InOutView) const
 {
 	if (bShaderResourceValid)
diff --git a/Engine/Plugins/Runtime/nDisplay/Source/DisplayCluster/Private/Render/Viewport/DisplayClusterViewport_OpenColorIO.h b/Engine/Plugins/Runtime/nDisplay/Source/DisplayCluster/Private/Render/Viewport/DisplayClusterViewport_OpenColorIO.h
--- a/Engine/Plugins/Runtime/nDisplay/Source/DisplayCluster/Private/Render/Viewport/DisplayClusterViewport_OpenColorIO.h
+++ b/Engine/Plugins/Runtime/nDisplay/Source/DisplayCluster/Private/Render/Viewport/DisplayClusterViewport_OpenColorIO.h
@@ -12,28 +12,17 @@
 #include "RHIResources.h"
 #include "Templates/SharedPointer.h"
 
-#include "UObject/GCObject.h"
-
 /**
  * nDisplay OCIO implementation.
  * 
  */
 class FDisplayClusterViewport_OpenColorIO
 	: public TSharedFromThis<FDisplayClusterViewport_OpenColorIO, ESPMode::ThreadSafe>
-	, public FGCObject
 {
 public:
 	FDisplayClusterViewport_OpenColorIO(const FOpenColorIOColorConversionSettings& InDisplayConfiguration);
 	virtual ~FDisplayClusterViewport_OpenColorIO();
 
-	//~Begin FGCObject interface
-	virtual void AddReferencedObjects(FReferenceCollector& Collector) override;
-	virtual FString GetReferencerName() const override
-	{
-		return TEXT("FDisplayClusterViewport_OpenColorIO");
-	}
-	//~End FGCObject interface
-
 public:
 	/** Update render thread resources. */
 	void UpdateOpenColorIORenderPassResources();
