nDisplay: Fixed issue where attempting to change the ICVFX camera component's chromakey color using the color picker window would cause the details panel to immediately become empty, as well as causing the color picker in the ICVFX panel's color grading drawer to stop working. Issue was caused by a customization forcing a details panel refresh on an interactive change, which prevented the color picker from working properly.

diff --git a/Engine/Plugins/Runtime/nDisplay/Source/DisplayClusterConfigurator/Private/Views/Details/DisplayClusterEditorPropertyReferenceTypeCustomization.cpp b/Engine/Plugins/Runtime/nDisplay/Source/DisplayClusterConfigurator/Private/Views/Details/DisplayClusterEditorPropertyReferenceTypeCustomization.cpp
--- a/Engine/Plugins/Runtime/nDisplay/Source/DisplayClusterConfigurator/Private/Views/Details/DisplayClusterEditorPropertyReferenceTypeCustomization.cpp
+++ b/Engine/Plugins/Runtime/nDisplay/Source/DisplayClusterConfigurator/Private/Views/Details/DisplayClusterEditorPropertyReferenceTypeCustomization.cpp
@@ -141,8 +141,9 @@ void FDisplayClusterEditorPropertyReferenceTypeCustomization::CustomizeChildren(
 					ReferencedPropertyHandle->SetToolTipText(InPropertyHandle->GetToolTipText());
 				}
 
-				ReferencedPropertyHandle->SetOnPropertyValueChanged(FSimpleDelegate::CreateSP(this, &FDisplayClusterEditorPropertyReferenceTypeCustomization::OnReferencedPropertyValueChanged));
-				
+				ReferencedPropertyHandle->SetOnPropertyValueChangedWithData(TDelegate<void(const FPropertyChangedEvent&)>::CreateSP(
+					this, &FDisplayClusterEditorPropertyReferenceTypeCustomization::OnReferencedPropertyValueChanged));
+
 				const bool bHasShowOnlyInners = InPropertyHandle->HasMetaData(TEXT("ShowOnlyInnerProperties")) || InPropertyHandle->GetInstanceMetaData(TEXT("ShowOnlyInnerProperties"));
 				if (bHasShowOnlyInners)
 				{
@@ -174,10 +175,11 @@ void FDisplayClusterEditorPropertyReferenceTypeCustomization::CustomizeChildren(
 	}
 }
 
-void FDisplayClusterEditorPropertyReferenceTypeCustomization::OnReferencedPropertyValueChanged()
+void FDisplayClusterEditorPropertyReferenceTypeCustomization::OnReferencedPropertyValueChanged(const FPropertyChangedEvent& PropertyChangedEvent)
 {
-	// When a referenced property is changed, we have to trigger layout refresh
-	if (PropertyUtilities.IsValid())
+	// When a referenced property is changed, we have to trigger layout refresh, as long as the change is not an interactive change
+	const bool bIsInteractiveChange = (PropertyChangedEvent.ChangeType & EPropertyChangeType::Interactive) != 0;
+	if (PropertyUtilities.IsValid() && !bIsInteractiveChange)
 	{
 		PropertyUtilities.Pin()->ForceRefresh();
 	}
diff --git a/Engine/Plugins/Runtime/nDisplay/Source/DisplayClusterConfigurator/Private/Views/Details/DisplayClusterEditorPropertyReferenceTypeCustomization.h b/Engine/Plugins/Runtime/nDisplay/Source/DisplayClusterConfigurator/Private/Views/Details/DisplayClusterEditorPropertyReferenceTypeCustomization.h
--- a/Engine/Plugins/Runtime/nDisplay/Source/DisplayClusterConfigurator/Private/Views/Details/DisplayClusterEditorPropertyReferenceTypeCustomization.h
+++ b/Engine/Plugins/Runtime/nDisplay/Source/DisplayClusterConfigurator/Private/Views/Details/DisplayClusterEditorPropertyReferenceTypeCustomization.h
@@ -35,7 +35,7 @@ private:
 	TSharedPtr<IPropertyHandle> FindRootPropertyHandle(const FString& PropertyPath, TSharedRef<IPropertyHandle>& InPropertyHandle);
 	TSharedPtr<IPropertyHandle> GetChildPropertyHandle(const FString& PropertyName, TSharedPtr<IPropertyHandle>& PropertyHandle);
 
-	void OnReferencedPropertyValueChanged();
+	void OnReferencedPropertyValueChanged(const FPropertyChangedEvent& PropertyChangedEvent);
 
 	bool IsListType(const TSharedPtr<IPropertyHandle>& PropertyHandle);
 
