nDisplay: Fixed issue where the buttons on the left panel of the color grading drawer would overlap other UI elements if the splitter was dragged too far left. Buttons are now implemented with a toolbar builder, which allows buttons to be condensed to an overflow menu when the toolbar becomes too small.

diff --git a/Engine/Plugins/Runtime/nDisplay/Source/DisplayClusterColorGrading/Private/DisplayClusterColorGradingCommands.cpp b/Engine/Plugins/Runtime/nDisplay/Source/DisplayClusterColorGrading/Private/DisplayClusterColorGradingCommands.cpp
--- a/Engine/Plugins/Runtime/nDisplay/Source/DisplayClusterColorGrading/Private/DisplayClusterColorGradingCommands.cpp
+++ b/Engine/Plugins/Runtime/nDisplay/Source/DisplayClusterColorGrading/Private/DisplayClusterColorGradingCommands.cpp
@@ -11,6 +11,9 @@ void FDisplayClusterColorGradingCommands::RegisterCommands()
 	UI_COMMAND(ColorWheelSliderOrientationHorizontal, "Right", "Puts the color wheel sliders to the right of the color wheel", EUserInterfaceActionType::RadioButton, FInputChord());
 	UI_COMMAND(ColorWheelSliderOrientationVertical, "Below", "Puts the color wheel sliders below the color wheel", EUserInterfaceActionType::RadioButton, FInputChord());
 
+	UI_COMMAND(ColorGradingMode, "Color Grading Mode", "Sets the drawer to display the color grading properties of the selected items", EUserInterfaceActionType::ToggleButton, FInputChord());
+	UI_COMMAND(DetailViewMode, "Detail View Mode", "Sets the drawer to display the additional in-camera VFX properties of the selected items", EUserInterfaceActionType::ToggleButton, FInputChord());
+
 	UI_COMMAND(OpenColorGradingDrawer, "Open Color Grading Drawer", "Opens the color grading drawer from the status bar", EUserInterfaceActionType::Button, FInputChord(EModifierKey::Control, EKeys::SpaceBar));
 }
 
diff --git a/Engine/Plugins/Runtime/nDisplay/Source/DisplayClusterColorGrading/Private/DisplayClusterColorGradingCommands.h b/Engine/Plugins/Runtime/nDisplay/Source/DisplayClusterColorGrading/Private/DisplayClusterColorGradingCommands.h
--- a/Engine/Plugins/Runtime/nDisplay/Source/DisplayClusterColorGrading/Private/DisplayClusterColorGradingCommands.h
+++ b/Engine/Plugins/Runtime/nDisplay/Source/DisplayClusterColorGrading/Private/DisplayClusterColorGradingCommands.h
@@ -3,6 +3,7 @@
 #pragma once
 
 #include "CoreMinimal.h"
+#include "DisplayClusterColorGradingStyle.h"
 #include "Framework/Commands/Commands.h"
 #include "Styling/AppStyle.h"
 
@@ -13,7 +14,7 @@ class FDisplayClusterColorGradingCommands
 public:
 	FDisplayClusterColorGradingCommands()
 		: TCommands<FDisplayClusterColorGradingCommands>(TEXT("DisplayClusterColorGrading"),
-			NSLOCTEXT("Contexts", "DisplayClusterColorGrading", "Display Cluster Color Grading"), NAME_None, FAppStyle::GetAppStyleSetName())
+			NSLOCTEXT("Contexts", "DisplayClusterColorGrading", "Display Cluster Color Grading"), NAME_None, FDisplayClusterColorGradingStyle::Get().GetStyleSetName())
 	{ }
 
 	virtual void RegisterCommands() override;
@@ -24,5 +25,8 @@ public:
 	TSharedPtr<FUICommandInfo> ColorWheelSliderOrientationHorizontal;
 	TSharedPtr<FUICommandInfo> ColorWheelSliderOrientationVertical;
 
+	TSharedPtr<FUICommandInfo> ColorGradingMode;
+	TSharedPtr<FUICommandInfo> DetailViewMode;
+
 	TSharedPtr<FUICommandInfo> OpenColorGradingDrawer;
 };
\ No newline at end of file
diff --git a/Engine/Plugins/Runtime/nDisplay/Source/DisplayClusterColorGrading/Private/DisplayClusterColorGradingStyle.h b/Engine/Plugins/Runtime/nDisplay/Source/DisplayClusterColorGrading/Private/DisplayClusterColorGradingStyle.h
--- a/Engine/Plugins/Runtime/nDisplay/Source/DisplayClusterColorGrading/Private/DisplayClusterColorGradingStyle.h
+++ b/Engine/Plugins/Runtime/nDisplay/Source/DisplayClusterColorGrading/Private/DisplayClusterColorGradingStyle.h
@@ -3,9 +3,11 @@
 #pragma once
 
 #include "CoreMinimal.h"
+#include "Styling/AppStyle.h"
 #include "Styling/SlateStyleRegistry.h"
 #include "Styling/SlateStyleMacros.h"
 #include "Styling/SlateStyle.h"
+#include "Styling/ToolBarStyle.h"
 
 
 /** Styleset for the nDisplay color grading UI elements */
@@ -18,16 +20,30 @@ public:
 	{
 		const FVector2D Icon16x16(16.0f, 16.0f);
 
+		SetParentStyleName(FAppStyle::GetAppStyleSetName());
+
 		// Set miscellaneous icons
 		SetContentRoot(FPaths::EnginePluginsDir() / TEXT("Runtime/nDisplay/Content/Icons/"));
 		SetCoreContentRoot(FPaths::EngineContentDir() / TEXT("Editor/Slate"));
 
 		Set("ColorGradingDrawer.Icon", new IMAGE_BRUSH_SVG("Components/nDisplayCamera", Icon16x16));
-		Set("ColorGradingDrawer.ColorGradingMode", new IMAGE_BRUSH_SVG("OperatorPanel/Colors", Icon16x16));
-		Set("ColorGradingDrawer.DetailViewMode", new IMAGE_BRUSH_SVG("OperatorPanel/Sliders", Icon16x16));
+		Set("DisplayClusterColorGrading.ColorGradingMode", new IMAGE_BRUSH_SVG("OperatorPanel/Colors", Icon16x16));
+		Set("DisplayClusterColorGrading.DetailViewMode", new IMAGE_BRUSH_SVG("OperatorPanel/Sliders", Icon16x16));
 		Set("ColorGradingDrawer.Viewports", new CORE_IMAGE_BRUSH_SVG("Starship/Common/Viewports", Icon16x16));
 		Set("ColorGradingDrawer.Nodes", new IMAGE_BRUSH_SVG("Cluster/ClusterNode", Icon16x16));
 
+		// Ultra slim toolbar style with smaller icons and smaller padding for the color grading drawer toolbars
+		FToolBarStyle UltraSlimToolbarStyle = FAppStyle::Get().GetWidgetStyle<FToolBarStyle>("SlimToolBar");
+
+		UltraSlimToolbarStyle
+			.SetBackgroundPadding(FMargin(4.0f, 0.0f, 0.0f, 0.0f))
+			.SetIconSize(Icon16x16);
+
+		UltraSlimToolbarStyle.ToggleButton
+			.SetPadding(FMargin(16.0f, 2.0f));
+
+		Set("UltraSlimToolBar", UltraSlimToolbarStyle);
+
 		FSlateStyleRegistry::RegisterSlateStyle(*this);
 	}
 
diff --git a/Engine/Plugins/Runtime/nDisplay/Source/DisplayClusterColorGrading/Private/Drawer/SDisplayClusterColorGradingDrawer.cpp b/Engine/Plugins/Runtime/nDisplay/Source/DisplayClusterColorGrading/Private/Drawer/SDisplayClusterColorGradingDrawer.cpp
--- a/Engine/Plugins/Runtime/nDisplay/Source/DisplayClusterColorGrading/Private/Drawer/SDisplayClusterColorGradingDrawer.cpp
+++ b/Engine/Plugins/Runtime/nDisplay/Source/DisplayClusterColorGrading/Private/Drawer/SDisplayClusterColorGradingDrawer.cpp
@@ -19,6 +19,7 @@
 
 #include "ColorCorrectRegion.h"
 #include "ColorCorrectWindow.h"
+#include "DisplayClusterColorGradingCommands.h"
 #include "Engine/PostProcessVolume.h"
 #include "Framework/MultiBox/MultiBoxBuilder.h"
 #include "Modules/ModuleManager.h"
@@ -72,6 +73,9 @@ void SDisplayClusterColorGradingDrawer::Construct(const FArguments& InArgs, bool
 	ColorGradingDataModel = MakeShared<FDisplayClusterColorGradingDataModel>();
 	ColorGradingDataModel->OnDataModelGenerated().AddSP(this, &SDisplayClusterColorGradingDrawer::OnColorGradingDataModelGenerated);
 
+	CommandList = MakeShared<FUICommandList>();
+	BindCommands();
+
 	bIsInDrawer = bInIsInDrawer;
 	OperatorViewModel = IDisplayClusterOperator::Get().GetOperatorViewModel();
 	OperatorViewModel->OnActiveRootActorChanged().AddSP(this, &SDisplayClusterColorGradingDrawer::OnActiveRootActorChanged);
@@ -152,43 +156,9 @@ void SDisplayClusterColorGradingDrawer::Construct(const FArguments& InArgs, bool
 					.BorderImage(bIsInDrawer ? FStyleDefaults::GetNoBrush() : FAppStyle::Get().GetBrush("Brushes.Panel"))
 					[
 						SNew(SBox)
-						.HeightOverride(24.0f)
+						.HeightOverride(28.0f)
 						[
-							SNew(SHorizontalBox)
-
-							+ SHorizontalBox::Slot()
-							.AutoWidth()
-							.VAlign(VAlign_Center)
-							.Padding(2, 0, 2, 0)
-							[
-								SNew(SCheckBox)
-								.Style(FAppStyle::Get(), "ToggleButtonCheckbox")
-								.ToolTipText(LOCTEXT("ColorGradingDrawerModeToolTip", "Sets the drawer to display the color grading properties of the selected items"))
-								.OnCheckStateChanged(this, &SDisplayClusterColorGradingDrawer::OnDrawerModeSelected, EDisplayClusterColorGradingDrawerMode::ColorGrading)
-								.IsChecked(this, &SDisplayClusterColorGradingDrawer::IsDrawerModeSelected, EDisplayClusterColorGradingDrawerMode::ColorGrading)
-								[
-									SNew(SImage)
-									.ColorAndOpacity(FSlateColor::UseForeground())
-									.Image(FDisplayClusterColorGradingStyle::Get().GetBrush("ColorGradingDrawer.ColorGradingMode"))
-								]
-							]
-
-							+ SHorizontalBox::Slot()
-							.AutoWidth()
-							.VAlign(VAlign_Center)
-							.Padding(2, 0, 2, 0)
-							[
-								SNew(SCheckBox)
-								.Style(FAppStyle::Get(), "ToggleButtonCheckbox")
-								.ToolTipText(LOCTEXT("DetailsViewDrawerModeToolTip", "Sets the drawer to display the additional in-camera VFX properties of the selected items"))
-								.OnCheckStateChanged(this, &SDisplayClusterColorGradingDrawer::OnDrawerModeSelected, EDisplayClusterColorGradingDrawerMode::DetailsView)
-								.IsChecked(this, &SDisplayClusterColorGradingDrawer::IsDrawerModeSelected, EDisplayClusterColorGradingDrawerMode::DetailsView)
-								[
-									SNew(SImage)
-									.ColorAndOpacity(FSlateColor::UseForeground())
-									.Image(FDisplayClusterColorGradingStyle::Get().GetBrush("ColorGradingDrawer.DetailViewMode"))
-								]
-							]
+							CreateListPanelToolbar()
 						]
 					]
 				]
@@ -506,6 +476,23 @@ void SDisplayClusterColorGradingDrawer::SetDrawerStateToDefault()
 	}
 }
 
+void SDisplayClusterColorGradingDrawer::BindCommands()
+{
+	CommandList->MapAction(
+		FDisplayClusterColorGradingCommands::Get().ColorGradingMode,
+		FExecuteAction::CreateSP(this, &SDisplayClusterColorGradingDrawer::OnDrawerModeSelected, EDisplayClusterColorGradingDrawerMode::ColorGrading),
+		FCanExecuteAction(),
+		FIsActionChecked::CreateSP(this, &SDisplayClusterColorGradingDrawer::IsDrawerModeSelected, EDisplayClusterColorGradingDrawerMode::ColorGrading)
+	);
+
+	CommandList->MapAction(
+		FDisplayClusterColorGradingCommands::Get().DetailViewMode,
+		FExecuteAction::CreateSP(this, &SDisplayClusterColorGradingDrawer::OnDrawerModeSelected, EDisplayClusterColorGradingDrawerMode::DetailsView),
+		FCanExecuteAction(),
+		FIsActionChecked::CreateSP(this, &SDisplayClusterColorGradingDrawer::IsDrawerModeSelected, EDisplayClusterColorGradingDrawerMode::DetailsView)
+	);
+}
+
 TSharedRef<SWidget> SDisplayClusterColorGradingDrawer::CreateDockInLayoutButton()
 {
 	if (bIsInDrawer)
@@ -541,6 +528,17 @@ TSharedRef<SWidget> SDisplayClusterColorGradingDrawer::CreateDockInLayoutButton(
 	return SNullWidget::NullWidget;
 }
 
+TSharedRef<SWidget> SDisplayClusterColorGradingDrawer::CreateListPanelToolbar() const
+{
+	FSlimHorizontalToolBarBuilder Toolbar(CommandList, FMultiBoxCustomization::None);
+
+	Toolbar.AddToolBarButton(FDisplayClusterColorGradingCommands::Get().ColorGradingMode, NAME_None, FText::GetEmpty());
+	Toolbar.AddToolBarButton(FDisplayClusterColorGradingCommands::Get().DetailViewMode, NAME_None, FText::GetEmpty());
+	Toolbar.SetStyle(&FDisplayClusterColorGradingStyle::Get(), "UltraSlimToolBar");
+
+	return Toolbar.MakeWidget();
+}
+
 FText SDisplayClusterColorGradingDrawer::GetColorGradingListName(int32 ListIndex) const
 {
 	switch (ListIndex)
@@ -881,9 +879,9 @@ void SDisplayClusterColorGradingDrawer::FillColorGradingGroupToolBar()
 	}
 }
 
-ECheckBoxState SDisplayClusterColorGradingDrawer::IsDrawerModeSelected(EDisplayClusterColorGradingDrawerMode InDrawerMode) const
+bool SDisplayClusterColorGradingDrawer::IsDrawerModeSelected(EDisplayClusterColorGradingDrawerMode InDrawerMode) const
 {
-	return CurrentDrawerMode == InDrawerMode ? ECheckBoxState::Checked : ECheckBoxState::Unchecked;
+	return CurrentDrawerMode == InDrawerMode;
 }
 
 EVisibility SDisplayClusterColorGradingDrawer::GetDrawerModeVisibility(EDisplayClusterColorGradingDrawerMode InDrawerMode) const
@@ -891,12 +889,9 @@ EVisibility SDisplayClusterColorGradingDrawer::GetDrawerModeVisibility(EDisplayC
 	return CurrentDrawerMode == InDrawerMode ? EVisibility::Visible : EVisibility::Collapsed;
 }
 
-void SDisplayClusterColorGradingDrawer::OnDrawerModeSelected(ECheckBoxState State, EDisplayClusterColorGradingDrawerMode InDrawerMode)
+void SDisplayClusterColorGradingDrawer::OnDrawerModeSelected(EDisplayClusterColorGradingDrawerMode InDrawerMode)
 {
-	if (State == ECheckBoxState::Checked)
-	{
-		CurrentDrawerMode = InDrawerMode;
-	}
+	CurrentDrawerMode = InDrawerMode;
 }
 
 EVisibility SDisplayClusterColorGradingDrawer::GetColorGradingGroupToolBarVisibility() const
diff --git a/Engine/Plugins/Runtime/nDisplay/Source/DisplayClusterColorGrading/Private/Drawer/SDisplayClusterColorGradingDrawer.h b/Engine/Plugins/Runtime/nDisplay/Source/DisplayClusterColorGrading/Private/Drawer/SDisplayClusterColorGradingDrawer.h
--- a/Engine/Plugins/Runtime/nDisplay/Source/DisplayClusterColorGrading/Private/Drawer/SDisplayClusterColorGradingDrawer.h
+++ b/Engine/Plugins/Runtime/nDisplay/Source/DisplayClusterColorGrading/Private/Drawer/SDisplayClusterColorGradingDrawer.h
@@ -56,9 +56,15 @@ public:
 	void SetDrawerStateToDefault();
 
 private:
+	/** Binds commands to the drawer's command list */
+	void BindCommands();
+
 	/** Creates the button used to dock the drawer in the operator panel */
 	TSharedRef<SWidget> CreateDockInLayoutButton();
 
+	/** Creates the toolbar used in the object list panel of the color grading drawer */
+	TSharedRef<SWidget> CreateListPanelToolbar() const;
+
 	/** Gets the display name of the color grading list at the specified index */
 	FText GetColorGradingListName(int32 ListIndex) const;
 
@@ -96,13 +102,13 @@ private:
 	void FillColorGradingGroupToolBar();
 
 	/** Gets whether the specified drawer mode is currently selected */
-	ECheckBoxState IsDrawerModeSelected(EDisplayClusterColorGradingDrawerMode InDrawerMode) const;
+	bool IsDrawerModeSelected(EDisplayClusterColorGradingDrawerMode InDrawerMode) const;
 
 	/** Gets the visibility state for the specified drawer mode */
 	EVisibility GetDrawerModeVisibility(EDisplayClusterColorGradingDrawerMode InDrawerMode) const;
 
 	/** Raised when the specified drawer mode checked state is changed */
-	void OnDrawerModeSelected(ECheckBoxState State, EDisplayClusterColorGradingDrawerMode InDrawerMode);
+	void OnDrawerModeSelected(EDisplayClusterColorGradingDrawerMode InDrawerMode);
 
 	/** Gets the visibility state of the color grading group toolbar */
 	EVisibility GetColorGradingGroupToolBarVisibility() const;
@@ -173,6 +179,9 @@ private:
 	TSharedPtr<SDisplayClusterColorGradingDetailsPanel> DetailsPanel;
 	TArray<TSharedPtr<SInlineEditableTextBlock>> ColorGradingGroupTextBlocks;
 
+	/** Commands used by the drawer */
+	TSharedPtr<FUICommandList> CommandList;
+
 	/** The color grading data model for the currently selected objects */
 	TSharedPtr<FDisplayClusterColorGradingDataModel> ColorGradingDataModel;
 
