nDisplay: Expose light card properties to sequencer.

diff --git a/Engine/Plugins/Runtime/nDisplay/Source/DisplayCluster/Public/DisplayClusterLightCardActor.h b/Engine/Plugins/Runtime/nDisplay/Source/DisplayCluster/Public/DisplayClusterLightCardActor.h
--- a/Engine/Plugins/Runtime/nDisplay/Source/DisplayCluster/Public/DisplayClusterLightCardActor.h
+++ b/Engine/Plugins/Runtime/nDisplay/Source/DisplayCluster/Public/DisplayClusterLightCardActor.h
@@ -193,79 +193,79 @@ public:
 public:
 
 	/** Radius of light card polar coordinates. Does not include the effect of RadialOffset */
-	UPROPERTY(EditAnywhere, BlueprintReadWrite, Category = "Orientation", meta = (EditCondition = "!bIsUVLightCard", HideEditConditionToggle, EditConditionHides))
+	UPROPERTY(EditAnywhere, BlueprintReadWrite, Interp, Category = "Orientation", meta = (EditCondition = "!bIsUVLightCard", HideEditConditionToggle, EditConditionHides))
 	double DistanceFromCenter;
 
 	/** Related to the Azimuth of light card polar coordinates */
-	UPROPERTY(EditAnywhere, BlueprintReadWrite, Category = "Orientation", meta = (UIMin = 0, ClampMin = 0, UIMax = 360, ClampMax = 360, EditCondition = "!bIsUVLightCard", HideEditConditionToggle, EditConditionHides))
+	UPROPERTY(EditAnywhere, BlueprintReadWrite, Interp, Category = "Orientation", meta = (UIMin = 0, ClampMin = 0, UIMax = 360, ClampMax = 360, EditCondition = "!bIsUVLightCard", HideEditConditionToggle, EditConditionHides))
 	double Longitude;
 
 	/** Related to the Elevation of light card polar coordinates */
-	UPROPERTY(EditAnywhere, BlueprintReadWrite, Category = "Orientation", meta = (UIMin = -90, ClampMin = -90, UIMax = 90, ClampMax = 90, EditCondition = "!bIsUVLightCard", HideEditConditionToggle, EditConditionHides))
+	UPROPERTY(EditAnywhere, BlueprintReadWrite, Interp, Category = "Orientation", meta = (UIMin = -90, ClampMin = -90, UIMax = 90, ClampMax = 90, EditCondition = "!bIsUVLightCard", HideEditConditionToggle, EditConditionHides))
 	double Latitude;
 
 	/** The UV coordinates of the light card, if it is in UV space */
-	UPROPERTY(EditAnywhere, BlueprintReadWrite, Category = "Orientation", meta = (DisplayName = "UV Coodinates", EditCondition = "bIsUVLightCard", HideEditConditionToggle, EditConditionHides))
+	UPROPERTY(EditAnywhere, BlueprintReadWrite, Interp, Category = "Orientation", meta = (DisplayName = "UV Coodinates", EditCondition = "bIsUVLightCard", HideEditConditionToggle, EditConditionHides))
 	FVector2D UVCoordinates = FVector2D(0.5, 0.5);
 
 	/** Roll rotation of light card around its plane axis */
-	UPROPERTY(EditAnywhere, BlueprintReadWrite, Category = "Orientation", meta = (UIMin = -360, ClampMin = -360, UIMax = 360, ClampMax = 360))
+	UPROPERTY(EditAnywhere, BlueprintReadWrite, Interp, Category = "Orientation", meta = (UIMin = -360, ClampMin = -360, UIMax = 360, ClampMax = 360))
 	double Spin;
 
-	UPROPERTY(EditAnywhere, BlueprintReadWrite, Category = "Orientation", meta = (UIMin = -360, ClampMin = -360, UIMax = 360, ClampMax = 360, EditCondition = "!bIsUVLightCard", HideEditConditionToggle, EditConditionHides))
+	UPROPERTY(EditAnywhere, BlueprintReadWrite, Interp, Category = "Orientation", meta = (UIMin = -360, ClampMin = -360, UIMax = 360, ClampMax = 360, EditCondition = "!bIsUVLightCard", HideEditConditionToggle, EditConditionHides))
 	double Pitch;
 
-	UPROPERTY(EditAnywhere, BlueprintReadWrite, Category = "Orientation", meta = (UIMin = -360, ClampMin = -360, UIMax = 360, ClampMax = 360, EditCondition = "!bIsUVLightCard", HideEditConditionToggle, EditConditionHides))
+	UPROPERTY(EditAnywhere, BlueprintReadWrite, Interp, Category = "Orientation", meta = (UIMin = -360, ClampMin = -360, UIMax = 360, ClampMax = 360, EditCondition = "!bIsUVLightCard", HideEditConditionToggle, EditConditionHides))
 	double Yaw;
 
-	UPROPERTY(EditAnywhere, BlueprintReadWrite, Category = "Orientation")
+	UPROPERTY(EditAnywhere, BlueprintReadWrite, Interp, Category = "Orientation")
 	FVector2D Scale;
 
 	/** Used by the flush constraint to offset the location of the light card form the wall */
-	UPROPERTY(EditAnywhere, BlueprintReadWrite, Category = "Orientation", meta = (EditCondition = "!bIsUVLightCard", HideEditConditionToggle, EditConditionHides))
+	UPROPERTY(EditAnywhere, BlueprintReadWrite, Interp, Category = "Orientation", meta = (EditCondition = "!bIsUVLightCard", HideEditConditionToggle, EditConditionHides))
 	double RadialOffset;
 
 	/** Indicates whether the light card is always made to be flush to a stage wall or not */
-	UPROPERTY(EditAnywhere, BlueprintReadWrite, Category = "Orientation", meta = (EditCondition = "!bIsUVLightCard", HideEditConditionToggle, EditConditionHides))
+	UPROPERTY(EditAnywhere, BlueprintReadWrite, Interp, Category = "Orientation", meta = (EditCondition = "!bIsUVLightCard", HideEditConditionToggle, EditConditionHides))
 	bool bAlwaysFlushToWall;
 
-	UPROPERTY(EditAnywhere, BlueprintReadWrite, Category = "Appearance")
+	UPROPERTY(EditAnywhere, BlueprintReadWrite, Interp, Category = "Appearance")
 	EDisplayClusterLightCardMask Mask;
 
-	UPROPERTY(EditAnywhere, BlueprintReadWrite, Category = "Appearance")
+	UPROPERTY(EditAnywhere, BlueprintReadWrite, Interp, Category = "Appearance")
 	TObjectPtr<UTexture> Texture;
 
 	/** Light card color, before any modifier is applied */
-	UPROPERTY(EditAnywhere, BlueprintReadWrite, Category = "Appearance")
+	UPROPERTY(EditAnywhere, BlueprintReadWrite, Interp, Category = "Appearance")
 	FLinearColor Color;
 
-	UPROPERTY(EditAnywhere, BlueprintReadWrite, Category = "Appearance", meta = (UIMin = 0, ClampMin = 0, UIMax = 10000, ClampMax = 10000))
+	UPROPERTY(EditAnywhere, BlueprintReadWrite, Interp, Category = "Appearance", meta = (UIMin = 0, ClampMin = 0, UIMax = 10000, ClampMax = 10000))
 	float Temperature;
 
-	UPROPERTY(EditAnywhere, BlueprintReadWrite, Category = "Appearance", meta = (UIMin = -1, ClampMin = -1, UIMax = 1, ClampMax = 1))
+	UPROPERTY(EditAnywhere, BlueprintReadWrite, Interp, Category = "Appearance", meta = (UIMin = -1, ClampMin = -1, UIMax = 1, ClampMax = 1))
 	float Tint;
 
 	/** 2^Exposure color value multiplier */
-	UPROPERTY(EditAnywhere, BlueprintReadWrite, Category = "Appearance", meta = (UIMin = -100, ClampMin = -100, UIMax = 100, ClampMax = 100))
+	UPROPERTY(EditAnywhere, BlueprintReadWrite, Interp, Category = "Appearance", meta = (UIMin = -100, ClampMin = -100, UIMax = 100, ClampMax = 100))
 	float Exposure;
 
 	/** Linear color value multiplier */
-	UPROPERTY(EditAnywhere, BlueprintReadWrite, Category = "Appearance", meta = (UIMin = 0, ClampMin = 0))
+	UPROPERTY(EditAnywhere, BlueprintReadWrite, Interp, Category = "Appearance", meta = (UIMin = 0, ClampMin = 0))
 	float Gain;
 
 	/** Linear alpha multiplier */
-	UPROPERTY(EditAnywhere, BlueprintReadWrite, Category = "Appearance", meta = (UIMin = 0, ClampMin = 0, UIMax = 1, ClampMax = 1))
+	UPROPERTY(EditAnywhere, BlueprintReadWrite, Interp, Category = "Appearance", meta = (UIMin = 0, ClampMin = 0, UIMax = 1, ClampMax = 1))
 	float Opacity;
 
 	/** Feathers in the alpha from the edges */
-	UPROPERTY(EditAnywhere, BlueprintReadWrite, Category = "Appearance", meta = (UIMin = 0, ClampMin = 0))
+	UPROPERTY(EditAnywhere, BlueprintReadWrite, Interp, Category = "Appearance", meta = (UIMin = 0, ClampMin = 0))
 	float Feathering;
 
 	/** Settings related to an alpha gradient effect */
-	UPROPERTY(EditAnywhere, BlueprintReadWrite, Category = "Appearance")
+	UPROPERTY(EditAnywhere, BlueprintReadWrite, Interp, Category = "Appearance")
 	FLightCardAlphaGradientSettings AlphaGradient;
 
-	/** A flag that controls wether the light card's location and rotation are locked to its "owning" root actor */
+	/** A flag that controls whether the light card's location and rotation are locked to its "owning" root actor */
 	UPROPERTY()
 	bool bLockToOwningRootActor = true;
 
@@ -274,7 +274,7 @@ public:
 	bool bIsUVLightCard = false;
 
 	/** Polygon points when using this type of mask */
-	UPROPERTY(EditAnywhere, BlueprintReadWrite, Category = "Appearance")
+	UPROPERTY(EditAnywhere, BlueprintReadWrite, Interp, Category = "Appearance")
 	TArray<FVector2D> Polygon;
 
 	/** Used to flag this light card as a proxy of a "real" light card. Used by the LightCard Editor */
