nDisplay: Fixed a crash when the VE is still in use after the ViewportManager has already been destroyed.

diff --git a/Engine/Plugins/Runtime/nDisplay/Source/DisplayCluster/Private/Render/Viewport/DisplayClusterViewportManager.cpp b/Engine/Plugins/Runtime/nDisplay/Source/DisplayCluster/Private/Render/Viewport/DisplayClusterViewportManager.cpp
--- a/Engine/Plugins/Runtime/nDisplay/Source/DisplayCluster/Private/Render/Viewport/DisplayClusterViewportManager.cpp
+++ b/Engine/Plugins/Runtime/nDisplay/Source/DisplayCluster/Private/Render/Viewport/DisplayClusterViewportManager.cpp
@@ -93,7 +93,13 @@ FDisplayClusterViewportManager::~FDisplayClusterViewportManager()
 		LightCardManager.Reset();
 	}
 
-	ViewportManagerViewExtension.Reset();
+	if (ViewportManagerViewExtension.IsValid())
+	{
+		// Force release of VE data since this TSharedPtr<> may also be held by other resources at this time.
+		// So Reset() doesn't actually release it right away in some situations.
+		ViewportManagerViewExtension->Release();
+		ViewportManagerViewExtension.Reset();
+	}
 
 	Configuration.Reset();
 
diff --git a/Engine/Plugins/Runtime/nDisplay/Source/DisplayCluster/Private/Render/Viewport/DisplayClusterViewportManagerProxy.cpp b/Engine/Plugins/Runtime/nDisplay/Source/DisplayCluster/Private/Render/Viewport/DisplayClusterViewportManagerProxy.cpp
--- a/Engine/Plugins/Runtime/nDisplay/Source/DisplayCluster/Private/Render/Viewport/DisplayClusterViewportManagerProxy.cpp
+++ b/Engine/Plugins/Runtime/nDisplay/Source/DisplayCluster/Private/Render/Viewport/DisplayClusterViewportManagerProxy.cpp
@@ -16,6 +16,8 @@
 #include "Render/Viewport/LightCard/DisplayClusterViewportLightCardManagerProxy.h"
 
 #include "Render/Viewport/DisplayClusterViewportProxy.h"
+#include "Render/Viewport/DisplayClusterViewportManagerViewExtension.h"
+
 #include "Render/Viewport/RenderTarget/DisplayClusterRenderTargetManager.h"
 #include "Render/Viewport/RenderTarget/DisplayClusterRenderTargetResource.h"
 #include "Render/Viewport/Postprocess/DisplayClusterViewportPostProcessManager.h"
@@ -98,7 +100,13 @@ void FDisplayClusterViewportManagerProxy::Release_RenderThread()
 
 	LightCardManagerProxy.Reset();
 
-	ViewportManagerViewExtension.Reset();
+	if (ViewportManagerViewExtension.IsValid())
+	{
+		// Force release of VE data since this TSharedPtr<> may also be held by other resources at this time.
+		// So Reset() doesn't actually release it right away in some situations.
+		ViewportManagerViewExtension->Release_RenderThread();
+		ViewportManagerViewExtension.Reset();
+	}
 }
 
 void FDisplayClusterViewportManagerProxy::Initialize(FDisplayClusterViewportManager& InViewportManager)
diff --git a/Engine/Plugins/Runtime/nDisplay/Source/DisplayCluster/Private/Render/Viewport/DisplayClusterViewportManagerViewExtension.cpp b/Engine/Plugins/Runtime/nDisplay/Source/DisplayCluster/Private/Render/Viewport/DisplayClusterViewportManagerViewExtension.cpp
--- a/Engine/Plugins/Runtime/nDisplay/Source/DisplayCluster/Private/Render/Viewport/DisplayClusterViewportManagerViewExtension.cpp
+++ b/Engine/Plugins/Runtime/nDisplay/Source/DisplayCluster/Private/Render/Viewport/DisplayClusterViewportManagerViewExtension.cpp
@@ -32,9 +32,33 @@ FDisplayClusterViewportManagerViewExtension::FDisplayClusterViewportManagerViewE
 
 FDisplayClusterViewportManagerViewExtension::~FDisplayClusterViewportManagerViewExtension()
 {
+	ViewportManagerProxy.Reset();
 	UnregisterCallbacks();
 }
 
+bool FDisplayClusterViewportManagerViewExtension::IsActive() const
+{
+	return ViewportManagerProxy.IsValid();
+}
+
+void FDisplayClusterViewportManagerViewExtension::Release()
+{
+	check(IsInGameThread());
+
+	ViewportManagerProxy.Reset();
+}
+
+void FDisplayClusterViewportManagerViewExtension::Release_RenderThread()
+{
+	check(IsInRenderingThread());
+
+	ViewportManagerProxy.Reset();
+	ViewportProxies.Empty();
+
+	UnregisterCallbacks();
+}
+
+
 void FDisplayClusterViewportManagerViewExtension::PreRenderView_RenderThread(FRDGBuilder& GraphBuilder, FSceneView& InView)
 {
 	if (InView.bIsViewInfo)
@@ -46,12 +70,17 @@ void FDisplayClusterViewportManagerViewExtension::PreRenderView_RenderThread(FRD
 
 void FDisplayClusterViewportManagerViewExtension::SubscribeToPostProcessingPass(EPostProcessingPass PassId, FAfterPassCallbackDelegateArray& InOutPassCallbacks, bool bIsPassEnabled)
 {
+	if (!IsActive())
+	{
+		return;
+	}
+
 	switch (PassId)
 	{
 	case EPostProcessingPass::SSRInput:
 		for (const FViewportProxy& ViewportIt : ViewportProxies)
 		{
-			if (ViewportIt.ViewportProxy.IsValid() && ViewportIt.ViewportProxy->ShouldUsePostProcessPassAfterSSRInput())
+			if (ViewportIt.ViewportProxy.IsValid() && ViewportIt.ViewportProxy.Pin()->ShouldUsePostProcessPassAfterSSRInput())
 			{
 				InOutPassCallbacks.Add(FAfterPassCallbackDelegate::CreateRaw(this, &FDisplayClusterViewportManagerViewExtension::PostProcessPassAfterSSRInput_RenderThread));
 				break;
@@ -62,7 +91,7 @@ void FDisplayClusterViewportManagerViewExtension::SubscribeToPostProcessingPass(
 	case EPostProcessingPass::FXAA:
 		for (const FViewportProxy& ViewportIt : ViewportProxies)
 		{
-			if (ViewportIt.ViewportProxy.IsValid() && ViewportIt.ViewportProxy->ShouldUsePostProcessPassAfterFXAA())
+			if (ViewportIt.ViewportProxy.IsValid() && ViewportIt.ViewportProxy.Pin()->ShouldUsePostProcessPassAfterFXAA())
 			{
 				InOutPassCallbacks.Add(FAfterPassCallbackDelegate::CreateRaw(this, &FDisplayClusterViewportManagerViewExtension::PostProcessPassAfterFXAA_RenderThread));
 				break;
@@ -77,8 +106,13 @@ void FDisplayClusterViewportManagerViewExtension::SubscribeToPostProcessingPass(
 
 FScreenPassTexture FDisplayClusterViewportManagerViewExtension::PostProcessPassAfterFXAA_RenderThread(FRDGBuilder& GraphBuilder, const FSceneView& View, const FPostProcessMaterialInputs& Inputs)
 {
+	if (!IsActive())
+	{
+		return ReturnUntouchedSceneColorForPostProcessing(Inputs);
+	}
+
 	uint32 ContextNum = 0;
-	if (FDisplayClusterViewportProxy* ViewportProxyPtr = ViewportManagerProxy->ImplFindViewport_RenderThread(View.StereoViewIndex, &ContextNum))
+	if (FDisplayClusterViewportProxy* ViewportProxyPtr = ViewportManagerProxy.Pin()->ImplFindViewport_RenderThread(View.StereoViewIndex, &ContextNum))
 	{
 		if (ViewportProxyPtr->ShouldUsePostProcessPassAfterFXAA())
 		{
@@ -91,8 +125,13 @@ FScreenPassTexture FDisplayClusterViewportManagerViewExtension::PostProcessPassA
 
 FScreenPassTexture FDisplayClusterViewportManagerViewExtension::PostProcessPassAfterSSRInput_RenderThread(FRDGBuilder& GraphBuilder, const FSceneView& View, const FPostProcessMaterialInputs& Inputs)
 {
+	if (!IsActive())
+	{
+		return ReturnUntouchedSceneColorForPostProcessing(Inputs);
+	}
+
 	uint32 ContextNum = 0;
-	if (FDisplayClusterViewportProxy* ViewportProxyPtr = ViewportManagerProxy->ImplFindViewport_RenderThread(View.StereoViewIndex, &ContextNum))
+	if (FDisplayClusterViewportProxy* ViewportProxyPtr = ViewportManagerProxy.Pin()->ImplFindViewport_RenderThread(View.StereoViewIndex, &ContextNum))
 	{
 		if (ViewportProxyPtr->ShouldUsePostProcessPassAfterSSRInput())
 		{
@@ -122,6 +161,11 @@ FScreenPassTexture FDisplayClusterViewportManagerViewExtension::ReturnUntouchedS
 
 bool FDisplayClusterViewportManagerViewExtension::IsActiveThisFrame_Internal(const FSceneViewExtensionContext& Context) const
 {
+	if (!IsActive())
+	{
+		return false;
+	}
+
 	static const FDisplayClusterSceneViewExtensionContext DCViewExtensionContext;
 	if (Context.IsA(MoveTempIfPossible(DCViewExtensionContext)))
 	{
@@ -138,6 +182,11 @@ bool FDisplayClusterViewportManagerViewExtension::IsActiveThisFrame_Internal(con
 
 void FDisplayClusterViewportManagerViewExtension::PreRenderViewFamily_RenderThread(FRDGBuilder& GraphBuilder, FSceneViewFamily& InViewFamily)
 {
+	if (!IsActive())
+	{
+		return;
+	}
+
 	ViewportProxies.Empty();
 
 	// Get all viewport proxies that are uses in this view family
@@ -147,7 +196,7 @@ void FDisplayClusterViewportManagerViewExtension::PreRenderViewFamily_RenderThre
 		{
 			FViewportProxy NewViewportProxy;
 
-			if (IDisplayClusterViewportProxy* ViewportProxyPtr = ViewportManagerProxy->FindViewport_RenderThread(SceneView->StereoViewIndex, &NewViewportProxy.ViewportProxyContext.ContextNum))
+			if (IDisplayClusterViewportProxy* ViewportProxyPtr = ViewportManagerProxy.Pin()->FindViewport_RenderThread(SceneView->StereoViewIndex, &NewViewportProxy.ViewportProxyContext.ContextNum))
 			{
 				NewViewportProxy.ViewportProxyContext.ViewFamilyProfileDescription = InViewFamily.ProfileDescription;
 
@@ -162,11 +211,19 @@ void FDisplayClusterViewportManagerViewExtension::PreRenderViewFamily_RenderThre
 
 void FDisplayClusterViewportManagerViewExtension::PostRenderViewFamily_RenderThread(FRDGBuilder& GraphBuilder, FSceneViewFamily& InViewFamily)
 {
+	if (!IsActive())
+	{
+		return;
+	}
+
 	for (FViewportProxy& ViewportIt : ViewportProxies)
 	{
 		if (const FSceneView* SceneView = InViewFamily.Views[ViewportIt.ViewIndex])
 		{
-			ViewportIt.ViewportProxy->OnPostRenderViewFamily_RenderThread(GraphBuilder, InViewFamily, *SceneView, ViewportIt.ViewportProxyContext);
+			if (ViewportIt.ViewportProxy.IsValid())
+			{
+				ViewportIt.ViewportProxy.Pin()->OnPostRenderViewFamily_RenderThread(GraphBuilder, InViewFamily, *SceneView, ViewportIt.ViewportProxyContext);
+			}
 		}
 	}
 
@@ -176,11 +233,17 @@ void FDisplayClusterViewportManagerViewExtension::PostRenderViewFamily_RenderThr
 //////////////////////////////////////////////////////////////////////////////////////////////
 void FDisplayClusterViewportManagerViewExtension::OnResolvedSceneColor_RenderThread(FRDGBuilder& GraphBuilder, const FSceneTextures& SceneTextures)
 {
-	check(IsInRenderingThread());
+	if (!IsActive())
+	{
+		return;
+	}
 
 	for (FViewportProxy& ViewportIt : ViewportProxies)
 	{
-		ViewportIt.ViewportProxy->OnResolvedSceneColor_RenderThread(GraphBuilder, SceneTextures, ViewportIt.ViewportProxyContext);
+		if (ViewportIt.ViewportProxy.IsValid())
+		{
+			ViewportIt.ViewportProxy.Pin()->OnResolvedSceneColor_RenderThread(GraphBuilder, SceneTextures, ViewportIt.ViewportProxyContext);
+		}
 	}
 }
 
diff --git a/Engine/Plugins/Runtime/nDisplay/Source/DisplayCluster/Private/Render/Viewport/DisplayClusterViewportManagerViewExtension.h b/Engine/Plugins/Runtime/nDisplay/Source/DisplayCluster/Private/Render/Viewport/DisplayClusterViewportManagerViewExtension.h
--- a/Engine/Plugins/Runtime/nDisplay/Source/DisplayCluster/Private/Render/Viewport/DisplayClusterViewportManagerViewExtension.h
+++ b/Engine/Plugins/Runtime/nDisplay/Source/DisplayCluster/Private/Render/Viewport/DisplayClusterViewportManagerViewExtension.h
@@ -52,6 +52,13 @@ public:
 	 * @return - Screen pass texture
 	 */
 	static FScreenPassTexture ReturnUntouchedSceneColorForPostProcessing(const FPostProcessMaterialInputs& InOutInputs);
+
+	/** Release from game thread.*/
+	void Release();
+
+	/** Release from render thread.*/
+	void Release_RenderThread();
+
 private:
 	/** nDisplay VE Callback [subscribed to Renderer:ResolvedSceneColorCallbacks].
 	 *
@@ -88,17 +95,20 @@ private:
 	/** UnRegister callbacks to external modules (Renderer:ResolvedSceneColor). */
 	void UnregisterCallbacks();
 
+	/** Return true if the VE is safe to use. */
+	bool IsActive() const;
+
 private:
 	FDelegateHandle ResolvedSceneColorCallbackHandle;
 
 private:
 	const FDisplayClusterViewportManager* ViewportManager;
-	const TSharedPtr<FDisplayClusterViewportManagerProxy, ESPMode::ThreadSafe> ViewportManagerProxy;
+	TWeakPtr<FDisplayClusterViewportManagerProxy, ESPMode::ThreadSafe> ViewportManagerProxy;
 
 	struct FViewportProxy
 	{
 		// The ViewportProxy must exist until this object is removed.
-		TSharedPtr<FDisplayClusterViewportProxy, ESPMode::ThreadSafe> ViewportProxy;
+		TWeakPtr<FDisplayClusterViewportProxy, ESPMode::ThreadSafe> ViewportProxy;
 
 		// The index of view in viewfamily collection
 		int32 ViewIndex = 0;
