Fixed missing dependency on DisplayClusterConfiguration in the DisplayClusterEditor module.

diff --git a/Engine/Plugins/Runtime/nDisplay/Source/DisplayClusterEditor/DisplayClusterEditor.Build.cs b/Engine/Plugins/Runtime/nDisplay/Source/DisplayClusterEditor/DisplayClusterEditor.Build.cs
--- a/Engine/Plugins/Runtime/nDisplay/Source/DisplayClusterEditor/DisplayClusterEditor.Build.cs
+++ b/Engine/Plugins/Runtime/nDisplay/Source/DisplayClusterEditor/DisplayClusterEditor.Build.cs
@@ -10,6 +10,7 @@ public class DisplayClusterEditor : ModuleRules
 		PrivateDependencyModuleNames.AddRange(
 			new string[] {
 				"DisplayCluster",
+				"DisplayClusterConfiguration",
 
 				"Core",
 				"CoreUObject",
