nDisplay: Improve orientation of frustum to match its widest fov direction with the render target's widest dimension.

This should avoid loss of detail in certain cases that the FoV wider dimension ended up on the narrower dimension of the render target.
It also implements hysteresis with a +/-5% margin to avoid a situation of jumping back and forth between rotated and unrotated frustums.

diff --git a/Engine/Plugins/Runtime/nDisplay/Source/DisplayClusterShaders/Private/WarpBlend/DisplayClusterWarpBlend_GeometryContext.h b/Engine/Plugins/Runtime/nDisplay/Source/DisplayClusterShaders/Private/WarpBlend/DisplayClusterWarpBlend_GeometryContext.h
--- a/Engine/Plugins/Runtime/nDisplay/Source/DisplayClusterShaders/Private/WarpBlend/DisplayClusterWarpBlend_GeometryContext.h
+++ b/Engine/Plugins/Runtime/nDisplay/Source/DisplayClusterShaders/Private/WarpBlend/DisplayClusterWarpBlend_GeometryContext.h
@@ -39,5 +39,8 @@ public:
 	// Force update caches
 	bool bGeometryChanged = false;
 
+	// Indicates if the frustum was rotated from base algorithm to better fit context size
+	// Used to implement hysteresis and avoid potential jumps back and forth.
+	bool bRotateFrustumToFitContextSize = false;
 };
 
diff --git a/Engine/Plugins/Runtime/nDisplay/Source/DisplayClusterShaders/Private/WarpBlend/Math/DisplayClusterWarpBlendMath_Frustum.cpp b/Engine/Plugins/Runtime/nDisplay/Source/DisplayClusterShaders/Private/WarpBlend/Math/DisplayClusterWarpBlendMath_Frustum.cpp
--- a/Engine/Plugins/Runtime/nDisplay/Source/DisplayClusterShaders/Private/WarpBlend/Math/DisplayClusterWarpBlendMath_Frustum.cpp
+++ b/Engine/Plugins/Runtime/nDisplay/Source/DisplayClusterShaders/Private/WarpBlend/Math/DisplayClusterWarpBlendMath_Frustum.cpp
@@ -10,6 +10,18 @@
 #include "StaticMeshResources.h"
 #include "ProceduralMeshComponent.h"
 
+float GDisplayClusterFrustumOrientationFitHysteresis = 0.05;
+static FAutoConsoleVariableRef CVarDisplayClusterFrustumOrientationFitHysteresis(
+	TEXT("nDisplay.render.FrustumOrientationFitHysteresis"),
+	GDisplayClusterFrustumOrientationFitHysteresis,
+	TEXT(
+		"Hysteresis margin when fitting frustum orientation with render target aspect ratio.\n"
+		"Values closer to 0 will allow immediate fit, while values close to 1 will require a \n"
+		"higher degree of mismatch before reorienting the frustum for a better fit."
+	),
+	ECVF_Default
+);
+
 FMatrix GetProjectionMatrixAssymetric(IDisplayClusterViewport* InViewport, const uint32 InContextNum, const FDisplayClusterWarpEye& InEye, const FDisplayClusterWarpContext& InContext)
 {
 	const float Left   = InContext.ProjectionAngles.Left;
@@ -22,6 +34,90 @@ FMatrix GetProjectionMatrixAssymetric(IDisplayClusterViewport* InViewport, const
 	return InViewport->GetContexts()[InContextNum].ProjectionMatrix;
 }
 
+bool FDisplayClusterWarpBlendMath_Frustum::CalcFrustum(IDisplayClusterViewport* InViewport, const uint32 InContextNum, FDisplayClusterWarpContext& OutFrustum)
+{
+	TRACE_CPUPROFILER_EVENT_SCOPE(DCWarpBlendMath_Frustum::CalcFrustum);
+
+	if (!GeometryContext.GeometryProxy.bIsGeometryValid)
+	{
+		// Bad source geometry
+		return false;
+	}
+
+	// Get view base:
+	ImplCalcView();
+
+	check(InViewport);
+
+	if (bFindBestProjectionType == false)
+	{
+		ImplCalcViewProjection();
+		ImplCalcFrustum();
+	}
+	else
+	{
+		// Frustum projection with auto-fix (back-side view planes)
+		// Projection type changed runtime. Validate frustum points, all must be over view plane:
+		EDisplayClusterWarpBlendProjectionType BaseProjectionType = ProjectionType;
+		EDisplayClusterWarpBlendFrustumType    BaseFrustumType = FrustumType;
+
+		// Optimize for PerfectCPU:
+		if (FrustumType == EDisplayClusterWarpBlendFrustumType::FULL && GeometryContext.GeometryProxy.GeometryType == EDisplayClusterWarpGeometryType::WarpMap)
+		{
+			while (true)
+			{
+				SetParameter(EDisplayClusterWarpBlendFrustumType::LOD);
+
+				// Fast check for bad view
+				ImplCalcViewProjection();
+				if (ImplCalcFrustum())
+				{
+					break;
+				}
+
+				if (!ImplUpdateProjectionType())
+				{
+					break;
+				}
+			}
+		}
+
+		// Restore base parameter
+		SetParameter(BaseFrustumType);
+
+		// Search valid projection mode:
+		if (BaseProjectionType == ProjectionType)
+		{
+
+			while (true)
+			{
+				// Full check for bad projection:
+				ImplCalcViewProjection();
+
+				if (ImplCalcFrustum())
+				{
+					break;
+				}
+
+				if (!ImplUpdateProjectionType())
+				{
+					break;
+				}
+			}
+		}
+	}
+
+	// Fit frustum to context size
+	ImplFitFrustumToContextSize(InViewport->GetContexts()[InContextNum].ContextSize);
+
+	ImplBuildFrustum(InViewport, InContextNum);
+
+	// Return result back:
+	OutFrustum = Frustum;
+
+	return true;
+}
+
 void FDisplayClusterWarpBlendMath_Frustum::ImplBuildFrustum(IDisplayClusterViewport* InViewport, const uint32 InContextNum)
 {
 	// These matrices were copied from LocalPlayer.cpp.
@@ -160,3 +256,146 @@ bool FDisplayClusterWarpBlendMath_Frustum::ImplCalcFrustum_LOD_WarpMap()
 	return bAllPointsInFrustum;
 }
 
+bool FDisplayClusterWarpBlendMath_Frustum::ImplCalcFrustum()
+{
+	// Reset extend of the frustum
+	Frustum.ProjectionAngles.Top = -FLT_MAX;
+	Frustum.ProjectionAngles.Bottom = FLT_MAX;
+	Frustum.ProjectionAngles.Left = FLT_MAX;
+	Frustum.ProjectionAngles.Right = -FLT_MAX;
+
+	GeometryToFrustum = GeometryContext.GeometryToOrigin * World2Local;
+
+	bool bResult = false;
+
+	//Compute rendering frustum with current method
+	switch (FrustumType)
+	{
+	case EDisplayClusterWarpBlendFrustumType::AABB:
+		bResult = ImplCalcFrustum_AABB();
+		break;
+
+	case EDisplayClusterWarpBlendFrustumType::FULL:
+		bResult = ImplCalcFrustum_FULL();
+		break;
+
+	case EDisplayClusterWarpBlendFrustumType::LOD:
+		bResult = ImplCalcFrustum_LOD();
+		break;
+
+	default:
+		break;
+	}
+
+	return bResult;
+}
+
+void FDisplayClusterWarpBlendMath_Frustum::ImplCalcViewProjection()
+{
+	// Customize view direction
+	FVector ViewDir = ViewDirection;
+
+	switch (ProjectionType)
+	{
+	case EDisplayClusterWarpBlendProjectionType::StaticSurfacePlane:
+		ViewDir = GeometryContext.SurfaceViewPlane;
+		break;
+
+	case EDisplayClusterWarpBlendProjectionType::RuntimeStaticSurfacePlaneInverted:
+		// Frustum projection fix (back-side view planes)
+		ViewDir = -GeometryContext.SurfaceViewPlane;
+		break;
+
+	case EDisplayClusterWarpBlendProjectionType::StaticSurfaceNormal:
+		// Use fixed surface view normal
+		ViewDir = GeometryContext.SurfaceViewNormal;
+		break;
+
+	case EDisplayClusterWarpBlendProjectionType::RuntimeStaticSurfaceNormalInverted:
+		// Frustum projection fix (back-side view planes)
+		ViewDir = -GeometryContext.SurfaceViewNormal;
+		break;
+
+	default:
+		break;
+	}
+
+	FVector const NewX = ViewDir.GetSafeNormal();
+
+	if (FMath::Abs(NewX.Z) < (1.f - KINDA_SMALL_NUMBER))
+	{
+		Local2World = FRotationMatrix::MakeFromXZ(NewX, FVector(0.f, 0.f, 1.f));
+	}
+	else
+	{
+		Local2World = FRotationMatrix::MakeFromXY(NewX, FVector(0.f, 1.f, 0.f));
+	}
+
+	if (GeometryContext.bRotateFrustumToFitContextSize)
+	{
+		const FQuat RotateBy90(NewX, FMath::DegreesToRadians(90));
+		Local2World *= FRotationMatrix(RotateBy90.Rotator());
+	}
+
+	Local2World.SetOrigin(EyeOrigin); // Finally set view origin to eye location
+	World2Local = Local2World.Inverse();
+}
+
+void FDisplayClusterWarpBlendMath_Frustum::ImplFitFrustumToContextSize(const FIntPoint& ContextSize)
+{
+	// See if it should roll by 90 degrees to get a better fit with the context size aspect ratio.
+	
+	// Calculate the frustum Horizontal to Vertical aspect ratio.
+
+	const float HFoV = FMath::Abs(Frustum.ProjectionAngles.Right - Frustum.ProjectionAngles.Left);
+	const float VFoV = FMath::Abs(Frustum.ProjectionAngles.Top - Frustum.ProjectionAngles.Bottom);
+
+	// Validate the values to protect from division by zero
+
+	const bool bValidValues =
+		!FMath::IsNearlyZero(HFoV)
+		&& !FMath::IsNearlyZero(VFoV)
+		&& (ContextSize.X > 0)
+		&& (ContextSize.Y > 0);
+
+	if (!bValidValues)
+	{
+		return;
+	}
+
+	const bool bFrustumIsWide = HFoV > VFoV;
+	const bool bContextSizeIsWide = ContextSize.X > ContextSize.Y;
+
+	// If both aspect ratios are of the same type already, then there is nothing else to do.
+	if (bContextSizeIsWide == bFrustumIsWide)
+	{
+		return;
+	}
+
+	// If we are here, we must decide if we should toggle the 90 degree roll.
+	// The toggle will have hysteresis to avoid jumping back and forth.
+	// We will use the frustum aspect ratio to determine this
+
+	// Don't flip the rotation if we're inside hystersis
+	if (FMath::IsNearlyEqual(HFoV/VFoV, 1.0, GDisplayClusterFrustumOrientationFitHysteresis))
+	{
+		return;
+	}
+
+	// If we're here, it is because we should toggle the rotation fit.
+	GeometryContext.bRotateFrustumToFitContextSize = !GeometryContext.bRotateFrustumToFitContextSize;
+
+	// So we re-calculate the frustum with the new parameters.
+	ImplCalcViewProjection();
+	const bool bValidFrustum = ImplCalcFrustum();
+
+	// We only update bRotateFrustumToFitContextSize if the frustum is valid, whether it was valid before the change or not.
+	// So if not valid, we roll back the change and recalculate the frustum.
+	if (!bValidFrustum)
+	{
+		GeometryContext.bRotateFrustumToFitContextSize = !GeometryContext.bRotateFrustumToFitContextSize;
+
+		ImplCalcViewProjection();
+		ImplCalcFrustum();
+	}
+}
\ No newline at end of file
diff --git a/Engine/Plugins/Runtime/nDisplay/Source/DisplayClusterShaders/Private/WarpBlend/Math/DisplayClusterWarpBlendMath_Frustum.h b/Engine/Plugins/Runtime/nDisplay/Source/DisplayClusterShaders/Private/WarpBlend/Math/DisplayClusterWarpBlendMath_Frustum.h
--- a/Engine/Plugins/Runtime/nDisplay/Source/DisplayClusterShaders/Private/WarpBlend/Math/DisplayClusterWarpBlendMath_Frustum.h
+++ b/Engine/Plugins/Runtime/nDisplay/Source/DisplayClusterShaders/Private/WarpBlend/Math/DisplayClusterWarpBlendMath_Frustum.h
@@ -18,81 +18,7 @@ public:
 	{ }
 
 public:
-	bool CalcFrustum(IDisplayClusterViewport* InViewport, const uint32 InContextNum, FDisplayClusterWarpContext& OutFrustum)
-	{
-		if (!GeometryContext.GeometryProxy.bIsGeometryValid)
-		{
-			// Bad source geometry
-			return false;
-		}
-
-		// Get view base:
-		ImplCalcView();
-
-		if (bFindBestProjectionType == false)
-		{
-			ImplCalcViewProjection();
-			ImplCalcFrustum();
-		}
-		else
-		{
-			// Frustum projection with auto-fix (back-side view planes)
-			// Projection type changed runtime. Validate frustum points, all must be over view plane:
-			EDisplayClusterWarpBlendProjectionType BaseProjectionType = ProjectionType;
-			EDisplayClusterWarpBlendFrustumType    BaseFrustumType = FrustumType;
-
-			// Optimize for PerfectCPU:
-			if (FrustumType == EDisplayClusterWarpBlendFrustumType::FULL && GeometryContext.GeometryProxy.GeometryType == EDisplayClusterWarpGeometryType::WarpMap)
-			{
-				while (true)
-				{
-					SetParameter(EDisplayClusterWarpBlendFrustumType::LOD);
-
-					// Fast check for bad view
-					ImplCalcViewProjection();
-					if (ImplCalcFrustum())
-					{
-						break;
-					}
-
-					if (!ImplUpdateProjectionType())
-					{
-						break;
-					}
-				}
-			}
-
-			// Restore base parameter
-			SetParameter(BaseFrustumType);
-
-			// Search valid projection mode:
-			if (BaseProjectionType == ProjectionType)
-			{
-
-				while (true)
-				{
-					//Full check for bad projection:
-					ImplCalcViewProjection();
-					if (ImplCalcFrustum())
-					{
-						break;
-					}
-
-					if (!ImplUpdateProjectionType())
-					{
-						break;
-					}
-				}
-			}
-		}
-
-		ImplBuildFrustum(InViewport, InContextNum);
-
-		// Return result back:
-		OutFrustum = Frustum;
-
-		return true;
-	}
+	bool CalcFrustum(IDisplayClusterViewport* InViewport, const uint32 InContextNum, FDisplayClusterWarpContext& OutFrustum);
 
 private:
 	EDisplayClusterWarpBlendFrustumType       FrustumType = EDisplayClusterWarpBlendFrustumType::AABB;
@@ -195,87 +121,11 @@ private:
 	FMatrix Local2World;
 	FMatrix World2Local;
 
-	void ImplCalcViewProjection()
-	{
-		// CUstomize view direction
-		FVector ViewDir = ViewDirection;
-
-		switch (ProjectionType)
-		{
-		case EDisplayClusterWarpBlendProjectionType::StaticSurfacePlane:
-			ViewDir = GeometryContext.SurfaceViewPlane;
-			break;
-
-		case EDisplayClusterWarpBlendProjectionType::RuntimeStaticSurfacePlaneInverted:
-			// Frustum projection fix (back-side view planes)
-			ViewDir = -GeometryContext.SurfaceViewPlane;
-			break;
-
-		case EDisplayClusterWarpBlendProjectionType::StaticSurfaceNormal:
-			// Use fixed surface view normal
-			ViewDir = GeometryContext.SurfaceViewNormal;
-			break;
-
-		case EDisplayClusterWarpBlendProjectionType::RuntimeStaticSurfaceNormalInverted:
-			// Frustum projection fix (back-side view planes)
-			ViewDir = -GeometryContext.SurfaceViewNormal;
-			break;
-
-		default:
-			break;
-		}
-
-		FVector const NewX = ViewDir.GetSafeNormal();
-		// make sure we don't ever pick the same as NewX
-		if(FMath::Abs(NewX.Z) < (1.f - KINDA_SMALL_NUMBER))
-		{
-			Local2World = FRotationMatrix::MakeFromXZ(NewX, FVector(0.f, 0.f, 1.f));
-		}
-		else
-		{
-			Local2World = FRotationMatrix::MakeFromXY(NewX, FVector(0.f, 1.f, 0.f));
-		}
-
-		Local2World.SetOrigin(EyeOrigin); // Finally set view origin to eye location
-
-		World2Local = Local2World.Inverse();
-	}
+	void ImplCalcViewProjection();
 
 	FMatrix GeometryToFrustum;
 
-	bool ImplCalcFrustum()
-	{
-		// Reset extend of the frustum
-		Frustum.ProjectionAngles.Top = -FLT_MAX;
-		Frustum.ProjectionAngles.Bottom = FLT_MAX;
-		Frustum.ProjectionAngles.Left = FLT_MAX;
-		Frustum.ProjectionAngles.Right = -FLT_MAX;
-
-		GeometryToFrustum = GeometryContext.GeometryToOrigin * World2Local;
-
-		bool bResult = false;
-
-		//Compute rendering frustum with current method
-		switch (FrustumType)
-		{
-		case EDisplayClusterWarpBlendFrustumType::AABB:
-			bResult = ImplCalcFrustum_AABB();
-			break;
-
-		case EDisplayClusterWarpBlendFrustumType::FULL:
-			bResult = ImplCalcFrustum_FULL();
-			break;
-
-		case EDisplayClusterWarpBlendFrustumType::LOD:
-			bResult = ImplCalcFrustum_LOD();
-			break;
-
-		default:
-			break;
-		}
-
-		return bResult;
-	}
+	bool ImplCalcFrustum();
 
 	void ImplBuildFrustum(IDisplayClusterViewport* InViewport, const uint32 InContextNum);
 
@@ -392,6 +242,13 @@ private:
 
 	bool ImplCalcFrustum_LOD_WarpMap();
 
+	/** 
+	 * When necessary, rotates by 90 deg. the frustum to better fit the aspect ratio of the viewport.
+	 * 
+	 * @param ContextSize Viewport size (Width x Height).
+	 */
+	void ImplFitFrustumToContextSize(const FIntPoint& ContextSize);
+
 private:
 	FDisplayClusterWarpBlend_GeometryContext& GeometryContext;
 
