nDisplay: Added ability to alphabetically sort nDisplay cluster nodes and viewports in the tree view in the nDisplay Editor

diff --git a/Engine/Plugins/Runtime/nDisplay/Source/DisplayClusterConfigurator/Private/Views/TreeViews/Cluster/DisplayClusterConfiguratorViewCluster.cpp b/Engine/Plugins/Runtime/nDisplay/Source/DisplayClusterConfigurator/Private/Views/TreeViews/Cluster/DisplayClusterConfiguratorViewCluster.cpp
--- a/Engine/Plugins/Runtime/nDisplay/Source/DisplayClusterConfigurator/Private/Views/TreeViews/Cluster/DisplayClusterConfiguratorViewCluster.cpp
+++ b/Engine/Plugins/Runtime/nDisplay/Source/DisplayClusterConfigurator/Private/Views/TreeViews/Cluster/DisplayClusterConfiguratorViewCluster.cpp
@@ -16,7 +16,6 @@
 #include "Framework/MultiBox/MultiBoxBuilder.h"
 #include "UICommandList_Pinnable.h"
 #include "HAL/PlatformApplicationMisc.h"
-#include "UnrealExporter.h"
 #include "Exporters/Exporter.h"
 #include "ScopedTransaction.h"
 #include "Widgets/Images/SImage.h"
@@ -46,7 +45,7 @@ TSharedRef<SWidget> FDisplayClusterConfiguratorViewCluster::CreateWidget()
 	return FDisplayClusterConfiguratorViewTree::CreateWidget();
 }
 
-void FDisplayClusterConfiguratorViewCluster::ConstructColumns(TArray<SHeaderRow::FColumn::FArguments>& OutColumnArgs) const
+void FDisplayClusterConfiguratorViewCluster::ConstructColumns(TArray<SHeaderRow::FColumn::FArguments>& OutColumnArgs)
 {
 	// For the cluster view, use the group column to indicate which host the cluster nodes belong to with a colored band
 	OutColumnArgs.Add(SHeaderRow::Column(FDisplayClusterConfiguratorViewCluster::Columns::Host)
diff --git a/Engine/Plugins/Runtime/nDisplay/Source/DisplayClusterConfigurator/Private/Views/TreeViews/Cluster/DisplayClusterConfiguratorViewCluster.h b/Engine/Plugins/Runtime/nDisplay/Source/DisplayClusterConfigurator/Private/Views/TreeViews/Cluster/DisplayClusterConfiguratorViewCluster.h
--- a/Engine/Plugins/Runtime/nDisplay/Source/DisplayClusterConfigurator/Private/Views/TreeViews/Cluster/DisplayClusterConfiguratorViewCluster.h
+++ b/Engine/Plugins/Runtime/nDisplay/Source/DisplayClusterConfigurator/Private/Views/TreeViews/Cluster/DisplayClusterConfiguratorViewCluster.h
@@ -26,7 +26,7 @@ public:
 
 	//~ Begin FDisplayClusterConfiguratorViewTree interface
 	virtual TSharedRef<SWidget> CreateWidget() override;
-	virtual void ConstructColumns(TArray<SHeaderRow::FColumn::FArguments>& OutColumnArgs) const override;
+	virtual void ConstructColumns(TArray<SHeaderRow::FColumn::FArguments>& OutColumnArgs) override;
 
 	virtual void FillContextMenu(FMenuBuilder& MenuBuilder) override;
 	virtual void BindPinnableCommands(FUICommandList_Pinnable& CommandList) override;
diff --git a/Engine/Plugins/Runtime/nDisplay/Source/DisplayClusterConfigurator/Private/Views/TreeViews/DisplayClusterConfiguratorViewTree.cpp b/Engine/Plugins/Runtime/nDisplay/Source/DisplayClusterConfigurator/Private/Views/TreeViews/DisplayClusterConfiguratorViewTree.cpp
--- a/Engine/Plugins/Runtime/nDisplay/Source/DisplayClusterConfigurator/Private/Views/TreeViews/DisplayClusterConfiguratorViewTree.cpp
+++ b/Engine/Plugins/Runtime/nDisplay/Source/DisplayClusterConfigurator/Private/Views/TreeViews/DisplayClusterConfiguratorViewTree.cpp
@@ -5,11 +5,9 @@
 #include "DisplayClusterConfiguratorBlueprintEditor.h"
 #include "Views/TreeViews/DisplayClusterConfiguratorTreeItem.h"
 #include "Views/TreeViews/DisplayClusterConfiguratorTreeBuilder.h"
+#include "Views/TreeViews/IDisplayClusterConfiguratorViewTree.h"
 #include "Views/TreeViews/SDisplayClusterConfiguratorViewTree.h"
 
-#include "Widgets/Images/SImage.h"
-#include "Widgets/Layout/SSpacer.h"
-
 #define LOCTEXT_NAMESPACE "FDisplayClusterConfiguratorViewTree"
 
 const FName IDisplayClusterConfiguratorViewTree::Columns::Item("Item");
@@ -20,6 +18,9 @@ FDisplayClusterConfiguratorViewTree::FDisplayClusterConfiguratorViewTree(const T
 	, bEnabled(false)
 {
 	GEditor->RegisterForUndo(this);
+
+	SortedColumn = IDisplayClusterConfiguratorViewTree::Columns::Item;
+	SortMode = EColumnSortMode::Ascending;
 }
 
 FDisplayClusterConfiguratorViewTree::~FDisplayClusterConfiguratorViewTree()
@@ -87,11 +88,14 @@ UDisplayClusterConfigurationData* FDisplayClusterConfiguratorViewTree::GetEditor
 	return ToolkitPtr.Pin()->GetEditorData();
 }
 
-void FDisplayClusterConfiguratorViewTree::ConstructColumns(TArray<SHeaderRow::FColumn::FArguments>& OutColumnArgs) const
+void FDisplayClusterConfiguratorViewTree::ConstructColumns(TArray<SHeaderRow::FColumn::FArguments>& OutColumnArgs)
 {
 	OutColumnArgs.Add(SHeaderRow::Column(IDisplayClusterConfiguratorViewTree::Columns::Item)
 		.DefaultLabel(LOCTEXT("DisplayClusterConfiguratorNameLabel", "Items"))
-		.FillWidth(0.5f));
+		.FillWidth(0.5f)
+		.InitialSortMode(EColumnSortMode::Ascending)
+		.SortMode(this, &FDisplayClusterConfiguratorViewTree::GetSortMode, IDisplayClusterConfiguratorViewTree::Columns::Item)
+		.OnSort(this, &FDisplayClusterConfiguratorViewTree::OnSort));
 }
 
 void FDisplayClusterConfiguratorViewTree::SetHoveredItem(const TSharedRef<IDisplayClusterConfiguratorTreeItem>& InTreeItem)
@@ -159,6 +163,67 @@ EDisplayClusterConfiguratorTreeFilterResult FDisplayClusterConfiguratorViewTree:
 	return ItemFilterResult;
 }
 
+void FDisplayClusterConfiguratorViewTree::Sort(const TArray<TSharedPtr<IDisplayClusterConfiguratorTreeItem>>& InItems, TArray<TSharedPtr<IDisplayClusterConfiguratorTreeItem>>& OutSortedItems)
+{
+	for (const TSharedPtr<IDisplayClusterConfiguratorTreeItem>& Item : InItems)
+	{
+		SortItem(Item);
+		OutSortedItems.Add(Item);
+	}
+
+	if (SortMode != EColumnSortMode::None)
+	{
+		if (SortedColumn == IDisplayClusterConfiguratorViewTree::Columns::Item)
+		{
+			Algo::SortBy(OutSortedItems, [this](const TSharedPtr<IDisplayClusterConfiguratorTreeItem>& Item) { return Item->GetRowItemName(); }, FNameLexicalLess());
+		}
+
+		if (SortMode == EColumnSortMode::Descending)
+		{
+			Algo::Reverse(OutSortedItems);
+		}
+	}
+}
+
+void FDisplayClusterConfiguratorViewTree::SortItem(const TSharedPtr<IDisplayClusterConfiguratorTreeItem>& InItem)
+{
+	if (InItem.IsValid() && InItem->GetChildren().Num() && SortMode != EColumnSortMode::None)
+	{
+		for (const TSharedPtr<IDisplayClusterConfiguratorTreeItem>& Child : InItem->GetChildren())
+		{
+			SortItem(Child);
+		}
+
+		if (SortedColumn == IDisplayClusterConfiguratorViewTree::Columns::Item)
+		{
+			Algo::SortBy(InItem->GetChildren(), [this](const TSharedPtr<IDisplayClusterConfiguratorTreeItem>& Child) { return Child->GetRowItemName(); }, FNameLexicalLess());
+		}
+
+		if (SortMode == EColumnSortMode::Descending)
+		{
+			Algo::Reverse(InItem->GetChildren());
+		}
+	}
+}
+
+EColumnSortMode::Type FDisplayClusterConfiguratorViewTree::GetSortMode(const FName InColumnId) const
+{
+	if (SortedColumn != InColumnId)
+	{
+		return EColumnSortMode::None;
+	}
+
+	return SortMode;
+}
+
+void FDisplayClusterConfiguratorViewTree::OnSort(const EColumnSortPriority::Type InSortPriority, const FName& InColumnId, const EColumnSortMode::Type InSortMode)
+{
+	SortedColumn = InColumnId;
+	SortMode = InSortMode;
+
+	ViewTree->ApplySortAndFilter();
+}
+
 void FDisplayClusterConfiguratorViewTree::RebuildTree()
 {
 	ViewTree->RebuildTree();
diff --git a/Engine/Plugins/Runtime/nDisplay/Source/DisplayClusterConfigurator/Private/Views/TreeViews/DisplayClusterConfiguratorViewTree.h b/Engine/Plugins/Runtime/nDisplay/Source/DisplayClusterConfigurator/Private/Views/TreeViews/DisplayClusterConfiguratorViewTree.h
--- a/Engine/Plugins/Runtime/nDisplay/Source/DisplayClusterConfigurator/Private/Views/TreeViews/DisplayClusterConfiguratorViewTree.h
+++ b/Engine/Plugins/Runtime/nDisplay/Source/DisplayClusterConfigurator/Private/Views/TreeViews/DisplayClusterConfiguratorViewTree.h
@@ -40,7 +40,7 @@ public:
 	//~ IDisplayClusterConfiguratorViewTree
 	virtual void RebuildTree() override;
 	virtual UDisplayClusterConfigurationData* GetEditorData() const override;
-	virtual void ConstructColumns(TArray<SHeaderRow::FColumn::FArguments>& OutColumnArgs) const override;
+	virtual void ConstructColumns(TArray<SHeaderRow::FColumn::FArguments>& OutColumnArgs) override;
 
 	virtual void SetHoveredItem(const TSharedRef<IDisplayClusterConfiguratorTreeItem>& InTreeItem) override;
 	virtual void ClearHoveredItem() override;
@@ -51,6 +51,7 @@ public:
 	virtual void GetSelectedObjects(TArray<UObject*>& OutObjects) const override;
 	virtual void FindAndSelectObjects(const TArray<UObject*>& ObjectsToSelect) override;
 	virtual void Filter(const FDisplayClusterConfiguratorTreeFilterArgs& InArgs, const TArray<TSharedPtr<IDisplayClusterConfiguratorTreeItem>>& InItems, TArray<TSharedPtr<IDisplayClusterConfiguratorTreeItem>>& OutFilteredItems) override;
+	virtual void Sort(const TArray<TSharedPtr<IDisplayClusterConfiguratorTreeItem>>& InItems, TArray<TSharedPtr<IDisplayClusterConfiguratorTreeItem>>& OutSortedItems) override;
 
 	virtual FDelegateHandle RegisterOnHoveredItemSet(const FOnHoveredItemSetDelegate& Delegate) override;
 	virtual void UnregisterOnHoveredItemSet(FDelegateHandle DelegateHandle) override;
@@ -74,6 +75,15 @@ protected:
 
 	virtual EDisplayClusterConfiguratorTreeFilterResult FilterItem(const TSharedPtr<IDisplayClusterConfiguratorTreeItem>& InItem, const FDisplayClusterConfiguratorTreeFilterArgs& InArgs, TArray<TSharedPtr<IDisplayClusterConfiguratorTreeItem>>& OutFilteredItems);
 
+	/** Applies the current sort configuration to the specified item, sorting the item's children as necessary */
+	virtual void SortItem(const TSharedPtr<IDisplayClusterConfiguratorTreeItem>& InItem);
+
+	/** Gets the current sort mode for the specified column */
+	virtual EColumnSortMode::Type GetSortMode(const FName InColumnId) const;
+
+	/** Raised when a column in the tree view is sorted */
+	virtual void OnSort(const EColumnSortPriority::Type InSortPriority, const FName& InColumnId, const EColumnSortMode::Type InSortMode);
+
 protected:
 	TWeakPtr<FDisplayClusterConfiguratorBlueprintEditor> ToolkitPtr;
 
@@ -89,4 +99,10 @@ protected:
 	TSharedPtr<SDisplayClusterConfiguratorViewTree> ViewTree;
 
 	bool bEnabled;
+
+	/** The name of the column in the tree view that is currently being sorted */
+	FName SortedColumn;
+
+	/** The current sort mode of the column being sorted in the tree view */
+	EColumnSortMode::Type SortMode;
 };
diff --git a/Engine/Plugins/Runtime/nDisplay/Source/DisplayClusterConfigurator/Private/Views/TreeViews/SDisplayClusterConfiguratorViewTree.cpp b/Engine/Plugins/Runtime/nDisplay/Source/DisplayClusterConfigurator/Private/Views/TreeViews/SDisplayClusterConfiguratorViewTree.cpp
--- a/Engine/Plugins/Runtime/nDisplay/Source/DisplayClusterConfigurator/Private/Views/TreeViews/SDisplayClusterConfiguratorViewTree.cpp
+++ b/Engine/Plugins/Runtime/nDisplay/Source/DisplayClusterConfigurator/Private/Views/TreeViews/SDisplayClusterConfiguratorViewTree.cpp
@@ -13,7 +13,6 @@
 #include "Modules/ModuleManager.h"
 #include "Misc/TextFilterExpressionEvaluator.h"
 #include "UICommandList_Pinnable.h"
-#include "UObject/PackageReload.h"
 #include "UObject/UObjectGlobals.h"
 #include "Types/ISlateMetaData.h"
 #include "Widgets/Views/STreeView.h"
@@ -288,17 +287,20 @@ void SDisplayClusterConfiguratorViewTree::RebuildTree()
 
 	FDisplayClusterConfiguratorTreeBuilderOutput Output(Items, LinearItems);
 	BuilderPtr.Pin()->Build(Output);
-	ApplyFilter();
+	ApplySortAndFilter();
 }
 
-void SDisplayClusterConfiguratorViewTree::ApplyFilter()
+void SDisplayClusterConfiguratorViewTree::ApplySortAndFilter()
 {
 	TextFilterPtr->SetFilterText(FilterText);
 
 	FilteredItems.Empty();
 
+	TArray<TSharedPtr<IDisplayClusterConfiguratorTreeItem>> SortedItems;
+	ViewTreePtr.Pin()->Sort(Items, SortedItems);
+
 	FDisplayClusterConfiguratorTreeFilterArgs FilterArgs(!FilterText.IsEmpty() ? TextFilterPtr : nullptr);
-	ViewTreePtr.Pin()->Filter(FilterArgs, Items, FilteredItems);
+	ViewTreePtr.Pin()->Filter(FilterArgs, SortedItems, FilteredItems);
 
 	if (!FilterText.IsEmpty())
 	{
@@ -422,7 +424,7 @@ void SDisplayClusterConfiguratorViewTree::BindCommands()
 void SDisplayClusterConfiguratorViewTree::OnFilterTextChanged(const FText& SearchText)
 {
 	FilterText = SearchText;
-	ApplyFilter();
+	ApplySortAndFilter();
 }
 
 void SDisplayClusterConfiguratorViewTree::OnConfigReloaded()
diff --git a/Engine/Plugins/Runtime/nDisplay/Source/DisplayClusterConfigurator/Private/Views/TreeViews/SDisplayClusterConfiguratorViewTree.h b/Engine/Plugins/Runtime/nDisplay/Source/DisplayClusterConfigurator/Private/Views/TreeViews/SDisplayClusterConfiguratorViewTree.h
--- a/Engine/Plugins/Runtime/nDisplay/Source/DisplayClusterConfigurator/Private/Views/TreeViews/SDisplayClusterConfiguratorViewTree.h
+++ b/Engine/Plugins/Runtime/nDisplay/Source/DisplayClusterConfigurator/Private/Views/TreeViews/SDisplayClusterConfiguratorViewTree.h
@@ -62,16 +62,14 @@ public:
 public:
 	virtual void Refresh();
 
-	virtual void RefreshFilter(){}
-
 	/** Creates the tree control and then populates */
 	virtual void CreateTreeColumns();
 
 	/** Function to build the Config tree widgets from the source config object */
 	void RebuildTree();
 
-	/** Apply filtering to the tree */
-	virtual void ApplyFilter();
+	/** Apply sorting and filtering to the tree */
+	virtual void ApplySortAndFilter();
 
 	/** Set the initial expansion state of the tree items */
 	virtual void SetInitialExpansionState();
diff --git a/Engine/Plugins/Runtime/nDisplay/Source/DisplayClusterConfigurator/Public/Views/TreeViews/IDisplayClusterConfiguratorViewTree.h b/Engine/Plugins/Runtime/nDisplay/Source/DisplayClusterConfigurator/Public/Views/TreeViews/IDisplayClusterConfiguratorViewTree.h
--- a/Engine/Plugins/Runtime/nDisplay/Source/DisplayClusterConfigurator/Public/Views/TreeViews/IDisplayClusterConfiguratorViewTree.h
+++ b/Engine/Plugins/Runtime/nDisplay/Source/DisplayClusterConfigurator/Public/Views/TreeViews/IDisplayClusterConfiguratorViewTree.h
@@ -98,7 +98,7 @@ public:
 	 * Generates a list of columns to add to the tree view.
 	 * @param OutColumnArgs - The list of column arguments to construct the columns from.
 	 */
-	virtual void ConstructColumns(TArray<SHeaderRow::FColumn::FArguments>& OutColumnArgs) const = 0;
+	virtual void ConstructColumns(TArray<SHeaderRow::FColumn::FArguments>& OutColumnArgs) = 0;
 
 	/**
 	 * Sets currently hovered tree item
@@ -138,7 +138,7 @@ public:
 	virtual void GetSelectedObjects(TArray<UObject*>& OutObjects) const = 0;
 
 	/**
-	 * Finds any tree items that represet the specified objects and selects them.
+	 * Finds any tree items that represent the specified objects and selects them.
 	 * @param ObjectsToSelect - The objects to select
 	 */
 	virtual void FindAndSelectObjects(const TArray<UObject*>& ObjectsToSelect) = 0;
@@ -151,6 +151,12 @@ public:
 	 */
 	virtual void Filter(const FDisplayClusterConfiguratorTreeFilterArgs& InArgs, const TArray<TSharedPtr<IDisplayClusterConfiguratorTreeItem>>& InItems, TArray<TSharedPtr<IDisplayClusterConfiguratorTreeItem>>& OutFilteredItems) = 0;
 
+	/**
+	 * Sorts the specified list of items based on the current sort configuration and outputs the sorted list.
+	 * @param InItems - The list of items to sort
+	 * @param OutSortedItems - The output list of sorted items
+	 */
+	virtual void Sort(const TArray<TSharedPtr<IDisplayClusterConfiguratorTreeItem>>& InItems, TArray<TSharedPtr<IDisplayClusterConfiguratorTreeItem>>& OutSortedItems) = 0;
 
 	/**
 	 * Register the hovered tree item delegate
