nDisplay: Allow light card labels to be toggled between editor and game.

diff --git a/Engine/Plugins/Runtime/nDisplay/Source/DisplayCluster/Private/Game/EngineClasses/Scene/DisplayClusterLabelComponent.cpp b/Engine/Plugins/Runtime/nDisplay/Source/DisplayCluster/Private/Game/EngineClasses/Scene/DisplayClusterLabelComponent.cpp
--- a/Engine/Plugins/Runtime/nDisplay/Source/DisplayCluster/Private/Game/EngineClasses/Scene/DisplayClusterLabelComponent.cpp
+++ b/Engine/Plugins/Runtime/nDisplay/Source/DisplayCluster/Private/Game/EngineClasses/Scene/DisplayClusterLabelComponent.cpp
@@ -33,6 +33,14 @@ UDisplayClusterLabelComponent::UDisplayClusterLabelComponent()
 	SetVisibility(false);
 }
 
+void UDisplayClusterLabelComponent::SetLabelConfiguration(const FDisplayClusterLabelConfiguration& InConfiguration)
+{
+	SetRootActor(InConfiguration.RootActor);
+	SetWidgetScale(InConfiguration.Scale);
+	SetVisibility(InConfiguration.bVisible);
+	LabelFlags = InConfiguration.LabelFlags;
+}
+
 void UDisplayClusterLabelComponent::SetRootActor(ADisplayClusterRootActor* InActor)
 {
 	RootActor = InActor;
@@ -49,6 +57,21 @@ float UDisplayClusterLabelComponent::GetWidgetScale() const
 	return WidgetScale;
 }
 
+void UDisplayClusterLabelComponent::SetLabelFlags(EDisplayClusterLabelFlags InFlags)
+{
+	EnumAddFlags(LabelFlags, InFlags);
+}
+
+void UDisplayClusterLabelComponent::ClearLabelFlags(EDisplayClusterLabelFlags InFlags)
+{
+	EnumRemoveFlags(LabelFlags, InFlags);
+}
+
+bool UDisplayClusterLabelComponent::HasAnyLabelFlags(EDisplayClusterLabelFlags InFlags) const
+{
+	return EnumHasAnyFlags(GetLabelFlags(), InFlags);
+}
+
 void UDisplayClusterLabelComponent::TickComponent(float DeltaTime, ELevelTick TickType,
                                                   FActorComponentTickFunction* ThisTickFunction)
 {
@@ -110,12 +133,12 @@ void UDisplayClusterLabelComponent::UpdateWidgetComponent()
 		}
 
 #if WITH_EDITOR
-		const bool bInGame = GEditor == nullptr;
+		const bool bIsGame = !GIsEditor;
 #else
-		constexpr bool bInGame = true;
+		const bool bIsGame = true;
 #endif
-
-		if (bInGame)
+		
+		if (bIsGame)
 		{
 			const USceneComponent* Viewpoint = RootActor.IsValid() ? RootActor->GetCommonViewPoint() : nullptr;
 			if (Viewpoint != nullptr)
@@ -167,8 +190,25 @@ void UDisplayClusterLabelComponent::CheckForVisibilityChange()
 	// MU doesn't seem to register nested sub component visibility changes, and OnVisibilityChanged won't
 	// fire on the MU client so we have to manually check here.
 
-	if (GetVisibleFlag() != WidgetComponent->GetVisibleFlag())
+#if WITH_EDITOR
+	const bool bIsGame = !GIsEditor;
+#else
+	const bool bIsGame = true;
+#endif
+	
+	bool bDisplayWidget = GetVisibleFlag();
+	
+	if (bIsGame && !HasAnyLabelFlags(EDisplayClusterLabelFlags::DisplayInGame))
 	{
-		WidgetComponent->SetVisibility(GetVisibleFlag());
+		bDisplayWidget = false;
 	}
-}
+	else if (!bIsGame && !HasAnyLabelFlags(EDisplayClusterLabelFlags::DisplayInEditor))
+	{
+		bDisplayWidget = false;
+	}
+	
+	if (bDisplayWidget != WidgetComponent->GetVisibleFlag())
+	{
+		WidgetComponent->SetVisibility(bDisplayWidget);
+	}
+}
\ No newline at end of file
diff --git a/Engine/Plugins/Runtime/nDisplay/Source/DisplayCluster/Private/Game/EngineClasses/Scene/DisplayClusterLightCardActor.cpp b/Engine/Plugins/Runtime/nDisplay/Source/DisplayCluster/Private/Game/EngineClasses/Scene/DisplayClusterLightCardActor.cpp
--- a/Engine/Plugins/Runtime/nDisplay/Source/DisplayCluster/Private/Game/EngineClasses/Scene/DisplayClusterLightCardActor.cpp
+++ b/Engine/Plugins/Runtime/nDisplay/Source/DisplayCluster/Private/Game/EngineClasses/Scene/DisplayClusterLightCardActor.cpp
@@ -577,7 +577,7 @@ void ADisplayClusterLightCardActor::SetIsLightCardFlag(bool bNewFlagValue)
 	bIsLightCardFlag = bNewFlagValue;
 }
 
-void ADisplayClusterLightCardActor::ShowLightCardLabel(bool bValue, float ScaleValue, ADisplayClusterRootActor* InRootActor)
+void ADisplayClusterLightCardActor::ShowLightCardLabel(const FDisplayClusterLabelConfiguration& InLabelConfiguration)
 {
 	if (IsUVActor())
 	{
@@ -589,9 +589,19 @@ void ADisplayClusterLightCardActor::ShowLightCardLabel(bool bValue, float ScaleV
 	LabelComponent->Modify(false);
 	LightCardComponent->Modify(false);
 #endif
-	LabelComponent->SetVisibility(bValue, true);
-	LabelComponent->SetRootActor(InRootActor);
-	LabelComponent->SetWidgetScale(ScaleValue);
+	LabelComponent->SetLabelConfiguration(InLabelConfiguration);
+}
+
+void ADisplayClusterLightCardActor::ShowLightCardLabel(bool bValue, float ScaleValue, ADisplayClusterRootActor* InRootActor)
+{
+	FDisplayClusterLabelConfiguration LabelConfiguration;
+	{
+		LabelConfiguration.bVisible = bValue;
+		LabelConfiguration.Scale = ScaleValue;
+		LabelConfiguration.RootActor = InRootActor;
+		LabelConfiguration.LabelFlags = LabelComponent->GetLabelFlags();
+	}
+	ShowLightCardLabel(MoveTemp(LabelConfiguration));
 }
 
 void ADisplayClusterLightCardActor::SetWeakRootActorOwner(ADisplayClusterRootActor* InRootActor)
diff --git a/Engine/Plugins/Runtime/nDisplay/Source/DisplayCluster/Public/Components/DisplayClusterLabelComponent.h b/Engine/Plugins/Runtime/nDisplay/Source/DisplayCluster/Public/Components/DisplayClusterLabelComponent.h
--- a/Engine/Plugins/Runtime/nDisplay/Source/DisplayCluster/Public/Components/DisplayClusterLabelComponent.h
+++ b/Engine/Plugins/Runtime/nDisplay/Source/DisplayCluster/Public/Components/DisplayClusterLabelComponent.h
@@ -4,6 +4,7 @@
 
 #include "CoreMinimal.h"
 #include "Components/SceneComponent.h"
+#include "Components/DisplayClusterLabelConfiguration.h"
 
 #include "DisplayClusterLabelComponent.generated.h"
 
@@ -26,8 +27,12 @@ class DISPLAYCLUSTER_API UDisplayClusterLabelComponent : public USceneComponent
 	GENERATED_BODY()
 
 public:
+	
 	UDisplayClusterLabelComponent();
 
+	/** Set the configuration for this label updating multiple properties at once */
+	void SetLabelConfiguration(const FDisplayClusterLabelConfiguration& InConfiguration);
+	
 	/** Set the root actor this label should use when determining transform */
 	void SetRootActor(ADisplayClusterRootActor* InActor);
 
@@ -37,6 +42,18 @@ public:
 	/** Return the widget scale */
 	float GetWidgetScale() const;
 
+	/** Updates flags for this label */
+	void SetLabelFlags(EDisplayClusterLabelFlags InFlags);
+	
+	/** Clears all label flags */
+	void ClearLabelFlags(EDisplayClusterLabelFlags InFlags);
+	
+	/** Retrieve label flags */
+	EDisplayClusterLabelFlags GetLabelFlags() const { return LabelFlags; }
+	
+	/** Checks for existence of a flag on this label */
+	bool HasAnyLabelFlags(EDisplayClusterLabelFlags InFlags) const;
+	
 	virtual void TickComponent(float DeltaTime, ELevelTick TickType,
 	                           FActorComponentTickFunction* ThisTickFunction) override;
 
@@ -73,4 +90,8 @@ private:
 	/** A uniform scale to apply to the text which will keep consistency across distance from the label to the root actor view */
 	UPROPERTY(EditAnywhere, BlueprintReadWrite, Category=Label, meta = (AllowPrivateAccess))
 	float WidgetScale = 1.f;
+
+	/** Current flags set for this label component */
+	UPROPERTY()
+	EDisplayClusterLabelFlags LabelFlags;
 };
diff --git a/Engine/Plugins/Runtime/nDisplay/Source/DisplayCluster/Public/Components/DisplayClusterLabelConfiguration.h b/Engine/Plugins/Runtime/nDisplay/Source/DisplayCluster/Public/Components/DisplayClusterLabelConfiguration.h
new file mode 100644
--- /dev/null
+++ b/Engine/Plugins/Runtime/nDisplay/Source/DisplayCluster/Public/Components/DisplayClusterLabelConfiguration.h
@@ -0,0 +1,44 @@
+// Copyright Epic Games, Inc. All Rights Reserved.
+
+#pragma once
+
+#include "CoreMinimal.h"
+#include "UObject/ObjectPtr.h"
+
+#include "DisplayClusterLabelConfiguration.generated.h"
+
+class ADisplayClusterRootActor;
+
+UENUM()
+enum class EDisplayClusterLabelFlags : uint8
+{
+	/** No label flags */
+	None              = 0,
+	/** Label allowed to be displayed in -game */
+	DisplayInGame     = 1 << 0,
+	/** Label allowed to be displayed in editor */
+	DisplayInEditor   = 1 << 1
+};
+ENUM_CLASS_FLAGS(EDisplayClusterLabelFlags)
+
+/** Generic args a label can recognize */
+struct FDisplayClusterLabelConfiguration
+{
+	/** The root actor */
+	TObjectPtr<ADisplayClusterRootActor> RootActor;
+		
+	/**
+	 * The scale to apply to the label
+	 */
+	float Scale = 1.f;
+
+	/**
+	 * Should the label be visible
+	 */
+	bool bVisible = false;
+
+	/**
+	 * Flags for the label. Can be used to fine-tune visibility settings
+	 */
+	EDisplayClusterLabelFlags LabelFlags;
+};
\ No newline at end of file
diff --git a/Engine/Plugins/Runtime/nDisplay/Source/DisplayCluster/Public/DisplayClusterLightCardActor.h b/Engine/Plugins/Runtime/nDisplay/Source/DisplayCluster/Public/DisplayClusterLightCardActor.h
--- a/Engine/Plugins/Runtime/nDisplay/Source/DisplayCluster/Public/DisplayClusterLightCardActor.h
+++ b/Engine/Plugins/Runtime/nDisplay/Source/DisplayCluster/Public/DisplayClusterLightCardActor.h
@@ -2,12 +2,14 @@
 
 #pragma once
 
+#include "Components/DisplayClusterLabelConfiguration.h"
 #include "StageActor/IDisplayClusterStageActor.h"
 
 #include "GameFramework/Actor.h"
 
 #include "DisplayClusterLightCardActor.generated.h"
 
+enum class EDisplayClusterLabelFlags : uint8;
 class ADisplayClusterRootActor;
 class UActorComponent;
 class UDisplayClusterLabelComponent;
@@ -115,8 +117,12 @@ public:
 
 	/** If this light card is considered a flag */
 	bool IsLightCardFlag() const { return bIsLightCardFlag; }
+
+	/** Configure light card labels */
+	void ShowLightCardLabel(const FDisplayClusterLabelConfiguration& InLabelConfiguration);
 	
 	/** Show or hide the light card label  */
+	UE_DEPRECATED(5.3, "Use the ShowLightCardLabel function which accepts FDisplayClusterLabelConfiguration arguments.")
 	void ShowLightCardLabel(bool bValue, float ScaleValue, ADisplayClusterRootActor* InRootActor);
 
 	/** Set a weak root actor owner. This should be used on legacy light cards that don't have StageActorComponent->RootActor set */
diff --git a/Engine/Plugins/Runtime/nDisplay/Source/DisplayClusterLightCardEditor/Private/DisplayClusterLightCardEditor.cpp b/Engine/Plugins/Runtime/nDisplay/Source/DisplayClusterLightCardEditor/Private/DisplayClusterLightCardEditor.cpp
--- a/Engine/Plugins/Runtime/nDisplay/Source/DisplayClusterLightCardEditor/Private/DisplayClusterLightCardEditor.cpp
+++ b/Engine/Plugins/Runtime/nDisplay/Source/DisplayClusterLightCardEditor/Private/DisplayClusterLightCardEditor.cpp
@@ -25,6 +25,7 @@
 #include "Blueprints/DisplayClusterBlueprintLib.h"
 #include "Components/DisplayClusterCameraComponent.h"
 #include "Components/DisplayClusterICVFXCameraComponent.h"
+#include "Components/DisplayClusterLabelComponent.h"
 
 #include "ContentBrowserModule.h"
 #include "FileHelpers.h"
@@ -231,8 +232,10 @@ AActor* FDisplayClusterLightCardEditor::SpawnActor(TSubclassOf<AActor> InActorCl
 		SpawnArgs.RootActor = ActiveRootActor.Get();
 		SpawnArgs.Template = InTemplate;
 		SpawnArgs.Level = InLevel;
-		SpawnArgs.AddLightCardArgs.bShowLabels = Settings->bDisplayLightCardLabels;
-		SpawnArgs.AddLightCardArgs.LabelScale = Settings->LightCardLabelScale;
+		SpawnArgs.AddLightCardArgs.LabelConfiguration.bVisible = Settings->bDisplayLightCardLabels;
+		SpawnArgs.AddLightCardArgs.LabelConfiguration.Scale = Settings->LightCardLabelScale;
+		SpawnArgs.AddLightCardArgs.LabelConfiguration.LabelFlags = Settings->LightCardLabelFlags;
+		SpawnArgs.AddLightCardArgs.LabelConfiguration.RootActor = ActiveRootActor.Get();
 	
 		if (ViewportView.IsValid())
 		{
@@ -456,8 +459,10 @@ void FDisplayClusterLightCardEditor::AddLightCardsToActor(const TArray<ADisplayC
 	if (ActiveRootActor.IsValid())
 	{
 		FDisplayClusterLightCardEditorHelper::FAddLightCardArgs AddLightCardArgs;
-		AddLightCardArgs.bShowLabels = ShouldShowLightCardLabels();
-		AddLightCardArgs.LabelScale = *GetLightCardLabelScale();
+		AddLightCardArgs.LabelConfiguration.bVisible = ShouldShowLightCardLabels();
+		AddLightCardArgs.LabelConfiguration.Scale = *GetLightCardLabelScale();
+		AddLightCardArgs.LabelConfiguration.LabelFlags = GetLightCardLabelFlags();
+		AddLightCardArgs.LabelConfiguration.RootActor = ActiveRootActor.Get();
 		
 		FDisplayClusterLightCardEditorHelper::AddLightCardsToRootActor(LightCards, ActiveRootActor.Get(), MoveTemp(AddLightCardArgs));
 		
@@ -763,8 +768,15 @@ void FDisplayClusterLightCardEditor::RemoveActors(
 							LayersSubsystem->RemoveActorFromLayer(LightCard, LightCardLayerName);
 						}
 					}
-					
-					LightCard->ShowLightCardLabel(false, *GetLightCardLabelScale(), ActiveRootActor.Get());
+
+					FDisplayClusterLabelConfiguration LabelArgs;
+					{
+						LabelArgs.bVisible = false;
+						LabelArgs.Scale = *GetLightCardLabelScale();
+						LabelArgs.RootActor = ActiveRootActor.Get();
+						LabelArgs.LabelFlags = GetLightCardLabelFlags();
+					}
+					LightCard->ShowLightCardLabel(MoveTemp(LabelArgs));
 				}
 			}
 
@@ -910,6 +922,7 @@ void FDisplayClusterLightCardEditor::ShowLightCardLabels(bool bVisible)
 
 	IDisplayClusterLightCardEditor::FLabelArgs Args;
 	Args.bVisible = bVisible;
+	Args.LabelFlags = GetLightCardLabelFlags();
 	Args.Scale = *GetLightCardLabelScale();
 	Args.RootActor = GetActiveRootActor().Get();
 	
@@ -923,6 +936,32 @@ bool FDisplayClusterLightCardEditor::ShouldShowLightCardLabels() const
 	return GetDefault<UDisplayClusterLightCardEditorProjectSettings>()->bDisplayLightCardLabels;
 }
 
+void FDisplayClusterLightCardEditor::ToggleLightCardLabelFlag(EDisplayClusterLabelFlags InFlag)
+{
+	SaveLightCardLabelFlags(GetLightCardLabelFlags() ^ InFlag);
+}
+
+void FDisplayClusterLightCardEditor::SaveLightCardLabelFlags(EDisplayClusterLabelFlags InFlags)
+{
+	GetMutableDefault<UDisplayClusterLightCardEditorProjectSettings>()->LightCardLabelFlags = InFlags;
+	ShowLightCardLabels(ShouldShowLightCardLabels());
+}
+
+bool FDisplayClusterLightCardEditor::LightCardLabelHasFlag(EDisplayClusterLabelFlags FlagToCheck) const
+{
+	return EnumHasAnyFlags(GetLightCardLabelFlags(), FlagToCheck);
+}
+
+bool FDisplayClusterLightCardEditor::LightCardLabelDoesNotHaveFlag(EDisplayClusterLabelFlags FlagToCheck) const
+{
+	return !LightCardLabelHasFlag(FlagToCheck);
+}
+
+EDisplayClusterLabelFlags FDisplayClusterLightCardEditor::GetLightCardLabelFlags() const
+{
+	return GetDefault<UDisplayClusterLightCardEditorProjectSettings>()->LightCardLabelFlags;
+}
+
 TOptional<float> FDisplayClusterLightCardEditor::GetLightCardLabelScale() const
 {
 	return GetDefault<UDisplayClusterLightCardEditorProjectSettings>()->LightCardLabelScale;
@@ -942,6 +981,7 @@ void FDisplayClusterLightCardEditor::SetLightCardLabelScale(float NewValue)
 	IDisplayClusterLightCardEditor::FLabelArgs Args;
 	Args.bVisible = ShouldShowLightCardLabels();
 	Args.Scale = NewValue;
+	Args.LabelFlags = GetLightCardLabelFlags();
 	Args.RootActor = GetActiveRootActor().Get();
 	
 	LightCardEditorModule.ShowLabels(MoveTemp(Args));
@@ -1289,7 +1329,9 @@ TSharedRef<SWidget> FDisplayClusterLightCardEditor::GenerateLabelsMenu()
 	MenuBuilder.BeginSection("Labels", LOCTEXT("LabelsMenuHeader", "Labels"));
 	{
 		MenuBuilder.AddMenuEntry(FDisplayClusterLightCardEditorCommands::Get().ToggleLightCardLabels);
-	
+		MenuBuilder.AddMenuEntry(FDisplayClusterLightCardEditorCommands::Get().ToggleLabelsVisibleInEditor);
+		MenuBuilder.AddMenuEntry(FDisplayClusterLightCardEditorCommands::Get().ToggleLabelsHiddenInGame);
+		
 		MenuBuilder.AddWidget(
 		SNew(SBox)
 		.HAlign(HAlign_Right)
@@ -1557,6 +1599,18 @@ void FDisplayClusterLightCardEditor::BindCommands()
 		FExecuteAction::CreateSP(this, &FDisplayClusterLightCardEditor::ToggleLightCardLabels),
 		FCanExecuteAction(),
 		FIsActionChecked::CreateSP(this, &FDisplayClusterLightCardEditor::ShouldShowLightCardLabels));
+
+	CommandList->MapAction(
+		FDisplayClusterLightCardEditorCommands::Get().ToggleLabelsVisibleInEditor,
+		FExecuteAction::CreateSP(this, &FDisplayClusterLightCardEditor::ToggleLightCardLabelFlag, EDisplayClusterLabelFlags::DisplayInEditor),
+		FCanExecuteAction(),
+		FIsActionChecked::CreateSP(this, &FDisplayClusterLightCardEditor::LightCardLabelHasFlag, EDisplayClusterLabelFlags::DisplayInEditor));
+
+	CommandList->MapAction(
+		FDisplayClusterLightCardEditorCommands::Get().ToggleLabelsHiddenInGame,
+		FExecuteAction::CreateSP(this, &FDisplayClusterLightCardEditor::ToggleLightCardLabelFlag, EDisplayClusterLabelFlags::DisplayInGame),
+		FCanExecuteAction(),
+		FIsActionChecked::CreateSP(this, &FDisplayClusterLightCardEditor::LightCardLabelDoesNotHaveFlag, EDisplayClusterLabelFlags::DisplayInGame));
 }
 
 void FDisplayClusterLightCardEditor::RegisterToolbarExtensions()
diff --git a/Engine/Plugins/Runtime/nDisplay/Source/DisplayClusterLightCardEditor/Private/DisplayClusterLightCardEditor.h b/Engine/Plugins/Runtime/nDisplay/Source/DisplayClusterLightCardEditor/Private/DisplayClusterLightCardEditor.h
--- a/Engine/Plugins/Runtime/nDisplay/Source/DisplayClusterLightCardEditor/Private/DisplayClusterLightCardEditor.h
+++ b/Engine/Plugins/Runtime/nDisplay/Source/DisplayClusterLightCardEditor/Private/DisplayClusterLightCardEditor.h
@@ -11,6 +11,7 @@
 #include "EditorUndoClient.h"
 #include "IDisplayClusterOperatorApp.h"
 
+enum class EDisplayClusterLabelFlags : uint8;
 class FLayoutExtender;
 class FMenuBuilder;
 class FSpawnTabArgs;
@@ -184,6 +185,21 @@ public:
 	/** If light card labels are currently toggled */
 	bool ShouldShowLightCardLabels() const;
 
+	/** Toggles light cards label flag */
+	void ToggleLightCardLabelFlag(EDisplayClusterLabelFlags InFlag);
+	
+	/** Save the light card label flags to settings */
+	void SaveLightCardLabelFlags(EDisplayClusterLabelFlags InFlags);
+	
+	/** If the light card label has a specific flag */
+	bool LightCardLabelHasFlag(EDisplayClusterLabelFlags FlagToCheck) const;
+
+	/** If the light card label is missing a specific flag */
+	bool LightCardLabelDoesNotHaveFlag(EDisplayClusterLabelFlags FlagToCheck) const;
+	
+	/** The current light card label flags */
+	EDisplayClusterLabelFlags GetLightCardLabelFlags() const;
+	
 	/** Return the current light card label scale */
 	TOptional<float> GetLightCardLabelScale() const;
 
diff --git a/Engine/Plugins/Runtime/nDisplay/Source/DisplayClusterLightCardEditor/Private/DisplayClusterLightCardEditorCommands.cpp b/Engine/Plugins/Runtime/nDisplay/Source/DisplayClusterLightCardEditor/Private/DisplayClusterLightCardEditorCommands.cpp
--- a/Engine/Plugins/Runtime/nDisplay/Source/DisplayClusterLightCardEditor/Private/DisplayClusterLightCardEditorCommands.cpp
+++ b/Engine/Plugins/Runtime/nDisplay/Source/DisplayClusterLightCardEditor/Private/DisplayClusterLightCardEditorCommands.cpp
@@ -42,8 +42,12 @@ void FDisplayClusterLightCardEditorCommands::RegisterCommands()
 	UI_COMMAND(DrawLightCard, "Draw Light Card", "Draw polygon light card on viewport", EUserInterfaceActionType::ToggleButton, FInputChord());
 
 	UI_COMMAND(ToggleAllLabels, "Labels", "Display labels in the preview and wall", EUserInterfaceActionType::ToggleButton, FInputChord());
-	UI_COMMAND(ToggleLightCardLabels, "Light Card Labels", "Display labels on light cards in the preview and wall", EUserInterfaceActionType::ToggleButton, FInputChord());
-
+	UI_COMMAND(ToggleLabelsVisibleInEditor, "Labels Visible in Editor", "Display labels in the editor", EUserInterfaceActionType::ToggleButton, FInputChord());
+	UI_COMMAND(ToggleLabelsHiddenInGame, "Labels Hidden in Game", "Hide labels from displaying on the wall when running as -game", EUserInterfaceActionType::ToggleButton, FInputChord());
+	
+	// @todo label types -- ToggleLightCardLabels is meant for only light card types and the user text should be updated to reflect that when labels are supported on more than just light cards
+	UI_COMMAND(ToggleLightCardLabels, "Show Labels", "Display labels on light cards in the preview and wall", EUserInterfaceActionType::ToggleButton, FInputChord());
+	
 	UI_COMMAND(ToggleIconVisibility, "Icons", "Display icons in the light card editor where applicable", EUserInterfaceActionType::ToggleButton, FInputChord());
 }
 
diff --git a/Engine/Plugins/Runtime/nDisplay/Source/DisplayClusterLightCardEditor/Private/DisplayClusterLightCardEditorCommands.h b/Engine/Plugins/Runtime/nDisplay/Source/DisplayClusterLightCardEditor/Private/DisplayClusterLightCardEditorCommands.h
--- a/Engine/Plugins/Runtime/nDisplay/Source/DisplayClusterLightCardEditor/Private/DisplayClusterLightCardEditorCommands.h
+++ b/Engine/Plugins/Runtime/nDisplay/Source/DisplayClusterLightCardEditor/Private/DisplayClusterLightCardEditorCommands.h
@@ -54,6 +54,8 @@ public:
 
 	TSharedPtr<FUICommandInfo> ToggleAllLabels;
 	TSharedPtr<FUICommandInfo> ToggleLightCardLabels;
+	TSharedPtr<FUICommandInfo> ToggleLabelsVisibleInEditor;
+	TSharedPtr<FUICommandInfo> ToggleLabelsHiddenInGame;
 
 	TSharedPtr<FUICommandInfo> ToggleIconVisibility;
 };
diff --git a/Engine/Plugins/Runtime/nDisplay/Source/DisplayClusterLightCardEditor/Private/DisplayClusterLightCardEditorModule.cpp b/Engine/Plugins/Runtime/nDisplay/Source/DisplayClusterLightCardEditor/Private/DisplayClusterLightCardEditorModule.cpp
--- a/Engine/Plugins/Runtime/nDisplay/Source/DisplayClusterLightCardEditor/Private/DisplayClusterLightCardEditorModule.cpp
+++ b/Engine/Plugins/Runtime/nDisplay/Source/DisplayClusterLightCardEditor/Private/DisplayClusterLightCardEditorModule.cpp
@@ -75,7 +75,7 @@ void FDisplayClusterLightCardEditorModule::ShowLabels(const FLabelArgs& InArgs)
 	for (ADisplayClusterLightCardActor* LightCardActor : RootActorLightCardActors)
 	{
 		LightCardActor->Modify(false);
-		LightCardActor->ShowLightCardLabel(InArgs.bVisible, InArgs.Scale, InArgs.RootActor);
+		LightCardActor->ShowLightCardLabel(InArgs);
 	}
 }
 
diff --git a/Engine/Plugins/Runtime/nDisplay/Source/DisplayClusterLightCardEditor/Private/Settings/DisplayClusterLightCardEditorSettings.cpp b/Engine/Plugins/Runtime/nDisplay/Source/DisplayClusterLightCardEditor/Private/Settings/DisplayClusterLightCardEditorSettings.cpp
--- a/Engine/Plugins/Runtime/nDisplay/Source/DisplayClusterLightCardEditor/Private/Settings/DisplayClusterLightCardEditorSettings.cpp
+++ b/Engine/Plugins/Runtime/nDisplay/Source/DisplayClusterLightCardEditor/Private/Settings/DisplayClusterLightCardEditorSettings.cpp
@@ -7,6 +7,7 @@
 #include "DisplayClusterMeshProjectionRenderer.h"
 
 #include "DisplayClusterLightCardActor.h"
+#include "Components/DisplayClusterLabelComponent.h"
 
 #include "Editor/UnrealEdTypes.h"
 #include "Engine/Texture.h"
@@ -20,6 +21,7 @@ UDisplayClusterLightCardEditorProjectSettings::UDisplayClusterLightCardEditorPro
 	LightCardTemplateDefaultPath.Path = TEXT("/Game/VP/LightCards");
 	LightCardLabelScale = 1.f;
 	bDisplayLightCardLabels = false;
+	LightCardLabelFlags = EDisplayClusterLabelFlags::DisplayInGame | EDisplayClusterLabelFlags::DisplayInEditor;
 }
 
 const FName FDisplayClusterLightCardEditorRecentItem::Type_LightCard = "LightCard";
diff --git a/Engine/Plugins/Runtime/nDisplay/Source/DisplayClusterLightCardEditor/Private/Settings/DisplayClusterLightCardEditorSettings.h b/Engine/Plugins/Runtime/nDisplay/Source/DisplayClusterLightCardEditor/Private/Settings/DisplayClusterLightCardEditorSettings.h
--- a/Engine/Plugins/Runtime/nDisplay/Source/DisplayClusterLightCardEditor/Private/Settings/DisplayClusterLightCardEditorSettings.h
+++ b/Engine/Plugins/Runtime/nDisplay/Source/DisplayClusterLightCardEditor/Private/Settings/DisplayClusterLightCardEditorSettings.h
@@ -6,6 +6,7 @@
 
 #include "DisplayClusterLightCardEditorSettings.generated.h"
 
+enum class EDisplayClusterLabelFlags : uint8;
 class UDisplayClusterLightCardTemplate;
 struct FSlateBrush;
 
@@ -36,6 +37,10 @@ public:
 	UPROPERTY()
 	bool bDisplayLightCardLabels;
 	
+	/** Flags for the label */
+	UPROPERTY()
+	EDisplayClusterLabelFlags LightCardLabelFlags;
+	
 	/** The scale to use for light card labels */
 	UPROPERTY(config, EditAnywhere, Category = Labels)
 	float LightCardLabelScale;
diff --git a/Engine/Plugins/Runtime/nDisplay/Source/DisplayClusterLightCardEditor/Public/IDisplayClusterLightCardEditor.h b/Engine/Plugins/Runtime/nDisplay/Source/DisplayClusterLightCardEditor/Public/IDisplayClusterLightCardEditor.h
--- a/Engine/Plugins/Runtime/nDisplay/Source/DisplayClusterLightCardEditor/Public/IDisplayClusterLightCardEditor.h
+++ b/Engine/Plugins/Runtime/nDisplay/Source/DisplayClusterLightCardEditor/Public/IDisplayClusterLightCardEditor.h
@@ -3,6 +3,7 @@
 #pragma once
 
 #include "CoreMinimal.h"
+#include "Components/DisplayClusterLabelComponent.h"
 #include "Modules/ModuleInterface.h"
 #include "Modules/ModuleManager.h"
 
@@ -40,24 +41,9 @@ public:
 		return FModuleManager::Get().IsModuleLoaded(ModuleName);
 	}
 
-	struct FLabelArgs
+	struct FLabelArgs : FDisplayClusterLabelConfiguration
 	{
 		// TODO: Specify type (Light Card vs CCR)
-		
-		/**
-		 * [Required] The root actor to display labels for
-		 */
-		class ADisplayClusterRootActor* RootActor = nullptr;
-		
-		/**
-		 * The scale to apply to the label
-		 */
-		float Scale = 1.f;
-
-		/**
-		 * Should the label be visible
-		 */
-		bool bVisible = false;
 	};
 	
 	/**
diff --git a/Engine/Plugins/Runtime/nDisplay/Source/DisplayClusterScenePreview/DisplayClusterScenePreview.Build.cs b/Engine/Plugins/Runtime/nDisplay/Source/DisplayClusterScenePreview/DisplayClusterScenePreview.Build.cs
--- a/Engine/Plugins/Runtime/nDisplay/Source/DisplayClusterScenePreview/DisplayClusterScenePreview.Build.cs
+++ b/Engine/Plugins/Runtime/nDisplay/Source/DisplayClusterScenePreview/DisplayClusterScenePreview.Build.cs
@@ -10,12 +10,12 @@ public class DisplayClusterScenePreview : ModuleRules
 			new string[] {
 				"DisplayClusterLightCardEditorShaders",
 				"DisplayClusterLightCardExtender",
-				"RHI"
+				"RHI",
+				"DisplayCluster"
 			});
 
 		PrivateDependencyModuleNames.AddRange(
 			new string[] {
-				"DisplayCluster",
 				"DisplayClusterConfiguration",
 
 				"Core",
diff --git a/Engine/Plugins/Runtime/nDisplay/Source/DisplayClusterScenePreview/Private/DisplayClusterLightCardEditorHelper.cpp b/Engine/Plugins/Runtime/nDisplay/Source/DisplayClusterScenePreview/Private/DisplayClusterLightCardEditorHelper.cpp
--- a/Engine/Plugins/Runtime/nDisplay/Source/DisplayClusterScenePreview/Private/DisplayClusterLightCardEditorHelper.cpp
+++ b/Engine/Plugins/Runtime/nDisplay/Source/DisplayClusterScenePreview/Private/DisplayClusterLightCardEditorHelper.cpp
@@ -882,7 +882,10 @@ AActor* FDisplayClusterLightCardEditorHelper::SpawnStageActor(const FSpawnActorA
 			}
 		}
 
-		NewLightCard->ShowLightCardLabel(InSpawnArgs.AddLightCardArgs.bShowLabels, InSpawnArgs.AddLightCardArgs.LabelScale, RootActor);
+		FDisplayClusterLabelConfiguration LabelConfiguration = InSpawnArgs.AddLightCardArgs.LabelConfiguration;
+		LabelConfiguration.RootActor = RootActor;
+
+		NewLightCard->ShowLightCardLabel(LabelConfiguration);
 		
 		if (ADisplayClusterChromakeyCardActor* ChromakeyCardActor = Cast<ADisplayClusterChromakeyCardActor>(NewActor))
 		{
diff --git a/Engine/Plugins/Runtime/nDisplay/Source/DisplayClusterScenePreview/Public/DisplayClusterLightCardEditorHelper.h b/Engine/Plugins/Runtime/nDisplay/Source/DisplayClusterScenePreview/Public/DisplayClusterLightCardEditorHelper.h
--- a/Engine/Plugins/Runtime/nDisplay/Source/DisplayClusterScenePreview/Public/DisplayClusterLightCardEditorHelper.h
+++ b/Engine/Plugins/Runtime/nDisplay/Source/DisplayClusterScenePreview/Public/DisplayClusterLightCardEditorHelper.h
@@ -4,6 +4,8 @@
 
 #include "StageActor/DisplayClusterWeakStageActorPtr.h"
 
+#include "Components/DisplayClusterLabelConfiguration.h"
+
 #include "DisplayClusterMeshProjectionRenderer.h"
 
 #include "SceneView.h"
@@ -234,13 +236,18 @@ public:
 
 	struct FAddLightCardArgs
 	{
-		bool bShowLabels;
-		float LabelScale;
+		UE_DEPRECATED(5.3, "Use LabelConfiguration to configure label visibility.")
+		bool bShowLabels = false;
+		UE_DEPRECATED(5.3, "Use LabelConfiguration to configure label scale.")
+		float LabelScale = 1.f;
+
+		/** Label configuration to use for the new light card */
+		FDisplayClusterLabelConfiguration LabelConfiguration;
 
 		FAddLightCardArgs():
-		bShowLabels(false),
-		LabelScale(1.f)
-		{}
+			LabelConfiguration()
+		{
+		}
 	};
 	
 	struct FSpawnActorArgs
