nDisplay: Fix changing from blueprint post process enabled on preview not working.

diff --git a/Engine/Plugins/Runtime/nDisplay/Source/DisplayCluster/Private/Game/EngineClasses/Scene/DisplayClusterRootActor.cpp b/Engine/Plugins/Runtime/nDisplay/Source/DisplayCluster/Private/Game/EngineClasses/Scene/DisplayClusterRootActor.cpp
--- a/Engine/Plugins/Runtime/nDisplay/Source/DisplayCluster/Private/Game/EngineClasses/Scene/DisplayClusterRootActor.cpp
+++ b/Engine/Plugins/Runtime/nDisplay/Source/DisplayCluster/Private/Game/EngineClasses/Scene/DisplayClusterRootActor.cpp
@@ -704,6 +704,20 @@ void ADisplayClusterRootActor::UpdateProceduralMeshComponentData(const UProcedur
 	}
 }
 
+void ADisplayClusterRootActor::SetPreviewEnablePostProcess(const bool bNewPreviewEnablePostProcess)
+{
+#if WITH_EDITOR
+	if (bNewPreviewEnablePostProcess != bPreviewEnablePostProcess)
+	{
+		bPreviewEnablePostProcess = bNewPreviewEnablePostProcess;
+
+		ResetPreviewComponents_Editor(true);
+		PreviewRenderFrame.Reset();
+	}
+#endif // WITH_EDITOR
+}
+
+
 bool ADisplayClusterRootActor::BuildHierarchy()
 {
 	check(CurrentConfigData);
diff --git a/Engine/Plugins/Runtime/nDisplay/Source/DisplayCluster/Public/DisplayClusterRootActor.h b/Engine/Plugins/Runtime/nDisplay/Source/DisplayCluster/Public/DisplayClusterRootActor.h
--- a/Engine/Plugins/Runtime/nDisplay/Source/DisplayCluster/Public/DisplayClusterRootActor.h
+++ b/Engine/Plugins/Runtime/nDisplay/Source/DisplayCluster/Public/DisplayClusterRootActor.h
@@ -177,6 +177,15 @@ public:
 	UFUNCTION(BlueprintCallable, meta = (DisplayName = "Update ProceduralMeshComponent data"), Category = "NDisplay|Components")
 	void UpdateProceduralMeshComponentData(const UProceduralMeshComponent* InProceduralMeshComponent = nullptr);
 
+
+	/**
+	* Blueprint setter for the PreviewEnablePostProcess property. Makes sure preview pipeline is updated properly.
+	*
+	* @param bNewPreviewEnablePostProcess - Desired new property value.
+	*/
+	UFUNCTION(BlueprintSetter)
+	void SetPreviewEnablePostProcess(const bool bNewPreviewEnablePostProcess);
+
 public:
 	IDisplayClusterViewportManager* GetViewportManager() const
 	{
@@ -348,7 +357,7 @@ public:
 	float PreviewRenderTargetRatioMult = 0.25;
 
 	/** Enable PostProcess for preview. */
-	UPROPERTY(EditAnywhere, BlueprintReadWrite, Category = "Editor Preview", meta = (DisplayName = "Enable Post Process", EditCondition = "bPreviewEnable"))
+	UPROPERTY(EditAnywhere, BlueprintReadWrite, Category = "Editor Preview", meta = (DisplayName = "Enable Post Process", EditCondition = "bPreviewEnable"), BlueprintSetter = SetPreviewEnablePostProcess)
 	bool bPreviewEnablePostProcess = false;
 
 	/** Freeze preview render.  This will impact editor performance. */
