OCIO: Ensure the fallback invalid pass is shown when relevant.

The fallback pass was added in CL 23599229 to as a visible display of an invalid config. It was however accidentally disabled by a fix later for UE-175058. This CL further enables it on nDisplay as well.

diff --git a/Engine/Plugins/Compositing/OpenColorIO/Source/OpenColorIO/Private/OpenColorIOColorTransform.cpp b/Engine/Plugins/Compositing/OpenColorIO/Source/OpenColorIO/Private/OpenColorIOColorTransform.cpp
--- a/Engine/Plugins/Compositing/OpenColorIO/Source/OpenColorIO/Private/OpenColorIOColorTransform.cpp
+++ b/Engine/Plugins/Compositing/OpenColorIO/Source/OpenColorIO/Private/OpenColorIOColorTransform.cpp
@@ -544,8 +544,8 @@ bool UOpenColorIOColorTransform::AreRenderResourcesReady() const
 		}
 	}
 
-	// Transform shaders are always required, and the all of the required ones should have finished compiling.
-	if (NumShadersRequired == 0 || NumShadersRequired != NumShadersReady)
+	// All of the required shaders should have finished compiling.
+	if (NumShadersRequired != NumShadersReady)
 	{
 		return false;
 	}
diff --git a/Engine/Plugins/Runtime/nDisplay/Source/DisplayCluster/Private/Render/Viewport/Configuration/DisplayClusterViewportConfigurationHelpers_OpenColorIO.cpp b/Engine/Plugins/Runtime/nDisplay/Source/DisplayCluster/Private/Render/Viewport/Configuration/DisplayClusterViewportConfigurationHelpers_OpenColorIO.cpp
--- a/Engine/Plugins/Runtime/nDisplay/Source/DisplayCluster/Private/Render/Viewport/Configuration/DisplayClusterViewportConfigurationHelpers_OpenColorIO.cpp
+++ b/Engine/Plugins/Runtime/nDisplay/Source/DisplayCluster/Private/Render/Viewport/Configuration/DisplayClusterViewportConfigurationHelpers_OpenColorIO.cpp
@@ -85,24 +85,24 @@ bool FDisplayClusterViewportConfigurationHelpers_OpenColorIO::UpdateChromakeyVie
 	return false;
 }
 
-void FDisplayClusterViewportConfigurationHelpers_OpenColorIO::ApplyOCIOConfiguration(FDisplayClusterViewport& DstViewport, const FOpenColorIOColorConversionSettings& InConfiguration)
+void FDisplayClusterViewportConfigurationHelpers_OpenColorIO::ApplyOCIOConfiguration(FDisplayClusterViewport& DstViewport, const FOpenColorIOColorConversionSettings& InConversionSettings)
 {
 	// Check OCIO before apply
-	if(!InConfiguration.IsValid())
+	if(!InConversionSettings.IsValid())
 	{
 		DisableOCIOConfiguration(DstViewport);
 
 		return;
 	}
 
-	if (DstViewport.OpenColorIO.IsValid() && DstViewport.OpenColorIO->IsDisplayConfigurationEquals(InConfiguration))
+	if (DstViewport.OpenColorIO.IsValid() && DstViewport.OpenColorIO->IsConversionSettingsEqual(InConversionSettings))
 	{
 		// Already assigned
 		return;
 	}
 
 	DstViewport.OpenColorIO.Reset();
-	DstViewport.OpenColorIO = MakeShared<FDisplayClusterViewport_OpenColorIO>(InConfiguration);
+	DstViewport.OpenColorIO = MakeShared<FDisplayClusterViewport_OpenColorIO>(InConversionSettings);
 }
 
 void FDisplayClusterViewportConfigurationHelpers_OpenColorIO::DisableOCIOConfiguration(FDisplayClusterViewport& DstViewport)
diff --git a/Engine/Plugins/Runtime/nDisplay/Source/DisplayCluster/Private/Render/Viewport/Configuration/DisplayClusterViewportConfigurationHelpers_OpenColorIO.h b/Engine/Plugins/Runtime/nDisplay/Source/DisplayCluster/Private/Render/Viewport/Configuration/DisplayClusterViewportConfigurationHelpers_OpenColorIO.h
--- a/Engine/Plugins/Runtime/nDisplay/Source/DisplayCluster/Private/Render/Viewport/Configuration/DisplayClusterViewportConfigurationHelpers_OpenColorIO.h
+++ b/Engine/Plugins/Runtime/nDisplay/Source/DisplayCluster/Private/Render/Viewport/Configuration/DisplayClusterViewportConfigurationHelpers_OpenColorIO.h
@@ -81,11 +81,11 @@ private:
 	/** Apply OCIO to viewport.
 	 *
 	 * @param DstViewport       - Dest viewport object
-	 * @param InConfiguration   - OCIO configuration
+	 * @param InConversionSettings   - OCIO conversion settings
 	 *
 	 * @return - none
 	 */
-	static void ApplyOCIOConfiguration(FDisplayClusterViewport& DstViewport, const FOpenColorIOColorConversionSettings& InConfiguration);
+	static void ApplyOCIOConfiguration(FDisplayClusterViewport& DstViewport, const FOpenColorIOColorConversionSettings& InConversionSettings);
 
 	/** Reset OCIO for viewport.
 	 *
diff --git a/Engine/Plugins/Runtime/nDisplay/Source/DisplayCluster/Private/Render/Viewport/DisplayClusterViewport.cpp b/Engine/Plugins/Runtime/nDisplay/Source/DisplayCluster/Private/Render/Viewport/DisplayClusterViewport.cpp
--- a/Engine/Plugins/Runtime/nDisplay/Source/DisplayCluster/Private/Render/Viewport/DisplayClusterViewport.cpp
+++ b/Engine/Plugins/Runtime/nDisplay/Source/DisplayCluster/Private/Render/Viewport/DisplayClusterViewport.cpp
@@ -109,7 +109,7 @@ bool FDisplayClusterViewport::IsOpenColorIOEquals(const FDisplayClusterViewport&
 			return true;
 		}
 
-		if (OpenColorIO->IsDisplayConfigurationEquals(InViewport.OpenColorIO->GetDisplayConfiguration()))
+		if (OpenColorIO->IsConversionSettingsEqual(InViewport.OpenColorIO->GetConversionSettings()))
 		{
 			return true;
 		}
diff --git a/Engine/Plugins/Runtime/nDisplay/Source/DisplayCluster/Private/Render/Viewport/DisplayClusterViewport_OpenColorIO.cpp b/Engine/Plugins/Runtime/nDisplay/Source/DisplayCluster/Private/Render/Viewport/DisplayClusterViewport_OpenColorIO.cpp
--- a/Engine/Plugins/Runtime/nDisplay/Source/DisplayCluster/Private/Render/Viewport/DisplayClusterViewport_OpenColorIO.cpp
+++ b/Engine/Plugins/Runtime/nDisplay/Source/DisplayCluster/Private/Render/Viewport/DisplayClusterViewport_OpenColorIO.cpp
@@ -21,7 +21,7 @@
 // FDisplayClusterViewport_OpenColorIO
 //////////////////////////////////////////////////////////////////////////
 FDisplayClusterViewport_OpenColorIO::FDisplayClusterViewport_OpenColorIO(const FOpenColorIOColorConversionSettings& InDisplayConfiguration)
-	: DisplayConfiguration(InDisplayConfiguration)
+	: ConversionSettings(InDisplayConfiguration)
 { }
 
 FDisplayClusterViewport_OpenColorIO::~FDisplayClusterViewport_OpenColorIO()
@@ -46,7 +46,9 @@ void FDisplayClusterViewport_OpenColorIO::SetupSceneView(FSceneViewFamily& InOut
 void FDisplayClusterViewport_OpenColorIO::UpdateOpenColorIORenderPassResources()
 {
 	bShaderResourceValid = false;
-	if (!bConfigurationDataValid)
+	bIsEnabled = ConversionSettings.IsValid() ? ConversionSettings.ConfigurationSource->IsTransformReady(ConversionSettings) : false;
+
+	if (!bIsEnabled)
 	{
 		// The data in this configuration was not correctly initialized in the previous frame. Ignore this OCIO
 		return;
@@ -55,13 +57,13 @@ void FDisplayClusterViewport_OpenColorIO::UpdateOpenColorIORenderPassResources()
 	FOpenColorIOTransformResource* ShaderResource = nullptr;
 	TSortedMap<int32, FTextureResource*> TransformTextureResources;
 
-	if (DisplayConfiguration.ConfigurationSource != nullptr)
+	if (ConversionSettings.ConfigurationSource != nullptr)
 	{
 		const ERHIFeatureLevel::Type FeatureLevel = GMaxRHIFeatureLevel;
 
-		const bool bFoundTransform = DisplayConfiguration.ConfigurationSource->GetRenderResources(
+		const bool bFoundTransform = ConversionSettings.ConfigurationSource->GetRenderResources(
 			FeatureLevel
-			, DisplayConfiguration
+			, ConversionSettings
 			, ShaderResource
 			, TransformTextureResources);
 
@@ -75,9 +77,6 @@ void FDisplayClusterViewport_OpenColorIO::UpdateOpenColorIORenderPassResources()
 
 				//Invalidate shader resource
 				ShaderResource = nullptr;
-
-				// Data in this configuration is corrupted. Ignore this configuration for next frames
-				bConfigurationDataValid = false;
 			}
 			else
 			{
@@ -104,12 +103,12 @@ void FDisplayClusterViewport_OpenColorIO::UpdateOpenColorIORenderPassResources()
 
 bool FDisplayClusterViewport_OpenColorIO::IsEnabled_RenderThread() const
 {
-	return CachedResourcesRenderThread.ShaderResource != nullptr;
+	return bIsEnabled;
 }
 
-bool FDisplayClusterViewport_OpenColorIO::IsDisplayConfigurationEquals(const FOpenColorIOColorConversionSettings& InDisplayConfiguration) const
+bool FDisplayClusterViewport_OpenColorIO::IsConversionSettingsEqual(const FOpenColorIOColorConversionSettings& InConversionSettings) const
 {
-	return DisplayConfiguration.ToString() == InDisplayConfiguration.ToString();
+	return ConversionSettings.ToString() == InConversionSettings.ToString();
 }
 
 bool FDisplayClusterViewport_OpenColorIO::AddPass_RenderThread(FRDGBuilder& GraphBuilder, const FDisplayClusterViewport_Context& InViewportContext,
diff --git a/Engine/Plugins/Runtime/nDisplay/Source/DisplayCluster/Private/Render/Viewport/DisplayClusterViewport_OpenColorIO.h b/Engine/Plugins/Runtime/nDisplay/Source/DisplayCluster/Private/Render/Viewport/DisplayClusterViewport_OpenColorIO.h
--- a/Engine/Plugins/Runtime/nDisplay/Source/DisplayCluster/Private/Render/Viewport/DisplayClusterViewport_OpenColorIO.h
+++ b/Engine/Plugins/Runtime/nDisplay/Source/DisplayCluster/Private/Render/Viewport/DisplayClusterViewport_OpenColorIO.h
@@ -55,16 +55,16 @@ public:
 
 	/** Compare two OCIO configurations.
 	 *
-	 * @param InDisplayConfiguration - configuration to compare with.
+	 * @param InConversionSettings - configuration to compare with.
 	 *
 	 * @return - true if equal.
 	 */
-	bool IsDisplayConfigurationEquals(const FOpenColorIOColorConversionSettings& InDisplayConfiguration) const;
+	bool IsConversionSettingsEqual(const FOpenColorIOColorConversionSettings& InConversionSettings) const;
 
-	/** Get current OCIO configuration. */
-	const FOpenColorIOColorConversionSettings& GetDisplayConfiguration() const
+	/** Get current OCIO conversion settings. */
+	const FOpenColorIOColorConversionSettings& GetConversionSettings() const
 	{
-		return DisplayConfiguration;
+		return ConversionSettings;
 	}
 
 private:
@@ -72,11 +72,11 @@ private:
 	FOpenColorIORenderPassResources CachedResourcesRenderThread;
 
 	/** Configuration to apply during post render callback. */
-	FOpenColorIOColorConversionSettings DisplayConfiguration;
+	FOpenColorIOColorConversionSettings ConversionSettings;
 
 	/** Shader resources state. */
 	bool bShaderResourceValid = true;
 
 	/** Configuration data state. */
-	bool bConfigurationDataValid = true;
+	TAtomic<bool> bIsEnabled = true;
 };
