[nDisplay] Fixed traces flooding with vblank events

diff --git a/Engine/Plugins/Runtime/nDisplay/Source/DisplayCluster/Private/Render/Monitoring/DisplayClusterVblankMonitor.cpp b/Engine/Plugins/Runtime/nDisplay/Source/DisplayCluster/Private/Render/Monitoring/DisplayClusterVblankMonitor.cpp
--- a/Engine/Plugins/Runtime/nDisplay/Source/DisplayCluster/Private/Render/Monitoring/DisplayClusterVblankMonitor.cpp
+++ b/Engine/Plugins/Runtime/nDisplay/Source/DisplayCluster/Private/Render/Monitoring/DisplayClusterVblankMonitor.cpp
@@ -32,7 +32,12 @@ uint32 FDisplayClusterVBlankMonitor::Run()
 		if (SyncHelper.IsWaitForVBlankSupported())
 		{
 			// Wait for V-blank
-			SyncHelper.WaitForVBlank();
+			if (!SyncHelper.WaitForVBlank())
+			{
+				// Nothing to do if the underlying output subsystem is not ready for monitoring yet
+				FPlatformProcess::SleepNoStats(1E-3);
+				continue;
+			}
 
 			// Get current timestamp
 			const double CurrentTime = FPlatformTime::Seconds();
@@ -71,8 +76,8 @@ uint32 FDisplayClusterVBlankMonitor::Run()
 		}
 		else
 		{
-			TRACE_CPUPROFILER_EVENT_SCOPE(VBLANK - not available);
-			FPlatformProcess::SleepNoStats(1.f);
+			// No need to keep the monitoring thread alive as v-blank monitoring is not supported
+			break;
 		}
 	}
 
