nDisplay: added new feature - inner frustum overlap
- The overlapping area of two or more internal frustums is shown as a chromakey.
- Reworked ICVFX multi-pass render architecture.

diff --git a/Engine/Plugins/Runtime/nDisplay/Shaders/Private/MPCDIOverlayShaders.usf b/Engine/Plugins/Runtime/nDisplay/Shaders/Private/MPCDIOverlayShaders.usf
--- a/Engine/Plugins/Runtime/nDisplay/Shaders/Private/MPCDIOverlayShaders.usf
+++ b/Engine/Plugins/Runtime/nDisplay/Shaders/Private/MPCDIOverlayShaders.usf
@@ -64,7 +64,54 @@ float4 ComposeRenderPassTwoOverlays(float4 InOverlay1, float4 InOverlay2)
 ////////////////////////////////////////////////////////////////////////////////
 float2 TileUV(float2 InUV, float InScale)
 {
-    return frac(InUV * InScale);
+	return frac(InUV * InScale);
+}
+
+float GetFeatherAlpha(float2 EdgeOffset, float Feather)
+{
+	float Alpha = 1;
+
+	Alpha *= smoothstep(0, Feather, EdgeOffset.x); // horizontal component
+	Alpha *= smoothstep(0, Feather, EdgeOffset.y); // vertical component
+
+	return Alpha;
+}
+
+float GetInnerCameraSoftEdgeAlpha(float2 CameraUV, float4 InInnerCameraSoftEdge)
+{
+
+	float2 EdgeOffset = float2(1,1);
+
+	// InInnerCameraSoftEdge.x values are used to pass left/right soft edge values
+	if(CameraUV.x < InInnerCameraSoftEdge.x) // left 
+	{
+		float CameraTextureAlphaLeft = Pow2(saturate(abs(CameraUV.x) / InInnerCameraSoftEdge.x));
+		EdgeOffset.x = CameraTextureAlphaLeft;
+	}
+	else
+	if(CameraUV.x > (1 - InInnerCameraSoftEdge.x)) // right
+	{
+		float CameraTextureAlphaRight = Pow2(saturate(abs(1 - CameraUV.x) / InInnerCameraSoftEdge.x));
+		EdgeOffset.x = CameraTextureAlphaRight;
+	}
+
+	// InInnerCameraSoftEdge.y values are used to pass top/bottom soft edge values
+	if(CameraUV.y < (InInnerCameraSoftEdge.y)) // top
+	{
+		float CameraTextureAlphaTop = Pow2(saturate(abs(CameraUV.y) / InInnerCameraSoftEdge.y));
+		EdgeOffset.y = CameraTextureAlphaTop;
+	}
+	else
+	if(CameraUV.y > (1 - InInnerCameraSoftEdge.y)) // bottom
+	{
+		float CameraTextureAlphaBottom = Pow2(saturate(abs(1 - CameraUV.y) / InInnerCameraSoftEdge.y));
+		EdgeOffset.y = CameraTextureAlphaBottom;
+	}
+
+	// InInnerCameraSoftEdge.z are used to pass Feather soft edge values
+	float Feather = InInnerCameraSoftEdge.z;
+
+	return GetFeatherAlpha(EdgeOffset, Feather);
 }
 
 float4 GetChromakeyMarkerColor(float2 InMarkersUV)
@@ -107,77 +154,116 @@ float4 GetInnerCameraColor(float4 WarpedUV, OutputVS IN)
 	float4 CameraColor = float4(CameraBaseColor, 0);
 #endif
 
-	float2 ToEdge = 0;
-	float Weight = 0;
-
-	// *** Soft Edges **
+	// ********* Compose camera alpha **********
 	if (CameraUVW.w > 0) // clip back plane
-	{//Defined texel:
-		ToEdge = (CameraUV.xy * 2) - 1.0f; // -1..1
-		Weight = 1 - max(abs(ToEdge.x), abs(ToEdge.y));
-		// Near clip Plane tests
-		//! UV.z <= UV.w &&
+	{
+		//Defined texel:
+		float2 ToEdge = (CameraUV.xy * 2) - 1.0f; // -1..1
+		float Weight = 1 - max(abs(ToEdge.x), abs(ToEdge.y));
+
+		// Near clip Plane tests 
 		if (Weight >= 0)
 		{
-			CameraColor.w = 1;
-			
-			//InnerCameraSoftEdge.xyzw values are used to pass left/top/right/bottom soft edge values
-			if(CameraUV.x < InnerCameraSoftEdge.x) // left 
-			{
-				float CameraTextureAlphaLeft = Pow2(saturate(abs(CameraUV.x) / InnerCameraSoftEdge.x));
-				CameraColor.w *= CameraTextureAlphaLeft;
-			}
-			else
-			if(CameraUV.x > (1 - InnerCameraSoftEdge.z)) // right
-			{
-				float CameraTextureAlphaRight = Pow2(saturate(abs(1 - CameraUV.x) / InnerCameraSoftEdge.z));
-				CameraColor.w *= CameraTextureAlphaRight;
-			}
+			// *** InCamera Soft Edges ***
+			CameraColor.w = GetInnerCameraSoftEdgeAlpha(CameraUV, InnerCameraSoftEdge);
 
-			
-			if(CameraUV.y < (InnerCameraSoftEdge.y)) // top
+			// **** disable incamera render inside overlapped areas
+			#if !INNER_CAMERA_OVERLAP
+			if(OverlappedInnerCamerasNum > 0 && CameraColor.w < 1)
 			{
-				float CameraTextureAlphaTop = Pow2(saturate(abs(CameraUV.y) / InnerCameraSoftEdge.y));
-				CameraColor.w *= CameraTextureAlphaTop;
+				for(uint OverlappedCameraIndex = 0; OverlappedCameraIndex < OverlappedInnerCamerasNum; OverlappedCameraIndex++)
+				{
+					float4x4 InnerCameraProjectionMatrix2 = OverlappedInnerCamerasProjectionMatrices[OverlappedCameraIndex];
+
+					// Transform WarpedUV to Camera ScreenSpaceUV
+					float4 CameraUVW2 = mul(WarpedUV, InnerCameraProjectionMatrix2);
+					if (CameraUVW2.w > 0) // clip back plane
+					{
+						float2 CameraUV2 = CameraUVW2.xy / CameraUVW2.w;
+						float2 ToEdge2 = (CameraUV2.xy * 2) - 1.0f; // -1..1
+						float Weight2 = 1 - max(abs(ToEdge2.x), abs(ToEdge2.y));
+						if (Weight2 >= 0)
+						{
+							float UnderCameraAlpha = GetInnerCameraSoftEdgeAlpha(CameraUV2, OverlappedInnerCameraSoftEdges[OverlappedCameraIndex]);
+
+							//!
+							if(UnderCameraAlpha > 0.1)
+							{
+								// This pixel was overlapped. Do not render
+								CameraColor.w = 0;
+							}
+						}
+					}
+				}
 			}
-			else
-			if(CameraUV.y > (1 - InnerCameraSoftEdge.w)) // bottom
-			{
-				float CameraTextureAlphaBottom = Pow2(saturate(abs(1 - CameraUV.y) / InnerCameraSoftEdge.w));
-				CameraColor.w *= CameraTextureAlphaBottom;
-			}
-		}
-	}
-
-#if CHROMAKEYFRAMECOLOR || CHROMAKEY
-	// Use frame color by default
-	float4 OverlayChromakeyColor = float4(ChromakeyColor.xyz,0);
-
-#if CHROMAKEY
-	float4 ChromakeyCameraTextureColor = ChromakeyCameraTexture.Sample(ChromakeyCameraSampler, CameraUV);
-
-	// Get alpha from chromakey rtt frame
-	OverlayChromakeyColor.w = ChromakeyCameraTextureColor.w;
-#endif
-
-#if CHROMAKEY_MARKER
-	float4 ChromakeyMarkerTextureColor = GetChromakeyMarkerColor(IN.UV_Chromakey.xy);
+			#endif
 
-	// Blend marker color with chromakey
-	OverlayChromakeyColor.xyz = lerp(OverlayChromakeyColor.xyz, ChromakeyMarkerTextureColor.xyz, ChromakeyMarkerTextureColor.w);
-#endif
+			// ********* Compose camera colors **********
+			#if CHROMAKEYFRAMECOLOR || CHROMAKEY
+				// Use frame color by default
+				float4 OverlayChromakeyColor = float4(ChromakeyColor.xyz,0);
+
+				#if CHROMAKEY
+					float4 ChromakeyCameraTextureColor = ChromakeyCameraTexture.Sample(ChromakeyCameraSampler, CameraUV);
+
+					// Get alpha from chromakey rtt frame
+					OverlayChromakeyColor.w = ChromakeyCameraTextureColor.w;
+				#endif
+
+				#if CHROMAKEY_MARKER
+					float4 ChromakeyMarkerTextureColor = GetChromakeyMarkerColor(IN.UV_Chromakey.xy);
+
+					// Blend marker color with chromakey
+					OverlayChromakeyColor.xyz = lerp(OverlayChromakeyColor.xyz, ChromakeyMarkerTextureColor.xyz, ChromakeyMarkerTextureColor.w);
+				#endif
+
+				// ********* When camera overlap we render chromakey **********
+				#if INNER_CAMERA_OVERLAP
+					if(CameraColor.w >= 0)
+					{
+						OverlayChromakeyColor.w = 0;
+
+						// disable chromakey
+						float OverlayAlphaMult = 0;
+
+						for(uint OverlappedCameraIndex = 0; OverlappedCameraIndex < OverlappedInnerCamerasNum; OverlappedCameraIndex++)
+						{
+							float4x4 InnerCameraProjectionMatrix2 = OverlappedInnerCamerasProjectionMatrices[OverlappedCameraIndex];
+
+							// Transform WarpedUV to Camera ScreenSpaceUV
+							float4 CameraUVW2 = mul(WarpedUV, InnerCameraProjectionMatrix2);
+							if (CameraUVW2.w > 0) // clip back plane
+							{
+								float2 CameraUV2 = CameraUVW2.xy / CameraUVW2.w;
+								float2 ToEdge2 = (CameraUV2.xy * 2) - 1.0f; // -1..1
+								float Weight2 = 1 - max(abs(ToEdge2.x), abs(ToEdge2.y));
+								if (Weight2 >= 0)
+								{
+									// This pixel was overlapped
+									// force to render Chromakey
+									OverlayAlphaMult = max(OverlayAlphaMult, GetInnerCameraSoftEdgeAlpha(CameraUV2, OverlappedInnerCameraSoftEdges[OverlappedCameraIndex]));
+								}
+							}
+						}
+
+						CameraColor.w = CameraColor.w * OverlayAlphaMult;
+					}
+				#endif
+
+				CameraColor.xyz = lerp(CameraColor.xyz, OverlayChromakeyColor.xyz, 1-OverlayChromakeyColor.w);
 
-	CameraColor.xyz = lerp(CameraColor.xyz, OverlayChromakeyColor.xyz, 1-OverlayChromakeyColor.w);
-#endif
+			#endif
 
-	// apply border
-	if((InnerCameraBorderThickness > 0) && (CameraUVW.w > 0) && (Weight >= 0))
-	{
-		if((ToEdge.x < (InnerCameraBorderThickness / InnerCameraFrameAspectRatio - 1)) || (ToEdge.x > (1 - InnerCameraBorderThickness / InnerCameraFrameAspectRatio)) ||
-		   (ToEdge.y < (InnerCameraBorderThickness - 1)) || (ToEdge.y > (1 - InnerCameraBorderThickness)))
-		{
-			CameraColor.rgb = InnerCameraBorderColor.xyz;
-			CameraColor.w = 1.0f;
+			// *** InCamera Border ***
+			if(InnerCameraBorderThickness>0)
+			{
+				if((ToEdge.x < (InnerCameraBorderThickness / InnerCameraFrameAspectRatio - 1)) || (ToEdge.x > (1 - InnerCameraBorderThickness / InnerCameraFrameAspectRatio)) ||
+					(ToEdge.y < (InnerCameraBorderThickness - 1)) || (ToEdge.y > (1 - InnerCameraBorderThickness)))
+				{
+					CameraColor.rgb = InnerCameraBorderColor.xyz;
+					CameraColor.w = 1.0f;
+				}
+			}
 		}
 	}
 
diff --git a/Engine/Plugins/Runtime/nDisplay/Shaders/Private/MPCDIUtils.ush b/Engine/Plugins/Runtime/nDisplay/Shaders/Private/MPCDIUtils.ush
--- a/Engine/Plugins/Runtime/nDisplay/Shaders/Private/MPCDIUtils.ush
+++ b/Engine/Plugins/Runtime/nDisplay/Shaders/Private/MPCDIUtils.ush
@@ -39,6 +39,13 @@ float4x4 ViewportTextureProjectionMatrix;
 float4x4 OverlayProjectionMatrix;
 float4x4 InnerCameraProjectionMatrix;
 
+// The total number of InCameras, which are rendered below.
+int OverlappedInnerCamerasNum;
+
+// Matrices and SoftEdges for InCameras, which are rendered below.
+float4x4 OverlappedInnerCamerasProjectionMatrices[MAX_INNER_CAMERAS_AMOUNT];
+float4 OverlappedInnerCameraSoftEdges[MAX_INNER_CAMERAS_AMOUNT];
+
 float AlphaEmbeddedGamma;
 
 float4 InnerCameraSoftEdge;
@@ -143,4 +150,4 @@ float4 ApplyBlending(float3 color, OutputVS IN, float ColorAlphaValue)
 #else //NOT ALPHAMAP_BLENDING
 	return float4(color, ColorAlphaValue);
 #endif // ALPHAMAP_BLENDING
-}
\ No newline at end of file
+}
diff --git a/Engine/Plugins/Runtime/nDisplay/Source/DisplayCluster/Private/Render/Viewport/Configuration/DisplayClusterViewportConfigurationHelpers.cpp b/Engine/Plugins/Runtime/nDisplay/Source/DisplayCluster/Private/Render/Viewport/Configuration/DisplayClusterViewportConfigurationHelpers.cpp
--- a/Engine/Plugins/Runtime/nDisplay/Source/DisplayCluster/Private/Render/Viewport/Configuration/DisplayClusterViewportConfigurationHelpers.cpp
+++ b/Engine/Plugins/Runtime/nDisplay/Source/DisplayCluster/Private/Render/Viewport/Configuration/DisplayClusterViewportConfigurationHelpers.cpp
@@ -276,14 +276,14 @@ void FDisplayClusterViewportConfigurationHelpers::UpdateBaseViewportSetting(FDis
 			case EDisplayClusterConfigurationICVFX_OverrideLightcardRenderMode::Over:
 
 				EnumAddFlags(TargetFlags, EDisplayClusterViewportICVFXFlags::OverrideLightcardMode);
-				DstViewport.RenderSettingsICVFX.ICVFX.LightcardMode = EDisplayClusterShaderParametersICVFX_LightcardRenderMode::Over;
+				DstViewport.RenderSettingsICVFX.ICVFX.LightCardMode = EDisplayClusterShaderParametersICVFX_LightCardRenderMode::Over;
 				break;
 
 				// Over lightcard over incamera frame  for this viewport
 			case EDisplayClusterConfigurationICVFX_OverrideLightcardRenderMode::Under:
 
 				EnumAddFlags(TargetFlags, EDisplayClusterViewportICVFXFlags::OverrideLightcardMode);
-				DstViewport.RenderSettingsICVFX.ICVFX.LightcardMode = EDisplayClusterShaderParametersICVFX_LightcardRenderMode::Under;
+				DstViewport.RenderSettingsICVFX.ICVFX.LightCardMode = EDisplayClusterShaderParametersICVFX_LightCardRenderMode::Under;
 				break;
 
 			case EDisplayClusterConfigurationICVFX_OverrideLightcardRenderMode::Disabled:
diff --git a/Engine/Plugins/Runtime/nDisplay/Source/DisplayCluster/Private/Render/Viewport/Configuration/DisplayClusterViewportConfigurationHelpers_ICVFX.cpp b/Engine/Plugins/Runtime/nDisplay/Source/DisplayCluster/Private/Render/Viewport/Configuration/DisplayClusterViewportConfigurationHelpers_ICVFX.cpp
--- a/Engine/Plugins/Runtime/nDisplay/Source/DisplayCluster/Private/Render/Viewport/Configuration/DisplayClusterViewportConfigurationHelpers_ICVFX.cpp
+++ b/Engine/Plugins/Runtime/nDisplay/Source/DisplayCluster/Private/Render/Viewport/Configuration/DisplayClusterViewportConfigurationHelpers_ICVFX.cpp
@@ -489,8 +489,10 @@ FDisplayClusterShaderParameters_ICVFX::FCameraSettings FDisplayClusterViewportCo
 	// remap values from 0-1 GUI range into acceptable 0.0 - 0.25 shader range
 	Result.SoftEdge.X = FMath::GetMappedRangeValueClamped(FVector2D(0.0, 1.0f), FVector2D(0.0, 0.25), CameraSettings.SoftEdge.Horizontal) / Overscan; // Left
 	Result.SoftEdge.Y = FMath::GetMappedRangeValueClamped(FVector2D(0.0, 1.0f), FVector2D(0.0, 0.25), CameraSettings.SoftEdge.Vertical) / Overscan; // Top
-	Result.SoftEdge.Z = FMath::GetMappedRangeValueClamped(FVector2D(0.0, 1.0f), FVector2D(0.0, 0.25), CameraSettings.SoftEdge.Horizontal) / Overscan; // right
-	Result.SoftEdge.W = FMath::GetMappedRangeValueClamped(FVector2D(0.0, 1.0f), FVector2D(0.0, 0.25), CameraSettings.SoftEdge.Vertical) / Overscan; // bottom
+
+	// ZW now used in other way
+	// Z for new parameter Feather
+	Result.SoftEdge.Z = CameraSettings.SoftEdge.Feather;
 
 	// default - percents
 	const float ConvertToPercent = 0.01f;
@@ -712,6 +714,8 @@ void FDisplayClusterViewportConfigurationHelpers_ICVFX::UpdateChromakeyViewportS
 
 void FDisplayClusterViewportConfigurationHelpers_ICVFX::UpdateCameraSettings_Chromakey(FDisplayClusterShaderParameters_ICVFX::FCameraSettings& InOutCameraSettings, const FDisplayClusterConfigurationICVFX_ChromakeySettings& InChromakeySettings, FDisplayClusterViewport* InChromakeyViewport)
 {
+	InOutCameraSettings.ChromakeyColor = InChromakeySettings.ChromakeyColor;
+
 	if (InChromakeySettings.bEnable)
 	{
 		if (InChromakeySettings.ChromakeyRenderTexture.bEnable)
@@ -727,13 +731,11 @@ void FDisplayClusterViewportConfigurationHelpers_ICVFX::UpdateCameraSettings_Chr
 
 				InOutCameraSettings.ChromakeySource = EDisplayClusterShaderParametersICVFX_ChromakeySource::ChromakeyLayers;
 				InOutCameraSettings.Chromakey.ViewportId = InChromakeyViewport->GetId();
-				InOutCameraSettings.ChromakeyColor = InChromakeySettings.ChromakeyColor;
 			}
 		}
 		else
 		{
 			InOutCameraSettings.ChromakeySource = EDisplayClusterShaderParametersICVFX_ChromakeySource::FrameColor;
-			InOutCameraSettings.ChromakeyColor = InChromakeySettings.ChromakeyColor;
 		}
 	}
 	else
@@ -853,10 +855,10 @@ void FDisplayClusterViewportConfigurationHelpers_ICVFX::UpdateLightcardViewportS
 		{
 		default:
 		case EDisplayClusterConfigurationICVFX_LightcardRenderMode::Over:
-			BaseViewport.RenderSettingsICVFX.ICVFX.LightcardMode = EDisplayClusterShaderParametersICVFX_LightcardRenderMode::Over;
+			BaseViewport.RenderSettingsICVFX.ICVFX.LightCardMode = EDisplayClusterShaderParametersICVFX_LightCardRenderMode::Over;
 			break;
 		case EDisplayClusterConfigurationICVFX_LightcardRenderMode::Under:
-			BaseViewport.RenderSettingsICVFX.ICVFX.LightcardMode = EDisplayClusterShaderParametersICVFX_LightcardRenderMode::Under;
+			BaseViewport.RenderSettingsICVFX.ICVFX.LightCardMode = EDisplayClusterShaderParametersICVFX_LightCardRenderMode::Under;
 			break;
 		};
 	}
@@ -870,13 +872,13 @@ void FDisplayClusterViewportConfigurationHelpers_ICVFX::UpdateLightcardViewportS
 		}
 		else
 		{
-			BaseViewport.RenderSettingsICVFX.ICVFX.Lightcard.ViewportId = DstViewport.GetId();
+			BaseViewport.RenderSettingsICVFX.ICVFX.LightCard.ViewportId = DstViewport.GetId();
 		}
 	}
 
 	if (EnumHasAnyFlags(DstViewport.RenderSettingsICVFX.RuntimeFlags, EDisplayClusterViewportRuntimeICVFXFlags::UVLightcard))
 	{
-		BaseViewport.RenderSettingsICVFX.ICVFX.UVLightcard.ViewportId = DstViewport.GetId();
+		BaseViewport.RenderSettingsICVFX.ICVFX.UVLightCard.ViewportId = DstViewport.GetId();
 	}
 }
 
diff --git a/Engine/Plugins/Runtime/nDisplay/Source/DisplayCluster/Private/Render/Viewport/Configuration/DisplayClusterViewportConfigurationICVFX.cpp b/Engine/Plugins/Runtime/nDisplay/Source/DisplayCluster/Private/Render/Viewport/Configuration/DisplayClusterViewportConfigurationICVFX.cpp
--- a/Engine/Plugins/Runtime/nDisplay/Source/DisplayCluster/Private/Render/Viewport/Configuration/DisplayClusterViewportConfigurationICVFX.cpp
+++ b/Engine/Plugins/Runtime/nDisplay/Source/DisplayCluster/Private/Render/Viewport/Configuration/DisplayClusterViewportConfigurationICVFX.cpp
@@ -418,6 +418,10 @@ void FDisplayClusterViewportConfigurationICVFX::Update()
 				{
 					// Sort cameras by render order for each target
 					TargetIt->RenderSettingsICVFX.ICVFX.SortCamerasRenderOrder();
+
+					// Setup incamera overlap mode
+					TargetIt->RenderSettingsICVFX.ICVFX.CameraOverlappingRenderMode = StageSettings.bEnableInnerFrustumChromakeyOverlap ?
+						EDisplayClusterShaderParametersICVFX_CameraOverlappingRenderMode::FinalPass : EDisplayClusterShaderParametersICVFX_CameraOverlappingRenderMode::None;
 				}
 			}
 
diff --git a/Engine/Plugins/Runtime/nDisplay/Source/DisplayCluster/Private/Render/Viewport/DisplayClusterViewportManagerProxy.cpp b/Engine/Plugins/Runtime/nDisplay/Source/DisplayCluster/Private/Render/Viewport/DisplayClusterViewportManagerProxy.cpp
--- a/Engine/Plugins/Runtime/nDisplay/Source/DisplayCluster/Private/Render/Viewport/DisplayClusterViewportManagerProxy.cpp
+++ b/Engine/Plugins/Runtime/nDisplay/Source/DisplayCluster/Private/Render/Viewport/DisplayClusterViewportManagerProxy.cpp
@@ -174,6 +174,8 @@ void FDisplayClusterViewportManagerProxy::DeleteViewport_RenderThread(const TSha
 	}
 }
 
+DECLARE_GPU_STAT_NAMED(nDisplay_ViewportManager_UpdateRenderFrameSettings, TEXT("nDisplay ViewportManager::UpdateRenderFrameSettings"));
+
 void FDisplayClusterViewportManagerProxy::ImplUpdateRenderFrameSettings(const FDisplayClusterRenderFrameSettings& InRenderFrameSettings,
 	const TSharedPtr<FDisplayClusterViewportManagerViewExtension, ESPMode::ThreadSafe>& InViewportManagerViewExtension)
 {
@@ -185,6 +187,9 @@ void FDisplayClusterViewportManagerProxy::ImplUpdateRenderFrameSettings(const FD
 	ENQUEUE_RENDER_COMMAND(DisplayClusterUpdateRenderFrameSettings)(
 		[ViewportManagerProxy = SharedThis(this), Settings, InViewportManagerViewExtension](FRHICommandListImmediate& RHICmdList)
 	{
+		SCOPED_GPU_STAT(RHICmdList, nDisplay_ViewportManager_UpdateRenderFrameSettings);
+		SCOPED_DRAW_EVENT(RHICmdList, nDisplay_ViewportManager_UpdateRenderFrameSettings);
+
 		ViewportManagerProxy->RenderFrameSettings = *Settings;
 		ViewportManagerProxy->ViewportManagerViewExtension = InViewportManagerViewExtension;
 		delete Settings;
@@ -194,6 +199,8 @@ void FDisplayClusterViewportManagerProxy::ImplUpdateRenderFrameSettings(const FD
 	});
 }
 
+DECLARE_GPU_STAT_NAMED(nDisplay_ViewportManager_UpdateViewports, TEXT("nDisplay ViewportManager::UpdateViewports"));
+
 void FDisplayClusterViewportManagerProxy::ImplUpdateViewports(const TArray<FDisplayClusterViewport*>& InViewports)
 {
 	check(IsInGameThread());
@@ -208,6 +215,9 @@ void FDisplayClusterViewportManagerProxy::ImplUpdateViewports(const TArray<FDisp
 	ENQUEUE_RENDER_COMMAND(DisplayClusterUpdateViewports)(
 		[ProxiesData = std::move(ViewportProxiesData)](FRHICommandListImmediate& RHICmdList)
 	{
+		SCOPED_GPU_STAT(RHICmdList, nDisplay_ViewportManager_UpdateViewports);
+		SCOPED_DRAW_EVENT(RHICmdList, nDisplay_ViewportManager_UpdateViewports);
+
 		for (FDisplayClusterViewportProxyData* It : ProxiesData)
 		{
 			It->UpdateProxy_RenderThread();
@@ -218,6 +228,9 @@ void FDisplayClusterViewportManagerProxy::ImplUpdateViewports(const TArray<FDisp
 
 
 DECLARE_GPU_STAT_NAMED(nDisplay_ViewportManager_RenderFrame, TEXT("nDisplay ViewportManager::RenderFrame"));
+DECLARE_GPU_STAT_NAMED(nDisplay_ViewportManager_CrossGPUTransfer, TEXT("nDisplay ViewportManager::CrossGPUTransfer"));
+DECLARE_GPU_STAT_NAMED(nDisplay_ViewportManager_UpdateDeferredResources, TEXT("nDisplay ViewportManager::UpdateDeferredResources"));
+DECLARE_GPU_STAT_NAMED(nDisplay_ViewportManager_WarpBlend, TEXT("nDisplay ViewportManager::WarpBlend"));
 
 void FDisplayClusterViewportManagerProxy::ImplRenderFrame(FViewport* InViewport)
 {
@@ -239,6 +252,9 @@ void FDisplayClusterViewportManagerProxy::ImplRenderFrame(FViewport* InViewport)
 	ENQUEUE_RENDER_COMMAND(DisplayClusterRenderFrame_CrossGPUTransfer)(
 		[InViewportManagerProxy = SharedThis(this), InViewport](FRHICommandListImmediate& RHICmdList)
 	{
+		SCOPED_GPU_STAT(RHICmdList, nDisplay_ViewportManager_CrossGPUTransfer);
+		SCOPED_DRAW_EVENT(RHICmdList, nDisplay_ViewportManager_CrossGPUTransfer);
+
 		const FDisplayClusterViewportManagerProxy* ViewportManagerProxy = &InViewportManagerProxy.Get();
 
 		// mGPU not used for in-editor rendering
@@ -258,6 +274,9 @@ void FDisplayClusterViewportManagerProxy::ImplRenderFrame(FViewport* InViewport)
 	ENQUEUE_RENDER_COMMAND(DisplayClusterRenderFrame_UpdateDeferredResources)(
 		[ViewportManagerProxy = SharedThis(this), InViewport](FRHICommandListImmediate& RHICmdList)
 	{
+		SCOPED_GPU_STAT(RHICmdList, nDisplay_ViewportManager_UpdateDeferredResources);
+		SCOPED_DRAW_EVENT(RHICmdList, nDisplay_ViewportManager_UpdateDeferredResources);
+
 		// Update viewports resources: vp/texture overlay, OCIO, blur, nummips, etc
 		ViewportManagerProxy->UpdateDeferredResources_RenderThread(RHICmdList);
 	});
@@ -265,6 +284,9 @@ void FDisplayClusterViewportManagerProxy::ImplRenderFrame(FViewport* InViewport)
 	ENQUEUE_RENDER_COMMAND(DisplayClusterRenderFrame_WarpBlend)(
 		[InViewportManagerProxy = SharedThis(this), InViewport](FRHICommandListImmediate& RHICmdList)
 	{
+		SCOPED_GPU_STAT(RHICmdList, nDisplay_ViewportManager_WarpBlend);
+		SCOPED_DRAW_EVENT(RHICmdList, nDisplay_ViewportManager_WarpBlend);
+
 		const FDisplayClusterViewportManagerProxy* ViewportManagerProxy = &InViewportManagerProxy.Get();
 
 		if (ViewportManagerProxy->PostProcessManager.IsValid())
diff --git a/Engine/Plugins/Runtime/nDisplay/Source/DisplayCluster/Public/DisplayClusterRootActor.h b/Engine/Plugins/Runtime/nDisplay/Source/DisplayCluster/Public/DisplayClusterRootActor.h
--- a/Engine/Plugins/Runtime/nDisplay/Source/DisplayCluster/Public/DisplayClusterRootActor.h
+++ b/Engine/Plugins/Runtime/nDisplay/Source/DisplayCluster/Public/DisplayClusterRootActor.h
@@ -228,8 +228,11 @@ private:
 	UPROPERTY(EditAnywhere, Transient, Category = "In Camera VFX", meta = (PropertyPath = "CurrentConfigData.StageSettings.bEnableInnerFrustums"))
 	FDisplayClusterEditorPropertyReference EnableInnerFrustumsRef;
 
+	UPROPERTY(EditAnywhere, Transient, Category = "In Camera VFX", meta = (PropertyPath = "CurrentConfigData.StageSettings.bEnableInnerFrustumChromakeyOverlap "))
+	FDisplayClusterEditorPropertyReference ShowInnerFrustumOverlapsRef;
+
 	UPROPERTY(EditAnywhere, Transient, Category = "Color Grading", meta = (PropertyPath = "CurrentConfigData.StageSettings.EnableColorGrading"))
-	FDisplayClusterEditorPropertyReference EnableColorGradingRef;
+	FDisplayClusterEditorPropertyReference EnableInnerFrustumChromakeyOverlapRef;
 
 	UPROPERTY(EditAnywhere, Transient, Category = "Color Grading", meta = (PropertyPath = "CurrentConfigData.StageSettings.EntireClusterColorGrading"))
 	FDisplayClusterEditorPropertyReference ClusterColorGradingRef;
diff --git a/Engine/Plugins/Runtime/nDisplay/Source/DisplayClusterConfiguration/Public/DisplayClusterConfigurationTypes_ICVFX.h b/Engine/Plugins/Runtime/nDisplay/Source/DisplayClusterConfiguration/Public/DisplayClusterConfigurationTypes_ICVFX.h
--- a/Engine/Plugins/Runtime/nDisplay/Source/DisplayClusterConfiguration/Public/DisplayClusterConfigurationTypes_ICVFX.h
+++ b/Engine/Plugins/Runtime/nDisplay/Source/DisplayClusterConfiguration/Public/DisplayClusterConfigurationTypes_ICVFX.h
@@ -478,6 +478,10 @@ struct DISPLAYCLUSTERCONFIGURATION_API FDisplayClusterConfigurationICVFX_CameraS
 	/** Adjust blur amount to the left and right side edges of the inner frustum. */
 	UPROPERTY(EditAnywhere, BlueprintReadWrite, Category = NDisplay, meta = (DisplayName = "Sides", ClampMin = "0.0", UIMin = "0.0", ClampMax = "1.0", UIMax = "1.0"))
 	float Horizontal = 0.f;
+
+	/** Feather. */
+	UPROPERTY(EditAnywhere, BlueprintReadWrite, Category = NDisplay, meta = (DisplayName = "Feather", ClampMin = "0.0", UIMin = "0.0", ClampMax = "1.0", UIMax = "1.0"))
+	float Feather = 1.f;
 };
 
 USTRUCT(BlueprintType)
@@ -651,6 +655,10 @@ public:
 	UPROPERTY(EditAnywhere, BlueprintReadWrite, Category = "In Camera VFX", meta = (DisplayName = "Enable Inner Frustum"))
 	bool bEnableInnerFrustums = true;
 
+	/** Render the chromakey where the inner frustum overlaps. */
+	UPROPERTY(EditAnywhere, BlueprintReadWrite, Category = "In Camera VFX", meta = (DisplayName = "Enable Inner Frustum Chromakey Overlap"))
+	bool bEnableInnerFrustumChromakeyOverlap = false;
+
 	/** Default incameras RTT texture size. */
 	UPROPERTY(EditAnywhere, BlueprintReadWrite, Category = Configuration)
 	FDisplayClusterConfigurationICVFX_Size DefaultFrameSize;
diff --git a/Engine/Plugins/Runtime/nDisplay/Source/DisplayClusterMedia/Private/Synchronization/LatencyQueue/DisplayClusterFrameQueueItem.cpp b/Engine/Plugins/Runtime/nDisplay/Source/DisplayClusterMedia/Private/Synchronization/LatencyQueue/DisplayClusterFrameQueueItem.cpp
--- a/Engine/Plugins/Runtime/nDisplay/Source/DisplayClusterMedia/Private/Synchronization/LatencyQueue/DisplayClusterFrameQueueItem.cpp
+++ b/Engine/Plugins/Runtime/nDisplay/Source/DisplayClusterMedia/Private/Synchronization/LatencyQueue/DisplayClusterFrameQueueItem.cpp
@@ -97,24 +97,13 @@ void FDisplayClusterFrameQueueItem::LoadData(FRHICommandListImmediate& RHICmdLis
 					return Item.Resource.ViewportId.Equals(CameraSettings.Resource.ViewportId, ESearchCase::IgnoreCase);
 				});
 
-			// Update camera settings
 			if (TargetCameraSettings)
 			{
-				TargetCameraSettings->SoftEdge                 = CameraSettings.SoftEdge;
-				TargetCameraSettings->InnerCameraBorderColor      = CameraSettings.InnerCameraBorderColor;
-				TargetCameraSettings->InnerCameraBorderThickness  = CameraSettings.InnerCameraBorderThickness;
-				TargetCameraSettings->InnerCameraFrameAspectRatio = CameraSettings.InnerCameraFrameAspectRatio;
-				TargetCameraSettings->Local2WorldTransform     = CameraSettings.Local2WorldTransform;
-				TargetCameraSettings->CameraViewRotation       = CameraSettings.CameraViewRotation;
-				TargetCameraSettings->CameraViewLocation       = CameraSettings.CameraViewLocation;
-				TargetCameraSettings->CameraPrjMatrix          = CameraSettings.CameraPrjMatrix;
-				TargetCameraSettings->ChromakeySource          = CameraSettings.ChromakeySource;
-				TargetCameraSettings->ChromakeyColor           = CameraSettings.ChromakeyColor;
-				TargetCameraSettings->ChromakeyMarkersColor    = CameraSettings.ChromakeyMarkersColor;
-				TargetCameraSettings->ChromakeyMarkersScale    = CameraSettings.ChromakeyMarkersScale;
-				TargetCameraSettings->ChromakeyMarkersDistance = CameraSettings.ChromakeyMarkersDistance;
-				TargetCameraSettings->ChromakeyMarkersOffset   = CameraSettings.ChromakeyMarkersOffset;
-				TargetCameraSettings->RenderOrder              = CameraSettings.RenderOrder;
+				// Do not copy render resources (because they are not copied in the source code either.)
+				const bool bIncludeResources = false;
+
+				// Updating the camera settings is done with a new function that copies all settings
+				TargetCameraSettings->SetCameraSettings(CameraSettings, bIncludeResources);
 			}
 		}
 	}
diff --git a/Engine/Plugins/Runtime/nDisplay/Source/DisplayClusterProjection/Private/Policy/MPCDI/DisplayClusterProjectionMPCDIPolicy.cpp b/Engine/Plugins/Runtime/nDisplay/Source/DisplayClusterProjection/Private/Policy/MPCDI/DisplayClusterProjectionMPCDIPolicy.cpp
--- a/Engine/Plugins/Runtime/nDisplay/Source/DisplayClusterProjection/Private/Policy/MPCDI/DisplayClusterProjectionMPCDIPolicy.cpp
+++ b/Engine/Plugins/Runtime/nDisplay/Source/DisplayClusterProjection/Private/Policy/MPCDI/DisplayClusterProjectionMPCDIPolicy.cpp
@@ -299,12 +299,13 @@ void FDisplayClusterProjectionMPCDIPolicy::ApplyWarpBlend_RenderThread(FRHIComma
 							// matched camera reference, update context for current eye
 							const FDisplayClusterViewport_Context& InContext = SrcViewportContexts[ContextNum];
 
-							FDisplayClusterShaderParametersICVFX_CameraContext CameraContext;
-							CameraContext.CameraViewLocation = InContext.ViewLocation;
-							CameraContext.CameraViewRotation = InContext.ViewRotation;
-							CameraContext.CameraPrjMatrix = InContext.ProjectionMatrix;
+							// The context stores data in local space. For the shader they must be converted to world space
+							FDisplayClusterShaderParametersICVFX_CameraViewProjection LocalSpaceViewProjection;
+							LocalSpaceViewProjection.ViewLocation = InContext.ViewLocation;
+							LocalSpaceViewProjection.ViewRotation = InContext.ViewRotation;
+							LocalSpaceViewProjection.PrjMatrix    = InContext.ProjectionMatrix;
 
-							CameraSettings->UpdateCameraContext(CameraContext);
+							CameraSettings->SetViewProjection(LocalSpaceViewProjection);
 						}
 					}
 				}
diff --git a/Engine/Plugins/Runtime/nDisplay/Source/DisplayClusterShaders/Private/Shaders/DisplayClusterShadersWarpblend_ICVFX.cpp b/Engine/Plugins/Runtime/nDisplay/Source/DisplayClusterShaders/Private/Shaders/DisplayClusterShadersWarpblend_ICVFX.cpp
--- a/Engine/Plugins/Runtime/nDisplay/Source/DisplayClusterShaders/Private/Shaders/DisplayClusterShadersWarpblend_ICVFX.cpp
+++ b/Engine/Plugins/Runtime/nDisplay/Source/DisplayClusterShaders/Private/Shaders/DisplayClusterShadersWarpblend_ICVFX.cpp
@@ -29,7 +29,6 @@ enum class EVarIcvfxMPCDIShaderType : uint8
 	Default = 0,
 	DefaultNoBlend,
 	Passthrough,
-	ForceMultiPassRender,
 	Disable,
 };
 
@@ -40,7 +39,6 @@ static TAutoConsoleVariable<int32> CVarIcvfxMPCDIShaderType(
 	TEXT(" 0: Warp shader (used by default)\n")
 	TEXT(" 1: Warp shader with disabled blend maps\n")
 	TEXT(" 2: Passthrough shader\n")
-	TEXT(" 3: Force Multi-Pass render shaders\n")
 	TEXT(" 4: Disable all warp shaders\n"),
 	ECVF_RenderThreadSafe
 );
@@ -61,6 +59,61 @@ static FAutoConsoleVariableRef CVarDisplayClusterShadersICVFXEnableUVLightCard(
 	ECVF_RenderThreadSafe
 );
 
+int32 GDisplayClusterShadersICVFXOverlapInnerFrustumSoftEdges = 1;
+static FAutoConsoleVariableRef CVarDisplayClusterShadersICVFXOverlapInnerFrustumSoftEdges(
+	TEXT("nDisplay.render.icvfx.OverlapInnerFrustumSoftEdges"),
+	GDisplayClusterShadersICVFXOverlapInnerFrustumSoftEdges,
+	TEXT("Enable SoftEdges for Inner Frustum overlap (0 - disable).\n"),
+	ECVF_RenderThreadSafe
+);
+
+int32 GDisplayClusterShadersICVFXEnableInnerFrustumUnderOverlap = 0;
+static FAutoConsoleVariableRef CVarDisplayClusterShadersICVFXEnableInnerFrustumUnderOverlap(
+	TEXT("nDisplay.render.icvfx.EnableInnerFrustumUnderOverlap"),
+	GDisplayClusterShadersICVFXEnableInnerFrustumUnderOverlap,
+	TEXT("Enable Inner Frustum render under overlap area (0 - disable).\n"),
+	ECVF_RenderThreadSafe
+);
+
+/**
+ * ICVFX is rendered in several passes:
+ */
+enum class EIcvfxPassRenderPass : uint8
+{
+	None = 0,
+
+	// viewport+overlayUnder+camera1
+	Base,
+
+	// Second and next cameras additive passes
+	InnerFrustumIterator,
+
+	// LightCard overlay
+	LightCardOver,
+
+	// render overlaps for inner cameras
+	InnerFrustumChromakeyOverlapIterator,
+};
+
+enum class EIcvfxPassRenderState : uint32
+{
+	None = 0,
+	BlendDisabled = 1 << 0,
+
+	LightCardOver = 1 << 1,
+	LightCardUnder = 1 << 2,
+	UVLightCardOver = 1 << 3,
+	UVLightCardUnder = 1 << 4,
+
+	EnableInnerFrustumChromakeyOverlap = 1 << 5,
+};
+ENUM_CLASS_FLAGS(EIcvfxPassRenderState);
+
+DECLARE_GPU_STAT_NAMED(nDisplay_IcvfxRenderPass_Base, TEXT("nDisplay Icvfx::RenderPass::Base"));
+DECLARE_GPU_STAT_NAMED(nDisplay_IcvfxRenderPass_InnerFrustumIterator, TEXT("nDisplay Icvfx::RenderPass::InnerFrustumIterator"));
+DECLARE_GPU_STAT_NAMED(nDisplay_IcvfxRenderPass_LightCardOver, TEXT("nDisplay Icvfx::RenderPass::LightCardOver"));
+DECLARE_GPU_STAT_NAMED(nDisplay_IcvfxRenderPass_InnerFrustumChromakeyOverlapIterator, TEXT("nDisplay Icvfx::RenderPass::InnerFrustumChromakeyOverlapIterator"));
+
 enum class EIcvfxShaderType : uint8
 {
 	Passthrough,  // Viewport frustum (no warpblend, only frustum math)
@@ -80,7 +133,8 @@ namespace IcvfxShaderPermutation
 	class FIcvfxShaderUVLightCardUnder : SHADER_PERMUTATION_BOOL("UVLIGHTCARD_UNDER");
 	class FIcvfxShaderUVLightCardOver : SHADER_PERMUTATION_BOOL("UVLIGHTCARD_OVER");
 
-	class FIcvfxShaderInnerCamera      : SHADER_PERMUTATION_BOOL("INNER_CAMERA");
+	class FIcvfxShaderInnerCamera         : SHADER_PERMUTATION_BOOL("INNER_CAMERA");
+	class FIcvfxShaderInnerCameraOverlap  : SHADER_PERMUTATION_BOOL("INNER_CAMERA_OVERLAP");
 
 	class FIcvfxShaderChromakey           : SHADER_PERMUTATION_BOOL("CHROMAKEY");
 	class FIcvfxShaderChromakeyFrameColor : SHADER_PERMUTATION_BOOL("CHROMAKEYFRAMECOLOR");
@@ -104,6 +158,8 @@ namespace IcvfxShaderPermutation
 		FIcvfxShaderUVLightCardOver,
 
 		FIcvfxShaderInnerCamera,
+		FIcvfxShaderInnerCameraOverlap,
+
 		FIcvfxShaderChromakey,
 		FIcvfxShaderChromakeyFrameColor,
 		FIcvfxShaderChromakeyMarker,
@@ -117,12 +173,29 @@ namespace IcvfxShaderPermutation
 	{
 		if (!PermutationVector.Get<FIcvfxShaderMeshWarp>())
 		{
-			if (PermutationVector.Get<FIcvfxShaderChromakey>() || PermutationVector.Get<FIcvfxShaderChromakeyFrameColor>())
+			// UVLightCard require UV_Chromakey (Mesh Warp)
+			if (PermutationVector.Get<FIcvfxShaderUVLightCardUnder>() || PermutationVector.Get<FIcvfxShaderUVLightCardOver>())
+			{
+				return false;
+			}
+
+			// Chromakey marker require UV_Chromakey (Mesh Warp)
+			if (PermutationVector.Get<FIcvfxShaderChromakeyMarker>())
+			{
+				return false;
+			}
+		}
+
+		if (PermutationVector.Get<FIcvfxShaderInnerCameraOverlap>())
+		{
+			// Overlap use chromakeyframecolor as render base
+			if (!PermutationVector.Get<FIcvfxShaderChromakeyFrameColor>())
 			{
 				return false;
 			}
 		}
 
+		// only one type of chromakey type can by used
 		if (PermutationVector.Get<FIcvfxShaderChromakey>() && PermutationVector.Get<FIcvfxShaderChromakeyFrameColor>())
 		{
 			return false;
@@ -157,10 +230,12 @@ namespace IcvfxShaderPermutation
 
 		if (!PermutationVector.Get<FIcvfxShaderInnerCamera>())
 		{
-			if (PermutationVector.Get<FIcvfxShaderChromakey>() || PermutationVector.Get<FIcvfxShaderChromakeyFrameColor>())
-		{
-			return false;
-		}
+			// All innner camera vectors
+			if (PermutationVector.Get<FIcvfxShaderInnerCameraOverlap>()
+			|| PermutationVector.Get<FIcvfxShaderChromakey>() || PermutationVector.Get<FIcvfxShaderChromakeyFrameColor>() || PermutationVector.Get<FIcvfxShaderChromakeyMarker>())
+			{
+				return false;
+			}
 		}
 
 		if (PermutationVector.Get<FIcvfxShaderChromakeyMarker>())
@@ -176,15 +251,6 @@ namespace IcvfxShaderPermutation
 			return false;
 		}
 
-		// UVLightCard require UV_Chromakey (Mesh Warp)
-		if (!PermutationVector.Get<FIcvfxShaderMeshWarp>())
-		{
-			if (PermutationVector.Get<FIcvfxShaderUVLightCardUnder>() || PermutationVector.Get<FIcvfxShaderUVLightCardOver>())
-			{
-				return false;
-			}
-		}
-
 		return true;
 	}
 
@@ -244,6 +310,11 @@ BEGIN_SHADER_PARAMETER_STRUCT(FIcvfxPixelShaderParameters, )
 	SHADER_PARAMETER(float, ChromakeyMarkerScale)
 	SHADER_PARAMETER(float, ChromakeyMarkerDistance)
 	SHADER_PARAMETER(FVector2f, ChromakeyMarkerOffset)
+
+	SHADER_PARAMETER_ARRAY(FMatrix44f, OverlappedInnerCamerasProjectionMatrices, [MAX_INNER_CAMERAS_AMOUNT])
+	SHADER_PARAMETER_ARRAY(FVector4f, OverlappedInnerCameraSoftEdges, [MAX_INNER_CAMERAS_AMOUNT])
+	SHADER_PARAMETER(int, OverlappedInnerCamerasNum)
+
 END_SHADER_PARAMETER_STRUCT()
 
 class FIcvfxWarpVS : public FGlobalShader
@@ -276,6 +347,13 @@ public:
 	{
 		return IcvfxShaderPermutation::ShouldCompileCommonPSPermutation(FPermutationDomain(Parameters.PermutationId));
 	}
+
+	static void ModifyCompilationEnvironment(const FGlobalShaderPermutationParameters& Parameters, FShaderCompilerEnvironment& OutEnvironment)
+	{
+		FGlobalShader::ModifyCompilationEnvironment(Parameters, OutEnvironment);
+
+		OutEnvironment.SetDefine(TEXT("MAX_INNER_CAMERAS_AMOUNT"), MAX_INNER_CAMERAS_AMOUNT);
+	}
 };
 
 IMPLEMENT_GLOBAL_SHADER(FIcvfxWarpPS, ICVFX_ShaderFileName, "IcvfxWarpPS", SF_Pixel);
@@ -291,6 +369,13 @@ public:
 	{
 		return true;
 	}
+
+	static void ModifyCompilationEnvironment(const FGlobalShaderPermutationParameters& Parameters, FShaderCompilerEnvironment& OutEnvironment)
+	{
+		FGlobalShader::ModifyCompilationEnvironment(Parameters, OutEnvironment);
+
+		OutEnvironment.SetDefine(TEXT("MAX_INNER_CAMERAS_AMOUNT"), MAX_INNER_CAMERAS_AMOUNT);
+	}
 };
 
 IMPLEMENT_GLOBAL_SHADER(FIcvfxPassthroughPS, ICVFX_ShaderFileName, "Passthrough_PS", SF_Pixel);
@@ -311,127 +396,48 @@ public:
 	FIcvfxPassRenderer(const FDisplayClusterShaderParameters_WarpBlend& InWarpBlendParameters, const FDisplayClusterShaderParameters_ICVFX& InICVFXParameters)
 		: WarpBlendParameters(InWarpBlendParameters)
 		, ICVFXParameters(InICVFXParameters)
-	{}
-
-private:
-	const FDisplayClusterShaderParameters_WarpBlend& WarpBlendParameters;
-	const FDisplayClusterShaderParameters_ICVFX& ICVFXParameters;
-
-	EIcvfxShaderType PixelShaderType;
-
-	FMatrix LocalUVMatrix;
-
-	bool bIsBlendDisabled = false;
-	bool bForceMultiCameraPass = false;
-	int32 CameraIndex = 0;
-
-	//LightCards
-private:
-	inline bool IsLightcardOverUsed() const
-	{
-		return ICVFXParameters.IsLightcardOverUsed() && GDisplayClusterShadersICVFXEnableLightCard;
-	}
-
-	inline bool IsLightcardUnderUsed() const
-	{
-		return ICVFXParameters.IsLightcardUnderUsed() && GDisplayClusterShadersICVFXEnableLightCard;
-	}
-
-	inline bool IsUVLightcardOverUsed() const
 	{
-		return ICVFXParameters.IsUVLightcardOverUsed() && GDisplayClusterShadersICVFXEnableUVLightCard;
-	}
-
-	inline bool IsUVLightcardUnderUsed() const
-	{
-		return ICVFXParameters.IsUVLightcardUnderUsed() && GDisplayClusterShadersICVFXEnableUVLightCard;
-	}
-
-private:
-	inline bool IsMultiPassRender() const
-	{
-		return bForceMultiCameraPass || ICVFXParameters.IsMultiCamerasUsed();
-	}
-
-	inline bool IsFirstPass() const
-	{
-		return CameraIndex == 0;
-	}
-
-	inline bool IsLastPass() const
-	{
-		int32 TotalUsedCameras = ICVFXParameters.Cameras.Num();
-
-		if (bForceMultiCameraPass)
+		if (ICVFXParameters.IsLightCardOverUsed() && GDisplayClusterShadersICVFXEnableLightCard)
 		{
-			// Last pass is N or N+1
-			// 0 - viewport+overlayUnder
-			// 1 - first camera additive pass
-			// 2 - second camera additive pass
-			// N - camera N
-			// N+1 - overlayOver (optional)
-			return (IsLightcardOverUsed() || IsUVLightcardOverUsed()) ? ( CameraIndex == (TotalUsedCameras + 1) ) : ( CameraIndex == TotalUsedCameras );
+			EnumAddFlags(RenderState, EIcvfxPassRenderState::LightCardOver);
 		}
 
-		if (TotalUsedCameras>1)
+		if (ICVFXParameters.IsLightCardUnderUsed() && GDisplayClusterShadersICVFXEnableLightCard)
 		{
-			// Multi-cam
-			// 0   - viewport+overlayUnder+camera1
-			// 1   - camera2
-			// N-1 - camera N
-			// N   - overlayOver (optional)
-			return (IsLightcardOverUsed() || IsUVLightcardOverUsed()) ? ( CameraIndex == TotalUsedCameras ) : ( CameraIndex == (TotalUsedCameras-1) );
+			EnumAddFlags(RenderState, EIcvfxPassRenderState::LightCardUnder);
 		}
 
-		// By default render single camera in one pass
-		return CameraIndex==0;
-	}
+		if (ICVFXParameters.IsUVLightCardOverUsed() && GDisplayClusterShadersICVFXEnableUVLightCard)
+		{
+			EnumAddFlags(RenderState, EIcvfxPassRenderState::UVLightCardOver);
+		}
 
-	inline bool IsLightCardRenderUnderInCamera() const
-	{
-		// Render OverlayUnder only on first pass
-		return IsFirstPass() && IsLightcardUnderUsed();
-	}
+		if (ICVFXParameters.IsUVLightCardUnderUsed() && GDisplayClusterShadersICVFXEnableUVLightCard)
+		{
+			EnumAddFlags(RenderState, EIcvfxPassRenderState::UVLightCardUnder);
+		}
 
-	inline bool IsLightCardRenderOverInCamera() const
-	{
-		if (IsMultiPassRender())
+		if (ICVFXParameters.CameraOverlappingRenderMode != EDisplayClusterShaderParametersICVFX_CameraOverlappingRenderMode::None)
 		{
-			// Render OverlayOver only on last additive pass
-			return IsLastPass() && IsLightcardOverUsed();
+			EnumAddFlags(RenderState, EIcvfxPassRenderState::EnableInnerFrustumChromakeyOverlap);
 		}
-		
-		return IsLightcardOverUsed();
 	}
 
-	inline bool IsUVLightCardRenderUnderInCamera() const
-	{
-		// Render UVLightCard Under only on first pass
-		return IsFirstPass() && IsUVLightcardUnderUsed();
-	}
+private:
+	const FDisplayClusterShaderParameters_WarpBlend& WarpBlendParameters;
+	const FDisplayClusterShaderParameters_ICVFX& ICVFXParameters;
 
-	inline bool IsUVLightCardRenderOverInCamera() const
-	{
-		if (IsMultiPassRender())
-		{
-			// Render UVLightCard Over only on last additive pass
-			return IsLastPass() && IsUVLightcardOverUsed();
-		}
+	EIcvfxShaderType PixelShaderType;
 
-		return IsUVLightcardOverUsed();
-	}
+	EIcvfxPassRenderPass  RenderPass = EIcvfxPassRenderPass::None;
+	EIcvfxPassRenderState RenderState = EIcvfxPassRenderState::None;
+	int32 ActiveCameraIndex = INDEX_NONE;
 
-	inline int32 GetUsedCameraIndex() const
-	{
-		if (bForceMultiCameraPass)
-		{
-			return CameraIndex - 1;
-		}
+	FMatrix LocalUVMatrix;
 
-		return CameraIndex;
-	}
+private:
 
-	FMatrix GetStereoMatrix() const
+	inline FMatrix GetStereoMatrix() const
 	{
 		FIntPoint WarpDataSrcSize = WarpBlendParameters.Src.Texture->GetSizeXY();
 		FIntPoint WarpDataDstSize = WarpBlendParameters.Dest.Texture->GetSizeXY();
@@ -445,7 +451,7 @@ private:
 		return StereoMatrix;
 	}
 
-	FIntRect GetViewportRect() const
+	inline FIntRect GetViewportRect() const
 	{
 		const int32 PosX = WarpBlendParameters.Dest.Rect.Min.X;
 		const int32 PosY = WarpBlendParameters.Dest.Rect.Min.Y;
@@ -455,7 +461,7 @@ private:
 		return FIntRect(FIntPoint(PosX, PosY), FIntPoint(PosX + SizeX, PosY + SizeY));
 	}
 
-	EIcvfxShaderType GetPixelShaderType()
+	inline EIcvfxShaderType GetPixelShaderType()
 	{
 		if (WarpBlendParameters.WarpInterface.IsValid())
 		{
@@ -479,7 +485,7 @@ public:
 
 	bool GetViewportParameters(FIcvfxRenderPassData& RenderPassData)
 	{
-		if (IsFirstPass())
+		if (RenderPass == EIcvfxPassRenderPass::Base)
 		{
 			// Input viewport texture
 			// Render only on first pass
@@ -584,61 +590,174 @@ public:
 	{
 		RenderPassData.PSParameters.OverlayProjectionMatrix = FMatrix44f(LocalUVMatrix);
 
-		if (IsLightCardRenderUnderInCamera() || IsLightCardRenderOverInCamera())
+		switch (RenderPass)
 		{
-			RenderPassData.PSParameters.LightCardSampler = TStaticSamplerState<SF_Trilinear, AM_Clamp, AM_Clamp, AM_Clamp>::GetRHI();
-			RenderPassData.PSParameters.LightCardTexture = ICVFXParameters.Lightcard.Texture;
-
-			if (IsLightCardRenderUnderInCamera())
+		case EIcvfxPassRenderPass::Base:
+			if (EnumHasAnyFlags(RenderState, EIcvfxPassRenderState::LightCardUnder))
 			{
+				// Rendering in one pass along with the first camera
 				RenderPassData.PSPermutationVector.Set<IcvfxShaderPermutation::FIcvfxShaderLightCardUnder>(true);
 			}
-			if (IsLightCardRenderOverInCamera())
+			else if (EnumHasAnyFlags(RenderState, EIcvfxPassRenderState::LightCardOver) && !ICVFXParameters.IsMultiCamerasUsed())
+			{
+				// If we have only one camera, the top LC layer is rendered in one pass
+				RenderPassData.PSPermutationVector.Set<IcvfxShaderPermutation::FIcvfxShaderLightCardOver>(true);
+			}
+		break;
+
+		case EIcvfxPassRenderPass::LightCardOver:
+			check(ICVFXParameters.IsMultiCamerasUsed());
+
+			if (EnumHasAllFlags(RenderState, EIcvfxPassRenderState::LightCardOver))
 			{
 				RenderPassData.PSPermutationVector.Set<IcvfxShaderPermutation::FIcvfxShaderLightCardOver>(true);
 			}
+			break;
+
+		default:
+			break;
 		}
 
-		return true;
+		if(RenderPassData.PSPermutationVector.Get<IcvfxShaderPermutation::FIcvfxShaderLightCardUnder>()
+		|| RenderPassData.PSPermutationVector.Get<IcvfxShaderPermutation::FIcvfxShaderLightCardOver>())
+		{
+			RenderPassData.PSParameters.LightCardSampler = TStaticSamplerState<SF_Trilinear, AM_Clamp, AM_Clamp, AM_Clamp>::GetRHI();
+			RenderPassData.PSParameters.LightCardTexture = ICVFXParameters.LightCard.Texture;
+
+			return true;
+		}
+
+		return false;
 	}
 
 	bool GetUVLightCardParameters(FIcvfxRenderPassData& RenderPassData)
 	{
-		if (IsUVLightCardRenderUnderInCamera() || IsUVLightCardRenderOverInCamera())
+		// UVLightCard require Mesh warp (UV_Chromakey)
+		if (!RenderPassData.VSPermutationVector.Get<IcvfxShaderPermutation::FIcvfxShaderMeshWarp>())
 		{
-			// UVLightCard require Mesh warp (UV_Chromakey)
-			if (RenderPassData.VSPermutationVector.Get<IcvfxShaderPermutation::FIcvfxShaderMeshWarp>())
+			return false;
+		}
+
+		switch (RenderPass)
+		{
+		case EIcvfxPassRenderPass::Base:
+			if (EnumHasAnyFlags(RenderState, EIcvfxPassRenderState::UVLightCardUnder))
 			{
-				RenderPassData.PSParameters.UVLightCardSampler = TStaticSamplerState<SF_Trilinear, AM_Clamp, AM_Clamp, AM_Clamp>::GetRHI();
-				RenderPassData.PSParameters.UVLightCardTexture = ICVFXParameters.UVLightcard.Texture;
+				RenderPassData.PSPermutationVector.Set<IcvfxShaderPermutation::FIcvfxShaderUVLightCardUnder>(true);
+			}
+			else if (EnumHasAnyFlags(RenderState, EIcvfxPassRenderState::UVLightCardOver) && !ICVFXParameters.IsMultiCamerasUsed())
+			{
+				// Render in one pass for single camera
+				RenderPassData.PSPermutationVector.Set<IcvfxShaderPermutation::FIcvfxShaderUVLightCardOver>(true);
+			}
+			break;
 
-				if (IsUVLightCardRenderUnderInCamera())
-				{
-					RenderPassData.PSPermutationVector.Set<IcvfxShaderPermutation::FIcvfxShaderUVLightCardUnder>(true);
-				}
-				if (IsUVLightCardRenderOverInCamera())
+		case EIcvfxPassRenderPass::LightCardOver:
+			check(ICVFXParameters.IsMultiCamerasUsed());
+
+			if (EnumHasAllFlags(RenderState, EIcvfxPassRenderState::UVLightCardOver))
+			{
+				RenderPassData.PSPermutationVector.Set<IcvfxShaderPermutation::FIcvfxShaderUVLightCardOver>(true);
+			}
+			break;
+
+		default:
+			break;
+		}
+
+		if (RenderPassData.PSPermutationVector.Get<IcvfxShaderPermutation::FIcvfxShaderUVLightCardUnder>()
+			|| RenderPassData.PSPermutationVector.Get<IcvfxShaderPermutation::FIcvfxShaderUVLightCardOver>())
+		{
+			RenderPassData.PSParameters.UVLightCardSampler = TStaticSamplerState<SF_Trilinear, AM_Clamp, AM_Clamp, AM_Clamp>::GetRHI();
+			RenderPassData.PSParameters.UVLightCardTexture = ICVFXParameters.UVLightCard.Texture;
+
+			return true;
+		}
+
+		return false;
+	}
+
+	int32 GetOverlappedInnerCameras(FIcvfxRenderPassData& RenderPassData)
+	{
+		static const FMatrix Game2Render(
+			FPlane(0, 0, 1, 0),
+			FPlane(1, 0, 0, 0),
+			FPlane(0, 1, 0, 0),
+			FPlane(0, 0, 0, 1));
+
+		// The total number of InCameras, which are rendered below.
+		int32 OutOverlapCameraNum = 0;
+
+		for (int32 CameraIndex = 0; CameraIndex < ActiveCameraIndex && ICVFXParameters.Cameras.IsValidIndex(CameraIndex) && CameraIndex < MAX_INNER_CAMERAS_AMOUNT; CameraIndex++)
+		{
+			const FDisplayClusterShaderParameters_ICVFX::FCameraSettings& Camera = ICVFXParameters.Cameras[CameraIndex];
+			const FRotationTranslationMatrix CameraTranslationMatrix(Camera.ViewProjection.ViewRotation, Camera.ViewProjection.ViewLocation);
+
+			const FMatrix WorldToCamera = CameraTranslationMatrix.Inverse();
+			const FMatrix UVMatrix = WorldToCamera * Game2Render * Camera.ViewProjection.PrjMatrix;
+			const FMatrix InnerCameraProjectionMatrix = UVMatrix * WarpBlendParameters.Context.TextureMatrix * WarpBlendParameters.Context.RegionMatrix;
+
+			RenderPassData.PSParameters.OverlappedInnerCamerasProjectionMatrices[CameraIndex] = FMatrix44f(InnerCameraProjectionMatrix);
+			RenderPassData.PSParameters.OverlappedInnerCameraSoftEdges[CameraIndex] = GDisplayClusterShadersICVFXOverlapInnerFrustumSoftEdges ? FVector4f(Camera.SoftEdge) : FVector4f(0, 0, 0, 0);
+
+			// Count total num of overlapped cameras
+			OutOverlapCameraNum++;
+		}
+
+		return OutOverlapCameraNum;
+	}
+	
+	EDisplayClusterShaderParametersICVFX_ChromakeySource GetCameraOverlapParameters(FIcvfxRenderPassData& RenderPassData)
+	{
+		int32 OverlapCameraNum = 0;
+		EDisplayClusterShaderParametersICVFX_ChromakeySource OutChromakeySource = EDisplayClusterShaderParametersICVFX_ChromakeySource::Unknown;
+
+		if (EnumHasAnyFlags(RenderState, EIcvfxPassRenderState::EnableInnerFrustumChromakeyOverlap) && ICVFXParameters.IsMultiCamerasUsed())
+		{
+			switch (RenderPass)
+			{
+			case EIcvfxPassRenderPass::InnerFrustumIterator:
+				// Hide Inner Frustum under overlap
+				if (!GDisplayClusterShadersICVFXEnableInnerFrustumUnderOverlap)
 				{
-					RenderPassData.PSPermutationVector.Set<IcvfxShaderPermutation::FIcvfxShaderUVLightCardOver>(true);
+					OverlapCameraNum = GetOverlappedInnerCameras(RenderPassData);
 				}
+				break;
 
-				return true;
+			case EIcvfxPassRenderPass::InnerFrustumChromakeyOverlapIterator:
+			{
+				// Use camera chromakey color to render overlap
+				OutChromakeySource = EDisplayClusterShaderParametersICVFX_ChromakeySource::FrameColor;
+
+				// Switch to special shader
+				RenderPassData.PSPermutationVector.Set<IcvfxShaderPermutation::FIcvfxShaderInnerCameraOverlap>(true);
+
+				// Get cameraa under active
+				OverlapCameraNum = GetOverlappedInnerCameras(RenderPassData);
+
+				break;
+			}
+
+			default:
+				break;
 			}
 		}
 
-		return false;
+		RenderPassData.PSParameters.OverlappedInnerCamerasNum = OverlapCameraNum;
+
+		return OutChromakeySource;
 	}
 
 	bool GetCameraParameters(FIcvfxRenderPassData& RenderPassData)
 	{
-		const int32 CurrentCameraIndex = GetUsedCameraIndex();
-		if (!ICVFXParameters.IsCameraUsed(CurrentCameraIndex))
+		if (!ICVFXParameters.IsCameraUsed(ActiveCameraIndex))
 		{
 			return false;
 		}
 
 		RenderPassData.PSPermutationVector.Set<IcvfxShaderPermutation::FIcvfxShaderInnerCamera>(true);
 
-		const FDisplayClusterShaderParameters_ICVFX::FCameraSettings& Camera = ICVFXParameters.Cameras[CurrentCameraIndex];
+		const FDisplayClusterShaderParameters_ICVFX::FCameraSettings& Camera = ICVFXParameters.Cameras[ActiveCameraIndex];
 
 		static const FMatrix Game2Render(
 			FPlane(0, 0, 1, 0),
@@ -646,11 +765,11 @@ public:
 			FPlane(0, 1, 0, 0),
 			FPlane(0, 0, 0, 1));
 
-		FRotationTranslationMatrix CameraTranslationMatrix(Camera.CameraViewRotation, Camera.CameraViewLocation);
+		const FRotationTranslationMatrix CameraTranslationMatrix(Camera.ViewProjection.ViewRotation, Camera.ViewProjection.ViewLocation);
 
-		FMatrix WorldToCamera = CameraTranslationMatrix.Inverse();
-		FMatrix UVMatrix = WorldToCamera * Game2Render * Camera.CameraPrjMatrix;
-		FMatrix InnerCameraProjectionMatrix = UVMatrix * WarpBlendParameters.Context.TextureMatrix * WarpBlendParameters.Context.RegionMatrix;
+		const FMatrix WorldToCamera = CameraTranslationMatrix.Inverse();
+		const FMatrix UVMatrix = WorldToCamera * Game2Render * Camera.ViewProjection.PrjMatrix;
+		const FMatrix InnerCameraProjectionMatrix = UVMatrix * WarpBlendParameters.Context.TextureMatrix * WarpBlendParameters.Context.RegionMatrix;
 
 		RenderPassData.PSParameters.InnerCameraTexture = Camera.Resource.Texture;
 		RenderPassData.PSParameters.InnerCameraSampler = TStaticSamplerState<SF_Trilinear, AM_Clamp, AM_Clamp, AM_Clamp>::GetRHI();
@@ -665,12 +784,16 @@ public:
 		return true;
 	}
 
-	bool GetCameraChromakeyParameters(FIcvfxRenderPassData& RenderPassData)
+	bool GetCameraChromakeyParameters(FIcvfxRenderPassData& RenderPassData, const EDisplayClusterShaderParametersICVFX_ChromakeySource InChromakeySource)
 	{
-		const int32 CurrentCameraIndex = GetUsedCameraIndex();
-		const FDisplayClusterShaderParameters_ICVFX::FCameraSettings& Camera = ICVFXParameters.Cameras[CurrentCameraIndex];
+		if (!ICVFXParameters.IsCameraUsed(ActiveCameraIndex))
+		{
+			return false;
+		}
 		
-		switch (Camera.ChromakeySource)
+		const FDisplayClusterShaderParameters_ICVFX::FCameraSettings& Camera = ICVFXParameters.Cameras[ActiveCameraIndex];
+		
+		switch (InChromakeySource == EDisplayClusterShaderParametersICVFX_ChromakeySource::Unknown ? Camera.ChromakeySource : InChromakeySource)
 		{
 		case EDisplayClusterShaderParametersICVFX_ChromakeySource::ChromakeyLayers:
 		{
@@ -701,8 +824,12 @@ public:
 
 	bool GetCameraChromakeyMarkerParameters(FIcvfxRenderPassData& RenderPassData)
 	{
-		const int32 CurrentCameraIndex = GetUsedCameraIndex();
-		const FDisplayClusterShaderParameters_ICVFX::FCameraSettings& Camera = ICVFXParameters.Cameras[CurrentCameraIndex];
+		if (!ICVFXParameters.IsCameraUsed(ActiveCameraIndex))
+		{
+			return false;
+		}
+
+		const FDisplayClusterShaderParameters_ICVFX::FCameraSettings& Camera = ICVFXParameters.Cameras[ActiveCameraIndex];
 
 		if (Camera.IsChromakeyMarkerUsed())
 		{
@@ -732,7 +859,7 @@ public:
 		{
 			GetWarpMapParameters(RenderPassData);
 
-			if (!bIsBlendDisabled)
+			if (!EnumHasAnyFlags(RenderState, EIcvfxPassRenderState::BlendDisabled))
 			{
 				GetWarpBlendParameters(RenderPassData);
 			}
@@ -742,8 +869,10 @@ public:
 
 			if (GetCameraParameters(RenderPassData))
 			{
+				const EDisplayClusterShaderParametersICVFX_ChromakeySource ChromakeySourceOverride = GetCameraOverlapParameters(RenderPassData);
+
 				// Chromakey inside cam:
-				if (GetCameraChromakeyParameters(RenderPassData))
+				if (GetCameraChromakeyParameters(RenderPassData, ChromakeySourceOverride))
 				{
 					//Support chromakey markers
 					GetCameraChromakeyMarkerParameters(RenderPassData);
@@ -900,24 +1029,24 @@ public:
 
 			GraphicsPSOInit.PrimitiveType = PT_TriangleList;
 
-			if (IsFirstPass())
+			switch (RenderPass)
 			{
+			case EIcvfxPassRenderPass::InnerFrustumIterator:
+			case EIcvfxPassRenderPass::InnerFrustumChromakeyOverlapIterator:
+				// Multicam rendering
+				GraphicsPSOInit.BlendState = TStaticBlendState <CW_RGBA, BO_Add, BF_SourceAlpha, BF_InverseSourceAlpha, BO_Add, BF_Zero, BF_One>::GetRHI();
+				break;
+
+			case EIcvfxPassRenderPass::LightCardOver:
+				// render pass for LC overlays
+				GraphicsPSOInit.BlendState = TStaticBlendState <CW_RGBA, BO_Add, BF_One, BF_SourceAlpha, BO_Add, BF_Zero, BF_One>::GetRHI();
+				break;
+
+			case EIcvfxPassRenderPass::Base:
+			default:
 				// First pass always override old viewport image
 				GraphicsPSOInit.BlendState = TStaticBlendState <>::GetRHI();
-			}
-			else
-			{
-				// Extra render pass: don't change target alpha channel
-				if (RenderPassData.PSPermutationVector.Get<IcvfxShaderPermutation::FIcvfxShaderInnerCamera>())
-				{
-					// Multicam rendering
-					GraphicsPSOInit.BlendState = TStaticBlendState <CW_RGBA, BO_Add, BF_SourceAlpha, BF_InverseSourceAlpha, BO_Add, BF_Zero, BF_One>::GetRHI();
-				}
-				else
-				{
-					// [optional] Final render pass for LC overlays
-					GraphicsPSOInit.BlendState = TStaticBlendState <CW_RGBA, BO_Add, BF_One, BF_SourceAlpha, BO_Add, BF_Zero, BF_One>::GetRHI();
-				}
+				break;
 			}
 
 			SetGraphicsPipelineState(RHICmdList, GraphicsPSOInit, 0);
@@ -948,28 +1077,25 @@ public:
 			return false;
 		}
 
-		bIsBlendDisabled = WarpBlendParameters.WarpInterface->GetTexture(EDisplayClusterWarpBlendTextureType::AlphaMap) == nullptr;
-		bForceMultiCameraPass = false;
+		if (WarpBlendParameters.WarpInterface->GetTexture(EDisplayClusterWarpBlendTextureType::AlphaMap) == nullptr)
+		{
+			EnumAddFlags(RenderState, EIcvfxPassRenderState::BlendDisabled);
+		};
 
 		const EVarIcvfxMPCDIShaderType ShaderType = (EVarIcvfxMPCDIShaderType)CVarIcvfxMPCDIShaderType.GetValueOnAnyThread();
 
 		switch (ShaderType)
 		{
 		case EVarIcvfxMPCDIShaderType::DefaultNoBlend:
-			bIsBlendDisabled = true;
+			EnumAddFlags(RenderState, EIcvfxPassRenderState::BlendDisabled);
 			break;
 
 		case EVarIcvfxMPCDIShaderType::Passthrough:
 			PixelShaderType = EIcvfxShaderType::Passthrough;
 			break;
 
-		case EVarIcvfxMPCDIShaderType::ForceMultiPassRender:
-			bForceMultiCameraPass = true;
-			break;
-
 		case EVarIcvfxMPCDIShaderType::Disable:
 			return false;
-			break;
 
 		case EVarIcvfxMPCDIShaderType::Default:
 		default:
@@ -979,6 +1105,135 @@ public:
 
 		return true;
 	}
+
+	// Render Full Icvfx Warp&Blend
+	bool DoRenderPass(FRHICommandListImmediate& RHICmdList)
+	{
+		switch (RenderPass)
+		{
+		case EIcvfxPassRenderPass::None:
+			RenderPass = EIcvfxPassRenderPass::Base;
+			if (ICVFXParameters.IsAnyCameraUsed())
+			{
+				ActiveCameraIndex = 0;
+			}
+			break;
+
+		case EIcvfxPassRenderPass::Base:
+			if (!ICVFXParameters.IsMultiCamerasUsed())
+			{
+				// Render one camera in single pass
+				RenderPass = EIcvfxPassRenderPass::None;
+				return false;
+			}
+
+			// render multi-cam
+			RenderPass = EIcvfxPassRenderPass::InnerFrustumIterator;
+			ActiveCameraIndex = 1;
+			break;
+
+		case EIcvfxPassRenderPass::InnerFrustumIterator:
+			if (!ICVFXParameters.IsCameraUsed(++ActiveCameraIndex))
+			{
+				// All cameras have been rendered.
+
+				if (EnumHasAnyFlags(RenderState, EIcvfxPassRenderState::LightCardOver | EIcvfxPassRenderState::UVLightCardOver))
+				{
+					// Render LC over cameras
+					ActiveCameraIndex = INDEX_NONE;
+					RenderPass = EIcvfxPassRenderPass::LightCardOver;
+					break;
+				}
+				else if (EnumHasAnyFlags(RenderState, EIcvfxPassRenderState::EnableInnerFrustumChromakeyOverlap))
+				{
+					// Render camera overlaps
+					ActiveCameraIndex = 1;
+					RenderPass = EIcvfxPassRenderPass::InnerFrustumChromakeyOverlapIterator;
+					break;
+				}
+
+				// Render completed
+				RenderPass = EIcvfxPassRenderPass::None;
+				return false;
+			}
+			break;
+
+		case EIcvfxPassRenderPass::LightCardOver:
+			
+			if (EnumHasAnyFlags(RenderState, EIcvfxPassRenderState::EnableInnerFrustumChromakeyOverlap))
+			{
+				// Render camera overlaps
+				ActiveCameraIndex = 1;
+				RenderPass = EIcvfxPassRenderPass::InnerFrustumChromakeyOverlapIterator;
+				break;
+			}
+
+			// Render completed
+			RenderPass = EIcvfxPassRenderPass::None;
+			return false;
+
+		case EIcvfxPassRenderPass::InnerFrustumChromakeyOverlapIterator:
+			if (!ICVFXParameters.IsCameraUsed(++ActiveCameraIndex))
+			{
+				// Render completed
+				RenderPass = EIcvfxPassRenderPass::None;
+				return false;
+			}
+			break;
+
+		default:
+			return false;
+		}
+
+		// Build and draw render-pass:
+		FIcvfxRenderPassData RenderPassData;
+		InitRenderPass(RenderPassData);
+
+		switch (RenderPass)
+		{
+		case EIcvfxPassRenderPass::Base:
+		{
+			SCOPED_GPU_STAT(RHICmdList, nDisplay_IcvfxRenderPass_Base);
+			SCOPED_DRAW_EVENT(RHICmdList, nDisplay_IcvfxRenderPass_Base);
+
+			RenderCurentPass(RHICmdList, RenderPassData);
+		}
+		break;
+
+		case EIcvfxPassRenderPass::InnerFrustumIterator:
+		{
+			SCOPED_GPU_STAT(RHICmdList, nDisplay_IcvfxRenderPass_InnerFrustumIterator);
+			SCOPED_DRAW_EVENT(RHICmdList, nDisplay_IcvfxRenderPass_InnerFrustumIterator);
+
+			RenderCurentPass(RHICmdList, RenderPassData);
+		}
+		break;
+
+		case EIcvfxPassRenderPass::InnerFrustumChromakeyOverlapIterator:
+		{
+			SCOPED_GPU_STAT(RHICmdList, nDisplay_IcvfxRenderPass_InnerFrustumChromakeyOverlapIterator);
+			SCOPED_DRAW_EVENT(RHICmdList, nDisplay_IcvfxRenderPass_InnerFrustumChromakeyOverlapIterator);
+
+			RenderCurentPass(RHICmdList, RenderPassData);
+		}
+		break;
+
+		case EIcvfxPassRenderPass::LightCardOver:
+		{
+			SCOPED_GPU_STAT(RHICmdList, nDisplay_IcvfxRenderPass_LightCardOver);
+			SCOPED_DRAW_EVENT(RHICmdList, nDisplay_IcvfxRenderPass_LightCardOver);
+
+			RenderCurentPass(RHICmdList, RenderPassData);
+		}
+		break;
+
+		default:
+			RenderCurentPass(RHICmdList, RenderPassData);
+			break;
+		}
+
+		return true;
+	}
 	
 	bool Render(FRHICommandListImmediate& RHICmdList)
 	{
@@ -1001,26 +1256,9 @@ public:
 			case EIcvfxShaderType::WarpAndBlend:
 			{
 				LocalUVMatrix = WarpBlendParameters.Context.UVMatrix * WarpBlendParameters.Context.TextureMatrix * WarpBlendParameters.Context.RegionMatrix;
-				{
-					// Render Full Icvfx Warp&Blend
-					while (true)
-					{
-						{
-							// Build and draw render-pass:
-							FIcvfxRenderPassData RenderPassData;
-							InitRenderPass(RenderPassData);
-							RenderCurentPass(RHICmdList, RenderPassData);
-						}
-
-						if (IsLastPass())
-						{
-							break;
-						}
-
-						// Switch to next render-pass
-						CameraIndex++;
-					}
-				}
+
+				// Multi-pass ICVFX rendering
+				while (DoRenderPass(RHICmdList));
 
 				break;
 			}
diff --git a/Engine/Plugins/Runtime/nDisplay/Source/DisplayClusterShaders/Private/Shaders/DisplayClusterShadersWarpblend_ICVFX.h b/Engine/Plugins/Runtime/nDisplay/Source/DisplayClusterShaders/Private/Shaders/DisplayClusterShadersWarpblend_ICVFX.h
--- a/Engine/Plugins/Runtime/nDisplay/Source/DisplayClusterShaders/Private/Shaders/DisplayClusterShadersWarpblend_ICVFX.h
+++ b/Engine/Plugins/Runtime/nDisplay/Source/DisplayClusterShaders/Private/Shaders/DisplayClusterShadersWarpblend_ICVFX.h
@@ -9,6 +9,7 @@
 class FDisplayClusterShaderParameters_ICVFX;
 struct FDisplayClusterShaderParameters_WarpBlend;
 
+#define MAX_INNER_CAMERAS_AMOUNT 10
 
 class FDisplayClusterShadersWarpblend_ICVFX
 {
diff --git a/Engine/Plugins/Runtime/nDisplay/Source/DisplayClusterShaders/Public/ShaderParameters/DisplayClusterShaderParameters_ICVFX.h b/Engine/Plugins/Runtime/nDisplay/Source/DisplayClusterShaders/Public/ShaderParameters/DisplayClusterShaderParameters_ICVFX.h
--- a/Engine/Plugins/Runtime/nDisplay/Source/DisplayClusterShaders/Public/ShaderParameters/DisplayClusterShaderParameters_ICVFX.h
+++ b/Engine/Plugins/Runtime/nDisplay/Source/DisplayClusterShaders/Public/ShaderParameters/DisplayClusterShaderParameters_ICVFX.h
@@ -5,6 +5,9 @@
 #include "CoreMinimal.h"
 #include "RHI.h"
 
+/**
+ * Rendering mode for chromakey
+ */
 enum class EDisplayClusterShaderParametersICVFX_ChromakeySource : uint8
 {
 	// Dont use chromakey
@@ -15,9 +18,14 @@ enum class EDisplayClusterShaderParametersICVFX_ChromakeySource : uint8
 
 	// Render specified layer from scene
 	ChromakeyLayers,
+
+	Unknown,
 };
 
-enum class EDisplayClusterShaderParametersICVFX_LightcardRenderMode : uint8
+/**
+ * The rendering order of the lightcard layer
+ */
+enum class EDisplayClusterShaderParametersICVFX_LightCardRenderMode : uint8
 {
 	// Render incamera frame over lightcard
 	Under = 0,
@@ -26,13 +34,33 @@ enum class EDisplayClusterShaderParametersICVFX_LightcardRenderMode : uint8
 	Over,
 };
 
-struct FDisplayClusterShaderParametersICVFX_CameraContext
+/**
+ * Rendering mode for overlapping areas of camera projections
+ */
+enum class EDisplayClusterShaderParametersICVFX_CameraOverlappingRenderMode : uint8
+{
+	// Disabled
+	None = 0,
+
+	// The overlapping area is rendered at the very end, on top of all the layers
+	FinalPass
+};
+
+/**
+ * Camera view projection data
+ */
+struct FDisplayClusterShaderParametersICVFX_CameraViewProjection
 {
-	FRotator CameraViewRotation;
-	FVector  CameraViewLocation;
-	FMatrix  CameraPrjMatrix;
+	FRotator ViewRotation;
+	FVector  ViewLocation;
+	FMatrix  PrjMatrix;
 };
 
+/**
+ * ICVFX rendering uses the resources of other viewports
+ * During initialization on a game thread, only the name of the viewport is saved here
+ * Later on the rendering thread, the resource reference is initialized
+ */
 struct FDisplayClusterShaderParametersICVFX_ViewportResource
 {
 public:
@@ -60,7 +88,9 @@ public:
 	FRHITexture* Texture = nullptr;
 };
 
-// Compact ICVFX data for viewport
+/**
+ * This is where the ICVFX data for the outer viewport is stored.
+ */
 class FDisplayClusterShaderParameters_ICVFX
 {
 public:
@@ -73,44 +103,44 @@ public:
 	}
 
 public:
-	bool IsLightcardOverUsed() const
+	inline bool IsLightCardOverUsed() const
 	{
-		return Lightcard.IsValid() && LightcardMode == EDisplayClusterShaderParametersICVFX_LightcardRenderMode::Over;
+		return LightCard.IsValid() && LightCardMode == EDisplayClusterShaderParametersICVFX_LightCardRenderMode::Over;
 	}
 
-	bool IsLightcardUnderUsed() const
+	inline bool IsLightCardUnderUsed() const
 	{
-		return Lightcard.IsValid() && LightcardMode == EDisplayClusterShaderParametersICVFX_LightcardRenderMode::Under;
+		return LightCard.IsValid() && LightCardMode == EDisplayClusterShaderParametersICVFX_LightCardRenderMode::Under;
 	}
 
-	bool IsUVLightcardOverUsed() const
+	inline bool IsUVLightCardOverUsed() const
 	{
-		return UVLightcard.IsValid() && LightcardMode == EDisplayClusterShaderParametersICVFX_LightcardRenderMode::Over;
+		return UVLightCard.IsValid() && LightCardMode == EDisplayClusterShaderParametersICVFX_LightCardRenderMode::Over;
 	}
 
-	bool IsUVLightcardUnderUsed() const
+	inline bool IsUVLightCardUnderUsed() const
 	{
-		return UVLightcard.IsValid() && LightcardMode == EDisplayClusterShaderParametersICVFX_LightcardRenderMode::Under;
+		return UVLightCard.IsValid() && LightCardMode == EDisplayClusterShaderParametersICVFX_LightCardRenderMode::Under;
 	}
 
-	bool IsCameraUsed(int32 CameraIndex) const
+	inline bool IsCameraUsed(int32 CameraIndex) const
 	{
-		return (Cameras.Num() > CameraIndex) && CameraIndex >= 0 && Cameras[CameraIndex].IsUsed();
+		return Cameras.IsValidIndex(CameraIndex) && Cameras[CameraIndex].IsUsed();
 	}
 
-	bool IsAnyCameraUsed() const
+	inline bool IsAnyCameraUsed() const
 	{
-		return (Cameras.Num() > 0);
+		return !Cameras.IsEmpty();
 	}
 
-	bool IsMultiCamerasUsed() const
+	inline bool IsMultiCamerasUsed() const
 	{
-		return (Cameras.Num() > 1);
+		return Cameras.Num() > 1;
 	}
 
-	bool IsValid()
+	inline bool IsValid()
 	{
-		return Lightcard.IsValid() || UVLightcard.IsValid() || IsAnyCameraUsed();
+		return LightCard.IsValid() || UVLightCard.IsValid() || IsAnyCameraUsed();
 	}
 
 public:
@@ -118,8 +148,8 @@ public:
 	{
 		Cameras.Empty();
 
-		UVLightcard.Reset();
-		Lightcard.Reset();
+		UVLightCard.Reset();
+		LightCard.Reset();
 	}
 
 	// Implement copy ref and arrays
@@ -129,9 +159,11 @@ public:
 
 		Cameras = InParameters.Cameras;
 
-		UVLightcard   = InParameters.UVLightcard;
-		Lightcard     = InParameters.Lightcard;
-		LightcardMode = InParameters.LightcardMode;
+		CameraOverlappingRenderMode = InParameters.CameraOverlappingRenderMode;
+
+		UVLightCard = InParameters.UVLightCard;
+		LightCard   = InParameters.LightCard;
+		LightCardMode = InParameters.LightCardMode;
 	}
 
 	inline void SortCamerasRenderOrder()
@@ -147,25 +179,27 @@ public:
 	}
 
 public:
-	// ICVFX Target only data
+	/**
+	 * Incamera render settings
+	 */
 	struct FCameraSettings
 	{
-		bool IsUsed() const
+		inline bool IsUsed() const
 		{
 			return (ChromakeySource == EDisplayClusterShaderParametersICVFX_ChromakeySource::FrameColor) || Resource.IsValid();
 		}
 
-		bool IsChromakeyMarkerUsed() const
+		inline bool IsChromakeyMarkerUsed() const
 		{
 			return ChromakeMarkerTextureRHI.IsValid();
 		}
 
-		void UpdateCameraContext(const FDisplayClusterShaderParametersICVFX_CameraContext& InContext)
+		inline void SetViewProjection(const FDisplayClusterShaderParametersICVFX_CameraViewProjection& InLocalSpaceViewProjection)
 		{
 			// Support icvfx stereo - update context from camera for each eye
-			CameraViewRotation = Local2WorldTransform.InverseTransformRotation(InContext.CameraViewRotation.Quaternion()).Rotator();
-			CameraViewLocation = Local2WorldTransform.InverseTransformPosition(InContext.CameraViewLocation);
-			CameraPrjMatrix = InContext.CameraPrjMatrix;
+			ViewProjection.ViewRotation = Local2WorldTransform.InverseTransformRotation(InLocalSpaceViewProjection.ViewRotation.Quaternion()).Rotator();
+			ViewProjection.ViewLocation = Local2WorldTransform.InverseTransformPosition(InLocalSpaceViewProjection.ViewLocation);
+			ViewProjection.PrjMatrix = InLocalSpaceViewProjection.PrjMatrix;
 		}
 
 		/**
@@ -185,21 +219,57 @@ public:
 			}
 		}
 
+		/**
+		 * Copying camera settings from the source camera.
+		 * The function is moved here from the DisplayClusterMedia module, because it must also be updated when new settings are added.
+		 * 
+		 * @param InCameraSettings - source camera settings
+		 * @param bExcludeResources - disables resource copying
+		 */
+		inline void SetCameraSettings(const FCameraSettings& InCameraSettings, const bool bIncludeResources)
+		{
+			if (bIncludeResources)
+			{
+				Resource = InCameraSettings.Resource;
+				Chromakey = InCameraSettings.Chromakey;
+				ChromakeMarkerTextureRHI = InCameraSettings.ChromakeMarkerTextureRHI;
+			}
+
+			SoftEdge = InCameraSettings.SoftEdge;
+
+			InnerCameraBorderColor      = InCameraSettings.InnerCameraBorderColor;
+			InnerCameraBorderThickness  = InCameraSettings.InnerCameraBorderThickness;
+			InnerCameraFrameAspectRatio = InCameraSettings.InnerCameraFrameAspectRatio;
+
+			Local2WorldTransform = InCameraSettings.Local2WorldTransform;
+
+			ViewProjection = InCameraSettings.ViewProjection;
+
+			ChromakeySource = InCameraSettings.ChromakeySource;
+			ChromakeyColor = InCameraSettings.ChromakeyColor;
+
+			ChromakeyMarkersColor = InCameraSettings.ChromakeyMarkersColor;
+			ChromakeyMarkersScale = InCameraSettings.ChromakeyMarkersScale;
+			ChromakeyMarkersDistance = InCameraSettings.ChromakeyMarkersDistance;
+			ChromakeyMarkersOffset = InCameraSettings.ChromakeyMarkersOffset;
+
+			RenderOrder = InCameraSettings.RenderOrder;
+		}
+
 	public:
+		// resource with the camera image
 		FDisplayClusterShaderParametersICVFX_ViewportResource Resource;
+
 		FVector4 SoftEdge;
 
 		FLinearColor InnerCameraBorderColor = FLinearColor::Black;
 		float InnerCameraBorderThickness = 0.1f;
 		float InnerCameraFrameAspectRatio = 1.0f;
 
-		// Camera Origin
 		FTransform Local2WorldTransform;
 
-		// Camera world
-		FRotator CameraViewRotation;
-		FVector  CameraViewLocation;
-		FMatrix  CameraPrjMatrix;
+		// Camera view projection data
+		FDisplayClusterShaderParametersICVFX_CameraViewProjection ViewProjection;
 
 		// Chromakey settings:
 		EDisplayClusterShaderParametersICVFX_ChromakeySource  ChromakeySource = EDisplayClusterShaderParametersICVFX_ChromakeySource::Disabled;
@@ -222,14 +292,14 @@ public:
 	template <typename Predicate>
 	void IterateViewportResourcesByPredicate(Predicate Pred)
 	{
-		if (Lightcard.IsDefined())
+		if (LightCard.IsDefined())
 		{
-			::Invoke(Pred, Lightcard);
+			::Invoke(Pred, LightCard);
 		}
 
-		if (UVLightcard.IsDefined())
+		if (UVLightCard.IsDefined())
 		{
-			::Invoke(Pred, UVLightcard);
+			::Invoke(Pred, UVLightCard);
 		}
 
 		for (FCameraSettings& CameraIt : Cameras)
@@ -250,7 +320,7 @@ public:
 	}
 
 	// Remove unused cameras from render
-	bool CleanupCamerasForRender()
+	inline bool CleanupCamerasForRender()
 	{
 		const TArray<FCameraSettings> InCameras = Cameras;
 		Cameras.Empty();
@@ -266,10 +336,14 @@ public:
 		return Cameras.Num() == InCameras.Num();
 	}
 
+	// All cameras that render on this viewport
 	TArray<FCameraSettings> Cameras;
 
-	// Lightcard settings
-	FDisplayClusterShaderParametersICVFX_ViewportResource    UVLightcard;
-	FDisplayClusterShaderParametersICVFX_ViewportResource    Lightcard;
-	EDisplayClusterShaderParametersICVFX_LightcardRenderMode LightcardMode = EDisplayClusterShaderParametersICVFX_LightcardRenderMode::Under;
+	// Rendering mode for overlapping areas of camera projections
+	EDisplayClusterShaderParametersICVFX_CameraOverlappingRenderMode CameraOverlappingRenderMode = EDisplayClusterShaderParametersICVFX_CameraOverlappingRenderMode::None;
+
+	// LightCard settings
+	FDisplayClusterShaderParametersICVFX_ViewportResource    UVLightCard;
+	FDisplayClusterShaderParametersICVFX_ViewportResource    LightCard;
+	EDisplayClusterShaderParametersICVFX_LightCardRenderMode LightCardMode = EDisplayClusterShaderParametersICVFX_LightCardRenderMode::Under;
 };
