[nDisplay] Expose bIsLightCardFlag and bIsUVLightCard with getters + setters for remote control access

diff --git a/Engine/Plugins/Runtime/nDisplay/Source/DisplayCluster/Private/Game/EngineClasses/Scene/DisplayClusterLightCardActor.cpp b/Engine/Plugins/Runtime/nDisplay/Source/DisplayCluster/Private/Game/EngineClasses/Scene/DisplayClusterLightCardActor.cpp
--- a/Engine/Plugins/Runtime/nDisplay/Source/DisplayCluster/Private/Game/EngineClasses/Scene/DisplayClusterLightCardActor.cpp
+++ b/Engine/Plugins/Runtime/nDisplay/Source/DisplayCluster/Private/Game/EngineClasses/Scene/DisplayClusterLightCardActor.cpp
@@ -577,6 +577,11 @@ void ADisplayClusterLightCardActor::SetIsLightCardFlag(bool bNewFlagValue)
 	bIsLightCardFlag = bNewFlagValue;
 }
 
+void ADisplayClusterLightCardActor::SetIsUVActor(bool bNewUVValue)
+{
+	bIsUVLightCard = bNewUVValue;
+}
+
 void ADisplayClusterLightCardActor::ShowLightCardLabel(const FDisplayClusterLabelConfiguration& InLabelConfiguration)
 {
 	if (IsUVActor())
diff --git a/Engine/Plugins/Runtime/nDisplay/Source/DisplayCluster/Public/DisplayClusterLightCardActor.h b/Engine/Plugins/Runtime/nDisplay/Source/DisplayCluster/Public/DisplayClusterLightCardActor.h
--- a/Engine/Plugins/Runtime/nDisplay/Source/DisplayCluster/Public/DisplayClusterLightCardActor.h
+++ b/Engine/Plugins/Runtime/nDisplay/Source/DisplayCluster/Public/DisplayClusterLightCardActor.h
@@ -118,6 +118,9 @@ public:
 	/** If this light card is considered a flag */
 	bool IsLightCardFlag() const { return bIsLightCardFlag; }
 
+	/** Configures this light card as a UV actor */
+	void SetIsUVActor(bool bNewUVValue);
+
 	/** Configure light card labels */
 	void ShowLightCardLabel(const FDisplayClusterLabelConfiguration& InLabelConfiguration);
 	
@@ -276,7 +279,7 @@ public:
 	bool bLockToOwningRootActor = true;
 
 	/** Indicates if the light card exists in 3D space or in UV space */
-	UPROPERTY()
+	UPROPERTY(Getter = IsUVActor, Setter = SetIsUVActor)
 	bool bIsUVLightCard = false;
 
 	/** Polygon points when using this type of mask */
@@ -334,7 +337,7 @@ protected:
 #endif
 
 	/** Indicates this light card should be considered a flag */
-	UPROPERTY()
+	UPROPERTY(Getter = IsLightCardFlag, Setter = SetIsLightCardFlag, meta = (AllowPrivateAccess))
 	bool bIsLightCardFlag = false;
 
 private:
