nDisplay: Reworked OCIO implementation:
OCIO is applied in different ways. It depends on the rendering workflow:
- When the viewport renders with a postprocess, OCIO must be done in between
- When the viewport is rendered without postprocessing, OCIO is applied last, to the RTT texture of the viewport

diff --git a/Engine/Plugins/Runtime/nDisplay/Source/DisplayCluster/Private/Render/Viewport/DisplayClusterViewportManagerViewExtension.cpp b/Engine/Plugins/Runtime/nDisplay/Source/DisplayCluster/Private/Render/Viewport/DisplayClusterViewportManagerViewExtension.cpp
--- a/Engine/Plugins/Runtime/nDisplay/Source/DisplayCluster/Private/Render/Viewport/DisplayClusterViewportManagerViewExtension.cpp
+++ b/Engine/Plugins/Runtime/nDisplay/Source/DisplayCluster/Private/Render/Viewport/DisplayClusterViewportManagerViewExtension.cpp
@@ -99,6 +99,17 @@ void FDisplayClusterViewportManagerViewExtension::SubscribeToPostProcessingPass(
 		}
 		break;
 
+	case EPostProcessingPass::Tonemap:
+		for (const FViewportProxy& ViewportIt : ViewportProxies)
+		{
+			if (ViewportIt.ViewportProxy.IsValid() && ViewportIt.ViewportProxy.Pin()->ShouldUsePostProcessPassTonemap())
+			{
+				InOutPassCallbacks.Add(FAfterPassCallbackDelegate::CreateRaw(this, &FDisplayClusterViewportManagerViewExtension::PostProcessPassAfterTonemap_RenderThread));
+				break;
+			}
+		}
+		break;
+
 	default:
 		break;
 	}
@@ -142,6 +153,25 @@ FScreenPassTexture FDisplayClusterViewportManagerViewExtension::PostProcessPassA
 	return ReturnUntouchedSceneColorForPostProcessing(Inputs);
 }
 
+FScreenPassTexture FDisplayClusterViewportManagerViewExtension::PostProcessPassAfterTonemap_RenderThread(FRDGBuilder& GraphBuilder, const FSceneView& View, const FPostProcessMaterialInputs& Inputs)
+{
+	if (!IsActive())
+	{
+		return ReturnUntouchedSceneColorForPostProcessing(Inputs);
+	}
+
+	uint32 ContextNum = 0;
+	if (FDisplayClusterViewportProxy* ViewportProxyPtr = ViewportManagerProxy.Pin()->ImplFindViewport_RenderThread(View.StereoViewIndex, &ContextNum))
+	{
+		if (ViewportProxyPtr->ShouldUsePostProcessPassTonemap())
+		{
+			return ViewportProxyPtr->OnPostProcessPassAfterTonemap_RenderThread(GraphBuilder, View, Inputs, ContextNum);
+		}
+	}
+
+	return ReturnUntouchedSceneColorForPostProcessing(Inputs);
+}
+
 /**
 * A helper function that extracts the right scene color texture, untouched, to be used further in post processing.
 */
diff --git a/Engine/Plugins/Runtime/nDisplay/Source/DisplayCluster/Private/Render/Viewport/DisplayClusterViewportManagerViewExtension.h b/Engine/Plugins/Runtime/nDisplay/Source/DisplayCluster/Private/Render/Viewport/DisplayClusterViewportManagerViewExtension.h
--- a/Engine/Plugins/Runtime/nDisplay/Source/DisplayCluster/Private/Render/Viewport/DisplayClusterViewportManagerViewExtension.h
+++ b/Engine/Plugins/Runtime/nDisplay/Source/DisplayCluster/Private/Render/Viewport/DisplayClusterViewportManagerViewExtension.h
@@ -89,6 +89,16 @@ private:
 	 */
 	FScreenPassTexture PostProcessPassAfterSSRInput_RenderThread(FRDGBuilder& GraphBuilder, const FSceneView& View, const FPostProcessMaterialInputs& Inputs);
 
+	/** Callback OnPostProcessPassAfterTonemap.
+	 *
+	 * @param GraphBuilder - RDG interface
+	 * @param View         - Scene View
+	 * @param Inputs       - PP Input resources
+	 *
+	 * @return - Screen pass texture
+	 */
+	FScreenPassTexture PostProcessPassAfterTonemap_RenderThread(FRDGBuilder& GraphBuilder, const FSceneView& View, const FPostProcessMaterialInputs& Inputs);
+
 	/** Register callbacks to external modules (Renderer:ResolvedSceneColor). */
 	void RegisterCallbacks();
 
diff --git a/Engine/Plugins/Runtime/nDisplay/Source/DisplayCluster/Private/Render/Viewport/DisplayClusterViewportProxy.cpp b/Engine/Plugins/Runtime/nDisplay/Source/DisplayCluster/Private/Render/Viewport/DisplayClusterViewportProxy.cpp
--- a/Engine/Plugins/Runtime/nDisplay/Source/DisplayCluster/Private/Render/Viewport/DisplayClusterViewportProxy.cpp
+++ b/Engine/Plugins/Runtime/nDisplay/Source/DisplayCluster/Private/Render/Viewport/DisplayClusterViewportProxy.cpp
@@ -458,6 +458,23 @@ bool FDisplayClusterViewportProxy::ImplGetResources_RenderThread(const EDisplayC
 	return false;
 }
 
+EDisplayClusterViewportOpenColorIOMode FDisplayClusterViewportProxy::GetOpenColorIOMode() const
+{
+	if (OpenColorIO.IsValid() && OpenColorIO->IsEnabled_RenderThread())
+	{
+		if (EnumHasAnyFlags(RenderSettingsICVFX.RuntimeFlags, EDisplayClusterViewportRuntimeICVFXFlags::UVLightcard | EDisplayClusterViewportRuntimeICVFXFlags::Lightcard | EDisplayClusterViewportRuntimeICVFXFlags::Chromakey))
+		{
+			// Rendering without post-processing, OCIO is applied last, to the RTT texture of the viewport
+			return EDisplayClusterViewportOpenColorIOMode::Resolved;
+		}
+
+		// By default, viewports render with a postprocess, OCIO must be done in between.
+		return EDisplayClusterViewportOpenColorIOMode::PostProcess;
+	}
+
+	return EDisplayClusterViewportOpenColorIOMode::None;
+}
+
 void FDisplayClusterViewportProxy::PostResolveViewport_RenderThread(FRHICommandListImmediate& RHICmdList) const
 {
 	check(IsInRenderingThread());
@@ -688,8 +705,8 @@ void FDisplayClusterViewportProxy::UpdateDeferredResources(FRHICommandListImmedi
 
 	bool bPass0Applied = false;
 
-	// Support OCIO on pass 0
-	if (OpenColorIO.IsValid() && OpenColorIO->IsEnabled_RenderThread())
+	// OCIO support on the first pass for an resolved RTT
+	if (GetOpenColorIOMode() == EDisplayClusterViewportOpenColorIOMode::Resolved)
 	{
 		TArray<FRHITexture2D*> Input, Output;
 		TArray<FIntRect> InputRects, OutputRects;
@@ -1064,6 +1081,16 @@ bool FDisplayClusterViewportProxy::ShouldUsePostProcessPassAfterFXAA() const
 	return false;
 }
 
+bool FDisplayClusterViewportProxy::ShouldUsePostProcessPassTonemap() const
+{
+	if(GetOpenColorIOMode() == EDisplayClusterViewportOpenColorIOMode::PostProcess)
+	{
+		return true;
+	}
+
+	return false;
+}
+
 void FDisplayClusterViewportProxy::OnResolvedSceneColor_RenderThread(FRDGBuilder& GraphBuilder, const FSceneTextures& SceneTextures, const FDisplayClusterViewportProxy_Context& InProxyContext)
 {
 	const uint32 InContextNum = InProxyContext.ContextNum;
@@ -1114,6 +1141,17 @@ FScreenPassTexture FDisplayClusterViewportProxy::OnPostProcessPassAfterFXAA_Rend
 	return OutScreenPassTexture;
 }
 
+FScreenPassTexture FDisplayClusterViewportProxy::OnPostProcessPassAfterTonemap_RenderThread(FRDGBuilder& GraphBuilder, const FSceneView& View, const FPostProcessMaterialInputs& Inputs, const uint32 ContextNum)
+{
+	// Perform OCIO rendering after the tonemapper
+	if (GetOpenColorIOMode() == EDisplayClusterViewportOpenColorIOMode::PostProcess)
+	{
+		return OpenColorIO->PostProcessPassAfterTonemap_RenderThread(GraphBuilder, GetContexts_RenderThread()[ContextNum], View, Inputs);
+	}
+
+	return FDisplayClusterViewportManagerViewExtension::ReturnUntouchedSceneColorForPostProcessing(Inputs);
+}
+
 void FDisplayClusterViewportProxy::OnPostRenderViewFamily_RenderThread(FRDGBuilder& GraphBuilder, FSceneViewFamily& InViewFamily, const FSceneView& InSceneView, const FDisplayClusterViewportProxy_Context& InProxyContext)
 {
 	const uint32 InContextNum = InProxyContext.ContextNum;
diff --git a/Engine/Plugins/Runtime/nDisplay/Source/DisplayCluster/Private/Render/Viewport/DisplayClusterViewportProxy.h b/Engine/Plugins/Runtime/nDisplay/Source/DisplayCluster/Private/Render/Viewport/DisplayClusterViewportProxy.h
--- a/Engine/Plugins/Runtime/nDisplay/Source/DisplayCluster/Private/Render/Viewport/DisplayClusterViewportProxy.h
+++ b/Engine/Plugins/Runtime/nDisplay/Source/DisplayCluster/Private/Render/Viewport/DisplayClusterViewportProxy.h
@@ -40,6 +40,20 @@ enum class EDisplayClusterTextureCopyMode : uint8
 	Alpha,
 };
 
+/**
+ * OCIO is applied in different ways. It depends on the rendering workflow
+ */
+enum class EDisplayClusterViewportOpenColorIOMode : uint8
+{
+	None = 0,
+
+	// When the viewport renders with a postprocess, OCIO must be done in between
+	PostProcess,
+
+	// When the viewport is rendered without postprocessing, OCIO is applied last, to the RTT texture of the viewport
+	Resolved
+};
+
 /**
  * nDisplay viewport proxy implementation.
  */
@@ -184,6 +198,20 @@ public:
 	 */
 	FScreenPassTexture OnPostProcessPassAfterSSRInput_RenderThread(FRDGBuilder& GraphBuilder, const FSceneView& View, const FPostProcessMaterialInputs& Inputs, const uint32 ContextNum);
 
+	/** Allow callback OnPostProcessPassAfterTonemap. */
+	bool ShouldUsePostProcessPassTonemap() const;
+
+	/** Callback OnPostProcessPassAfterTonemap.
+	 *
+	 * @param GraphBuilder - RDG interface
+	 * @param View         - Scene View
+	 * @param Inputs       - PP Input resources
+	 * @param ContextNum   - viewport context index
+	 *
+	 * @return - Screen pass texture
+	 */
+	FScreenPassTexture OnPostProcessPassAfterTonemap_RenderThread(FRDGBuilder& GraphBuilder, const FSceneView& View, const FPostProcessMaterialInputs& Inputs, const uint32 ContextNum);
+
 	/** Enable alpha channel for this viewport (useful for overlays with alpha channel: ChromaKey, LightCard). */
 	bool ShouldUseAlphaChannel_RenderThread() const;
 
@@ -277,6 +305,9 @@ private:
 	/** Check if there is an RTT source (internal or external) in this viewport proxy. */
 	bool IsInputRenderTargetResourceExists() const;
 
+	/** Returns the OCIO rendering type for the given viewport. */
+	EDisplayClusterViewportOpenColorIOMode GetOpenColorIOMode() const;
+
 protected:
 	friend FDisplayClusterViewportProxyData;
 	friend FDisplayClusterViewportManagerProxy;
diff --git a/Engine/Plugins/Runtime/nDisplay/Source/DisplayCluster/Private/Render/Viewport/DisplayClusterViewport_OpenColorIO.cpp b/Engine/Plugins/Runtime/nDisplay/Source/DisplayCluster/Private/Render/Viewport/DisplayClusterViewport_OpenColorIO.cpp
--- a/Engine/Plugins/Runtime/nDisplay/Source/DisplayCluster/Private/Render/Viewport/DisplayClusterViewport_OpenColorIO.cpp
+++ b/Engine/Plugins/Runtime/nDisplay/Source/DisplayCluster/Private/Render/Viewport/DisplayClusterViewport_OpenColorIO.cpp
@@ -16,6 +16,8 @@
 #include "SceneView.h"
 #include "ScreenPass.h"
 #include "CommonRenderResources.h"
+// for FPostProcessMaterialInputs
+#include "PostProcess/PostProcessMaterial.h"
 
 //////////////////////////////////////////////////////////////////////////
 // FDisplayClusterViewport_OpenColorIO
@@ -111,6 +113,18 @@ bool FDisplayClusterViewport_OpenColorIO::IsConversionSettingsEqual(const FOpenC
 	return ConversionSettings.ToString() == InConversionSettings.ToString();
 }
 
+float FDisplayClusterViewport_OpenColorIO::GetDisplayGamma(const FDisplayClusterViewport_Context& InViewportContext) const
+{
+	const float DefaultDisplayGamma = FOpenColorIODisplayExtension::DefaultDisplayGamma;
+	const float EngineDisplayGamma = InViewportContext.RenderThreadData.EngineDisplayGamma;
+
+	// There is a special case where post processing and tonemapper are disabled. In this case tonemapper applies a static display Inverse of Gamma which defaults to 2.2.
+	// In the case when Both PostProcessing and ToneMapper are disabled we apply gamma manually. In every other case we apply inverse gamma before applying OCIO.
+	const float DisplayGamma = (InViewportContext.RenderThreadData.EngineShowFlags.Tonemapper == 0) || (InViewportContext.RenderThreadData.EngineShowFlags.PostProcessing == 0) ? DefaultDisplayGamma : DefaultDisplayGamma / EngineDisplayGamma;
+
+	return DisplayGamma;
+}
+
 bool FDisplayClusterViewport_OpenColorIO::AddPass_RenderThread(FRDGBuilder& GraphBuilder, const FDisplayClusterViewport_Context& InViewportContext,
 	FRHITexture2D* InputTextureRHI, const FIntRect& InputRect, FRHITexture2D* OutputTextureRHI, const FIntRect& OutputRect) const
 {
@@ -121,13 +135,6 @@ bool FDisplayClusterViewport_OpenColorIO::AddPass_RenderThread(FRDGBuilder& Grap
 
 		FScreenPassRenderTarget Output = FScreenPassRenderTarget(OutputTexture, OutputRect, ERenderTargetLoadAction::EClear);
 
-		const float DefaultDisplayGamma = FOpenColorIODisplayExtension::DefaultDisplayGamma;
-		const float EngineDisplayGamma = InViewportContext.RenderThreadData.EngineDisplayGamma;
-
-		// There is a special case where post processing and tonemapper are disabled. In this case tonemapper applies a static display Inverse of Gamma which defaults to 2.2.
-		// In the case when Both PostProcessing and ToneMapper are disabled we apply gamma manually. In every other case we apply inverse gamma before applying OCIO.
-		const float DisplayGamma = (InViewportContext.RenderThreadData.EngineShowFlags.Tonemapper == 0) || (InViewportContext.RenderThreadData.EngineShowFlags.PostProcessing == 0) ? DefaultDisplayGamma : DefaultDisplayGamma / EngineDisplayGamma;
-
 		FOpenColorIORendering::AddPass_RenderThread(
 			GraphBuilder,
 			FScreenPassViewInfo(),
@@ -135,7 +142,7 @@ bool FDisplayClusterViewport_OpenColorIO::AddPass_RenderThread(FRDGBuilder& Grap
 			FScreenPassTexture(InputTexture),
 			Output,
 			CachedResourcesRenderThread,
-			DisplayGamma
+			GetDisplayGamma(InViewportContext)
 		);
 
 		return true;
@@ -143,3 +150,30 @@ bool FDisplayClusterViewport_OpenColorIO::AddPass_RenderThread(FRDGBuilder& Grap
 
 	return false;
 }
+
+// This is a copy of FOpenColorIODisplayExtension::PostProcessPassAfterTonemap_RenderThread()
+FScreenPassTexture FDisplayClusterViewport_OpenColorIO::PostProcessPassAfterTonemap_RenderThread(FRDGBuilder& GraphBuilder, const FDisplayClusterViewport_Context& InViewportContext, const FSceneView& View, const FPostProcessMaterialInputs& InOutInputs)
+{
+	const FScreenPassTexture& SceneColor = InOutInputs.GetInput(EPostProcessMaterialInput::SceneColor);
+	check(SceneColor.IsValid());
+	checkSlow(View.bIsViewInfo);
+	const FViewInfo& ViewInfo = static_cast<const FViewInfo&>(View);
+	FScreenPassRenderTarget Output = InOutInputs.OverrideOutput;
+
+	// If the override output is provided, it means that this is the last pass in post processing.
+	if (!Output.IsValid())
+	{
+		Output = FScreenPassRenderTarget::CreateFromInput(GraphBuilder, SceneColor, ViewInfo.GetOverwriteLoadAction(), TEXT("OCIORenderTarget"));
+	}
+
+	FOpenColorIORendering::AddPass_RenderThread(
+		GraphBuilder,
+		ViewInfo,
+		SceneColor,
+		Output,
+		CachedResourcesRenderThread,
+		GetDisplayGamma(InViewportContext)
+	);
+
+	return MoveTemp(Output);
+}
diff --git a/Engine/Plugins/Runtime/nDisplay/Source/DisplayCluster/Private/Render/Viewport/DisplayClusterViewport_OpenColorIO.h b/Engine/Plugins/Runtime/nDisplay/Source/DisplayCluster/Private/Render/Viewport/DisplayClusterViewport_OpenColorIO.h
--- a/Engine/Plugins/Runtime/nDisplay/Source/DisplayCluster/Private/Render/Viewport/DisplayClusterViewport_OpenColorIO.h
+++ b/Engine/Plugins/Runtime/nDisplay/Source/DisplayCluster/Private/Render/Viewport/DisplayClusterViewport_OpenColorIO.h
@@ -12,6 +12,13 @@
 #include "RHIResources.h"
 #include "Templates/SharedPointer.h"
 
+class FSceneView;
+class FSceneViewFamily;
+struct FScreenPassTexture;
+struct FPostProcessMaterialInputs;
+class FRDGBuilder;
+class FDisplayClusterViewport_Context;
+
 /**
  * nDisplay OCIO implementation.
  * 
@@ -34,7 +41,7 @@ public:
 	 *
 	 * @return - none.
 	 */
-	void SetupSceneView(class FSceneViewFamily& InOutViewFamily, class FSceneView& InOutView) const;
+	void SetupSceneView(FSceneViewFamily& InOutViewFamily, FSceneView& InOutView) const;
 
 	/** Add OCIO render pass.
 	 *
@@ -47,9 +54,12 @@ public:
 	 *
 	 * @return - true if success.
 	 */
-	bool AddPass_RenderThread(class FRDGBuilder& GraphBuilder, const class FDisplayClusterViewport_Context& InViewportContext,
+	bool AddPass_RenderThread(FRDGBuilder& GraphBuilder, const FDisplayClusterViewport_Context& InViewportContext,
 		FRHITexture2D* InputTextureRHI, const FIntRect& InputRect, FRHITexture2D* OutputTextureRHI, const FIntRect& OutputRect) const;
 
+	/* This is a copy of FOpenColorIODisplayExtension::PostProcessPassAfterTonemap_RenderThread() */
+	FScreenPassTexture PostProcessPassAfterTonemap_RenderThread(FRDGBuilder& GraphBuilder, const FDisplayClusterViewport_Context& InViewportContext, const FSceneView& View, const FPostProcessMaterialInputs& InOutInputs);
+
 	/** Is OCIO enabled for render thread */
 	bool IsEnabled_RenderThread() const;
 
@@ -67,6 +77,10 @@ public:
 		return ConversionSettings;
 	}
 
+private:
+	/* returns the DisplayGamma of OCIO for the viewport context.*/
+	float GetDisplayGamma(const FDisplayClusterViewport_Context& InViewportContext) const;
+
 private:
 	/** Cached pass resources required to apply conversion for render thread. */
 	FOpenColorIORenderPassResources CachedResourcesRenderThread;
diff --git a/Engine/Plugins/Runtime/nDisplay/Source/DisplayCluster/Private/Render/Viewport/RenderTarget/DisplayClusterRenderTargetManager.cpp b/Engine/Plugins/Runtime/nDisplay/Source/DisplayCluster/Private/Render/Viewport/RenderTarget/DisplayClusterRenderTargetManager.cpp
--- a/Engine/Plugins/Runtime/nDisplay/Source/DisplayCluster/Private/Render/Viewport/RenderTarget/DisplayClusterRenderTargetManager.cpp
+++ b/Engine/Plugins/Runtime/nDisplay/Source/DisplayCluster/Private/Render/Viewport/RenderTarget/DisplayClusterRenderTargetManager.cpp
@@ -12,28 +12,67 @@
 
 #include "Misc/DisplayClusterLog.h"
 
+
+int32 GDisplayClusterSceneColorFormat = 0;
+static FAutoConsoleVariableRef CVarDisplayClusterSceneColorFormat(
+	TEXT("nDisplay.render.SceneColorFormat"),
+	GDisplayClusterSceneColorFormat,
+	TEXT("Defines the memory layout (RGBA) used for the scene color\n"
+		"(affects performance, mostly through bandwidth, quality especially with translucency).\n"
+		" 0: Use the backbuffer format\n"
+		" 1: PF_FloatRGBA 64Bit (default, might be overkill, especially if translucency is mostly using SeparateTranslucency)\n"
+		" 2: PF_A32B32G32R32F 128Bit (unreasonable but good for testing)"),
+	ECVF_Default
+);
+
 ////////////////////////////////////////////////////////////////////////////////////////////////////////////
 namespace DisplayClusterRenderTargetManager
 {
 	static EPixelFormat ImplGetCustomFormat(EDisplayClusterViewportCaptureMode CaptureMode)
 	{
+		EPixelFormat OutPixelFormat = EPixelFormat::PF_Unknown;
+
 		// Pre-defined formats for targets:
 		switch (CaptureMode)
 		{
 		case EDisplayClusterViewportCaptureMode::Chromakey:
 		case EDisplayClusterViewportCaptureMode::Lightcard:
 			// The Chromakey and LightCard always uses PF_FloatRGBA for OCIO support.
-			return EPixelFormat::PF_FloatRGBA;
+			OutPixelFormat = EPixelFormat::PF_FloatRGBA;
+			break;
 
 		case EDisplayClusterViewportCaptureMode::MoviePipeline:
 			// Movie pipeline always use PF_FloatRGBA
-			return EPixelFormat::PF_FloatRGBA;
+			OutPixelFormat = EPixelFormat::PF_FloatRGBA;
+			break;
 
 		default:
+			switch (GDisplayClusterSceneColorFormat)
+			{
+			case 1:
+				OutPixelFormat = EPixelFormat::PF_FloatRGBA;
+				break;
+			case 2:
+				OutPixelFormat = EPixelFormat::PF_A32B32G32R32F;
+				break;
+
+			default:
+				// by default use the backbuffer format (when function returns PF_Unknown its mean use default scene format)
+				break;
+			}
 			break;
 		}
 
-		return EPixelFormat::PF_Unknown;
+		if (OutPixelFormat != EPixelFormat::PF_Unknown)
+		{
+			// Fallback in case the scene color selected isn't supported.
+			if (!GPixelFormats[OutPixelFormat].Supported)
+			{
+				OutPixelFormat = PF_FloatRGBA;
+			}
+		}
+
+		return OutPixelFormat;
 	}
 
 	static void ImplViewportTextureResourceLogging(const FDisplayClusterViewport* InViewport, const uint32 ContextNum, const FString& ResourceId, const FDisplayClusterViewportTextureResource* InTextureResource)
