nDisplay: Fixed issue that would cause crash when deleting an nDisplay config actor in the level while shaders are being compiled

diff --git a/Engine/Plugins/Runtime/nDisplay/Source/DisplayCluster/Private/Game/EngineClasses/Scene/DisplayClusterPreviewComponent.cpp b/Engine/Plugins/Runtime/nDisplay/Source/DisplayCluster/Private/Game/EngineClasses/Scene/DisplayClusterPreviewComponent.cpp
--- a/Engine/Plugins/Runtime/nDisplay/Source/DisplayCluster/Private/Game/EngineClasses/Scene/DisplayClusterPreviewComponent.cpp
+++ b/Engine/Plugins/Runtime/nDisplay/Source/DisplayCluster/Private/Game/EngineClasses/Scene/DisplayClusterPreviewComponent.cpp
@@ -287,7 +287,13 @@ void UDisplayClusterPreviewComponent::ReleasePreviewMaterial()
 {
 	if (PreviewMaterialInstance != nullptr)
 	{
-		PreviewMaterialInstance->SetTextureParameterValue(TEXT("Preview"), nullptr);
+		// Clear the material parameters to ensure that no pointers to the texture resources are being kept around
+		// Materials destroy their resources in BeginDestroy, so only clear the parameters BeginDestroy hasn't already been called
+		if (!PreviewMaterialInstance->HasAnyFlags(RF_BeginDestroyed))
+		{
+			PreviewMaterialInstance->ClearParameterValues();
+		}
+
 		PreviewMaterialInstance = nullptr;
 	}
 }
