Fix missing parameter
#rnx

diff --git a/Engine/Plugins/Runtime/nDisplay/Source/DisplayCluster/Private/Render/Viewport/DisplayClusterViewport_OpenColorIO.cpp b/Engine/Plugins/Runtime/nDisplay/Source/DisplayCluster/Private/Render/Viewport/DisplayClusterViewport_OpenColorIO.cpp
--- a/Engine/Plugins/Runtime/nDisplay/Source/DisplayCluster/Private/Render/Viewport/DisplayClusterViewport_OpenColorIO.cpp
+++ b/Engine/Plugins/Runtime/nDisplay/Source/DisplayCluster/Private/Render/Viewport/DisplayClusterViewport_OpenColorIO.cpp
@@ -169,6 +169,7 @@ FScreenPassTexture FDisplayClusterViewport_OpenColorIO::PostProcessPassAfterTone
 	FOpenColorIORendering::AddPass_RenderThread(
 		GraphBuilder,
 		ViewInfo,
+		ViewInfo.GetFeatureLevel(),
 		SceneColor,
 		Output,
 		CachedResourcesRenderThread,
