nDisplay: Fix ndisplay preview not taking into account exposure compensation.

View states are disable by default in the preview, but eye adaptation requires them.

Also fixed typo in details panel (unrelated to the main issue).

diff --git a/Engine/Plugins/Runtime/nDisplay/Source/DisplayCluster/Private/Render/Viewport/DisplayClusterViewport_EditorPreview.cpp b/Engine/Plugins/Runtime/nDisplay/Source/DisplayCluster/Private/Render/Viewport/DisplayClusterViewport_EditorPreview.cpp
--- a/Engine/Plugins/Runtime/nDisplay/Source/DisplayCluster/Private/Render/Viewport/DisplayClusterViewport_EditorPreview.cpp
+++ b/Engine/Plugins/Runtime/nDisplay/Source/DisplayCluster/Private/Render/Viewport/DisplayClusterViewport_EditorPreview.cpp
@@ -29,7 +29,7 @@
 #include "Render/Viewport/Configuration/DisplayClusterViewportConfigurationHelpers_Postprocess.h"
 
 ///////////////////////////////////////////////////////////////////////////////////////
-int32 GDisplayClusterPreviewEnableViewState = 0;
+int32 GDisplayClusterPreviewEnableViewState = 1;
 static FAutoConsoleVariableRef CVarDisplayClusterPreviewEnableViewState(
 	TEXT("nDisplay.preview.EnableViewState"),
 	GDisplayClusterPreviewEnableViewState,
diff --git a/Engine/Plugins/Runtime/nDisplay/Source/DisplayCluster/Public/DisplayClusterRootActor.h b/Engine/Plugins/Runtime/nDisplay/Source/DisplayCluster/Public/DisplayClusterRootActor.h
--- a/Engine/Plugins/Runtime/nDisplay/Source/DisplayCluster/Public/DisplayClusterRootActor.h
+++ b/Engine/Plugins/Runtime/nDisplay/Source/DisplayCluster/Public/DisplayClusterRootActor.h
@@ -240,7 +240,7 @@ private:
 	UPROPERTY(EditAnywhere, Transient, Category = "Color Grading", meta = (PropertyPath = "CurrentConfigData.StageSettings.PerViewportColorGrading"))
 	FDisplayClusterEditorPropertyReference PerViewportColorGradingRef;
 
-	UPROPERTY(EditAnywhere, Transient, Category = OCIO, meta = (PropertyPath = "CurrentConfigData.StageSettings.ViewportOCIO.AllViewportsOCIOConfiguration.bIsEnabled", DisplayName = "Enable Viewpport OCIO"))
+	UPROPERTY(EditAnywhere, Transient, Category = OCIO, meta = (PropertyPath = "CurrentConfigData.StageSettings.ViewportOCIO.AllViewportsOCIOConfiguration.bIsEnabled", DisplayName = "Enable Viewport OCIO"))
 	FDisplayClusterEditorPropertyReference EnableViewportOCIORef;
 
 	UPROPERTY(EditAnywhere, Transient, Category = OCIO, meta = (PropertyPath = "CurrentConfigData.StageSettings.ViewportOCIO.AllViewportsOCIOConfiguration.ColorConfiguration", DisplayName = "All Viewports Color Configuration"))
