Fixed shadowed variable errors.

diff --git a/Engine/Plugins/Runtime/nDisplay/Source/DisplayClusterProjection/Private/Policy/Domeprojection/Windows/DX11/DisplayClusterProjectionDomeprojectionViewAdapterDX11.cpp b/Engine/Plugins/Runtime/nDisplay/Source/DisplayClusterProjection/Private/Policy/Domeprojection/Windows/DX11/DisplayClusterProjectionDomeprojectionViewAdapterDX11.cpp
--- a/Engine/Plugins/Runtime/nDisplay/Source/DisplayClusterProjection/Private/Policy/Domeprojection/Windows/DX11/DisplayClusterProjectionDomeprojectionViewAdapterDX11.cpp
+++ b/Engine/Plugins/Runtime/nDisplay/Source/DisplayClusterProjection/Private/Policy/Domeprojection/Windows/DX11/DisplayClusterProjectionDomeprojectionViewAdapterDX11.cpp
@@ -221,20 +221,20 @@ bool FDisplayClusterProjectionDomeprojectionViewAdapterDX11::ApplyWarpBlend_Rend
 	return true;
 }
 
-void FDisplayClusterProjectionDomeprojectionViewAdapterDX11::FViewData::Release(FCriticalSection& DllAccessCS)
+void FDisplayClusterProjectionDomeprojectionViewAdapterDX11::FViewData::Release(FCriticalSection& InDllAccessCS)
 {
 	if (Context != nullptr)
 	{
-		FScopeLock lock(&DllAccessCS);
+		FScopeLock lock(&InDllAccessCS);
 		DisplayClusterProjectionDomeprojectionLibraryDX11::dpDestroyContextFunc(Context);
 		Context = nullptr;
 	}
 }
 
-bool FDisplayClusterProjectionDomeprojectionViewAdapterDX11::FViewData::Initialize(IDisplayClusterViewport* InViewport, const FString& InFile, FCriticalSection& DllAccessCS)
+bool FDisplayClusterProjectionDomeprojectionViewAdapterDX11::FViewData::Initialize(IDisplayClusterViewport* InViewport, const FString& InFile, FCriticalSection& InDllAccessCS)
 {
 	// Initialize Domeprojection DLL API
-	FScopeLock lock(&DllAccessCS);
+	FScopeLock lock(&InDllAccessCS);
 
 	if (!DisplayClusterProjectionDomeprojectionLibraryDX11::Initialize())
 	{
