nDisplay: Added a new tab to the ICVFX panel that displays the level instance root actor's details panel and components list, allowing users to edit any property on the level intance without having to go to the level editor window.

diff --git a/Engine/Plugins/Runtime/nDisplay/Source/DisplayClusterOperator/DisplayClusterOperator.Build.cs b/Engine/Plugins/Runtime/nDisplay/Source/DisplayClusterOperator/DisplayClusterOperator.Build.cs
--- a/Engine/Plugins/Runtime/nDisplay/Source/DisplayClusterOperator/DisplayClusterOperator.Build.cs
+++ b/Engine/Plugins/Runtime/nDisplay/Source/DisplayClusterOperator/DisplayClusterOperator.Build.cs
@@ -20,6 +20,8 @@ public class DisplayClusterOperator : ModuleRules
 				"Projects",
 				"Slate",
 				"SlateCore",
+				"SubobjectDataInterface",
+				"SubobjectEditor",
 				"ToolMenus",
 				"ToolWidgets",
 				"UnrealEd",
diff --git a/Engine/Plugins/Runtime/nDisplay/Source/DisplayClusterOperator/Private/DisplayClusterOperatorStyle.cpp b/Engine/Plugins/Runtime/nDisplay/Source/DisplayClusterOperator/Private/DisplayClusterOperatorStyle.cpp
--- a/Engine/Plugins/Runtime/nDisplay/Source/DisplayClusterOperator/Private/DisplayClusterOperatorStyle.cpp
+++ b/Engine/Plugins/Runtime/nDisplay/Source/DisplayClusterOperator/Private/DisplayClusterOperatorStyle.cpp
@@ -37,6 +37,7 @@ void FDisplayClusterOperatorStyle::Initialize()
 	// Icons
 	{
 		Set("OperatorPanel.Icon", new IMAGE_BRUSH_SVG("Icons/Components/nDisplayCamera", Icon16x16));
+		Set("RootActorPanel.Icon", new IMAGE_BRUSH_SVG("Icons/RootActor/nDisplay", Icon16x16));
 	}
 	
 	// Colors
diff --git a/Engine/Plugins/Runtime/nDisplay/Source/DisplayClusterOperator/Private/SDisplayClusterOperatorPanel.cpp b/Engine/Plugins/Runtime/nDisplay/Source/DisplayClusterOperator/Private/SDisplayClusterOperatorPanel.cpp
--- a/Engine/Plugins/Runtime/nDisplay/Source/DisplayClusterOperator/Private/SDisplayClusterOperatorPanel.cpp
+++ b/Engine/Plugins/Runtime/nDisplay/Source/DisplayClusterOperator/Private/SDisplayClusterOperatorPanel.cpp
@@ -4,8 +4,10 @@
 
 #include "DisplayClusterOperatorModule.h"
 #include "DisplayClusterOperatorStatusBarExtender.h"
+#include "DisplayClusterOperatorStyle.h"
 #include "IDisplayClusterOperator.h"
 #include "IDisplayClusterOperatorViewModel.h"
+#include "SDisplayClusterOperatorRootActorPanel.h"
 #include "SDisplayClusterOperatorStatusBar.h"
 #include "SDisplayClusterOperatorToolbar.h"
 
@@ -17,7 +19,6 @@
 #include "ToolMenus.h"
 #include "Framework/Docking/LayoutExtender.h"
 #include "Framework/Docking/TabManager.h"
-#include "Framework/MultiBox/MultiBoxBuilder.h"
 #include "Interfaces/IMainFrameModule.h"
 #include "Styling/AppStyle.h"
 #include "Widgets/Docking/SDockTab.h"
@@ -25,6 +26,7 @@
 #define LOCTEXT_NAMESPACE "SDisplayClusterOperatorPanel"
 
 const FName SDisplayClusterOperatorPanel::DetailsTabId = TEXT("OperatorDetails");
+const FName SDisplayClusterOperatorPanel::RootActorPanelId = TEXT("RootActorPanel");
 const FName SDisplayClusterOperatorPanel::PrimaryTabExtensionId = TEXT("PrimaryOperatorTabStack");
 const FName SDisplayClusterOperatorPanel::AuxilliaryTabExtensionId = TEXT("AuxilliaryOperatorTabStack");
 
@@ -56,6 +58,11 @@ void SDisplayClusterOperatorPanel::Construct(const FArguments& InArgs, const TSh
 		.SetIcon(FSlateIcon(FAppStyle::GetAppStyleSetName(), "LevelEditor.Tabs.Details"))
 		.SetGroup(AppMenuGroup);
 
+	TabManager->RegisterTabSpawner(RootActorPanelId, FOnSpawnTab::CreateSP(this, &SDisplayClusterOperatorPanel::SpawnRootActorPanelTab))
+		.SetDisplayName(LOCTEXT("RootActorPanelTabTitle", "nDisplay Settings"))
+		.SetIcon(FSlateIcon(FDisplayClusterOperatorStyle::Get().GetStyleSetName(), "RootActorPanel.Icon"))
+		.SetGroup(AppMenuGroup);
+
 	// Load any operator apps
 	{
 		FDisplayClusterOperatorModule& OperatorModule = FModuleManager::GetModuleChecked<FDisplayClusterOperatorModule>(IDisplayClusterOperator::ModuleName);
@@ -71,7 +78,7 @@ void SDisplayClusterOperatorPanel::Construct(const FArguments& InArgs, const TSh
 		}
 	}
 	
-	TSharedRef<FTabManager::FLayout> Layout = FTabManager::NewLayout("nDisplayOperatorLayout_v1")
+	TSharedRef<FTabManager::FLayout> Layout = FTabManager::NewLayout("nDisplayOperatorLayout_v2")
 		->AddArea
 		(
 			FTabManager::NewPrimaryArea()
@@ -106,7 +113,9 @@ void SDisplayClusterOperatorPanel::Construct(const FArguments& InArgs, const TSh
 						(
 							FTabManager::NewStack()
 							->AddTab(DetailsTabId, ETabState::OpenedTab)
-							->SetHideTabWell(true)
+							->AddTab(RootActorPanelId, ETabState::OpenedTab)
+							->SetForegroundTab(DetailsTabId)
+							->SetHideTabWell(false)
 						)
 						
 					)
@@ -258,6 +267,17 @@ TSharedRef<SDockTab> SDisplayClusterOperatorPanel::SpawnDetailsTab(const FSpawnT
 		];
 }
 
+TSharedRef<SDockTab> SDisplayClusterOperatorPanel::SpawnRootActorPanelTab(const FSpawnTabArgs& Args)
+{
+	SAssignNew(RootActorPanel, SDisplayClusterOperatorRootActorPanel, TabManager, RootActorPanelId);
+
+	return SNew(SDockTab)
+		.TabRole(ETabRole::PanelTab)
+		[
+			RootActorPanel.ToSharedRef()
+		];
+}
+
 void SDisplayClusterOperatorPanel::DisplayObjectsInDetailsPanel(const TArray<UObject*>& Objects)
 {
 	if (DetailsView.IsValid())
diff --git a/Engine/Plugins/Runtime/nDisplay/Source/DisplayClusterOperator/Private/SDisplayClusterOperatorPanel.h b/Engine/Plugins/Runtime/nDisplay/Source/DisplayClusterOperator/Private/SDisplayClusterOperatorPanel.h
--- a/Engine/Plugins/Runtime/nDisplay/Source/DisplayClusterOperator/Private/SDisplayClusterOperatorPanel.h
+++ b/Engine/Plugins/Runtime/nDisplay/Source/DisplayClusterOperator/Private/SDisplayClusterOperatorPanel.h
@@ -5,6 +5,7 @@
 #include "CoreMinimal.h"
 #include "Widgets/SCompoundWidget.h"
 
+class SDisplayClusterOperatorRootActorPanel;
 class FTabManager;
 class FLayoutExtender;
 class FSpawnTabArgs;
@@ -24,6 +25,9 @@ public:
 	/** The tab ID that the details panel lives in */
 	static const FName DetailsTabId;
 
+	/** The tab ID that the root actor properties panel lives in */
+	static const FName RootActorPanelId;
+
 	/** The ID of the main tab stack that can be extended by external tabs */
 	static const FName PrimaryTabExtensionId;
 
@@ -39,7 +43,7 @@ public:
 	void Construct(const FArguments& InArgs, const TSharedRef<FTabManager>& InTabManager, const TSharedPtr<SWindow>& WindowOwner);
 
 	//~ SWidget interface
-	virtual FReply OnKeyDown(const FGeometry& MyGeometry, const FKeyEvent& InKeyEvent);
+	virtual FReply OnKeyDown(const FGeometry& MyGeometry, const FKeyEvent& InKeyEvent) override;
 	//~ End SWidget interface
 
 	/** Toggles the state of a drawer with the specified ID, closing the drawer if it is open, and opening the drawer if it is closed */
@@ -55,6 +59,9 @@ private:
 	/** Creates a tab with the details panel in it */
 	TSharedRef<SDockTab> SpawnDetailsTab(const FSpawnTabArgs& Args);
 
+	/** Creates a tab with the root actor properties panel in it */
+	TSharedRef<SDockTab> SpawnRootActorPanelTab(const FSpawnTabArgs& Args);
+
 	/** Raised when something wants to display a list of objects in the operator's details panel */
 	void DisplayObjectsInDetailsPanel(const TArray<UObject*>& Objects);
 
@@ -77,6 +84,9 @@ private:
 	/** A reference to the operator panel's details view */
 	TSharedPtr<class SKismetInspector> DetailsView;
 
+	/** A reference to the operator panel's root actor panel */
+	TSharedPtr<SDisplayClusterOperatorRootActorPanel> RootActorPanel;
+
 	/** A reference to the operator panel's status bar widget */
 	TSharedPtr<SDisplayClusterOperatorStatusBar> StatusBar;
 
diff --git a/Engine/Plugins/Runtime/nDisplay/Source/DisplayClusterOperator/Private/SDisplayClusterOperatorRootActorPanel.cpp b/Engine/Plugins/Runtime/nDisplay/Source/DisplayClusterOperator/Private/SDisplayClusterOperatorRootActorPanel.cpp
new file mode 100644
--- /dev/null
+++ b/Engine/Plugins/Runtime/nDisplay/Source/DisplayClusterOperator/Private/SDisplayClusterOperatorRootActorPanel.cpp
@@ -0,0 +1,532 @@
+﻿// Copyright Epic Games, Inc. All Rights Reserved.
+
+#include "SDisplayClusterOperatorRootActorPanel.h"
+
+#include "DisplayClusterRootActor.h"
+#include "Editor/UnrealEdEngine.h"
+#include "IDisplayClusterOperator.h"
+#include "IDisplayClusterOperatorViewModel.h"
+#include "LevelEditor.h"
+#include "PropertyEditorModule.h"
+#include "SlateOptMacros.h"
+#include "SSubobjectEditorModule.h"
+#include "SSubobjectInstanceEditor.h"
+#include "UnrealEdGlobals.h"
+#include "Widgets/Text/SRichTextBlock.h"
+
+#define LOCTEXT_NAMESPACE "SDisplayClusterOperatorRootActorPanel"
+
+BEGIN_SLATE_FUNCTION_BUILD_OPTIMIZATION
+
+namespace UE::DisplayClusterOperatorRootActorPanel
+{
+	/** A widget to display a warning message */
+	class SWarningLabel : public SCompoundWidget
+	{
+	public:
+		SLATE_BEGIN_ARGS(SWarningLabel)
+			: _WarningText()
+			, _OnHyperlinkClicked()
+		{}
+			
+			/** The rich text to show in the warning */
+			SLATE_ATTRIBUTE(FText, WarningText)
+
+			/** Called when the hyperlink in the rich text is clicked */
+			SLATE_EVENT(FSlateHyperlinkRun::FOnClick, OnHyperlinkClicked)
+
+		SLATE_END_ARGS()
+
+		/** Constructs the widget */
+		void Construct(const FArguments& InArgs)
+		{
+			ChildSlot
+			[
+				SNew(SBorder)
+				.BorderImage(FAppStyle::Get().GetBrush("Brushes.Panel"))
+				[
+					SNew(SHorizontalBox)
+					+ SHorizontalBox::Slot()
+					.AutoWidth()
+					.HAlign(HAlign_Center)
+					.VAlign(VAlign_Center)
+					.Padding(2)
+					[
+						SNew(SImage)
+						.Image(FAppStyle::Get().GetBrush("Icons.Warning"))
+					]
+					+ SHorizontalBox::Slot()
+						.VAlign(VAlign_Center)
+						.Padding(2)
+						[
+							SNew(SRichTextBlock)
+							.DecoratorStyleSet(&FAppStyle::Get())
+							.Justification(ETextJustify::Left)
+							.TextStyle(FAppStyle::Get(), "DetailsView.BPMessageTextStyle")
+							.Text(InArgs._WarningText)
+							.AutoWrapText(true)
+							+ SRichTextBlock::HyperlinkDecorator(TEXT("HyperlinkDecorator"), InArgs._OnHyperlinkClicked)
+						]
+				]
+			];
+		}
+	};
+}
+
+SDisplayClusterOperatorRootActorPanel::~SDisplayClusterOperatorRootActorPanel()
+{
+	if (GEditor)
+	{
+		GEditor->UnregisterForUndo(this);
+	}
+
+	FCoreUObjectDelegates::OnObjectsReplaced.RemoveAll(this);
+
+	if (FLevelEditorModule* LevelEditor = FModuleManager::GetModulePtr<FLevelEditorModule>("LevelEditor"))
+	{
+		LevelEditor->OnComponentsEdited().RemoveAll(this);
+	}
+
+	if (OperatorViewModel.IsValid())
+	{
+		OperatorViewModel->OnActiveRootActorChanged().RemoveAll(this);
+	}
+}
+
+void SDisplayClusterOperatorRootActorPanel::Construct(const FArguments& InArgs, TSharedPtr<FTabManager> InTabManager, const FName& TabIdentifier)
+{
+	GEditor->RegisterForUndo(this);
+
+	FCoreUObjectDelegates::OnObjectsReplaced.AddRaw(this, &SDisplayClusterOperatorRootActorPanel::OnObjectsReplaced);
+
+	FLevelEditorModule& LevelEditor = FModuleManager::GetModuleChecked<FLevelEditorModule>("LevelEditor");
+	LevelEditor.OnComponentsEdited().AddSP(this, &SDisplayClusterOperatorRootActorPanel::OnComponentsEdited);
+
+	OperatorViewModel = IDisplayClusterOperator::Get().GetOperatorViewModel();
+	OperatorViewModel->OnActiveRootActorChanged().AddSP(this, &SDisplayClusterOperatorRootActorPanel::OnRootActorChanged);
+
+	FDetailsViewArgs DetailsViewArgs;
+	DetailsViewArgs.bUpdatesFromSelection = true;
+	DetailsViewArgs.bLockable = true;
+	DetailsViewArgs.bAllowFavoriteSystem = true;
+	DetailsViewArgs.NameAreaSettings = FDetailsViewArgs::ObjectsUseNameArea | FDetailsViewArgs::ComponentsAndActorsUseNameArea;
+	DetailsViewArgs.NotifyHook = GUnrealEd;
+	DetailsViewArgs.ViewIdentifier = TabIdentifier;
+	DetailsViewArgs.bCustomNameAreaLocation = true;
+	DetailsViewArgs.bCustomFilterAreaLocation = true;
+	DetailsViewArgs.DefaultsOnlyVisibility = EEditDefaultsOnlyNodeVisibility::Hide;
+	DetailsViewArgs.HostTabManager = InTabManager;
+	DetailsViewArgs.bShowSectionSelector = true;
+
+	FPropertyEditorModule& PropertyEditor = FModuleManager::LoadModuleChecked<FPropertyEditorModule>("PropertyEditor");
+	DetailsView = PropertyEditor.CreateDetailView(DetailsViewArgs);
+
+	DetailsView->SetIsPropertyVisibleDelegate(FIsPropertyVisible::CreateSP(this, &SDisplayClusterOperatorRootActorPanel::IsPropertyVisible));
+	DetailsView->SetIsPropertyReadOnlyDelegate(FIsPropertyReadOnly::CreateSP(this, &SDisplayClusterOperatorRootActorPanel::IsPropertyReadOnly));
+	DetailsView->SetIsPropertyEditingEnabledDelegate(FIsPropertyEditingEnabled::CreateSP(this, &SDisplayClusterOperatorRootActorPanel::IsPropertyEditingEnabled));
+
+	FModuleManager::LoadModuleChecked<FSubobjectEditorModule>("SubobjectEditor");
+
+	SubobjectEditor = SNew(SSubobjectInstanceEditor)
+		.ObjectContext(this, &SDisplayClusterOperatorRootActorPanel::GetRootActorContextObject)
+		.AllowEditing(this, &SDisplayClusterOperatorRootActorPanel::CanEditRootActorComponents)
+		.OnSelectionUpdated(this, &SDisplayClusterOperatorRootActorPanel::OnSelectedSubobjectsChanged);
+
+	TSharedRef<SWidget> SubobjectButtonsBox = SNew(SHorizontalBox)
+		+SHorizontalBox::Slot()
+		.VAlign(VAlign_Center)
+		.AutoWidth()
+		[
+			SubobjectEditor->GetToolButtonsBox().ToSharedRef()
+		];
+
+	DetailsView->SetNameAreaCustomContent(SubobjectButtonsBox);
+
+	ChildSlot
+	[
+		SNew(SVerticalBox)
+
+		+SVerticalBox::Slot()
+		.Padding(10.f, 4.f, 0.f, 0.f)
+		.AutoHeight()
+		[
+			DetailsView->GetNameAreaWidget().ToSharedRef()
+		]
+
+		+SVerticalBox::Slot()
+		[
+			SNew(SSplitter)
+			.MinimumSlotHeight(80.0f)
+			.Orientation(Orient_Vertical)
+			.Style(FAppStyle::Get(), "SplitterDark")
+			.PhysicalSplitterHandleSize(2.0f)
+
+			+SSplitter::Slot()
+			.Value(0.2f)
+			[
+				SNew(SBox)
+				.Padding(FMargin(2.0f, 0.0f, 2.0f, 0.0f))
+				.Visibility(this, &SDisplayClusterOperatorRootActorPanel::GetSubobjectEditorVisibility)
+				[
+					SubobjectEditor.ToSharedRef()
+				]
+			]
+
+			+SSplitter::Slot()
+			[
+				SNew( SVerticalBox )
+				+SVerticalBox::Slot()
+				.AutoHeight()
+				.Padding( FMargin(0,0,0,1) )
+				[
+					SNew(UE::DisplayClusterOperatorRootActorPanel::SWarningLabel)
+					.Visibility(this, &SDisplayClusterOperatorRootActorPanel::GetUCSComponentWarningVisibility)
+					.WarningText(LOCTEXT("BlueprintUCSComponentWarning", "Components created by the User Construction Script can only be edited in the <a id=\"HyperlinkDecorator\" style=\"DetailsView.BPMessageHyperlinkStyle\">Blueprint</>"))
+					.OnHyperlinkClicked(this, &SDisplayClusterOperatorRootActorPanel::OnComponentBlueprintHyperlinkClicked)
+				]
+				+SVerticalBox::Slot()
+				.AutoHeight()
+				.Padding( FMargin(0,0,0,1) )
+				[
+					SNew(UE::DisplayClusterOperatorRootActorPanel::SWarningLabel)
+					.Visibility(this, &SDisplayClusterOperatorRootActorPanel::GetInheritedComponentWarningVisibility)
+					.WarningText(LOCTEXT("BlueprintUneditableInheritedComponentWarning", "Components flagged as not editable when inherited must be edited in the <a id=\"HyperlinkDecorator\" style=\"DetailsView.BPMessageHyperlinkStyle\">Blueprint</>"))
+					.OnHyperlinkClicked(this, &SDisplayClusterOperatorRootActorPanel::OnComponentBlueprintHyperlinkClicked)
+				]
+				+SVerticalBox::Slot()
+				.AutoHeight()
+				.Padding( FMargin(0,0,0,1) )
+				[
+					SNew(UE::DisplayClusterOperatorRootActorPanel::SWarningLabel)
+					.Visibility(this, &SDisplayClusterOperatorRootActorPanel::GetNativeComponentWarningVisibility)
+					.WarningText(LOCTEXT("UneditableNativeComponentWarning", "Native components are editable when declared as a FProperty in <a id=\"HyperlinkDecorator\" style=\"DetailsView.BPMessageHyperlinkStyle\">C++</>"))
+					.OnHyperlinkClicked(this, &SDisplayClusterOperatorRootActorPanel::OnComponentBlueprintHyperlinkClicked)
+				]
+				+ SVerticalBox::Slot()
+				.AutoHeight()
+				[
+					DetailsView->GetFilterAreaWidget().ToSharedRef()
+				]
+				+ SVerticalBox::Slot()
+				[
+					DetailsView.ToSharedRef()
+				]
+			]
+		]
+	];
+
+	if (OperatorViewModel->GetRootActor())
+	{
+		DetailsView->SetObject(OperatorViewModel->GetRootActor());
+	}
+	else
+	{
+		// Force the details view to clear the set objects. This ensures the "Select an object to show details" message properly displays
+		DetailsView->SetObjects(TArray<UObject*>());
+	}
+}
+
+void SDisplayClusterOperatorRootActorPanel::PostUndo(bool bSuccess)
+{
+	if (SubobjectEditor.IsValid())
+	{
+		SubobjectEditor->UpdateTree();
+	}
+
+	if (DetailsView.IsValid())
+	{
+		// As long as there is a valid root actor, the details view should have a selected object to display
+		// In the case where the deletion of a root actor was undone, the details view will not remember its
+		// last selected object, so attempt to sync with the subobject editor
+		if (!DetailsView->GetSelectedObjects().Num())
+		{
+			// Make sure the details view is still synced with the subobject editor selection after the undo. If there are no
+			// selected subobjects, set the root actor as the displayed object in the details view
+			TArray<FSubobjectEditorTreeNodePtrType> SelectedSubobjects = SubobjectEditor->GetSelectedNodes();
+			if (SelectedSubobjects.Num())
+			{
+				OnSelectedSubobjectsChanged(SubobjectEditor->GetSelectedNodes());
+			}
+			else if (OperatorViewModel->GetRootActor())
+			{
+				DetailsView->SetObject(OperatorViewModel->GetRootActor());
+			}
+		}
+	}
+}
+
+void SDisplayClusterOperatorRootActorPanel::PostRedo(bool bSuccess)
+{
+	PostUndo(bSuccess);
+}
+
+UObject* SDisplayClusterOperatorRootActorPanel::GetRootActorContextObject() const
+{
+	if (OperatorViewModel.IsValid())
+	{
+		return OperatorViewModel->GetRootActor();
+	}
+
+	return nullptr;
+}
+
+EVisibility SDisplayClusterOperatorRootActorPanel::GetSubobjectEditorVisibility() const
+{
+	return OperatorViewModel->HasRootActor() ? EVisibility::Visible : EVisibility::Collapsed;
+}
+
+bool SDisplayClusterOperatorRootActorPanel::CanEditRootActorComponents() const
+{
+	return GEditor->PlayWorld == nullptr;
+}
+
+bool SDisplayClusterOperatorRootActorPanel::IsPropertyVisible(const FPropertyAndParent& PropertyAndParent) const
+{
+	if (PropertyAndParent.Property.HasAllPropertyFlags(CPF_DisableEditOnInstance))
+	{
+		return false;
+	}
+
+	return true;
+}
+
+bool SDisplayClusterOperatorRootActorPanel::IsPropertyReadOnly(const FPropertyAndParent& PropertyAndParent) const
+{
+	bool bIsReadOnly = false;
+	for (const FSubobjectEditorTreeNodePtrType& Node : SubobjectEditor->GetSelectedNodes())
+	{
+		const UActorComponent* Component = Node->GetComponentTemplate();
+		if (Component && Component->CreationMethod == EComponentCreationMethod::SimpleConstructionScript)
+		{
+			TSet<const FProperty*> UCSModifiedProperties;
+			Component->GetUCSModifiedProperties(UCSModifiedProperties);
+			if (UCSModifiedProperties.Contains(&PropertyAndParent.Property) ||
+				(PropertyAndParent.ParentProperties.Num() > 0 && UCSModifiedProperties.Contains(PropertyAndParent.ParentProperties[0])))
+			{
+				bIsReadOnly = true;
+				break;
+			}
+		}
+	}
+	return bIsReadOnly;
+}
+
+bool SDisplayClusterOperatorRootActorPanel::IsPropertyEditingEnabled() const
+{
+	FLevelEditorModule& LevelEditor = FModuleManager::GetModuleChecked<FLevelEditorModule>("LevelEditor");
+	if (!LevelEditor.AreObjectsEditable(DetailsView->GetSelectedObjects()))
+	{
+		return false;
+	}
+
+	bool bIsEditable = true;
+	for (const FSubobjectEditorTreeNodePtrType& Node : SubobjectEditor->GetSelectedNodes())
+	{
+		if (const FSubobjectData* Data = Node->GetDataSource())
+		{
+			bIsEditable = Data->CanEdit();
+			if (!bIsEditable)
+			{
+				break;
+			}
+		}
+	}
+	return bIsEditable;
+}
+
+bool SDisplayClusterOperatorRootActorPanel::IsBlueprintNotEditable(const UActorComponent* Component) const
+{
+	// Determine if it is locked out from a blueprint or from the native
+	UActorComponent* Archetype = CastChecked<UActorComponent>(Component->GetArchetype());
+	while (Archetype)
+	{
+		if (Archetype->GetOuter()->IsA<UBlueprintGeneratedClass>() || Archetype->GetOuter()->GetClass()->HasAllClassFlags(CLASS_CompiledFromBlueprint))
+		{
+			if (!Archetype->bEditableWhenInherited)
+			{
+				return true;
+			}
+
+			Archetype = CastChecked<UActorComponent>(Archetype->GetArchetype());
+		}
+		else
+		{
+			Archetype = nullptr;
+		}
+	}
+
+	return false;
+}
+
+EVisibility SDisplayClusterOperatorRootActorPanel::GetUCSComponentWarningVisibility() const
+{
+	bool bIsUneditableBlueprintComponent = false;
+
+	// Check to see if any selected components are inherited from blueprint
+	for (const FSubobjectEditorTreeNodePtrType& Node : SubobjectEditor->GetSelectedNodes())
+	{
+		if (const FSubobjectData* Data = Node->GetDataSource())
+		{
+			if (!Data->IsNativeComponent())
+			{
+				const UActorComponent* Component = Data->GetComponentTemplate();
+				bIsUneditableBlueprintComponent = Component ? Component->CreationMethod == EComponentCreationMethod::UserConstructionScript : false;
+				if (bIsUneditableBlueprintComponent)
+				{
+					break;
+				}
+			}
+		}
+
+	}
+
+	return bIsUneditableBlueprintComponent ? EVisibility::Visible : EVisibility::Collapsed;
+}
+
+EVisibility SDisplayClusterOperatorRootActorPanel::GetInheritedComponentWarningVisibility() const
+{
+	bool bIsUneditableBlueprintComponent = false;
+
+	// Check to see if any selected components are inherited from blueprint
+	for (const FSubobjectEditorTreeNodePtrType& Node : SubobjectEditor->GetSelectedNodes())
+	{
+		if (const FSubobjectData* Data = Node->GetDataSource())
+		{
+			if (!Data->IsNativeComponent())
+			{
+				if (const UActorComponent* Component = Data->GetComponentTemplate())
+				{
+					if (!Component->IsEditableWhenInherited() && Component->CreationMethod == EComponentCreationMethod::SimpleConstructionScript)
+					{
+						bIsUneditableBlueprintComponent = true;
+						break;
+					}
+				}
+			}
+			else if (!Data->CanEdit() && IsBlueprintNotEditable(Data->GetComponentTemplate()))
+			{
+				bIsUneditableBlueprintComponent = true;
+				break;
+			}
+		}
+	}
+
+	return bIsUneditableBlueprintComponent ? EVisibility::Visible : EVisibility::Collapsed;
+}
+
+EVisibility SDisplayClusterOperatorRootActorPanel::GetNativeComponentWarningVisibility() const
+{
+	bool bIsUneditableNative = false;
+	for (const FSubobjectEditorTreeNodePtrType& Node : SubobjectEditor->GetSelectedNodes())
+	{
+		if (const FSubobjectData* Data = Node->GetDataSource())
+		{
+			// Check to see if the component is native and not editable
+			if (Data->IsNativeComponent() && !Data->CanEdit() && !IsBlueprintNotEditable(Data->GetComponentTemplate()))
+			{
+				bIsUneditableNative = true;
+				break;
+			}
+		}
+
+	}
+
+	return bIsUneditableNative ? EVisibility::Visible : EVisibility::Collapsed;
+}
+
+void SDisplayClusterOperatorRootActorPanel::OnComponentBlueprintHyperlinkClicked(const FSlateHyperlinkRun::FMetadata& Metadata)
+{
+	UBlueprint* Blueprint = SubobjectEditor->GetBlueprint();
+	if (Blueprint)
+	{
+		// Open the blueprint
+		GEditor->EditObject(Blueprint);
+	}
+}
+
+void SDisplayClusterOperatorRootActorPanel::OnRootActorChanged(ADisplayClusterRootActor* RootActor)
+{
+	if (DetailsView.IsValid())
+	{
+		if (RootActor)
+		{
+			DetailsView->SetObject(RootActor);
+		}
+		else
+		{
+			// Force the details view to clear the set objects. This ensures the "Select an object to show details" message properly displays
+			DetailsView->SetObjects(TArray<UObject*>());
+		}
+	}
+
+	if (SubobjectEditor.IsValid())
+	{
+		SubobjectEditor->UpdateTree();
+	}
+}
+
+void SDisplayClusterOperatorRootActorPanel::OnSelectedSubobjectsChanged(const TArray<TSharedPtr<FSubobjectEditorTreeNode>>& SelectedNodes)
+{
+	if (SelectedNodes.Num() == 0)
+	{
+		return;
+	}
+
+	if (!DetailsView->IsLocked())
+	{
+		auto IsRootActor = [](const FSubobjectEditorTreeNodePtrType& Node)
+		{
+			if (Node.IsValid())
+			{
+				if (FSubobjectData* Data = Node->GetDataSource())
+				{
+					return Data->IsRootActor();
+				}
+			}
+
+			return false;
+		};
+
+		if (SelectedNodes.ContainsByPredicate(IsRootActor))
+		{
+			DetailsView->SetObject(OperatorViewModel->GetRootActor());
+		}
+		else
+		{
+			TArray<UObject*> Components;
+			for (const FSubobjectEditorTreeNodePtrType& Node : SelectedNodes)
+			{
+				if (Node.IsValid())
+				{
+					if (FSubobjectData* Data = Node->GetDataSource())
+					{
+						if (Data->IsComponent())
+						{
+							if (const UActorComponent* Component = Data->FindComponentInstanceInActor(OperatorViewModel->GetRootActor()))
+							{
+								Components.Add(const_cast<UActorComponent*>(Component));
+							}
+						}
+					}
+				}
+			}
+
+			DetailsView->SetObjects(Components);
+		}
+	}
+}
+
+void SDisplayClusterOperatorRootActorPanel::OnComponentsEdited()
+{
+	SubobjectEditor->UpdateTree();
+	DetailsView->ForceRefresh();
+}
+
+void SDisplayClusterOperatorRootActorPanel::OnObjectsReplaced(const TMap<UObject*, UObject*>& InReplacementObjects)
+{
+	SubobjectEditor->UpdateTree();
+}
+
+END_SLATE_FUNCTION_BUILD_OPTIMIZATION
+
+#undef LOCTEXT_NAMESPACE
\ No newline at end of file
diff --git a/Engine/Plugins/Runtime/nDisplay/Source/DisplayClusterOperator/Private/SDisplayClusterOperatorRootActorPanel.h b/Engine/Plugins/Runtime/nDisplay/Source/DisplayClusterOperator/Private/SDisplayClusterOperatorRootActorPanel.h
new file mode 100644
--- /dev/null
+++ b/Engine/Plugins/Runtime/nDisplay/Source/DisplayClusterOperator/Private/SDisplayClusterOperatorRootActorPanel.h
@@ -0,0 +1,92 @@
+﻿// Copyright Epic Games, Inc. All Rights Reserved.
+
+#pragma once
+
+#include "CoreMinimal.h"
+#include "EditorUndoClient.h"
+#include "Framework/Text/SlateHyperlinkRun.h"
+#include "Widgets/SCompoundWidget.h"
+
+struct FPropertyAndParent;
+class ADisplayClusterRootActor;
+class FSubobjectEditorTreeNode;
+class FTabManager;
+class IDisplayClusterOperatorViewModel;
+class IDetailsView;
+class SSubobjectEditor;
+class UActorComponent;
+
+/**
+ * A panel that mimics the appearance of SActorDetails to display the root actor's components and properties in the operator panel
+ */
+class DISPLAYCLUSTEROPERATOR_API SDisplayClusterOperatorRootActorPanel : public SCompoundWidget, public FEditorUndoClient
+{
+public:
+	virtual ~SDisplayClusterOperatorRootActorPanel() override;
+
+	SLATE_BEGIN_ARGS(SDisplayClusterOperatorRootActorPanel) {}
+	SLATE_END_ARGS()
+
+	/** Constructs this widget with InArgs */
+	void Construct(const FArguments& InArgs, TSharedPtr<FTabManager> InTabManager, const FName& TabIdentifier);
+
+	//~ FEditorUndoClient Interface
+	virtual void PostUndo(bool bSuccess) override;
+	virtual void PostRedo(bool bSuccess) override;
+
+private:
+	/** Gets the root actor displayed by the subobject editor and the details panel */
+	UObject* GetRootActorContextObject() const;
+
+	/** Gets the visibility of the subobject editor */
+	EVisibility GetSubobjectEditorVisibility() const;
+
+	/** Gets whether the root actor's components can be edited */
+	bool CanEditRootActorComponents() const;
+
+	/** Determines if a property is visible or not */
+	bool IsPropertyVisible(const FPropertyAndParent& PropertyAndParent) const;
+
+	/** Determines if a property is read only or not */
+	bool IsPropertyReadOnly(const FPropertyAndParent& PropertyAndParent) const;
+
+	/** Determines if property editing is enabled or not */
+	bool IsPropertyEditingEnabled() const;
+
+	/** Determines if the component is set as not editable through its blueprint */
+	bool IsBlueprintNotEditable(const UActorComponent* Component) const;
+
+	/** Gets the visibility of the warning label for user compile script generated components */
+	EVisibility GetUCSComponentWarningVisibility() const;
+
+	/** Gets the visibility of the warning label for inherited blueprint components */
+	EVisibility GetInheritedComponentWarningVisibility() const;
+
+	/** Gets the visibility of the warning label for native components */
+	EVisibility GetNativeComponentWarningVisibility() const;
+
+	/** Raised when a hyperlink is clicked in the warning labels for components */
+	void OnComponentBlueprintHyperlinkClicked(const FSlateHyperlinkRun::FMetadata& Metadata);
+
+	/** Raised when the operator panel's active root actor has been changed */
+	void OnRootActorChanged(ADisplayClusterRootActor* RootActor);
+
+	/** Raised when the selection in the subobject editor is changed */
+	void OnSelectedSubobjectsChanged(const TArray<TSharedPtr<FSubobjectEditorTreeNode>>& SelectedNodes);
+
+	/** Raised by the level editor when any components in the world have been edited */
+	void OnComponentsEdited();
+
+	/** Raised when any objects in the editor are being destroyed and replaced with new instances */
+	void OnObjectsReplaced(const TMap<UObject*, UObject*>& InReplacementObjects);
+
+private:
+	/** Subobject editor used to display the components of the root actor */
+	TSharedPtr<SSubobjectEditor> SubobjectEditor;
+
+	/** Details view widget used to display the actor's properties */
+	TSharedPtr<IDetailsView> DetailsView;
+
+	/** A reference to the operator panel's view model */
+	TSharedPtr<IDisplayClusterOperatorViewModel> OperatorViewModel;
+};
