nDisplay: Removed unecessary light card operations that were dragging down frame rate.

diff --git a/Engine/Plugins/Runtime/nDisplay/Source/DisplayCluster/Private/Game/EngineClasses/Scene/DisplayClusterLightCardActor.cpp b/Engine/Plugins/Runtime/nDisplay/Source/DisplayCluster/Private/Game/EngineClasses/Scene/DisplayClusterLightCardActor.cpp
--- a/Engine/Plugins/Runtime/nDisplay/Source/DisplayCluster/Private/Game/EngineClasses/Scene/DisplayClusterLightCardActor.cpp
+++ b/Engine/Plugins/Runtime/nDisplay/Source/DisplayCluster/Private/Game/EngineClasses/Scene/DisplayClusterLightCardActor.cpp
@@ -89,6 +89,9 @@ ADisplayClusterLightCardActor::ADisplayClusterLightCardActor(const FObjectInitia
 		LightCardComponent->SetCollisionProfileName(TEXT("OverlapAllDynamic"));
 		LightCardComponent->Mobility = EComponentMobility::Movable;
 		LightCardComponent->SetStaticMesh(PlaneObj.Object);
+
+		// Overlap events are expensive and seem to not be needed at the moment, so we disable them.
+		LightCardComponent->SetGenerateOverlapEvents(false);
 	}
 
 #if WITH_EDITOR
@@ -157,6 +160,8 @@ void ADisplayClusterLightCardActor::OnConstruction(const FTransform& Transform)
 
 void ADisplayClusterLightCardActor::Tick(float DeltaSeconds)
 {
+	TRACE_CPUPROFILER_EVENT_SCOPE(DCLightcard_Tick);
+
 	Super::Tick(DeltaSeconds);
 
 	if (!MainSpringArmComponent || !LightCardTransformerComponent || !LightCardComponent)
@@ -278,13 +283,13 @@ void ADisplayClusterLightCardActor::SetStaticMesh(UStaticMesh* InStaticMesh)
 
 void ADisplayClusterLightCardActor::UpdateLightCardMaterialInstance()
 {
+	TRACE_CPUPROFILER_EVENT_SCOPE(DCLightcard_UpdateLightCardMaterialInstance);
+
 	if (UMaterialInstanceDynamic* LightCardMaterialInstance = Cast<UMaterialInstanceDynamic>(LightCardComponent->GetMaterial(0)))
 	{
 		// Showing proxy with low opacity to make it less distracting when it doesn't line up well with its projection in the Light Card Editor.
 		constexpr float ProxyOpacity = 0.25;
 
-		LightCardMaterialInstance->ClearParameterValues();
-
 		LightCardMaterialInstance->SetVectorParameterValue(TEXT("CardColor"), Color);
 		LightCardMaterialInstance->SetScalarParameterValue(TEXT("Temperature"), Temperature);
 		LightCardMaterialInstance->SetScalarParameterValue(TEXT("Tint"), Tint);
@@ -489,6 +494,8 @@ void ADisplayClusterLightCardActor::UpdateUVIndicator()
 
 void ADisplayClusterLightCardActor::UpdateLightCardPositionToRootActor()
 {
+	TRACE_CPUPROFILER_EVENT_SCOPE(DCLightcard_UpdateLightCardPositionToRootActor);
+
 	if (!bLockToOwningRootActor || HasAnyFlags(RF_Transient) /* Proxies update from world instance positions instead */)
 	{
 		return;
diff --git a/Engine/Plugins/Runtime/nDisplay/Source/DisplayCluster/Private/Game/EngineClasses/Scene/DisplayClusterRootActor.cpp b/Engine/Plugins/Runtime/nDisplay/Source/DisplayCluster/Private/Game/EngineClasses/Scene/DisplayClusterRootActor.cpp
--- a/Engine/Plugins/Runtime/nDisplay/Source/DisplayCluster/Private/Game/EngineClasses/Scene/DisplayClusterRootActor.cpp
+++ b/Engine/Plugins/Runtime/nDisplay/Source/DisplayCluster/Private/Game/EngineClasses/Scene/DisplayClusterRootActor.cpp
@@ -737,6 +737,8 @@ bool ADisplayClusterRootActor::BuildHierarchy()
 
 void ADisplayClusterRootActor::SetLightCardOwnership()
 {
+	TRACE_CPUPROFILER_EVENT_SCOPE(DCRootActor_SetLightCardOwnership);
+
 	if (UDisplayClusterConfigurationData* CurrentData = GetConfigData())
 	{
 		FDisplayClusterConfigurationICVFX_VisibilityList& LightCardVisibilityList = CurrentData->StageSettings.Lightcard.ShowOnlyList;
@@ -868,6 +870,8 @@ bool ADisplayClusterRootActor::GetFlushPositionAndNormal(const FVector& WorldPos
 
 bool ADisplayClusterRootActor::MakeStageActorFlushToWall(const TScriptInterface<IDisplayClusterStageActor>& StageActor, double DesiredOffsetFromFlush)
 {
+	TRACE_CPUPROFILER_EVENT_SCOPE(DCRootActor_MakeStageActorFlushToWall);
+
 	if (StageActor.GetObject() && !StageActor->IsUVActor())
 	{
 		FVector Position;
diff --git a/Engine/Plugins/Runtime/nDisplay/Source/DisplayCluster/Private/Render/Viewport/LightCard/DisplayClusterViewportLightCardManager.cpp b/Engine/Plugins/Runtime/nDisplay/Source/DisplayCluster/Private/Render/Viewport/LightCard/DisplayClusterViewportLightCardManager.cpp
--- a/Engine/Plugins/Runtime/nDisplay/Source/DisplayCluster/Private/Render/Viewport/LightCard/DisplayClusterViewportLightCardManager.cpp
+++ b/Engine/Plugins/Runtime/nDisplay/Source/DisplayCluster/Private/Render/Viewport/LightCard/DisplayClusterViewportLightCardManager.cpp
@@ -18,7 +18,7 @@
 /** Console variable used to control the size of the UV light card map texture */
 static TAutoConsoleVariable<int32> CVarUVLightCardTextureSize(
 	TEXT("nDisplay.render.uvlightcards.UVTextureSize"),
-	512,
+	4096,
 	TEXT("The size of the texture UV light cards are rendered to.")
 );
 
