nDisplay: Fixed issue where editor would occationally crash when performing and undo while the ICVFX panel was open

diff --git a/Engine/Plugins/Runtime/nDisplay/Source/DisplayClusterLightCardEditor/Private/Outliner/SDisplayClusterLightCardOutliner.cpp b/Engine/Plugins/Runtime/nDisplay/Source/DisplayClusterLightCardEditor/Private/Outliner/SDisplayClusterLightCardOutliner.cpp
--- a/Engine/Plugins/Runtime/nDisplay/Source/DisplayClusterLightCardEditor/Private/Outliner/SDisplayClusterLightCardOutliner.cpp
+++ b/Engine/Plugins/Runtime/nDisplay/Source/DisplayClusterLightCardEditor/Private/Outliner/SDisplayClusterLightCardOutliner.cpp
@@ -61,7 +61,11 @@ void SDisplayClusterLightCardOutliner::SetRootActor(ADisplayClusterRootActor* Ne
 	
 	RootActor = NewRootActor;
 
-	CreateWorldOutliner();
+	FillActorList();
+	if (SceneOutliner.IsValid())
+	{
+		SceneOutliner->FullRefresh();
+	}
 	
 	// Select previously selected light cards. This fixes an issue where the details panel may clear in some cases
 	// and also supports maintaining the selection of shared light cards across different root actors.
