nDisplay:  Fixed crashes when viewport proxy has nullptr resource.

diff --git a/Engine/Plugins/Runtime/nDisplay/Source/DisplayCluster/Private/Render/Viewport/DisplayClusterViewportManagerProxy.cpp b/Engine/Plugins/Runtime/nDisplay/Source/DisplayCluster/Private/Render/Viewport/DisplayClusterViewportManagerProxy.cpp
--- a/Engine/Plugins/Runtime/nDisplay/Source/DisplayCluster/Private/Render/Viewport/DisplayClusterViewportManagerProxy.cpp
+++ b/Engine/Plugins/Runtime/nDisplay/Source/DisplayCluster/Private/Render/Viewport/DisplayClusterViewportManagerProxy.cpp
@@ -513,45 +513,76 @@ void FDisplayClusterViewportManagerProxy::DoCrossGPUTransfers_RenderThread(FRHIC
 	check(IsInRenderingThread());
 
 #if WITH_MGPU
+	if (GNumExplicitGPUsForRendering < 2)
+	{
+		// for mGPU we need at least 2 GPUs
+		return;
+	}
+
 	// Copy the view render results to all GPUs that are native to the viewport.
 	TArray<FTransferResourceParams> TransferResources;
 
 	for (const TSharedPtr<FDisplayClusterViewportProxy, ESPMode::ThreadSafe>& ViewportProxyIt : ClusterNodeViewportProxies)
 	{
-		// Skip a frozen viewport that has already been transferred between GPUs
-		// The first time freezing should do the transfer (RenderTargets must be assigned on the first pass)
-		const bool bSkipThisViewport = ViewportProxyIt->RenderSettings.bFreezeRendering && ViewportProxyIt->RenderTargets.Num() == 0;
-		if (ViewportProxyIt->RenderSettings.bEnableCrossGPUTransfer && !bSkipThisViewport)
+		if (!ViewportProxyIt.IsValid())
+		{
+			continue;
+		}
+
+		if (!ViewportProxyIt->RenderSettings.bEnableCrossGPUTransfer || ViewportProxyIt->RenderTargets.IsEmpty())
+		{
+			// Skip a frozen viewport that has already been transferred between GPUs
+			// The first time freezing should do the transfer (RenderTargets must be assigned on the first pass)
+			continue;
+		}
+
+		// Do cross GPU trasfers for all viewport contexts
+		for (const FDisplayClusterViewport_Context& ViewportContext : ViewportProxyIt->Contexts)
 		{
-			for (const FDisplayClusterViewport_Context& ViewportContext : ViewportProxyIt->Contexts)
+			if (!ViewportContext.bOverrideCrossGPUTransfer)
+			{
+				// Skip this context because it uses the default way
+				continue;
+			}
+
+			const int32 SrcGPUIndex = ViewportContext.RenderThreadData.GPUIndex;
+			if (SrcGPUIndex < 0)
+			{
+				// This viewport context does not use the custom GPU index for rendering
+				continue;
+			}
+
+			if (!ViewportProxyIt->RenderTargets.IsValidIndex(ViewportContext.ContextNum))
+			{
+				// RTT does not exist for this context
+				continue;
+			}
+
+			// Clamp the view rect by the rendertarget rect to prevent issues when resizing the viewport.
+			const FIntRect TransferRect = ViewportContext.RenderTargetRect;
+			if (TransferRect.Width() <= 0 || TransferRect.Height() <= 0)
 			{
-				const int32 SrcGPUIndex = ViewportContext.RenderThreadData.GPUIndex;
+				// Broken rect, skip
+				continue;
+			}
 
-				// Perform an MGPU transfer for a specific viewport context:
-				if (ViewportContext.bOverrideCrossGPUTransfer && SrcGPUIndex >= 0)
+			// Perform an MGPU transfer for a specific viewport context:
+			if (FDisplayClusterViewportRenderTargetResource* RenderTarget = ViewportProxyIt->RenderTargets[ViewportContext.ContextNum])
+			{
+				if (FRHITexture2D* TextureRHI = RenderTarget->GetViewportResource2DRHI())
 				{
-					// Clamp the view rect by the rendertarget rect to prevent issues when resizing the viewport.
-					const FIntRect TransferRect = ViewportContext.RenderTargetRect;
-					if (TransferRect.Width() > 0 && TransferRect.Height() > 0)
-					{
-						// Use optimized cross GPU transfer for this context
-						FRenderTarget* RenderTarget = ViewportProxyIt->RenderTargets[ViewportContext.ContextNum];
-						const FRHIGPUMask RenderTargetGPUMask = (GNumExplicitGPUsForRendering > 1 && RenderTarget) ? RenderTarget->GetGPUMask(RHICmdList) : FRHIGPUMask::GPU0();
+					const FRHIGPUMask RenderTargetGPUMask = RenderTarget->GetGPUMask(RHICmdList);
 
-						if (FRHITexture2D* TextureRHI = ViewportProxyIt->RenderTargets[ViewportContext.ContextNum]->GetViewportRenderTargetResourceRHI())
+					for (uint32 DestGPUIndex : RenderTargetGPUMask)
+					{
+						if (DestGPUIndex != SrcGPUIndex)
 						{
-							for (uint32 DestGPUIndex : RenderTargetGPUMask)
-							{
-								if (DestGPUIndex != SrcGPUIndex)
-								{
-									TransferResources.Add(FTransferResourceParams(
-										TextureRHI, TransferRect,
-										SrcGPUIndex, DestGPUIndex,
-										RenderFrameSettings.CrossGPUTransfer.bPullData,
-										RenderFrameSettings.CrossGPUTransfer.bLockSteps
-									));
-								}
-							}
+							TransferResources.Add(FTransferResourceParams(
+								TextureRHI, TransferRect,
+								SrcGPUIndex, DestGPUIndex,
+								RenderFrameSettings.CrossGPUTransfer.bPullData,
+								RenderFrameSettings.CrossGPUTransfer.bLockSteps
+							));
 						}
 					}
 				}
@@ -574,37 +605,44 @@ bool FDisplayClusterViewportManagerProxy::GetFrameTargets_RenderThread(TArray<FR
 	// Get any defined frame targets from first visible viewport
 	for (const TSharedPtr<FDisplayClusterViewportProxy, ESPMode::ThreadSafe>& ViewportProxyIt : ClusterNodeViewportProxies)
 	{
-		if (ViewportProxyIt.IsValid())
+		// Process only valid viewports with output frame resources
+		if (ViewportProxyIt.IsValid() && !ViewportProxyIt->OutputFrameTargetableResources.IsEmpty())
 		{
 			const TArray<FDisplayClusterViewportTextureResource*>& Frames = ViewportProxyIt->OutputFrameTargetableResources;
 			const TArray<FDisplayClusterViewportTextureResource*>& AdditionalFrames = ViewportProxyIt->AdditionalFrameTargetableResources;
 
-			if (Frames.Num() > 0)
+			for (int32 FrameIt = 0; FrameIt < Frames.Num(); FrameIt++)
 			{
-				OutFrameResources.Reserve(Frames.Num());
-				OutTargetOffsets.Reserve(Frames.Num());
-
-				bool bUseAdditionalFrameResources = (OutAdditionalFrameResources != nullptr) && (AdditionalFrames.Num() == Frames.Num());
-
-				if (bUseAdditionalFrameResources)
+				if (FRHITexture2D* FrameTexture = Frames[FrameIt] ? Frames[FrameIt]->GetViewportResource2DRHI() : nullptr)
 				{
-					OutAdditionalFrameResources->AddZeroed(AdditionalFrames.Num());
-				}
-
-				for (int32 FrameIt = 0; FrameIt < Frames.Num(); FrameIt++)
-				{
-					OutFrameResources.Add(Frames[FrameIt]->GetViewportResourceRHI());
+					OutFrameResources.Add(FrameTexture);
 					OutTargetOffsets.Add(Frames[FrameIt]->BackbufferFrameOffset);
 
-					if (bUseAdditionalFrameResources)
+					if (OutAdditionalFrameResources && AdditionalFrames.IsValidIndex(FrameIt))
 					{
-						(*OutAdditionalFrameResources)[FrameIt] = AdditionalFrames[FrameIt]->GetViewportResourceRHI();
+						if (FRHITexture2D* AdditionalFrameTexture = AdditionalFrames[FrameIt] ? AdditionalFrames[FrameIt]->GetViewportResource2DRHI() : nullptr)
+						{
+							OutAdditionalFrameResources->Add(AdditionalFrameTexture);
+						}
 					}
 				}
+			}
 
+			// if all resources are received in the same amount
+			const bool bValidAdditionalFrameResources = !OutAdditionalFrameResources || AdditionalFrames.IsEmpty() || (OutAdditionalFrameResources && OutAdditionalFrameResources->Num() == Frames.Num());
+			if (OutFrameResources.Num() == Frames.Num() && bValidAdditionalFrameResources)
+			{
 				return true;
 			}
 		}
+
+		// Resetting output values at the end of each cycle
+		OutFrameResources.Reset();
+		OutTargetOffsets.Reset();
+		if (OutAdditionalFrameResources)
+		{
+			OutAdditionalFrameResources->Reset();
+		}
 	}
 
 	// no visible viewports
diff --git a/Engine/Plugins/Runtime/nDisplay/Source/DisplayCluster/Private/Render/Viewport/DisplayClusterViewportProxy.cpp b/Engine/Plugins/Runtime/nDisplay/Source/DisplayCluster/Private/Render/Viewport/DisplayClusterViewportProxy.cpp
--- a/Engine/Plugins/Runtime/nDisplay/Source/DisplayCluster/Private/Render/Viewport/DisplayClusterViewportProxy.cpp
+++ b/Engine/Plugins/Runtime/nDisplay/Source/DisplayCluster/Private/Render/Viewport/DisplayClusterViewportProxy.cpp
@@ -123,65 +123,37 @@ namespace DisplayClusterViewportProxyHelpers
 	*
 	* @return true, if all resources are valid
 	*/
-	static bool GetViewportRHIResourcesImpl_RenderThread(const TArray<FDisplayClusterViewportTextureResource*>& InResources, TArray<FRHITexture2D*>& OutResources)
+	static bool GetViewportRHIResourcesImpl_RenderThread(const TArrayView<FDisplayClusterViewportResource*>& InResources, TArray<FRHITexture2D*>& OutResources)
 	{
-		check(OutResources.Num() == 0);
+		OutResources.Reset();
 
-		if (InResources.Num() > 0)
+		for (FDisplayClusterViewportResource* ViewportResourceIt : InResources)
 		{
-			for (int32 ResourceIndex = 0; ResourceIndex < InResources.Num(); ResourceIndex++)
-			{
-				if (InResources[ResourceIndex] != nullptr)
-				{
-					OutResources.Add(InResources[ResourceIndex]->GetViewportResourceRHI());
-				}
-
-			}
-
-			if (OutResources.Num() == InResources.Num())
+			if (FRHITexture2D* RHITexture2D = ViewportResourceIt ? ViewportResourceIt->GetViewportResource2DRHI() : nullptr)
 			{
-				return true;
+				// Collects only valid resources.
+				OutResources.Add(RHITexture2D);
 			}
+		}
 
+		if (OutResources.Num() != InResources.Num())
+		{
 			// Some resources lost
-			OutResources.Empty();
+			OutResources.Reset();
 		}
 
-		return false;
+		// returns success if the number of output resources is equal to the input and is not empty
+		return !OutResources.IsEmpty();
 	}
 
-	/**
-	* Get viewport RHI resources
-	*
-	* @param InResources - Array with DC resources (for mono num=1, for stereo num=2)
-	* @param OutResources - Array with RHI resources
-	*
-	* @return true, if all resources are valid
-	*/
-	static bool GetViewportRHIResourcesImpl_RenderThread(const TArray<FTextureRenderTargetResource*>& InResources, TArray<FRHITexture2D*>& OutResources)
+	static bool GetViewportRHIResourcesImpl_RenderThread(const TArray<FDisplayClusterViewportTextureResource*>& InResources, TArray<FRHITexture2D*>& OutResources)
 	{
-		check(OutResources.Num() == 0);
-
-		if (InResources.Num() > 0)
-		{
-			for (int32 ResourceIndex = 0; ResourceIndex < InResources.Num(); ResourceIndex++)
-			{
-				if (InResources[ResourceIndex] != nullptr)
-				{
-					OutResources.Add(InResources[ResourceIndex]->GetTexture2DRHI());
-				}
-			}
-
-			if (OutResources.Num() == InResources.Num())
-			{
-				return true;
-			}
-
-			// Some resources lost
-			OutResources.Empty();
-		}
+		return GetViewportRHIResourcesImpl_RenderThread(TArrayView<FDisplayClusterViewportResource*>((FDisplayClusterViewportResource**)(InResources.GetData()), InResources.Num()), OutResources);
+	}
 
-		return false;
+	static bool GetViewportRHIResourcesImpl_RenderThread(const TArray<FDisplayClusterViewportRenderTargetResource*>& InResources, TArray<FRHITexture2D*>& OutResources)
+	{
+		return GetViewportRHIResourcesImpl_RenderThread(TArrayView<FDisplayClusterViewportResource*>((FDisplayClusterViewportResource**)(InResources.GetData()), InResources.Num()), OutResources);
 	}
 
 	// Reset reference to deleted resource
@@ -395,17 +367,7 @@ bool FDisplayClusterViewportProxy::ImplGetResources_RenderThread(const EDisplayC
 			// 3. Finally Use InternalRTT
 			if (!bResult)
 			{
-				for (int32 ContextIndex = 0; ContextIndex < Contexts.Num(); ContextIndex++)
-				{
-					if (RenderTargets.IsValidIndex(ContextIndex))
-					{
-						if (FDisplayClusterViewportRenderTargetResource* Input = RenderTargets[ContextIndex])
-						{
-							OutResources.Add(Input->GetViewportRenderTargetResourceRHI());
-							bResult = true;
-						}
-					}
-				}
+				bResult = GetViewportRHIResourcesImpl_RenderThread(RenderTargets, OutResources);
 			}
 		}
 
@@ -515,8 +477,8 @@ void FDisplayClusterViewportProxy::ImplViewportRemap_RenderThread(FRHICommandLis
 				FDisplayClusterViewportTextureResource* Src = AdditionalFrameTargetableResources[ContextIt];
 				FDisplayClusterViewportTextureResource* Dst = OutputFrameTargetableResources[ContextIt];
 
-				FRHITexture2D* Input = Src ? Src->GetViewportResourceRHI() : nullptr;
-				FRHITexture2D* Output = Dst ? Dst->GetViewportResourceRHI() : nullptr;
+				FRHITexture2D* Input = Src ? Src->GetViewportResource2DRHI() : nullptr;
+				FRHITexture2D* Output = Dst ? Dst->GetViewportResource2DRHI() : nullptr;
 
 				if (Input && Output)
 				{
@@ -1192,26 +1154,25 @@ void FDisplayClusterViewportProxy::OnPostRenderViewFamily_RenderThread(FRDGBuild
 		{
 			// Restore alpha channed
 			EFXAAQuality FXAAQuality = EFXAAQuality::Q0;
-			if (GetFXAAQuality(RenderSettings.CaptureMode, FXAAQuality))
+			if (GetFXAAQuality(RenderSettings.CaptureMode, FXAAQuality) && InputShaderResources.IsValidIndex(InContextNum))
 			{
-				// Apply FXAA for RGB only
-				// Note: Add AA for alpha channel
+				if (FRHITexture2D* InputTextureRHI = InputShaderResources[InContextNum] ? InputShaderResources[InContextNum]->GetViewportResource2DRHI() : nullptr)
+				{
+					// Apply FXAA for RGB only
+					// Note: Add AA for alpha channel
 
-				// Copy Alpha channels back from'InputShaderResource' to 'InternalRenderTargetResource'
-				CopyResource_RenderThread(GraphBuilder, EDisplayClusterTextureCopyMode::Alpha, InContextNum, EDisplayClusterViewportResourceType::InputShaderResource, EDisplayClusterViewportResourceType::InternalRenderTargetResource);
+					// Copy Alpha channels back from'InputShaderResource' to 'InternalRenderTargetResource'
+					CopyResource_RenderThread(GraphBuilder, EDisplayClusterTextureCopyMode::Alpha, InContextNum, EDisplayClusterViewportResourceType::InputShaderResource, EDisplayClusterViewportResourceType::InternalRenderTargetResource);
 
-				checkSlow(InSceneView.bIsViewInfo);
-				const FViewInfo& ViewInfo = static_cast<const FViewInfo&>(InSceneView);
+					checkSlow(InSceneView.bIsViewInfo);
+					const FViewInfo& ViewInfo = static_cast<const FViewInfo&>(InSceneView);
 
-				// 1. Copy RGB channels from 'InternalRenderTargetResource' to 'InputShaderResource'
-				CopyResource_RenderThread(GraphBuilder, EDisplayClusterTextureCopyMode::RGB, InContextNum, EDisplayClusterViewportResourceType::InternalRenderTargetResource, EDisplayClusterViewportResourceType::InputShaderResource);
+					// 1. Copy RGB channels from 'InternalRenderTargetResource' to 'InputShaderResource'
+					CopyResource_RenderThread(GraphBuilder, EDisplayClusterTextureCopyMode::RGB, InContextNum, EDisplayClusterViewportResourceType::InternalRenderTargetResource, EDisplayClusterViewportResourceType::InputShaderResource);
 
-				check(InputShaderResources.IsValidIndex(InContextNum));
-				if (FDisplayClusterViewportTextureResource* InputResource = InputShaderResources[InContextNum])
-				{
 					// 2. FXAA render pass to 'InputShaderResource'
 					// Input is 'InputShaderResource'
-					FRDGTextureRef InputTexture = RegisterExternalTexture(GraphBuilder, InputResource->GetTextureRHI(), TEXT("DCViewportProxyFXAAResource"));
+					FRDGTextureRef InputTexture = RegisterExternalTexture(GraphBuilder, InputTextureRHI, TEXT("DCViewportProxyFXAAResource"));
 					GraphBuilder.SetTextureAccessFinal(InputTexture, ERHIAccess::RTV);
 
 					FFXAAInputs PassInputs;
diff --git a/Engine/Plugins/Runtime/nDisplay/Source/DisplayCluster/Private/Render/Viewport/RenderTarget/DisplayClusterRenderTargetResource.h b/Engine/Plugins/Runtime/nDisplay/Source/DisplayCluster/Private/Render/Viewport/RenderTarget/DisplayClusterRenderTargetResource.h
--- a/Engine/Plugins/Runtime/nDisplay/Source/DisplayCluster/Private/Render/Viewport/RenderTarget/DisplayClusterRenderTargetResource.h
+++ b/Engine/Plugins/Runtime/nDisplay/Source/DisplayCluster/Private/Render/Viewport/RenderTarget/DisplayClusterRenderTargetResource.h
@@ -115,14 +115,9 @@ public:
 		return ResourceState[(uint8)InState];
 	}
 
-	virtual const FTexture2DRHIRef& GetViewportResourceRHI() const
+	virtual FRHITexture2D* GetViewportResource2DRHI() const
 	{
-		return (const FTexture2DRHIRef&)TextureRHI;
-	}
-
-	virtual const FTexture2DRHIRef& GetViewportRenderTargetResourceRHI() const
-	{
-		return (const FTexture2DRHIRef&)TextureRHI;
+		return TextureRHI.IsValid() ? TextureRHI->GetTexture2D() : nullptr;
 	}
 
 	/** Returns the width of the texture in pixels. */
