nDisplay: Dragging a light card or CCW in the ICVFX window will now trigger sequencer key frames. Changing the time in sequencer will update level and ICVFX proxy positions accordingly.

Light cards work with sequencer by using the orientation properties such as latitude and longitude.
CCWs require the transform tracked instead.

Light card alpha gradient settings are also exposed to sequencer.

diff --git a/Engine/Plugins/Experimental/ColorCorrectRegions/Source/ColorCorrectRegions/Private/ColorCorrectWindow.cpp b/Engine/Plugins/Experimental/ColorCorrectRegions/Source/ColorCorrectRegions/Private/ColorCorrectWindow.cpp
--- a/Engine/Plugins/Experimental/ColorCorrectRegions/Source/ColorCorrectRegions/Private/ColorCorrectWindow.cpp
+++ b/Engine/Plugins/Experimental/ColorCorrectRegions/Source/ColorCorrectRegions/Private/ColorCorrectWindow.cpp
@@ -11,6 +11,10 @@
 #include "Misc/EnumRange.h"
 #include "UObject/ConstructorHelpers.h"
 
+#if WITH_EDITOR
+#include "IDisplayClusterLightCardExtenderModule.h"
+#endif
+
 ENUM_RANGE_BY_COUNT(EColorCorrectWindowType, EColorCorrectWindowType::MAX)
 
 AColorCorrectionWindow::AColorCorrectionWindow(const FObjectInitializer& ObjectInitializer)
@@ -50,6 +54,25 @@ AColorCorrectionWindow::AColorCorrectionWindow(const FObjectInitializer& ObjectI
 #if WITH_METADATA
 	CreateIcon();
 #endif
+
+#if WITH_EDITOR
+	if (!IsTemplate())
+	{
+		IDisplayClusterLightCardExtenderModule& LightCardExtenderModule = IDisplayClusterLightCardExtenderModule::Get();
+		LightCardExtenderModule.GetOnSequencerTimeChanged().AddUObject(this, &AColorCorrectionWindow::OnSequencerTimeChanged);
+	}
+#endif
+}
+
+AColorCorrectionWindow::~AColorCorrectionWindow()
+{
+#if WITH_EDITOR
+	if (!IsTemplate())
+	{
+		IDisplayClusterLightCardExtenderModule& LightCardExtenderModule = IDisplayClusterLightCardExtenderModule::Get();
+		LightCardExtenderModule.GetOnSequencerTimeChanged().RemoveAll(this);
+	}
+#endif
 }
 
 #if WITH_EDITOR
@@ -102,6 +125,11 @@ FName AColorCorrectionWindow::GetCustomIconName() const
 	return TEXT("CCW.OutlinerThumbnail");
 }
 
+void AColorCorrectionWindow::OnSequencerTimeChanged(TWeakPtr<ISequencer> InSequencer)
+{
+	UpdatePositionalParamsFromTransform();
+}
+
 #endif //WITH_EDITOR
 
 
@@ -308,4 +336,9 @@ void AColorCorrectionWindow::GetPositionalProperties(FPositionalPropertyArray& O
 	}
 }
 
+FName AColorCorrectionWindow::GetPositionalPropertiesMemberName() const
+{
+	return GET_MEMBER_NAME_CHECKED(AColorCorrectionWindow, PositionalParams);
+}
+
 #undef NOTIFY_PARAM_SETTER
diff --git a/Engine/Plugins/Experimental/ColorCorrectRegions/Source/ColorCorrectRegions/Public/ColorCorrectWindow.h b/Engine/Plugins/Experimental/ColorCorrectRegions/Source/ColorCorrectRegions/Public/ColorCorrectWindow.h
--- a/Engine/Plugins/Experimental/ColorCorrectRegions/Source/ColorCorrectRegions/Public/ColorCorrectWindow.h
+++ b/Engine/Plugins/Experimental/ColorCorrectRegions/Source/ColorCorrectRegions/Public/ColorCorrectWindow.h
@@ -26,6 +26,7 @@ class COLORCORRECTREGIONS_API AColorCorrectionWindow : public AColorCorrectRegio
 {
 	GENERATED_UCLASS_BODY()
 public:
+	virtual ~AColorCorrectionWindow() override;
 
 	UPROPERTY()
 	TArray<TObjectPtr<UStaticMeshComponent>> MeshComponents;
@@ -39,6 +40,9 @@ public:
 	virtual void PostEditChangeProperty(FPropertyChangedEvent& PropertyChangedEvent) override;
 	virtual void PostEditMove(bool bFinished) override;
 	virtual FName GetCustomIconName() const override;
+
+private:
+	void OnSequencerTimeChanged(TWeakPtr<class ISequencer> InSequencer);
 #endif
 
 private:
@@ -103,6 +107,7 @@ public:
 	virtual FDisplayClusterPositionalParams GetPositionalParams() const override;
 
 	virtual void GetPositionalProperties(FPositionalPropertyArray& OutPropertyPairs) const override;
+	virtual FName GetPositionalPropertiesMemberName() const override;
 	// ~End IDisplayClusterStageActor interface
 	
 protected:
diff --git a/Engine/Plugins/Runtime/nDisplay/Source/DisplayCluster/Public/DisplayClusterLightCardActor.h b/Engine/Plugins/Runtime/nDisplay/Source/DisplayCluster/Public/DisplayClusterLightCardActor.h
--- a/Engine/Plugins/Runtime/nDisplay/Source/DisplayCluster/Public/DisplayClusterLightCardActor.h
+++ b/Engine/Plugins/Runtime/nDisplay/Source/DisplayCluster/Public/DisplayClusterLightCardActor.h
@@ -37,19 +37,19 @@ struct FLightCardAlphaGradientSettings
 	GENERATED_BODY()
 
 	/** Enables/disables alpha gradient effect */
-	UPROPERTY(EditAnywhere, BlueprintReadWrite, Category = "Appearance")
+	UPROPERTY(EditAnywhere, Interp, BlueprintReadWrite, Category = "Appearance")
 	bool bEnableAlphaGradient = false;
 
 	/** Starting alpha value in the gradient */
-	UPROPERTY(EditAnywhere, BlueprintReadWrite, Category = "Appearance", meta = (EditCondition = "bEnableAlphaGradient"))
+	UPROPERTY(EditAnywhere, Interp, BlueprintReadWrite, Category = "Appearance", meta = (EditCondition = "bEnableAlphaGradient"))
 	float StartingAlpha = 0;
 
 	/** Ending alpha value in the gradient */
-	UPROPERTY(EditAnywhere, BlueprintReadWrite, Category = "Appearance", meta = (EditCondition = "bEnableAlphaGradient"))
+	UPROPERTY(EditAnywhere, Interp, BlueprintReadWrite, Category = "Appearance", meta = (EditCondition = "bEnableAlphaGradient"))
 	float EndingAlpha = 1;
 
 	/** The angle (degrees) determines the gradient direction. */
-	UPROPERTY(EditAnywhere, BlueprintReadWrite, Category = "Appearance", meta = (EditCondition = "bEnableAlphaGradient"))
+	UPROPERTY(EditAnywhere, Interp, BlueprintReadWrite, Category = "Appearance", meta = (EditCondition = "bEnableAlphaGradient"))
 	float Angle = 0;
 };
 
diff --git a/Engine/Plugins/Runtime/nDisplay/Source/DisplayClusterLightCardEditor/Private/DisplayClusterLightCardEditor.cpp b/Engine/Plugins/Runtime/nDisplay/Source/DisplayClusterLightCardEditor/Private/DisplayClusterLightCardEditor.cpp
--- a/Engine/Plugins/Runtime/nDisplay/Source/DisplayClusterLightCardEditor/Private/DisplayClusterLightCardEditor.cpp
+++ b/Engine/Plugins/Runtime/nDisplay/Source/DisplayClusterLightCardEditor/Private/DisplayClusterLightCardEditor.cpp
@@ -1820,7 +1820,11 @@ void FDisplayClusterLightCardEditor::OnActorPropertyChanged(UObject* ObjectBeing
                                                             FPropertyChangedEvent& PropertyChangedEvent)
 {
 	EDisplayClusterLightCardEditorProxyType ProxyType;
-	if (IsOurObject(ObjectBeingModified, ProxyType))
+	// When a level instance actor property is changed we want to update our proxy version here. However,
+	// sometimes the viewport client is responsible for the property change call and in that case we want to avoid
+	// double processing an update.
+	if (IsOurObject(ObjectBeingModified, ProxyType) && (!ViewportView.IsValid() ||
+		!ViewportView->GetLightCardEditorViewportClient()->HasCalledPostEditChangeProperty()))
 	{
 		IDisplayClusterStageActor* StageActor = Cast<IDisplayClusterStageActor>(ObjectBeingModified);
 		
diff --git a/Engine/Plugins/Runtime/nDisplay/Source/DisplayClusterLightCardEditor/Private/Viewport/DisplayClusterLightCardEditorViewportClient.cpp b/Engine/Plugins/Runtime/nDisplay/Source/DisplayClusterLightCardEditor/Private/Viewport/DisplayClusterLightCardEditorViewportClient.cpp
--- a/Engine/Plugins/Runtime/nDisplay/Source/DisplayClusterLightCardEditor/Private/Viewport/DisplayClusterLightCardEditorViewportClient.cpp
+++ b/Engine/Plugins/Runtime/nDisplay/Source/DisplayClusterLightCardEditor/Private/Viewport/DisplayClusterLightCardEditorViewportClient.cpp
@@ -22,6 +22,8 @@
 #include "SceneManagement.h"
 #include "Settings/DisplayClusterLightCardEditorSettings.h"
 
+#include "IDisplayClusterLightCardExtenderModule.h"
+
 #include "AudioDevice.h"
 #include "CameraController.h"
 #include "Components/BillboardComponent.h"
@@ -138,6 +140,9 @@ FDisplayClusterLightCardEditorViewportClient::FDisplayClusterLightCardEditorView
 
 	const int32 CameraSpeed = 3;
 	SetCameraSpeedSetting(CameraSpeed);
+
+	IDisplayClusterLightCardExtenderModule& LightCardExtenderModule = IDisplayClusterLightCardExtenderModule::Get();
+	LightCardExtenderModule.GetOnSequencerTimeChanged().AddRaw(this, &FDisplayClusterLightCardEditorViewportClient::OnSequencerTimeChanged);
 }
 
 FDisplayClusterLightCardEditorViewportClient::~FDisplayClusterLightCardEditorViewportClient()
@@ -146,6 +151,9 @@ FDisplayClusterLightCardEditorViewportClient::~FDisplayClusterLightCardEditorVie
 
 	EndTransaction();
 	UnsubscribeFromRootActor();
+
+	IDisplayClusterLightCardExtenderModule& LightCardExtenderModule = IDisplayClusterLightCardExtenderModule::Get();
+	LightCardExtenderModule.GetOnSequencerTimeChanged().RemoveAll(this);
 }
 
 FLinearColor FDisplayClusterLightCardEditorViewportClient::GetBackgroundColor() const
@@ -1266,6 +1274,11 @@ void FDisplayClusterLightCardEditorViewportClient::UnsubscribeFromRootActor()
 	}
 }
 
+void FDisplayClusterLightCardEditorViewportClient::OnSequencerTimeChanged(TWeakPtr<ISequencer> InSequencer)
+{
+	UpdateProxyTransforms();
+}
+
 void FDisplayClusterLightCardEditorViewportClient::ProcessClick(FSceneView& View, HHitProxy* HitProxy, FKey Key, EInputEvent Event, uint32 HitX, uint32 HitY)
 {
 	// Don't select light cards while drawing a new light card
@@ -2418,10 +2431,33 @@ void FDisplayClusterLightCardEditorViewportClient::PropagateActorTransform(const
 		) -> void
 		{
 			// Only change if values are different.
-			const FProperty* Property = ProxyPropertyPair.Value;
+			FProperty* Property = ProxyPropertyPair.Value;
 			if (!Property->Identical_InContainer(ProxyPropertyPair.Key, LevelInstancePropertyPair.Key))
 			{
+				FEditPropertyChain PropertyChain;
+				PropertyChain.AddTail(Property);
+				PropertyChain.SetActivePropertyNode(Property);
+
+				const FName PositionalParamsPropertyName = LevelInstanceStageActor->GetPositionalPropertiesMemberName();
+				if (!PositionalParamsPropertyName.IsNone())
+				{
+					if (FProperty* MemberProperty = LevelInstance->GetClass()->FindPropertyByName(PositionalParamsPropertyName))
+					{
+						PropertyChain.AddHead(MemberProperty);
+						PropertyChain.SetActiveMemberPropertyNode(MemberProperty);
+					}
+				}
+
+				// Required for sequencer to update property key frames
+				FCoreUObjectDelegates::OnPreObjectPropertyChanged.Broadcast(LevelInstance, PropertyChain);
+				
 				Property->CopyCompleteValue_InContainer(LevelInstancePropertyPair.Key, ProxyPropertyPair.Key);
+
+				bHasCalledPostEditChangeProperty = true;
+				FPropertyChangedEvent PropertyChangedEvent(Property, EPropertyChangeType::ValueSet);
+				FCoreUObjectDelegates::OnObjectPropertyChanged.Broadcast(LevelInstance, PropertyChangedEvent);
+				bHasCalledPostEditChangeProperty = false;
+				
 				InOutChangedProperties.Add(Property);
 			}
 		};
@@ -2434,9 +2470,35 @@ void FDisplayClusterLightCardEditorViewportClient::PropagateActorTransform(const
 		const FVector  RALevelLocation = RootActorLevelInstance.IsValid() ? RootActorLevelInstance->GetActorLocation() : FVector::ZeroVector;
 
 		const FTransform RALevelTransformNoScale(RALevelRotation, RALevelLocation, FVector::OneVector);
+		
+		FTransform NewTransform = ActorProxy.AsActorChecked()->GetTransform() * RALevelTransformNoScale;
+
+		FProperty* TransformProperty = LevelInstance->GetRootComponent() ?
+				LevelInstance->GetRootComponent()->GetClass()->FindPropertyByName(USceneComponent::GetRelativeLocationPropertyName())
+				: nullptr;
+
+		const bool bTransformChanged = !NewTransform.Equals(LevelInstance->GetActorTransform(), 0.f);
 
-		LevelInstance->SetActorTransform(ActorProxy.AsActorChecked()->GetTransform() * RALevelTransformNoScale);
+		if (TransformProperty && bTransformChanged)
+		{
+			FEditPropertyChain PropertyChain;
+			PropertyChain.AddTail(TransformProperty);
+			PropertyChain.SetActivePropertyNode(TransformProperty);
+			// Required for sequencer and CCWs.
+			FCoreUObjectDelegates::OnPreObjectPropertyChanged.Broadcast(LevelInstance, PropertyChain);
+		}
+		
+		LevelInstance->SetActorTransform(NewTransform);
 
+		if (TransformProperty && bTransformChanged)
+		{
+			// Required for sequencer and CCWs.
+			bHasCalledPostEditChangeProperty = true;
+			FPropertyChangedEvent PropertyChangedEvent(TransformProperty, EPropertyChangeType::ValueSet);
+			FCoreUObjectDelegates::OnObjectPropertyChanged.Broadcast(LevelInstance, PropertyChangedEvent);
+			bHasCalledPostEditChangeProperty = false;
+		}
+		
 		TArray<const FProperty*> ChangedProperties;
 		ChangedProperties.Reserve(ProxyPropertyPairs.Num());
 
diff --git a/Engine/Plugins/Runtime/nDisplay/Source/DisplayClusterLightCardEditor/Private/Viewport/DisplayClusterLightCardEditorViewportClient.h b/Engine/Plugins/Runtime/nDisplay/Source/DisplayClusterLightCardEditor/Private/Viewport/DisplayClusterLightCardEditorViewportClient.h
--- a/Engine/Plugins/Runtime/nDisplay/Source/DisplayClusterLightCardEditor/Private/Viewport/DisplayClusterLightCardEditorViewportClient.h
+++ b/Engine/Plugins/Runtime/nDisplay/Source/DisplayClusterLightCardEditor/Private/Viewport/DisplayClusterLightCardEditorViewportClient.h
@@ -162,6 +162,9 @@ public:
 	/** Requests that we exit light card drawing input mode (and go back to idle/normal) */
 	void ExitDrawingLightCardMode();
 
+	/** If the viewport client is responsible for the post edit change property call */
+	bool HasCalledPostEditChangeProperty() const { return bHasCalledPostEditChangeProperty; }
+
 private:
 	/** Initiates a transaction. */
 	void BeginTransaction(const FText& Description);
@@ -299,6 +302,9 @@ private:
 	
 	/** Unsubscribe post process and preview hooks from the root actor */
 	void UnsubscribeFromRootActor();
+
+	/** When the time in a sequencer has changed */
+	void OnSequencerTimeChanged(TWeakPtr<class ISequencer> InSequencer);
 	
 private:
 	TWeakPtr<FSceneViewport> SceneViewportPtr;
@@ -403,4 +409,7 @@ private:
 
 	/** A flag that disables drawing with the custom projection renderer and instead renders the viewport client with the editor's default renderer */
 	bool bDisableCustomRenderer = false;
+
+	/** True if the viewport client has triggered the post edit change property call */
+	bool bHasCalledPostEditChangeProperty = false;
 };
\ No newline at end of file
diff --git a/Engine/Plugins/Runtime/nDisplayModularFeatures/Source/DisplayClusterLightCardExtender/DisplayClusterLightCardExtender.build.cs b/Engine/Plugins/Runtime/nDisplayModularFeatures/Source/DisplayClusterLightCardExtender/DisplayClusterLightCardExtender.build.cs
--- a/Engine/Plugins/Runtime/nDisplayModularFeatures/Source/DisplayClusterLightCardExtender/DisplayClusterLightCardExtender.build.cs
+++ b/Engine/Plugins/Runtime/nDisplayModularFeatures/Source/DisplayClusterLightCardExtender/DisplayClusterLightCardExtender.build.cs
@@ -16,7 +16,12 @@ public class DisplayClusterLightCardExtender : ModuleRules
 
 		if (Target.Type == TargetType.Editor)
 		{
-			PrivateDependencyModuleNames.Add("UnrealEd");
+			PrivateDependencyModuleNames.AddRange(
+				new string[]
+				{
+					"Sequencer",
+					"UnrealEd"
+				});
 		}
 	}
 }
diff --git a/Engine/Plugins/Runtime/nDisplayModularFeatures/Source/DisplayClusterLightCardExtender/Private/DisplayClusterLightCardExtenderModule.cpp b/Engine/Plugins/Runtime/nDisplayModularFeatures/Source/DisplayClusterLightCardExtender/Private/DisplayClusterLightCardExtenderModule.cpp
--- a/Engine/Plugins/Runtime/nDisplayModularFeatures/Source/DisplayClusterLightCardExtender/Private/DisplayClusterLightCardExtenderModule.cpp
+++ b/Engine/Plugins/Runtime/nDisplayModularFeatures/Source/DisplayClusterLightCardExtender/Private/DisplayClusterLightCardExtenderModule.cpp
@@ -1,5 +1,65 @@
 // Copyright Epic Games, Inc. All Rights Reserved.
 
-#include "Modules/ModuleManager.h"
+#include "DisplayClusterLightCardExtenderModule.h"
 
-IMPLEMENT_MODULE(FDefaultModuleImpl, DisplayClusterLightCardExtender)
+#if WITH_EDITOR
+#include "ISequencerModule.h"
+#endif
+
+void FDisplayClusterLightCardExtenderModule::StartupModule()
+{
+#if WITH_EDITOR
+	ISequencerModule& SequencerModule = FModuleManager::Get().LoadModuleChecked<ISequencerModule>("Sequencer");
+	SequencerCreatedHandle = SequencerModule.RegisterOnSequencerCreated(FOnSequencerCreated::FDelegate::CreateRaw(this, &FDisplayClusterLightCardExtenderModule::OnSequencerCreated));
+#endif
+}
+
+void FDisplayClusterLightCardExtenderModule::ShutdownModule()
+{
+#if WITH_EDITOR
+	for (const TWeakPtr<ISequencer, ESPMode::ThreadSafe>& Sequencer : OpenSequencers)
+	{
+		if (Sequencer.IsValid())
+		{
+			Sequencer.Pin()->OnGlobalTimeChanged().RemoveAll(this);
+			Sequencer.Pin()->OnCloseEvent().RemoveAll(this);
+		}
+	}
+
+	OpenSequencers.Empty();
+	
+	if (FModuleManager::Get().IsModuleLoaded("Sequencer"))
+	{
+		ISequencerModule& SequencerModule = FModuleManager::Get().LoadModuleChecked<ISequencerModule>(TEXT("Sequencer"));
+		SequencerModule.UnregisterOnSequencerCreated(SequencerCreatedHandle);
+	}
+#endif
+}
+
+#if WITH_EDITOR
+void FDisplayClusterLightCardExtenderModule::OnSequencerCreated(TSharedRef<ISequencer> InSequencer)
+{
+	TWeakPtr<ISequencer> SequencerWeakPtr = InSequencer;
+	OpenSequencers.Add(SequencerWeakPtr);
+	
+	InSequencer->OnGlobalTimeChanged().AddRaw(this, &FDisplayClusterLightCardExtenderModule::OnSequencerTimeChanged, SequencerWeakPtr);
+	InSequencer->OnCloseEvent().AddRaw(this, &FDisplayClusterLightCardExtenderModule::OnSequencerClosed);
+}
+
+void FDisplayClusterLightCardExtenderModule::OnSequencerTimeChanged(TWeakPtr<ISequencer> InSequencer)
+{
+	OnSequencerTimeChangedDelegate.Broadcast(InSequencer);
+}
+
+void FDisplayClusterLightCardExtenderModule::OnSequencerClosed(TSharedRef<ISequencer> InSequencer)
+{
+	InSequencer->OnGlobalTimeChanged().RemoveAll(this);
+	InSequencer->OnCloseEvent().RemoveAll(this);
+	OpenSequencers.Remove(InSequencer);
+
+	// Positions could have reset to their original values.
+	OnSequencerTimeChangedDelegate.Broadcast(InSequencer);
+}
+#endif
+
+IMPLEMENT_MODULE(FDisplayClusterLightCardExtenderModule, DisplayClusterLightCardExtender)
diff --git a/Engine/Plugins/Runtime/nDisplayModularFeatures/Source/DisplayClusterLightCardExtender/Private/DisplayClusterLightCardExtenderModule.h b/Engine/Plugins/Runtime/nDisplayModularFeatures/Source/DisplayClusterLightCardExtender/Private/DisplayClusterLightCardExtenderModule.h
new file mode 100644
--- /dev/null
+++ b/Engine/Plugins/Runtime/nDisplayModularFeatures/Source/DisplayClusterLightCardExtender/Private/DisplayClusterLightCardExtenderModule.h
@@ -0,0 +1,38 @@
+// Copyright Epic Games, Inc. All Rights Reserved.
+
+#pragma once
+
+#include "IDisplayClusterLightCardExtenderModule.h"
+
+class ISequencer;
+
+class FDisplayClusterLightCardExtenderModule : public IDisplayClusterLightCardExtenderModule
+{
+public:
+	
+	/** IModuleInterface implementation start */
+	virtual void StartupModule() override;
+	virtual void ShutdownModule() override;
+	/** IModuleInterface implementation end */
+
+#if WITH_EDITOR
+	/** Broadcast when any sequencer has its time changed */
+	virtual FOnSequencerTimeChanged& GetOnSequencerTimeChanged() override { return OnSequencerTimeChangedDelegate; }
+
+private:
+	/** When a sequencer is opened */
+	void OnSequencerCreated(TSharedRef<ISequencer> InSequencer);
+	
+	/** Called when the global time is changed in sequencer */
+	void OnSequencerTimeChanged(TWeakPtr<ISequencer> InSequencer);
+	
+	/** When a sequencer is closed */
+	void OnSequencerClosed(TSharedRef<ISequencer> InSequencer);
+	
+private:
+	TArray<TWeakPtr<ISequencer>> OpenSequencers;
+	FDelegateHandle SequencerCreatedHandle;
+
+	FOnSequencerTimeChanged OnSequencerTimeChangedDelegate;
+#endif
+};
diff --git a/Engine/Plugins/Runtime/nDisplayModularFeatures/Source/DisplayClusterLightCardExtender/Private/StageActor/IDisplayClusterStageActor.cpp b/Engine/Plugins/Runtime/nDisplayModularFeatures/Source/DisplayClusterLightCardExtender/Private/StageActor/IDisplayClusterStageActor.cpp
--- a/Engine/Plugins/Runtime/nDisplayModularFeatures/Source/DisplayClusterLightCardExtender/Private/StageActor/IDisplayClusterStageActor.cpp
+++ b/Engine/Plugins/Runtime/nDisplayModularFeatures/Source/DisplayClusterLightCardExtender/Private/StageActor/IDisplayClusterStageActor.cpp
@@ -45,6 +45,11 @@ const TSet<FName>& IDisplayClusterStageActor::GetPositionalPropertyNames() const
 	return DefaultPositionalPropertyNames;
 }
 
+FName IDisplayClusterStageActor::GetPositionalPropertiesMemberName() const
+{
+	return NAME_None;
+}
+
 FBox IDisplayClusterStageActor::GetBoxBounds(bool bLocalSpace) const
 {
 	if (const AActor* Actor = FStageActorHelper::GetActorConst(this))
@@ -69,8 +74,25 @@ void IDisplayClusterStageActor::UpdateStageActorTransform()
 	const FTransform Transform = PositionalParamsToActorTransform(GetPositionalParams(), GetOrigin());
 	if (AActor* Actor = FStageActorHelper::GetActor(this))
 	{
+#if WITH_EDITOR
+		const bool bTransformChanged = !Actor->GetActorTransform().Equals(Transform, 0.f);
+#endif
+		
 		Actor->SetActorTransform(Transform);
 
+#if WITH_EDITOR
+		if (bTransformChanged)
+		{
+			if (FProperty* TransformProperty = Actor->GetRootComponent() ?
+				Actor->GetRootComponent()->GetClass()->FindPropertyByName(USceneComponent::GetRelativeLocationPropertyName())
+				: nullptr)
+			{
+				FPropertyChangedEvent PropertyChangedEvent(TransformProperty, EPropertyChangeType::ValueSet);
+				FCoreUObjectDelegates::OnObjectPropertyChanged.Broadcast(Actor, PropertyChangedEvent);
+			}
+		}
+#endif
+
 #if WITH_EDITOR
 		UpdateEditorGizmos();
 #endif
diff --git a/Engine/Plugins/Runtime/nDisplayModularFeatures/Source/DisplayClusterLightCardExtender/Public/IDisplayClusterLightCardExtenderModule.h b/Engine/Plugins/Runtime/nDisplayModularFeatures/Source/DisplayClusterLightCardExtender/Public/IDisplayClusterLightCardExtenderModule.h
new file mode 100644
--- /dev/null
+++ b/Engine/Plugins/Runtime/nDisplayModularFeatures/Source/DisplayClusterLightCardExtender/Public/IDisplayClusterLightCardExtenderModule.h
@@ -0,0 +1,42 @@
+// Copyright Epic Games, Inc. All Rights Reserved.
+
+#pragma once
+
+#include "CoreMinimal.h"
+#include "Modules/ModuleManager.h"
+
+class IDisplayClusterLightCardExtenderModule : public IModuleInterface
+{
+public:
+	static constexpr const TCHAR* ModuleName = TEXT("DisplayClusterLightCardExtender");
+
+public:
+	virtual ~IDisplayClusterLightCardExtenderModule() = default;
+
+	/**
+	* Singleton-like access to this module's interface.  This is just for convenience!
+	* Beware of calling this during the shutdown phase, though.  Your module might have been unloaded already.
+	*
+	* @return Returns singleton instance, loading the module on demand if needed
+	*/
+	static inline IDisplayClusterLightCardExtenderModule& Get()
+	{
+		return FModuleManager::LoadModuleChecked<IDisplayClusterLightCardExtenderModule>(ModuleName);
+	}
+
+	/**
+	* Checks to see if this module is loaded and ready.  It is only valid to call Get() if IsAvailable() returns true.
+	*
+	* @return True if the module is loaded and ready to use
+	*/
+	static inline bool IsAvailable()
+	{
+		return FModuleManager::Get().IsModuleLoaded(ModuleName);
+	}
+
+#if WITH_EDITOR
+	DECLARE_MULTICAST_DELEGATE_OneParam(FOnSequencerTimeChanged, TWeakPtr<class ISequencer>);
+	/** Broadcast when any sequencer has its time changed, including if one is closed */
+	virtual FOnSequencerTimeChanged& GetOnSequencerTimeChanged() = 0;
+#endif
+};
\ No newline at end of file
diff --git a/Engine/Plugins/Runtime/nDisplayModularFeatures/Source/DisplayClusterLightCardExtender/Public/StageActor/IDisplayClusterStageActor.h b/Engine/Plugins/Runtime/nDisplayModularFeatures/Source/DisplayClusterLightCardExtender/Public/StageActor/IDisplayClusterStageActor.h
--- a/Engine/Plugins/Runtime/nDisplayModularFeatures/Source/DisplayClusterLightCardExtender/Public/StageActor/IDisplayClusterStageActor.h
+++ b/Engine/Plugins/Runtime/nDisplayModularFeatures/Source/DisplayClusterLightCardExtender/Public/StageActor/IDisplayClusterStageActor.h
@@ -63,6 +63,9 @@ public:
 	/** Return property names defined for this stage actor for use in property notifies */
 	virtual const TSet<FName>& GetPositionalPropertyNames() const;
 
+	/** Return the property name for an owning struct containing the positional params properties, may be NAME_None if a struct isn't used */
+	virtual FName GetPositionalPropertiesMemberName() const;
+
 	/** Get an array of positional properties defined for this stage actor for use in property notifies. */
 	virtual void GetPositionalProperties(FPositionalPropertyArray& OutPropertyPairs) const = 0;
 
