nDisplay: Fixed issue where the viewport in the ICVFX panel would go black if the stage actor's editor preview was set to frozen

diff --git a/Engine/Plugins/Runtime/nDisplay/Source/DisplayCluster/Private/Game/EngineClasses/Scene/DisplayClusterPreviewComponent.cpp b/Engine/Plugins/Runtime/nDisplay/Source/DisplayCluster/Private/Game/EngineClasses/Scene/DisplayClusterPreviewComponent.cpp
--- a/Engine/Plugins/Runtime/nDisplay/Source/DisplayCluster/Private/Game/EngineClasses/Scene/DisplayClusterPreviewComponent.cpp
+++ b/Engine/Plugins/Runtime/nDisplay/Source/DisplayCluster/Private/Game/EngineClasses/Scene/DisplayClusterPreviewComponent.cpp
@@ -183,9 +183,15 @@ bool UDisplayClusterPreviewComponent::UpdatePreviewMesh()
 
 		// And search for new mesh reference
 		IDisplayClusterViewport* Viewport = GetCurrentViewport();
-		const bool bIsViewportEnabled = (RenderTarget != nullptr || RenderTargetPostProcess != nullptr)
-		&& Viewport != nullptr && Viewport->GetRenderSettings().bEnable && Viewport->GetProjectionPolicy().IsValid();
-		if (bIsViewportEnabled)
+
+		// Determine if the preview component needs to output a render or texture to the stage's screen meshes.
+		// It must have a renderable resource (preview render target or override texture), have a valid viewport configured,
+		// and have that viewport either actively being rendered to by the preview renderer OR have an override texture supplied
+		// externally
+		const bool bHasRenderableResource = RenderTarget != nullptr || RenderTargetPostProcess != nullptr || OverrideTexture != nullptr;
+		const bool bIsViewportValid = Viewport != nullptr && Viewport->GetProjectionPolicy().IsValid();
+		const bool bOutputToPreviewMesh = bHasRenderableResource && bIsViewportValid && (Viewport->GetRenderSettings().bEnable || OverrideTexture);
+		if (bOutputToPreviewMesh)
 		{
 			// Handle preview mesh:
 			if (Viewport->GetProjectionPolicy()->HasPreviewMesh())
