[nDisplay] Fixed crash when DCRA is not available in the level
- Changed default custom frustum size to match icvfx custom size (2560x1440) and avoid having a default size of 0,0
- Removed unneccessary USTRUCT constructors

diff --git a/Engine/Plugins/Runtime/nDisplay/Source/DisplayClusterConfiguration/Private/DisplayClusterConfigurationTypes.cpp b/Engine/Plugins/Runtime/nDisplay/Source/DisplayClusterConfiguration/Private/DisplayClusterConfigurationTypes.cpp
--- a/Engine/Plugins/Runtime/nDisplay/Source/DisplayClusterConfiguration/Private/DisplayClusterConfigurationTypes.cpp
+++ b/Engine/Plugins/Runtime/nDisplay/Source/DisplayClusterConfiguration/Private/DisplayClusterConfigurationTypes.cpp
@@ -162,10 +162,6 @@ FDisplayClusterConfigurationProjection::FDisplayClusterConfigurationProjection()
 	Type = TEXT("simple");
 }
 
-FDisplayClusterConfigurationPostprocess::FDisplayClusterConfigurationPostprocess()
-{
-
-}
 
 const float UDisplayClusterConfigurationViewport::ViewportMinimumSize = 1.0f;
 const float UDisplayClusterConfigurationViewport::ViewportMaximumSize = 15360.0f;
diff --git a/Engine/Plugins/Runtime/nDisplay/Source/DisplayClusterConfiguration/Private/DisplayClusterConfigurationTypes_ICVFX.cpp b/Engine/Plugins/Runtime/nDisplay/Source/DisplayClusterConfiguration/Private/DisplayClusterConfigurationTypes_ICVFX.cpp
--- a/Engine/Plugins/Runtime/nDisplay/Source/DisplayClusterConfiguration/Private/DisplayClusterConfigurationTypes_ICVFX.cpp
+++ b/Engine/Plugins/Runtime/nDisplay/Source/DisplayClusterConfiguration/Private/DisplayClusterConfigurationTypes_ICVFX.cpp
@@ -6,14 +6,10 @@
 // FDisplayClusterConfigurationICVFX_ChromakeyMarkers
 ///////////////////////////////////////////////////////////////////////////////////////
 FDisplayClusterConfigurationICVFX_ChromakeyMarkers::FDisplayClusterConfigurationICVFX_ChromakeyMarkers()
-	: MarkerColor(0, 0.25f, 0)
-	, MarkerTileOffset(0)
 {
 	// Default marker texture
-	{
-		const FString TexturePath = TEXT("/nDisplay/Textures/T_TrackingMarker_A.T_TrackingMarker_A");
-		MarkerTileRGBA = Cast<UTexture2D>(FSoftObjectPath(TexturePath).TryLoad());
-	}
+	const FString TexturePath = TEXT("/nDisplay/Textures/T_TrackingMarker_A.T_TrackingMarker_A");
+	MarkerTileRGBA = Cast<UTexture2D>(FSoftObjectPath(TexturePath).TryLoad());
 }
 
 ///////////////////////////////////////////////////////////////////////////////////////
@@ -25,14 +21,6 @@ FDisplayClusterConfigurationICVFX_CameraRenderSettings::FDisplayClusterConfigura
 	GenerateMips.bAutoGenerateMips = true;
 }
 
-///////////////////////////////////////////////////////////////////////////////////////
-// FDisplayClusterConfigurationICVFX_CameraCustomFrustum
-///////////////////////////////////////////////////////////////////////////////////////
-FDisplayClusterConfigurationICVFX_CameraCustomFrustum::FDisplayClusterConfigurationICVFX_CameraCustomFrustum() :
-	EstimatedOverscanResolution(ForceInitToZero), InnerFrustumResolution(ForceInitToZero), OverscanPixelsIncrease(0.f)
-{
-}
-
 ///////////////////////////////////////////////////////////////////////////////////////
 // FDisplayClusterConfigurationICVFX_LightcardCustomOCIO
 ///////////////////////////////////////////////////////////////////////////////////////
diff --git a/Engine/Plugins/Runtime/nDisplay/Source/DisplayClusterConfiguration/Public/DisplayClusterConfigurationTypes_Base.h b/Engine/Plugins/Runtime/nDisplay/Source/DisplayClusterConfiguration/Public/DisplayClusterConfigurationTypes_Base.h
--- a/Engine/Plugins/Runtime/nDisplay/Source/DisplayClusterConfiguration/Public/DisplayClusterConfigurationTypes_Base.h
+++ b/Engine/Plugins/Runtime/nDisplay/Source/DisplayClusterConfiguration/Public/DisplayClusterConfigurationTypes_Base.h
@@ -184,8 +184,6 @@ struct DISPLAYCLUSTERCONFIGURATION_API FDisplayClusterConfigurationPostprocess
 	GENERATED_BODY()
 
 public:
-	FDisplayClusterConfigurationPostprocess();
-
 	// Control postprocess rendering order. Bigger value rendered last
 	UPROPERTY(EditAnywhere, BlueprintReadOnly, Category = NDisplay)
 	int32 Order = -1;
diff --git a/Engine/Plugins/Runtime/nDisplay/Source/DisplayClusterConfiguration/Public/DisplayClusterConfigurationTypes_ICVFX.h b/Engine/Plugins/Runtime/nDisplay/Source/DisplayClusterConfiguration/Public/DisplayClusterConfigurationTypes_ICVFX.h
--- a/Engine/Plugins/Runtime/nDisplay/Source/DisplayClusterConfiguration/Public/DisplayClusterConfigurationTypes_ICVFX.h
+++ b/Engine/Plugins/Runtime/nDisplay/Source/DisplayClusterConfiguration/Public/DisplayClusterConfigurationTypes_ICVFX.h
@@ -244,7 +244,7 @@ public:
 
 	/** Marker Color */
 	UPROPERTY(EditAnywhere, BlueprintReadWrite, Category = NDisplay, meta = (EditCondition = "bEnable"))
-	FLinearColor MarkerColor;
+	FLinearColor MarkerColor = { 0.f, 0.25f, 0.f, 1.f };
 
 	/** Texture to use as the chromakey marker tile. */
 	UPROPERTY(EditAnywhere, BlueprintReadWrite, Category = NDisplay, meta = (EditCondition = "bEnable"))
@@ -260,7 +260,7 @@ public:
 
 	/** Offset value for the chromakey marker tiles, normalized to the tile distance.  Adjust placement of the chromakey markers within the composition of the camera framing.  Whole numbers will offset chromakey markers by a cyclical amount and have no visual change. */
 	UPROPERTY(EditAnywhere, BlueprintReadWrite, Category = NDisplay, meta = (EditCondition = "bEnable"))
-	FVector2D MarkerTileOffset;
+	FVector2D MarkerTileOffset = { 0, 0 };
 };
 
 USTRUCT(Blueprintable)
@@ -507,16 +507,14 @@ struct DISPLAYCLUSTERCONFIGURATION_API FDisplayClusterConfigurationICVFX_CameraC
 {
 	GENERATED_BODY()
 
-	FDisplayClusterConfigurationICVFX_CameraCustomFrustum();
-	
 	UPROPERTY(VisibleAnywhere, BlueprintReadOnly, Category = NDisplay)
-	FIntPoint EstimatedOverscanResolution;
+	FIntPoint EstimatedOverscanResolution = { 2560, 1440 };
 	
 	UPROPERTY(VisibleAnywhere, BlueprintReadOnly, Category = NDisplay)
-	FIntPoint InnerFrustumResolution;
+	FIntPoint InnerFrustumResolution = { 2560, 1440 };
 
 	UPROPERTY(BlueprintReadOnly, Category = NDisplay, meta = (HideInDetailPanel))
-	float OverscanPixelsIncrease;
+	float OverscanPixelsIncrease = 0.f;
 
 	/** Enable Custom Frustum. */
 	UPROPERTY(EditAnywhere, BlueprintReadWrite, Category = NDisplay, meta = (DisplayName = "Enable Inner Frustum Overscan"))
