nDisplay: Removing WITH_OCIO definition and use in display cluster modules, ensuring that the only valid definition is held by the engine third party library module.

#rnx

diff --git a/Engine/Plugins/Runtime/nDisplay/Source/DisplayCluster/DisplayCluster.Build.cs b/Engine/Plugins/Runtime/nDisplay/Source/DisplayCluster/DisplayCluster.Build.cs
--- a/Engine/Plugins/Runtime/nDisplay/Source/DisplayCluster/DisplayCluster.Build.cs
+++ b/Engine/Plugins/Runtime/nDisplay/Source/DisplayCluster/DisplayCluster.Build.cs
@@ -7,8 +7,6 @@ public class DisplayCluster : ModuleRules
 {
 	public DisplayCluster(ReadOnlyTargetRules ROTargetRules) : base(ROTargetRules)
 	{
-		PublicDefinitions.Add("WITH_OCIO=0");
-
 		PrivateIncludePaths.AddRange(
 			new string[] {
 				System.IO.Path.Combine(GetModuleDirectory("Renderer"), "Private"), //required for FPostProcessMaterialInputs
diff --git a/Engine/Plugins/Runtime/nDisplay/Source/DisplayClusterConfiguration/DisplayClusterConfiguration.Build.cs b/Engine/Plugins/Runtime/nDisplay/Source/DisplayClusterConfiguration/DisplayClusterConfiguration.Build.cs
--- a/Engine/Plugins/Runtime/nDisplay/Source/DisplayClusterConfiguration/DisplayClusterConfiguration.Build.cs
+++ b/Engine/Plugins/Runtime/nDisplay/Source/DisplayClusterConfiguration/DisplayClusterConfiguration.Build.cs
@@ -7,8 +7,6 @@ public class DisplayClusterConfiguration : ModuleRules
 {
 	public DisplayClusterConfiguration(ReadOnlyTargetRules ROTargetRules) : base(ROTargetRules)
 	{
-		PublicDefinitions.Add("WITH_OCIO=0");
-
 		PublicIncludePathModuleNames.AddRange(
 			new string[] {
 				"DisplayCluster",
diff --git a/Engine/Plugins/Runtime/nDisplay/Source/DisplayClusterTests/Private/DisplayCluster.spec.cpp b/Engine/Plugins/Runtime/nDisplay/Source/DisplayClusterTests/Private/DisplayCluster.spec.cpp
--- a/Engine/Plugins/Runtime/nDisplay/Source/DisplayClusterTests/Private/DisplayCluster.spec.cpp
+++ b/Engine/Plugins/Runtime/nDisplay/Source/DisplayClusterTests/Private/DisplayCluster.spec.cpp
@@ -14,10 +14,7 @@
 #include "Kismet2/ComponentEditorUtils.h"
 #include "Kismet2/KismetEditorUtilities.h"
 #include "Misc/AutomationTest.h"
-
-#if WITH_OCIO
 #include "OpenColorIOConfiguration.h"
-#endif
 
 #if WITH_DEV_AUTOMATION_TESTS
 
@@ -122,10 +119,8 @@ BEGIN_DEFINE_SPEC(FDisplayClusterSpec, "DisplayCluster", EAutomationTestFlags::A
 	UDisplayClusterConfigurationViewport* InstancedViewport;
 	UDisplayClusterXformComponent* InstancedComponent;
 
-#if WITH_OCIO
 	UOpenColorIOConfiguration* OCIOConfig;
 	UOpenColorIOConfiguration* AltOCIOConfig;
-#endif
 
 	AActor* AltActor;
 
@@ -1001,7 +996,6 @@ void FDisplayClusterSpec::DoOCIOConfigurationPropagationTests(TFunction<OwnerTyp
 			true); \
 	}
 
-#if WITH_OCIO
 	BeforeEach([this]()
 	{
 		OCIOConfig = NewObject<UOpenColorIOConfiguration>(ClusterAsset->GetPackage(), FName("TestOCIOConfig"), RF_Transient | RF_Public);
@@ -1012,15 +1006,12 @@ void FDisplayClusterSpec::DoOCIOConfigurationPropagationTests(TFunction<OwnerTyp
 		AltOCIOConfig->AddToRoot();
 		FAssetRegistryModule::AssetCreated(AltOCIOConfig);
 	});
-#endif
 	
 	const FString GroupName = FString::Printf(TEXT("Color grading settings (%s)"), *FieldNames.Last().ToString()); 
 	Describe(*GroupName, [this, FieldNames, BeforeTestAction, AssetOwnerAccessor, InstancedOwnerAccessor, OwnerName, bDoInstanceOverrideTest]
 	{
-#if WITH_OCIO
 		OCIO_CONFIG_TEST_INTERNAL(static_cast<UObject*>(OCIOConfig), static_cast<UObject*>(AltOCIOConfig),
 			TArray({ GET_MEMBER_NAME_CHECKED(FOpenColorIODisplayConfiguration, ColorConfiguration), GET_MEMBER_NAME_CHECKED(FOpenColorIOColorConversionSettings, ConfigurationSource) }));
-#endif
 		
 		OCIO_CONFIG_TEST_INTERNAL(FString("TestName"), FString("OverrideName"),
 			TArray({ GET_MEMBER_NAME_CHECKED(FOpenColorIODisplayConfiguration, ColorConfiguration), GET_MEMBER_NAME_CHECKED(FOpenColorIOColorConversionSettings, DestinationColorSpace),
@@ -1031,7 +1022,6 @@ void FDisplayClusterSpec::DoOCIOConfigurationPropagationTests(TFunction<OwnerTyp
 				GET_MEMBER_NAME_CHECKED(FOpenColorIOColorSpace, ColorSpaceName) }));
 	});
 
-#if WITH_OCIO
 	AfterEach([this]()
 	{
 		DisplayClusterTestUtils::CleanUpAsset(OCIOConfig);
@@ -1040,7 +1030,6 @@ void FDisplayClusterSpec::DoOCIOConfigurationPropagationTests(TFunction<OwnerTyp
 		DisplayClusterTestUtils::CleanUpAsset(AltOCIOConfig);
 		AltOCIOConfig = nullptr;
 	});
-#endif
 }
 
 ///////////////////////
