nDisplay: Fix LCs or CKs that were disabled by the DCRA not being properly excluded during render.

diff --git a/Engine/Plugins/Runtime/nDisplay/Source/DisplayCluster/Private/Game/EngineClasses/Scene/DisplayClusterRootActor.cpp b/Engine/Plugins/Runtime/nDisplay/Source/DisplayCluster/Private/Game/EngineClasses/Scene/DisplayClusterRootActor.cpp
--- a/Engine/Plugins/Runtime/nDisplay/Source/DisplayCluster/Private/Game/EngineClasses/Scene/DisplayClusterRootActor.cpp
+++ b/Engine/Plugins/Runtime/nDisplay/Source/DisplayCluster/Private/Game/EngineClasses/Scene/DisplayClusterRootActor.cpp
@@ -750,24 +750,16 @@ void ADisplayClusterRootActor::SetLightCardOwnership()
 		// Iterate over the VisibilityList.Actors array, looking for legacy cards and set the owner so the
 		// card can look the root actor up in certain situations, like adjusting labels when running as -game.
 		{
-			if (CurrentData->StageSettings.Lightcard.bEnable)
+			for (const TSoftObjectPtr<AActor>& Actor : LightCardVisibilityList.Actors)
 			{
-				for (const TSoftObjectPtr<AActor>& Actor : LightCardVisibilityList.Actors)
+				if (ADisplayClusterLightCardActor* LightCardActor = Actor.IsValid() ? Cast<ADisplayClusterLightCardActor>(Actor.Get()) : nullptr)
 				{
-					if (ADisplayClusterLightCardActor* LightCardActor = Actor.IsValid() ? Cast<ADisplayClusterLightCardActor>(Actor.Get()) : nullptr)
-					{
-						LightCardActor->SetWeakRootActorOwner(this);
-					}
+					LightCardActor->SetWeakRootActorOwner(this);
 				}
 			}
 		
 			for (UDisplayClusterICVFXCameraComponent* Camera : ICVFXComponents)
 			{
-				if (!Camera->CameraSettings.Chromakey.bEnable)
-				{
-					continue;
-				}
-				
 				FDisplayClusterConfigurationICVFX_VisibilityList& ChromakeyCards = Camera->CameraSettings.Chromakey.ChromakeyRenderTexture.ShowOnlyList;
 				ChromakeyCards.AutoAddedActors.Reset();
 				
@@ -792,11 +784,6 @@ void ADisplayClusterRootActor::SetLightCardOwnership()
 				{
 					for (UDisplayClusterICVFXCameraComponent* Camera : ICVFXComponents)
 					{
-						if (!Camera->CameraSettings.Chromakey.bEnable)
-						{
-							continue;
-						}
-
 						FDisplayClusterConfigurationICVFX_VisibilityList& ChromakeyCards = Camera->CameraSettings.Chromakey.ChromakeyRenderTexture.ShowOnlyList;
 
 						const UDisplayClusterChromakeyCardStageActorComponent* ChromakeyStageActor = Cast<UDisplayClusterChromakeyCardStageActorComponent>(ChromakeyCardActor->GetStageActorComponent());
@@ -819,7 +806,7 @@ void ADisplayClusterRootActor::SetLightCardOwnership()
 						}
 					}
 				}
-				else if (CurrentData->StageSettings.Lightcard.bEnable)
+				else
 				{
 					if (LightCardActor->GetStageActorComponent() &&
 						LightCardActor->GetStageActorComponent()->GetRootActor() == this)
