nDisplay: Add toolbar button to ICVFX panel and level toolbar detecting when the DCRA has viewports frozen. Clicking the button will unfreeze the viewports.

diff --git a/Engine/Plugins/Runtime/nDisplay/Source/DisplayCluster/Private/Game/EngineClasses/Scene/DisplayClusterRootActor.cpp b/Engine/Plugins/Runtime/nDisplay/Source/DisplayCluster/Private/Game/EngineClasses/Scene/DisplayClusterRootActor.cpp
--- a/Engine/Plugins/Runtime/nDisplay/Source/DisplayCluster/Private/Game/EngineClasses/Scene/DisplayClusterRootActor.cpp
+++ b/Engine/Plugins/Runtime/nDisplay/Source/DisplayCluster/Private/Game/EngineClasses/Scene/DisplayClusterRootActor.cpp
@@ -1392,8 +1392,15 @@ bool ADisplayClusterRootActor::SetFreezeOuterViewports(bool bEnable)
 		UE_LOG(LogDisplayClusterGame, Warning, TEXT("ADisplayClusterRootActor::SetFreezeOuterViewports failed because ConfigData was null"));
 		return false;
 	}
-
-	ConfigData->StageSettings.bFreezeRenderOuterViewports = bEnable;
+	
+	if (ConfigData->StageSettings.bFreezeRenderOuterViewports != bEnable)
+	{
+#if WITH_EDITOR
+		Modify();
+		ConfigData->Modify();
+#endif
+		ConfigData->StageSettings.bFreezeRenderOuterViewports = bEnable;
+	}
 
 	return true;
 }
diff --git a/Engine/Plugins/Runtime/nDisplay/Source/DisplayClusterConfigurator/Private/DisplayClusterConfiguratorModule.cpp b/Engine/Plugins/Runtime/nDisplay/Source/DisplayClusterConfigurator/Private/DisplayClusterConfiguratorModule.cpp
--- a/Engine/Plugins/Runtime/nDisplay/Source/DisplayClusterConfigurator/Private/DisplayClusterConfiguratorModule.cpp
+++ b/Engine/Plugins/Runtime/nDisplay/Source/DisplayClusterConfigurator/Private/DisplayClusterConfiguratorModule.cpp
@@ -43,6 +43,9 @@
 #include "HAL/IConsoleManager.h"
 #include "Modules/ModuleManager.h"
 #include "ActorFactories/ActorFactoryBlueprint.h"
+#include "ScopedTransaction.h"
+#include "SEditorViewportToolBarButton.h"
+#include "Subsystems/PanelExtensionSubsystem.h"
 
 #define REGISTER_PROPERTY_LAYOUT(PropertyType, CustomizationType) { \
 	const FName LayoutName = PropertyType::StaticStruct()->GetFName(); \
@@ -60,6 +63,11 @@
 
 #define LOCTEXT_NAMESPACE "DisplayClusterConfigurator"
 
+namespace UE::DisplayClusterConfiguratorModuleConstants
+{
+	const static FName LevelEditorViewportExtensionIdentifier = TEXT("DisplayClusterLevelViewportExtension");
+}
+
 void FDisplayClusterConfiguratorModule::StartupModule()
 {
 	FCoreDelegates::OnPostEngineInit.AddRaw(this, &FDisplayClusterConfiguratorModule::OnPostEngineInit);
@@ -135,10 +143,7 @@ void FDisplayClusterConfiguratorModule::ShutdownModule()
 	MenuExtensibilityManager.Reset();
 	ToolBarExtensibilityManager.Reset();
 
-	if (GEditor)
-	{
-		FDisplayClusterConfiguratorBlueprintEditor::UnregisterPanelExtensionFactory();
-	}
+	UnregisterPanelExtensions();
 
 	IKismetCompilerInterface& KismetCompilerModule = FModuleManager::GetModuleChecked<IKismetCompilerInterface>("KismetCompiler");
 	KismetCompilerModule.GetCompilers().Remove(&BlueprintCompiler);
@@ -152,7 +157,7 @@ void FDisplayClusterConfiguratorModule::ShutdownModule()
 
 void FDisplayClusterConfiguratorModule::OnPostEngineInit()
 {
-	FDisplayClusterConfiguratorBlueprintEditor::RegisterPanelExtensionFactory();
+	RegisterPanelExtensions();
 }
 
 const FDisplayClusterConfiguratorCommands& FDisplayClusterConfiguratorModule::GetCommands() const
@@ -314,12 +319,93 @@ void FDisplayClusterConfiguratorModule::UnregisterSectionMappings()
 	}
 }
 
+void FDisplayClusterConfiguratorModule::RegisterPanelExtensions()
+{
+	if (GEditor)
+	{
+		FDisplayClusterConfiguratorBlueprintEditor::RegisterPanelExtensionFactory();
+	
+		if (UPanelExtensionSubsystem* PanelExtensionSubsystem = GEditor->GetEditorSubsystem<UPanelExtensionSubsystem>())
+		{
+			if (!PanelExtensionSubsystem->IsPanelFactoryRegistered(UE::DisplayClusterConfiguratorModuleConstants::LevelEditorViewportExtensionIdentifier))
+			{
+				FPanelExtensionFactory LevelViewportToolbarExtensionWidget;
+				LevelViewportToolbarExtensionWidget.CreateExtensionWidget = FPanelExtensionFactory::FCreateExtensionWidget::CreateStatic(&FDisplayClusterConfiguratorModule::OnExtendLevelEditorViewportToolbar);
+				LevelViewportToolbarExtensionWidget.Identifier = UE::DisplayClusterConfiguratorModuleConstants::LevelEditorViewportExtensionIdentifier;
+
+				PanelExtensionSubsystem->RegisterPanelFactory(TEXT("LevelViewportToolBar.LeftExtension"), LevelViewportToolbarExtensionWidget);
+			}
+		}
+	}
+}
+
+void FDisplayClusterConfiguratorModule::UnregisterPanelExtensions()
+{
+	if (GEditor)
+	{
+		FDisplayClusterConfiguratorBlueprintEditor::UnregisterPanelExtensionFactory();
+
+		if (UPanelExtensionSubsystem* PanelExtensionSubsystem = GEditor->GetEditorSubsystem<UPanelExtensionSubsystem>())
+		{
+			PanelExtensionSubsystem->UnregisterPanelFactory(UE::DisplayClusterConfiguratorModuleConstants::LevelEditorViewportExtensionIdentifier);
+		}
+	}
+}
+
 TSharedPtr<FKismetCompilerContext> FDisplayClusterConfiguratorModule::GetCompilerForDisplayClusterBP(UBlueprint* BP,
                                                                                                      FCompilerResultsLog& InMessageLog, const FKismetCompilerOptions& InCompileOptions)
 {
 	return TSharedPtr<FKismetCompilerContext>(new FDisplayClusterConfiguratorKismetCompilerContext(CastChecked<UDisplayClusterBlueprint>(BP), InMessageLog, InCompileOptions));
 }
 
+TSharedRef<SWidget> FDisplayClusterConfiguratorModule::OnExtendLevelEditorViewportToolbar(FWeakObjectPtr ExtensionContext)
+{
+	return SNew(SEditorViewportToolBarButton)	
+	.ButtonType(EUserInterfaceActionType::Button)
+	.ButtonStyle(&FAppStyle::Get().GetWidgetStyle<FButtonStyle>("EditorViewportToolBar.WarningButton"))
+	.OnClicked_Static(OnViewportsFrozenWarningClicked)
+	.Visibility_Static(GetViewportsFrozenWarningVisibility)
+	.ToolTipText(LOCTEXT("DisplayClusterViewportsFrozenOff_ToolTip", "nDisplay viewports are frozen. Click to unfreeze the viewports."))
+	.Content()
+	[
+		SNew(STextBlock)
+		.TextStyle(&FAppStyle::Get().GetWidgetStyle<FTextBlockStyle>("SmallText"))
+		.Text(LOCTEXT("DisplayClusterViewportsFrozen", "nDisplay Viewports Frozen"))
+	];
+}
+
+FReply FDisplayClusterConfiguratorModule::OnViewportsFrozenWarningClicked()
+{
+	if (GEditor)
+	{
+		FScopedTransaction Transaction(LOCTEXT("UnfreezeViewports", "Unfreeze viewports"));
+		
+		for (TActorIterator<ADisplayClusterRootActor> It(GEditor->GetEditorWorldContext().World()); It; ++It)
+		{
+			It->SetFreezeOuterViewports(false);
+		}
+	}
+
+	return FReply::Handled();
+}
+
+EVisibility FDisplayClusterConfiguratorModule::GetViewportsFrozenWarningVisibility()
+{
+	if (GEditor)
+	{
+		for (TActorIterator<ADisplayClusterRootActor> It(GEditor->GetEditorWorldContext().World()); It; ++It)
+		{
+			const UDisplayClusterConfigurationData* ConfigData = It->GetConfigData();
+			if (ConfigData && ConfigData->StageSettings.bFreezeRenderOuterViewports)
+			{
+				return EVisibility::Visible;
+			}
+		}
+	}
+
+	return EVisibility::Collapsed;
+}
+
 IMPLEMENT_MODULE(FDisplayClusterConfiguratorModule, DisplayClusterConfigurator);
 
 #undef LOCTEXT_NAMESPACE
diff --git a/Engine/Plugins/Runtime/nDisplay/Source/DisplayClusterConfigurator/Private/DisplayClusterConfiguratorModule.h b/Engine/Plugins/Runtime/nDisplay/Source/DisplayClusterConfigurator/Private/DisplayClusterConfiguratorModule.h
--- a/Engine/Plugins/Runtime/nDisplay/Source/DisplayClusterConfigurator/Private/DisplayClusterConfiguratorModule.h
+++ b/Engine/Plugins/Runtime/nDisplay/Source/DisplayClusterConfigurator/Private/DisplayClusterConfiguratorModule.h
@@ -42,8 +42,23 @@ private:
 	void RegisterSectionMappings();
 	void UnregisterSectionMappings();
 
+	/** Register any panel extensions used by this module */
+	void RegisterPanelExtensions();
+
+	/** Unregister any active panel extensions used by this module */
+	void UnregisterPanelExtensions();
+
 	static TSharedPtr<FKismetCompilerContext> GetCompilerForDisplayClusterBP(UBlueprint* BP, FCompilerResultsLog& InMessageLog, const FKismetCompilerOptions& InCompileOptions);
 
+	/** Callback when the level editor viewport creates its toolbar */
+	static TSharedRef<SWidget> OnExtendLevelEditorViewportToolbar(FWeakObjectPtr ExtensionContext);
+
+	/** Callback when our viewports frozen button is clicked */
+	static FReply OnViewportsFrozenWarningClicked();
+
+	/** If the viewports frozen warning button should be displayed */
+	static EVisibility GetViewportsFrozenWarningVisibility();
+
 private:
 	TArray<TSharedPtr<IAssetTypeActions>> CreatedAssetTypeActions;
 	TSharedPtr<FExtensibilityManager> MenuExtensibilityManager;
diff --git a/Engine/Plugins/Runtime/nDisplay/Source/DisplayClusterLightCardEditor/Private/Viewport/DisplayClusterLightcardEditorViewport.cpp b/Engine/Plugins/Runtime/nDisplay/Source/DisplayClusterLightCardEditor/Private/Viewport/DisplayClusterLightcardEditorViewport.cpp
--- a/Engine/Plugins/Runtime/nDisplay/Source/DisplayClusterLightCardEditor/Private/Viewport/DisplayClusterLightcardEditorViewport.cpp
+++ b/Engine/Plugins/Runtime/nDisplay/Source/DisplayClusterLightCardEditor/Private/Viewport/DisplayClusterLightcardEditorViewport.cpp
@@ -10,8 +10,13 @@
 #include "LightCardTemplates/DisplayClusterLightCardTemplate.h"
 #include "LightCardTemplates/DisplayClusterLightCardTemplateDragDropOp.h"
 
+#include "DisplayClusterConfigurationTypes.h"
+
+#include "DisplayClusterRootActor.h"
+
 #include "EditorViewportCommands.h"
 #include "ScopedTransaction.h"
+#include "SEditorViewportToolBarButton.h"
 #include "SEditorViewportToolBarMenu.h"
 #include "STransformViewportToolbar.h"
 #include "Framework/Commands/GenericCommands.h"
@@ -84,6 +89,23 @@ public:
 					.OnGetMenuContent(this, &SDisplayClusterLightCardEditorViewportToolBar::GenerateShowMenu)
 				]
 				+ SHorizontalBox::Slot()
+				.AutoWidth()
+				.Padding(2.0f, 2.0f)
+				[
+					SNew(SEditorViewportToolBarButton)	
+					.ButtonType(EUserInterfaceActionType::Button)
+					.ButtonStyle(&FAppStyle::Get().GetWidgetStyle<FButtonStyle>("EditorViewportToolBar.WarningButton"))
+					.OnClicked(this, &SDisplayClusterLightCardEditorViewportToolBar::OnViewportsFrozenWarningClicked)
+					.Visibility(this, &SDisplayClusterLightCardEditorViewportToolBar::GetViewportsFrozenWarningVisibility)
+					.ToolTipText(LOCTEXT("DisplayClusterViewportsFrozenOff_ToolTip", "The nDisplay viewports are frozen. Click to unfreeze the viewports."))
+					.Content()
+					[
+						SNew(STextBlock)
+						.TextStyle(&FAppStyle::Get().GetWidgetStyle<FTextBlockStyle>("SmallText"))
+						.Text(LOCTEXT("DisplayClusterViewportsFrozen", "nDisplay Viewports Frozen"))
+					]
+				]
+				+ SHorizontalBox::Slot()
 				.Padding(3.0f, 1.0f)
 				.HAlign(HAlign_Right)
 				[
@@ -211,6 +233,39 @@ public:
 
 		return MenuBuilder.MakeWidget();
 	}
+
+	FReply OnViewportsFrozenWarningClicked()
+	{
+		if (EditorViewport.IsValid())
+		{
+			const TWeakObjectPtr<ADisplayClusterRootActor> RootActor = EditorViewport.Pin()->GetRootActor();
+			if (RootActor.IsValid())
+			{
+				FScopedTransaction Transaction(LOCTEXT("UnfreezeViewports", "Unfreeze viewports"));
+				RootActor->SetFreezeOuterViewports(false);
+			}
+		}
+
+		return FReply::Handled();
+	}
+	
+	EVisibility GetViewportsFrozenWarningVisibility() const
+	{
+		if (EditorViewport.IsValid())
+		{
+			const TWeakObjectPtr<ADisplayClusterRootActor> RootActor = EditorViewport.Pin()->GetRootActor();
+			if (RootActor.IsValid())
+			{
+				const UDisplayClusterConfigurationData* ConfigData = RootActor->GetConfigData();
+				if (ConfigData && ConfigData->StageSettings.bFreezeRenderOuterViewports)
+				{
+					return EVisibility::Visible;
+				}
+			}
+		}
+
+		return EVisibility::Collapsed;
+	}
 	
 	TSharedRef<SWidget> MakeTransformToolBar()
 	{
@@ -389,6 +444,16 @@ void SDisplayClusterLightCardEditorViewport::SetRootActor(ADisplayClusterRootAct
 	}
 }
 
+TWeakObjectPtr<ADisplayClusterRootActor> SDisplayClusterLightCardEditorViewport::GetRootActor() const
+{
+	if (LightCardEditorPtr.IsValid())
+	{
+		return LightCardEditorPtr.Pin()->GetActiveRootActor();
+	}
+
+	return nullptr;
+}
+
 void SDisplayClusterLightCardEditorViewport::SummonContextMenu()
 {
 	const FVector2D MousePos = FSlateApplication::Get().GetCursorPos();
diff --git a/Engine/Plugins/Runtime/nDisplay/Source/DisplayClusterLightCardEditor/Private/Viewport/DisplayClusterLightcardEditorViewport.h b/Engine/Plugins/Runtime/nDisplay/Source/DisplayClusterLightCardEditor/Private/Viewport/DisplayClusterLightcardEditorViewport.h
--- a/Engine/Plugins/Runtime/nDisplay/Source/DisplayClusterLightCardEditor/Private/Viewport/DisplayClusterLightcardEditorViewport.h
+++ b/Engine/Plugins/Runtime/nDisplay/Source/DisplayClusterLightCardEditor/Private/Viewport/DisplayClusterLightcardEditorViewport.h
@@ -47,6 +47,9 @@ public:
 	TWeakPtr<FDisplayClusterLightCardEditor> GetLightCardEditor() { return LightCardEditorPtr;}
 
 	void SetRootActor(ADisplayClusterRootActor* NewRootActor);
+
+	/** Return the active root actor for this viewport. */
+	TWeakObjectPtr<ADisplayClusterRootActor> GetRootActor() const;
 	
 	TSharedRef<FDisplayClusterLightCardEditorViewportClient> GetLightCardEditorViewportClient() const { return ViewportClient.ToSharedRef(); }
 
