[nDisplay] Fix DisplayClusterScenePreview failing to repopulate when component is modified

diff --git a/Engine/Plugins/Runtime/nDisplay/Source/DisplayClusterScenePreview/Private/DisplayClusterScenePreviewModule.cpp b/Engine/Plugins/Runtime/nDisplay/Source/DisplayClusterScenePreview/Private/DisplayClusterScenePreviewModule.cpp
--- a/Engine/Plugins/Runtime/nDisplay/Source/DisplayClusterScenePreview/Private/DisplayClusterScenePreviewModule.cpp
+++ b/Engine/Plugins/Runtime/nDisplay/Source/DisplayClusterScenePreview/Private/DisplayClusterScenePreviewModule.cpp
@@ -621,9 +621,22 @@ void FDisplayClusterScenePreviewModule::OnActorPropertyChanged(UObject* ObjectBe
 	for (TPair<int32, FRendererConfig>& ConfigPair : RendererConfigs)
 	{
 		FRendererConfig& Config = ConfigPair.Value;
-		if (Config.bAutoUpdateLightcards && Config.RootActor == ObjectBeingModified)
+		if (Config.bAutoUpdateLightcards)
 		{
-			Config.bIsSceneDirty = true;
+			if (Config.RootActor == ObjectBeingModified)
+			{
+				Config.bIsSceneDirty = true;
+				continue;
+			}
+
+			if (UActorComponent* Component = Cast<UActorComponent>(ObjectBeingModified))
+			{
+				if (Component->GetOwner() == Config.RootActor)
+				{
+					Config.bIsSceneDirty = true;
+					continue;
+				}
+			}
 		}
 	}
 }
