Fixed mismatched copy constructor/assignment operator in FDisplayClusterShaderParameters_ICVFX.

diff --git a/Engine/Plugins/Runtime/nDisplay/Source/DisplayClusterShaders/Public/ShaderParameters/DisplayClusterShaderParameters_ICVFX.h b/Engine/Plugins/Runtime/nDisplay/Source/DisplayClusterShaders/Public/ShaderParameters/DisplayClusterShaderParameters_ICVFX.h
--- a/Engine/Plugins/Runtime/nDisplay/Source/DisplayClusterShaders/Public/ShaderParameters/DisplayClusterShaderParameters_ICVFX.h
+++ b/Engine/Plugins/Runtime/nDisplay/Source/DisplayClusterShaders/Public/ShaderParameters/DisplayClusterShaderParameters_ICVFX.h
@@ -97,11 +97,6 @@ public:
 	FDisplayClusterShaderParameters_ICVFX()
 	{ }
 
-	FDisplayClusterShaderParameters_ICVFX(const FDisplayClusterShaderParameters_ICVFX& InParameters)
-	{
-		SetParameters(InParameters);
-	}
-
 public:
 	inline bool IsLightCardOverUsed() const
 	{
