nDisplay: Fix clearing texture not updating material parameter.

diff --git a/Engine/Plugins/Runtime/nDisplay/Source/DisplayCluster/Private/Game/EngineClasses/Scene/DisplayClusterLightCardActor.cpp b/Engine/Plugins/Runtime/nDisplay/Source/DisplayCluster/Private/Game/EngineClasses/Scene/DisplayClusterLightCardActor.cpp
--- a/Engine/Plugins/Runtime/nDisplay/Source/DisplayCluster/Private/Game/EngineClasses/Scene/DisplayClusterLightCardActor.cpp
+++ b/Engine/Plugins/Runtime/nDisplay/Source/DisplayCluster/Private/Game/EngineClasses/Scene/DisplayClusterLightCardActor.cpp
@@ -296,8 +296,39 @@ void ADisplayClusterLightCardActor::UpdateLightCardMaterialInstance()
 		LightCardMaterialInstance->SetScalarParameterValue(TEXT("Exposure"), Exposure);
 		LightCardMaterialInstance->SetScalarParameterValue(TEXT("Gain"), Gain);
 		LightCardMaterialInstance->SetScalarParameterValue(TEXT("Opacity"), IsProxy() ? ProxyOpacity : Opacity);
-		LightCardMaterialInstance->SetTextureParameterValue(TEXT("Texture"), Texture);
-		LightCardMaterialInstance->SetTextureParameterValue(TEXT("AlphaTexture"), PolygonMask);
+
+		// Here we assign the textures. These properties can become null but since SetTextureParameterValue(nullptr) is a no-op,
+		// when the textures are null we instead assign the parent material default texture.
+
+		if (Texture)
+		{
+			LightCardMaterialInstance->SetTextureParameterValue(TEXT("Texture"), Texture);
+		}
+		else
+		{
+			UTexture* DefaultTexture = nullptr;
+
+			if (ensure(LightCardMaterialInstance->Parent))
+			{
+				LightCardMaterialInstance->Parent->GetTextureParameterValue(TEXT("Texture"), DefaultTexture);
+				LightCardMaterialInstance->SetTextureParameterValue(TEXT("Texture"), DefaultTexture);
+			}
+		}
+
+		if (PolygonMask)
+		{
+			LightCardMaterialInstance->SetTextureParameterValue(TEXT("AlphaTexture"), PolygonMask);
+		}
+		else
+		{
+			UTexture* DefaultTexture = nullptr;
+
+			if (ensure(LightCardMaterialInstance->Parent))
+			{
+				LightCardMaterialInstance->Parent->GetTextureParameterValue(TEXT("AlphaTexture"), DefaultTexture);
+				LightCardMaterialInstance->SetTextureParameterValue(TEXT("AlphaTexture"), DefaultTexture);
+			}
+		}
 
 		LightCardMaterialInstance->SetScalarParameterValue(TEXT("AlphaGradient"), AlphaGradient.bEnableAlphaGradient);
 		LightCardMaterialInstance->SetScalarParameterValue(TEXT("AlphaGradientStartingAlpha"), AlphaGradient.StartingAlpha);
