nDisplay: Added global chromakey settings to stage actor that can be used by any ICVFX camera, as well as used for any overlap chromakey when enabled. Also deprecated the custom chromakey enable flag, instead using an enum in the parent chromakey settings that determines what kind of chromakey is used for the ICVFX camera.

diff --git a/Engine/Plugins/Runtime/nDisplay/Source/DisplayCluster/Private/Game/EngineClasses/Scene/DisplayClusterChromakeyCardActor.cpp b/Engine/Plugins/Runtime/nDisplay/Source/DisplayCluster/Private/Game/EngineClasses/Scene/DisplayClusterChromakeyCardActor.cpp
--- a/Engine/Plugins/Runtime/nDisplay/Source/DisplayCluster/Private/Game/EngineClasses/Scene/DisplayClusterChromakeyCardActor.cpp
+++ b/Engine/Plugins/Runtime/nDisplay/Source/DisplayCluster/Private/Game/EngineClasses/Scene/DisplayClusterChromakeyCardActor.cpp
@@ -81,7 +81,18 @@ void ADisplayClusterChromakeyCardActor::UpdateChromakeySettings()
 		{
 			if (IsReferencedByICVFXCamera(Component))
 			{
-				ChromaColor += Component->GetCameraSettingsICVFX().Chromakey.ChromakeyColor;
+				const FDisplayClusterConfigurationICVFX_ChromakeySettings& CameraChromakey = Component->GetCameraSettingsICVFX().Chromakey;
+
+				if (CameraChromakey.ChromakeySettingsSource == EDisplayClusterConfigurationICVFX_ChromakeySettingsSource::Viewport)
+				{
+					const FDisplayClusterConfigurationICVFX_GlobalChromakeySettings& GlobalChromakey = RootActor->GetStageSettings().GlobalChromakey;
+					ChromaColor += GlobalChromakey.ChromakeyColor;
+				}
+				else
+				{
+					ChromaColor += Component->GetCameraSettingsICVFX().Chromakey.ChromakeyColor;
+				}
+
 				Count++;
 			}
 		}
diff --git a/Engine/Plugins/Runtime/nDisplay/Source/DisplayCluster/Private/Game/EngineClasses/Scene/DisplayClusterICVFXCameraComponent.cpp b/Engine/Plugins/Runtime/nDisplay/Source/DisplayCluster/Private/Game/EngineClasses/Scene/DisplayClusterICVFXCameraComponent.cpp
--- a/Engine/Plugins/Runtime/nDisplay/Source/DisplayCluster/Private/Game/EngineClasses/Scene/DisplayClusterICVFXCameraComponent.cpp
+++ b/Engine/Plugins/Runtime/nDisplay/Source/DisplayCluster/Private/Game/EngineClasses/Scene/DisplayClusterICVFXCameraComponent.cpp
@@ -15,8 +15,35 @@
 #include "Misc/DisplayClusterGlobals.h"
 #include "Misc/Parse.h"
 #include "DisplayClusterEnums.h"
+#include "Version/DisplayClusterICVFXCameraCustomVersion.h"
 
 
+void UDisplayClusterICVFXCameraComponent::Serialize(FArchive& Ar)
+{
+	Super::Serialize(Ar);
+
+	Ar.UsingCustomVersion(FDisplayClusterICVFXCameraCustomVersion::GUID);
+}
+
+void UDisplayClusterICVFXCameraComponent::PostLoad()
+{
+	Super::PostLoad();
+
+	PRAGMA_DISABLE_DEPRECATION_WARNINGS
+	const int32 CustomVersion = GetLinkerCustomVersion(FDisplayClusterICVFXCameraCustomVersion::GUID);
+	if (CustomVersion < FDisplayClusterICVFXCameraCustomVersion::UpdateChromakeyConfig)
+	{
+		const bool bCustomChromakey = CameraSettings.Chromakey.ChromakeyRenderTexture.bEnable_DEPRECATED;
+		CameraSettings.Chromakey.ChromakeyType = bCustomChromakey ? 
+			EDisplayClusterConfigurationICVFX_ChromakeyType::CustomChromakey :
+			EDisplayClusterConfigurationICVFX_ChromakeyType::InnerFrustum;
+
+		// New ICVFX cameras default to the global chromakey settings, but for pre 5.3 cameras, the source must be set to the ICVFX camera
+		CameraSettings.Chromakey.ChromakeySettingsSource = EDisplayClusterConfigurationICVFX_ChromakeySettingsSource::ICVFXCamera;
+	}
+	PRAGMA_ENABLE_DEPRECATION_WARNINGS
+}
+
 void UDisplayClusterICVFXCameraComponent::GetDesiredView(FMinimalViewInfo& DesiredView)
 {
 	if (ADisplayClusterRootActor* RootActor = Cast<ADisplayClusterRootActor>(GetOwner()))
diff --git a/Engine/Plugins/Runtime/nDisplay/Source/DisplayCluster/Private/Render/Viewport/Configuration/DisplayClusterViewportConfigurationHelpers_ICVFX.cpp b/Engine/Plugins/Runtime/nDisplay/Source/DisplayCluster/Private/Render/Viewport/Configuration/DisplayClusterViewportConfigurationHelpers_ICVFX.cpp
--- a/Engine/Plugins/Runtime/nDisplay/Source/DisplayCluster/Private/Render/Viewport/Configuration/DisplayClusterViewportConfigurationHelpers_ICVFX.cpp
+++ b/Engine/Plugins/Runtime/nDisplay/Source/DisplayCluster/Private/Render/Viewport/Configuration/DisplayClusterViewportConfigurationHelpers_ICVFX.cpp
@@ -571,7 +571,7 @@ void FDisplayClusterViewportConfigurationHelpers_ICVFX::UpdateCameraViewportSett
 
 	// use framecolor instead of viewport rendering
 	const FDisplayClusterConfigurationICVFX_ChromakeySettings& ChromakeySettings = CameraSettings.Chromakey;
-	if (ChromakeySettings.bEnable && !ChromakeySettings.ChromakeyRenderTexture.bEnable)
+	if (ChromakeySettings.bEnable && ChromakeySettings.ChromakeyType == EDisplayClusterConfigurationICVFX_ChromakeyType::InnerFrustum)
 	{
 		DstViewport.RenderSettings.bSkipRendering = true;
 	}
@@ -656,7 +656,7 @@ void FDisplayClusterViewportConfigurationHelpers_ICVFX::UpdateChromakeyViewportS
 	const FDisplayClusterConfigurationICVFX_CameraSettings& CameraSettings = InCameraComponent.GetCameraSettingsICVFX();
 	const FDisplayClusterConfigurationICVFX_ChromakeySettings& ChromakeySettings = CameraSettings.Chromakey;
 
-	check(ChromakeySettings.bEnable && ChromakeySettings.ChromakeyRenderTexture.bEnable);
+	check(ChromakeySettings.bEnable && ChromakeySettings.ChromakeyType == EDisplayClusterConfigurationICVFX_ChromakeyType::CustomChromakey);
 
 	// Reset runtime flags from prev frame:
 	DstViewport.ResetRuntimeParameters();
@@ -712,13 +712,20 @@ void FDisplayClusterViewportConfigurationHelpers_ICVFX::UpdateChromakeyViewportS
 	}
 }
 
-void FDisplayClusterViewportConfigurationHelpers_ICVFX::UpdateCameraSettings_Chromakey(FDisplayClusterShaderParameters_ICVFX::FCameraSettings& InOutCameraSettings, const FDisplayClusterConfigurationICVFX_ChromakeySettings& InChromakeySettings, FDisplayClusterViewport* InChromakeyViewport)
+void FDisplayClusterViewportConfigurationHelpers_ICVFX::UpdateCameraSettings_Chromakey(FDisplayClusterShaderParameters_ICVFX::FCameraSettings& InOutCameraSettings, const FDisplayClusterConfigurationICVFX_GlobalChromakeySettings& InGlobalChromakeySettings, const FDisplayClusterConfigurationICVFX_ChromakeySettings& InChromakeySettings, FDisplayClusterViewport* InChromakeyViewport)
 {
-	InOutCameraSettings.ChromakeyColor = InChromakeySettings.ChromakeyColor;
+	if (InChromakeySettings.ChromakeySettingsSource == EDisplayClusterConfigurationICVFX_ChromakeySettingsSource::Viewport)
+	{
+		InOutCameraSettings.ChromakeyColor = InGlobalChromakeySettings.ChromakeyColor;
+	}
+	else
+	{
+		InOutCameraSettings.ChromakeyColor = InChromakeySettings.ChromakeyColor;
+	}
 
 	if (InChromakeySettings.bEnable)
 	{
-		if (InChromakeySettings.ChromakeyRenderTexture.bEnable)
+		if (InChromakeySettings.ChromakeyType == EDisplayClusterConfigurationICVFX_ChromakeyType::CustomChromakey)
 		{
 			if (InChromakeySettings.ChromakeyRenderTexture.bReplaceCameraViewport)
 			{
@@ -764,6 +771,26 @@ void FDisplayClusterViewportConfigurationHelpers_ICVFX::UpdateCameraSettings_Chr
 	}
 }
 
+void FDisplayClusterViewportConfigurationHelpers_ICVFX::UpdateCameraSettings_OverlapChromakeyMarkers(FDisplayClusterShaderParameters_ICVFX::FCameraSettings& InOutCameraSettings, const FDisplayClusterConfigurationICVFX_ChromakeyMarkers& InChromakeyMarkers)
+{
+	InOutCameraSettings.OverlapChromakeyMarkerTextureRHI.SafeRelease();
+
+	if (InChromakeyMarkers.bEnable && InChromakeyMarkers.MarkerTileRGBA != nullptr)
+	{
+		InOutCameraSettings.OverlapChromakeyMarkersColor = InChromakeyMarkers.MarkerColor;
+		InOutCameraSettings.OverlapChromakeyMarkersScale = InChromakeyMarkers.MarkerSizeScale;
+		InOutCameraSettings.OverlapChromakeyMarkersDistance = InChromakeyMarkers.MarkerTileDistance;
+		InOutCameraSettings.OverlapChromakeyMarkersOffset = InChromakeyMarkers.MarkerTileOffset;
+
+		// Assign texture RHI ref
+		FTextureResource* MarkersResource = InChromakeyMarkers.MarkerTileRGBA->GetResource();
+		if (MarkersResource)
+		{
+			InOutCameraSettings.OverlapChromakeyMarkerTextureRHI = MarkersResource->TextureRHI;
+		}
+	}
+}
+
 bool FDisplayClusterViewportConfigurationHelpers_ICVFX::IsShouldUseLightcard(const FDisplayClusterConfigurationICVFX_LightcardSettings& InLightcardSettings)
 {
 	if (InLightcardSettings.bEnable == false)
diff --git a/Engine/Plugins/Runtime/nDisplay/Source/DisplayCluster/Private/Render/Viewport/Configuration/DisplayClusterViewportConfigurationHelpers_ICVFX.h b/Engine/Plugins/Runtime/nDisplay/Source/DisplayCluster/Private/Render/Viewport/Configuration/DisplayClusterViewportConfigurationHelpers_ICVFX.h
--- a/Engine/Plugins/Runtime/nDisplay/Source/DisplayCluster/Private/Render/Viewport/Configuration/DisplayClusterViewportConfigurationHelpers_ICVFX.h
+++ b/Engine/Plugins/Runtime/nDisplay/Source/DisplayCluster/Private/Render/Viewport/Configuration/DisplayClusterViewportConfigurationHelpers_ICVFX.h
@@ -14,6 +14,7 @@ class ADisplayClusterRootActor;
 class UDisplayClusterICVFXCameraComponent;
 
 struct FDisplayClusterConfigurationICVFX_ChromakeySettings;
+struct FDisplayClusterConfigurationICVFX_GlobalChromakeySettings;
 struct FDisplayClusterConfigurationICVFX_ChromakeyMarkers;
 struct FDisplayClusterConfigurationICVFX_LightcardSettings;
 struct FDisplayClusterConfigurationICVFX_CameraCustomFrustum;
@@ -79,9 +80,15 @@ public:
 
 	static FDisplayClusterShaderParameters_ICVFX::FCameraSettings GetShaderParametersCameraSettings(const FDisplayClusterViewport& InCameraViewport, ADisplayClusterRootActor& RootActor, UDisplayClusterICVFXCameraComponent& InCameraComponent);
 	
-	static void UpdateCameraSettings_Chromakey(FDisplayClusterShaderParameters_ICVFX::FCameraSettings& InOutCameraSettings, const FDisplayClusterConfigurationICVFX_ChromakeySettings& InChromakeySettings, FDisplayClusterViewport* InChromakeyViewport);
+	/** Configures the camera's chromakey render settings to match the specified configuration settings  */
+	static void UpdateCameraSettings_Chromakey(FDisplayClusterShaderParameters_ICVFX::FCameraSettings& InOutCameraSettings, const FDisplayClusterConfigurationICVFX_GlobalChromakeySettings& InGlobalChromakeySettings, const FDisplayClusterConfigurationICVFX_ChromakeySettings& InChromakeySettings, FDisplayClusterViewport* InChromakeyViewport);
+
+	/** Configures the camera's chromakey marker render settings to match the specified configuration settings */
 	static void UpdateCameraSettings_ChromakeyMarkers(FDisplayClusterShaderParameters_ICVFX::FCameraSettings& InOutCameraSettings, const FDisplayClusterConfigurationICVFX_ChromakeyMarkers& InChromakeyMarkers);
 
+	/** Configures the camera's overlap chromakey marker render settings to match the specified configuration settings */
+	static void UpdateCameraSettings_OverlapChromakeyMarkers(FDisplayClusterShaderParameters_ICVFX::FCameraSettings& InOutCameraSettings, const FDisplayClusterConfigurationICVFX_ChromakeyMarkers& InChromakeyMarkers);
+
 	static void UpdateCameraCustomFrustum(FDisplayClusterViewport& DstViewport, const FDisplayClusterConfigurationICVFX_CameraCustomFrustum& InCameraCustomFrustum);
 	static void UpdateCameraViewportBufferRatio(FDisplayClusterViewport& DstViewport, const FDisplayClusterConfigurationICVFX_CameraSettings& CameraSettings);
 
diff --git a/Engine/Plugins/Runtime/nDisplay/Source/DisplayCluster/Private/Render/Viewport/Configuration/DisplayClusterViewportConfigurationICVFX.cpp b/Engine/Plugins/Runtime/nDisplay/Source/DisplayCluster/Private/Render/Viewport/Configuration/DisplayClusterViewportConfigurationICVFX.cpp
--- a/Engine/Plugins/Runtime/nDisplay/Source/DisplayCluster/Private/Render/Viewport/Configuration/DisplayClusterViewportConfigurationICVFX.cpp
+++ b/Engine/Plugins/Runtime/nDisplay/Source/DisplayCluster/Private/Render/Viewport/Configuration/DisplayClusterViewportConfigurationICVFX.cpp
@@ -98,7 +98,7 @@ public:
 		return false;
 	}
 
-	void AssignChromakeyToTargetViewport(FDisplayClusterViewport& DstViewport, UDisplayClusterICVFXCameraComponent& InCameraComponent)
+	void AssignOverlapChromakeyToTargetViewport(FDisplayClusterViewport& DstViewport, ADisplayClusterRootActor& InRootActor)
 	{
 		if (EnumHasAnyFlags(DstViewport.RenderSettingsICVFX.Flags, EDisplayClusterViewportICVFXFlags::DisableChromakey))
 		{
@@ -106,14 +106,41 @@ public:
 			return;
 		}
 
+		const FDisplayClusterConfigurationICVFX_StageSettings& StageSettings = InRootActor.GetStageSettings();
+		FDisplayClusterShaderParameters_ICVFX::FCameraSettings& DstCameraData = DstViewport.RenderSettingsICVFX.ICVFX.Cameras.Last();
+
+		DstCameraData.OverlapChromakeyColor = StageSettings.GlobalChromakey.ChromakeyColor;
+
+		if (!EnumHasAnyFlags(DstViewport.RenderSettingsICVFX.Flags, EDisplayClusterViewportICVFXFlags::DisableChromakeyMarkers))
+		{
+			FDisplayClusterViewportConfigurationHelpers_ICVFX::UpdateCameraSettings_OverlapChromakeyMarkers(DstCameraData, StageSettings.GlobalChromakey.ChromakeyMarkers);
+		}
+	}
+
+	void AssignChromakeyToTargetViewport(FDisplayClusterViewport& DstViewport, ADisplayClusterRootActor& InRootActor, UDisplayClusterICVFXCameraComponent& InCameraComponent)
+	{
+		if (EnumHasAnyFlags(DstViewport.RenderSettingsICVFX.Flags, EDisplayClusterViewportICVFXFlags::DisableChromakey))
+		{
+			// chromakey disabled for this viewport
+			return;
+		}
+
+		const FDisplayClusterConfigurationICVFX_StageSettings& StageSettings = InRootActor.GetStageSettings();
 		const FDisplayClusterConfigurationICVFX_CameraSettings& CameraSettings = InCameraComponent.GetCameraSettingsICVFX();
 		FDisplayClusterShaderParameters_ICVFX::FCameraSettings& DstCameraData = DstViewport.RenderSettingsICVFX.ICVFX.Cameras.Last();
 
-		FDisplayClusterViewportConfigurationHelpers_ICVFX::UpdateCameraSettings_Chromakey(DstCameraData, CameraSettings.Chromakey, ChromakeyViewport);
+		FDisplayClusterViewportConfigurationHelpers_ICVFX::UpdateCameraSettings_Chromakey(DstCameraData, StageSettings.GlobalChromakey, CameraSettings.Chromakey, ChromakeyViewport);
 
 		if (!EnumHasAnyFlags(DstViewport.RenderSettingsICVFX.Flags, EDisplayClusterViewportICVFXFlags::DisableChromakeyMarkers))
 		{
-			FDisplayClusterViewportConfigurationHelpers_ICVFX::UpdateCameraSettings_ChromakeyMarkers(DstCameraData, CameraSettings.Chromakey.ChromakeyMarkers);
+			if (CameraSettings.Chromakey.ChromakeySettingsSource == EDisplayClusterConfigurationICVFX_ChromakeySettingsSource::Viewport)
+			{
+				FDisplayClusterViewportConfigurationHelpers_ICVFX::UpdateCameraSettings_ChromakeyMarkers(DstCameraData, StageSettings.GlobalChromakey.ChromakeyMarkers);
+			}
+			else
+			{
+				FDisplayClusterViewportConfigurationHelpers_ICVFX::UpdateCameraSettings_ChromakeyMarkers(DstCameraData, CameraSettings.Chromakey.ChromakeyMarkers);
+			}
 		}
 	}
 
@@ -223,8 +250,17 @@ protected:
 		const FDisplayClusterConfigurationICVFX_CameraSettings& CameraSettings = CameraComponent.GetCameraSettingsICVFX();
 		const FDisplayClusterConfigurationICVFX_ChromakeySettings& ChromakeySettings = CameraSettings.Chromakey;
 
+		// Assign the overlap chromakey to all supported targets, always done regardless of custom chromakey settings
+		for (FDisplayClusterViewport* TargetIt : VisibleTargets)
+		{
+			if (TargetIt)
+			{
+				CameraConfiguration.AssignOverlapChromakeyToTargetViewport(*TargetIt, RootActor);
+			}
+		}
+
 		// Try create chromakey render on demand
-		if (ChromakeySettings.bEnable && ChromakeySettings.ChromakeyRenderTexture.bEnable)
+		if (ChromakeySettings.bEnable && ChromakeySettings.ChromakeyType == EDisplayClusterConfigurationICVFX_ChromakeyType::CustomChromakey)
 		{
 			if (ChromakeySettings.ChromakeyRenderTexture.Replace.bAllowReplace && ChromakeySettings.ChromakeyRenderTexture.Replace.SourceTexture == nullptr)
 			{
@@ -250,7 +286,7 @@ protected:
 		{
 			if (TargetIt)
 			{
-				CameraConfiguration.AssignChromakeyToTargetViewport(*TargetIt, CameraComponent);
+				CameraConfiguration.AssignChromakeyToTargetViewport(*TargetIt, RootActor, CameraComponent);
 			}
 		}
 
diff --git a/Engine/Plugins/Runtime/nDisplay/Source/DisplayCluster/Private/Version/DisplayClusterICVFXCameraCustomVersion.cpp b/Engine/Plugins/Runtime/nDisplay/Source/DisplayCluster/Private/Version/DisplayClusterICVFXCameraCustomVersion.cpp
new file mode 100644
--- /dev/null
+++ b/Engine/Plugins/Runtime/nDisplay/Source/DisplayCluster/Private/Version/DisplayClusterICVFXCameraCustomVersion.cpp
@@ -0,0 +1,8 @@
+// Copyright Epic Games, Inc. All Rights Reserved.
+
+#include "DisplayClusterICVFXCameraCustomVersion.h"
+#include "Serialization/CustomVersion.h"
+
+const FGuid FDisplayClusterICVFXCameraCustomVersion::GUID(0x02AA05D1, 0xAF1A4366, 0x92017933, 0x729C6419);
+// Register the custom version with core
+FCustomVersionRegistration GRegisterDisplayClusterConfigurationObjectVersion(FDisplayClusterICVFXCameraCustomVersion::GUID, FDisplayClusterICVFXCameraCustomVersion::LatestVersion, TEXT("DisplayClusterICVFXCameraVer"));
\ No newline at end of file
diff --git a/Engine/Plugins/Runtime/nDisplay/Source/DisplayCluster/Private/Version/DisplayClusterICVFXCameraCustomVersion.h b/Engine/Plugins/Runtime/nDisplay/Source/DisplayCluster/Private/Version/DisplayClusterICVFXCameraCustomVersion.h
new file mode 100644
--- /dev/null
+++ b/Engine/Plugins/Runtime/nDisplay/Source/DisplayCluster/Private/Version/DisplayClusterICVFXCameraCustomVersion.h
@@ -0,0 +1,27 @@
+// Copyright Epic Games, Inc. All Rights Reserved.
+
+#pragma once
+
+#include "CoreMinimal.h"
+#include "Misc/Guid.h"
+
+// Custom serialization version for UDisplayClusterICVFXCameraComponent
+struct FDisplayClusterICVFXCameraCustomVersion
+{
+	enum Type
+	{
+		// Before any version changes were made
+		BeforeCustomVersionWasAdded = 0,
+		UpdateChromakeyConfig = 1,
+
+		// -----<new versions can be added above this line>-------------------------------------------------
+		VersionPlusOne,
+		LatestVersion = VersionPlusOne - 1
+	};
+
+	// The GUID for this custom version number
+	const static FGuid GUID;
+
+private:
+	FDisplayClusterICVFXCameraCustomVersion() {}
+};
diff --git a/Engine/Plugins/Runtime/nDisplay/Source/DisplayCluster/Public/Components/DisplayClusterICVFXCameraComponent.h b/Engine/Plugins/Runtime/nDisplay/Source/DisplayCluster/Public/Components/DisplayClusterICVFXCameraComponent.h
--- a/Engine/Plugins/Runtime/nDisplay/Source/DisplayCluster/Public/Components/DisplayClusterICVFXCameraComponent.h
+++ b/Engine/Plugins/Runtime/nDisplay/Source/DisplayCluster/Public/Components/DisplayClusterICVFXCameraComponent.h
@@ -31,6 +31,9 @@ public:
 	UDisplayClusterICVFXCameraComponent(const FObjectInitializer& ObjectInitializer)
 	{ }
 
+	virtual void Serialize(FArchive& Ar) override;
+	virtual void PostLoad() override;
+
 public:
 	UPROPERTY(BlueprintReadWrite, EditAnywhere, Category = NDisplay)
 	FDisplayClusterConfigurationICVFX_CameraSettings CameraSettings;
@@ -121,6 +124,12 @@ private:
 	UPROPERTY(EditAnywhere, Transient, Category = Chromakey, meta = (PropertyPath = "CameraSettings.Chromakey.bEnable", EditConditionPath = "CameraSettings.bEnable"))
 	FDisplayClusterEditorPropertyReference ChromaKeyEnabledRef;
 
+	UPROPERTY(EditAnywhere, Transient, Category = Chromakey, meta = (PropertyPath = "CameraSettings.Chromakey.ChromakeyType", EditConditionPath = "CameraSettings.bEnable && CameraSettings.Chromakey.bEnable"))
+	FDisplayClusterEditorPropertyReference ChromakeyTypeRef;
+
+	UPROPERTY(EditAnywhere, Transient, Category = Chromakey, meta = (PropertyPath = "CameraSettings.Chromakey.ChromakeySettingsSource", EditConditionPath = "CameraSettings.bEnable && CameraSettings.Chromakey.bEnable"))
+	FDisplayClusterEditorPropertyReference ChromakeySettingsSourceRef;
+
 	UPROPERTY(EditAnywhere, Transient, Category = Chromakey, meta = (PropertyPath = "CameraSettings.Chromakey.ChromakeyColor", EditConditionPath = "CameraSettings.bEnable && CameraSettings.Chromakey.bEnable"))
 	FDisplayClusterEditorPropertyReference ChromakeyColorRef;
 
diff --git a/Engine/Plugins/Runtime/nDisplay/Source/DisplayCluster/Public/DisplayClusterRootActor.h b/Engine/Plugins/Runtime/nDisplay/Source/DisplayCluster/Public/DisplayClusterRootActor.h
--- a/Engine/Plugins/Runtime/nDisplay/Source/DisplayCluster/Public/DisplayClusterRootActor.h
+++ b/Engine/Plugins/Runtime/nDisplay/Source/DisplayCluster/Public/DisplayClusterRootActor.h
@@ -15,6 +15,7 @@
 #include "SceneInterface.h"
 
 #include "DisplayClusterConfigurationStrings.h"
+#include "DisplayClusterConfigurationTypes.h"
 #include "DisplayClusterConfigurationTypes_Viewport.h"
 #include "DisplayClusterConfigurationTypes_ICVFX.h"
 #include "DisplayClusterConfigurationTypes_OCIO.h"
@@ -228,7 +229,7 @@ private:
 	UPROPERTY(EditAnywhere, Transient, Category = "In Camera VFX", meta = (PropertyPath = "CurrentConfigData.StageSettings.bEnableInnerFrustums"))
 	FDisplayClusterEditorPropertyReference EnableInnerFrustumsRef;
 
-	UPROPERTY(EditAnywhere, Transient, Category = "In Camera VFX", meta = (PropertyPath = "CurrentConfigData.StageSettings.bEnableInnerFrustumChromakeyOverlap "))
+	UPROPERTY(EditAnywhere, Transient, Category = "In Camera VFX", meta = (PropertyPath = "CurrentConfigData.StageSettings.bEnableInnerFrustumChromakeyOverlap", DisplayName = "Enable Chromakey Overlap"))
 	FDisplayClusterEditorPropertyReference ShowInnerFrustumOverlapsRef;
 
 	UPROPERTY(EditAnywhere, Transient, Category = "Color Grading", meta = (PropertyPath = "CurrentConfigData.StageSettings.EnableColorGrading"))
@@ -249,6 +250,12 @@ private:
 	UPROPERTY(EditAnywhere, Transient, Category = OCIO, meta = (PropertyPath = "CurrentConfigData.StageSettings.ViewportOCIO.PerViewportOCIOProfiles", DisplayName = "Per-Viewport OCIO Overrides"))
 	FDisplayClusterEditorPropertyReference PerViewportOCIOProfilesRef;
 
+	UPROPERTY(EditAnywhere, Transient, Category = Chromakey, meta = (PropertyPath = "CurrentConfigData.StageSettings.GlobalChromakey.ChromakeyColor"))
+	FDisplayClusterEditorPropertyReference GlobalChromakeyColorRef;
+
+	UPROPERTY(EditAnywhere, Transient, Category = Chromakey, meta = (PropertyPath = "CurrentConfigData.StageSettings.GlobalChromakey.ChromakeyMarkers"))
+	FDisplayClusterEditorPropertyReference GlobalChromakeyMarkersRef;
+
 	UPROPERTY(EditAnywhere, Transient, Category = "Light Cards", meta = (PropertyPath = "CurrentConfigData.StageSettings.Lightcard.bEnable"))
 	FDisplayClusterEditorPropertyReference EnableLightcardsRef;
 
diff --git a/Engine/Plugins/Runtime/nDisplay/Source/DisplayClusterConfiguration/Public/DisplayClusterConfigurationTypes_Enums.h b/Engine/Plugins/Runtime/nDisplay/Source/DisplayClusterConfiguration/Public/DisplayClusterConfigurationTypes_Enums.h
--- a/Engine/Plugins/Runtime/nDisplay/Source/DisplayClusterConfiguration/Public/DisplayClusterConfigurationTypes_Enums.h
+++ b/Engine/Plugins/Runtime/nDisplay/Source/DisplayClusterConfiguration/Public/DisplayClusterConfigurationTypes_Enums.h
@@ -49,6 +49,28 @@ enum class EDisplayClusterConfiguration_PostRenderBlur : uint8
 	Dilate    UMETA(DisplayName = "Dilate"),
 };
 
+/** Indicates the type of chromakey that is used within an ICVFX camera's inner frustum */
+UENUM()
+enum class EDisplayClusterConfigurationICVFX_ChromakeyType : uint8
+{
+	/** The entire inner frustum is rendered as chromakey */
+	InnerFrustum       UMETA(DisplayName = "Inner Frustum"),
+
+	/** Only actors specified in the custom chromakey actor list are rendered as chromakey */
+	CustomChromakey    UMETA(DisplayName = "Custom Chromakey"),
+};
+
+/** Indicates the source of the chromakey settings when an ICVFX camera's inner frustum is rendered as chromakey */
+UENUM()
+enum class EDisplayClusterConfigurationICVFX_ChromakeySettingsSource : uint8
+{
+	/** The nDisplay stage actor's chromakey settings are used */
+	Viewport       UMETA(DisplayName = "Use nDisplay Viewport Chromakey"),
+	
+	/** The ICVFX camera's chromakey settings are used */
+	ICVFXCamera    UMETA(DisplayName = "Use ICVFX Camera Chromakey"),
+};
+
 UENUM()
 enum class EDisplayClusterConfigurationICVFX_ChromakeySource : uint8
 {
diff --git a/Engine/Plugins/Runtime/nDisplay/Source/DisplayClusterConfiguration/Public/DisplayClusterConfigurationTypes_ICVFX.h b/Engine/Plugins/Runtime/nDisplay/Source/DisplayClusterConfiguration/Public/DisplayClusterConfigurationTypes_ICVFX.h
--- a/Engine/Plugins/Runtime/nDisplay/Source/DisplayClusterConfiguration/Public/DisplayClusterConfigurationTypes_ICVFX.h
+++ b/Engine/Plugins/Runtime/nDisplay/Source/DisplayClusterConfiguration/Public/DisplayClusterConfigurationTypes_ICVFX.h
@@ -196,36 +196,48 @@ struct DISPLAYCLUSTERCONFIGURATION_API FDisplayClusterConfigurationICVFX_Chromak
 {
 	GENERATED_BODY()
 
+public:
+	// Note: We need to explicitly disable warnings on these constructors/operators for clang to be happy with deprecated variables
+	PRAGMA_DISABLE_DEPRECATION_WARNINGS
+	FDisplayClusterConfigurationICVFX_ChromakeyRenderSettings() = default;
+	~FDisplayClusterConfigurationICVFX_ChromakeyRenderSettings() = default;
+	FDisplayClusterConfigurationICVFX_ChromakeyRenderSettings(const FDisplayClusterConfigurationICVFX_ChromakeyRenderSettings&) = default;
+	FDisplayClusterConfigurationICVFX_ChromakeyRenderSettings(FDisplayClusterConfigurationICVFX_ChromakeyRenderSettings&&) = default;
+	FDisplayClusterConfigurationICVFX_ChromakeyRenderSettings& operator=(const FDisplayClusterConfigurationICVFX_ChromakeyRenderSettings&) = default;
+	FDisplayClusterConfigurationICVFX_ChromakeyRenderSettings& operator=(FDisplayClusterConfigurationICVFX_ChromakeyRenderSettings&&) = default;
+	PRAGMA_ENABLE_DEPRECATION_WARNINGS
+
 public:
 	/** Set to True to use custom chromakey content. */
-	UPROPERTY(EditAnywhere, BlueprintReadWrite, Category = NDisplay, meta = (DisplayName = "Use Custom Chromakey"))
-	bool bEnable = false;
+	UE_DEPRECATED(5.3, "Use the ChromakeyType enum in FDisplayClusterConfigurationICVFX_ChromakeySettings instead")
+	UPROPERTY(meta = (DeprecatedProperty, DeprecationMessage = "Use the Chromakey Type enum in the chromakey settings instead"))
+	bool bEnable_DEPRECATED = false;
 
 	// Replace the texture of the camera viewport from this chromakey RTT
-	UPROPERTY(BlueprintReadWrite,Category = NDisplay, meta = (EditCondition = "bEnable"))
+	UPROPERTY(BlueprintReadWrite,Category = NDisplay)
 	bool bReplaceCameraViewport = false;
 
 	// Performance: Use custom size (low-res) for chromakey RTT frame. Default size same as camera frame
-	UPROPERTY(BlueprintReadWrite, Category = NDisplay, meta = (EditCondition = "bEnable"))
+	UPROPERTY(BlueprintReadWrite, Category = NDisplay)
 	FDisplayClusterConfigurationICVFX_CustomSize CustomSize;
 
 	/** Content specified here will be overridden to use the chromakey color specified and include chromakey markers if enabled. */
-	UPROPERTY(EditAnywhere, BlueprintReadWrite, Category = NDisplay, meta = (EditCondition = "bEnable", DisplayName = "Custom Chromakey Content"))
+	UPROPERTY(EditAnywhere, BlueprintReadWrite, Category = NDisplay, meta = (DisplayName = "Custom Chromakey Content"))
 	FDisplayClusterConfigurationICVFX_VisibilityList ShowOnlyList;
 
 	// Replace viewport render from source texture
-	UPROPERTY(BlueprintReadWrite, Category = NDisplay, meta = (EditCondition = "bEnable"))
+	UPROPERTY(BlueprintReadWrite, Category = NDisplay)
 	FDisplayClusterConfigurationPostRender_Override Replace;
 
 	/** Apply blur to the Custom Chromakey content. */
-	UPROPERTY(EditAnywhere, BlueprintReadWrite, Category = NDisplay, meta = (EditCondition = "bEnable", DisplayName = "Post Process Blur"))
+	UPROPERTY(EditAnywhere, BlueprintReadWrite, Category = NDisplay, meta = (DisplayName = "Post Process Blur"))
 	FDisplayClusterConfigurationPostRender_BlurPostprocess PostprocessBlur;
 
-	UPROPERTY(BlueprintReadWrite, Category = NDisplay, meta = (EditCondition = "bEnable"))
+	UPROPERTY(BlueprintReadWrite, Category = NDisplay)
 	FDisplayClusterConfigurationPostRender_GenerateMips GenerateMips;
 
 	// Advanced render settings
-	UPROPERTY(BlueprintReadWrite, Category = NDisplay, meta = (EditCondition = "bEnable"))
+	UPROPERTY(BlueprintReadWrite, Category = NDisplay)
 	FDisplayClusterConfigurationICVFX_OverlayAdvancedRenderSettings AdvancedRenderSettings;
 };
 
@@ -279,6 +291,14 @@ public:
 	UPROPERTY(EditAnywhere, BlueprintReadWrite, Category = NDisplay, meta = (DisplayName = "Enable Inner Frustum Chromakey"))
 	bool bEnable = false;
 
+	/** The type of chromakey to use */
+	UPROPERTY(EditAnywhere, BlueprintReadWrite, Category = NDisplay)
+	EDisplayClusterConfigurationICVFX_ChromakeyType ChromakeyType = EDisplayClusterConfigurationICVFX_ChromakeyType::InnerFrustum;
+
+	/** The source of the chromakey settings, either the settings on the ICVFX camera or the global settings on the root actor */
+	UPROPERTY(EditAnywhere, BlueprintReadWrite, Category = NDisplay, meta = (DisplayName = "ICVFX Camera Chromakey"))
+	EDisplayClusterConfigurationICVFX_ChromakeySettingsSource ChromakeySettingsSource = EDisplayClusterConfigurationICVFX_ChromakeySettingsSource::Viewport;
+
 	/** Chromakey Color */
 	UPROPERTY(EditAnywhere, BlueprintReadWrite, Category = NDisplay, meta = (EditCondition = "bEnable"))
 	FLinearColor ChromakeyColor;
@@ -292,6 +312,28 @@ public:
 	FDisplayClusterConfigurationICVFX_ChromakeyMarkers ChromakeyMarkers;
 };
 
+/** Chromakey settings that are global for all ICVFX cameras in a root actor */
+USTRUCT(Blueprintable)
+struct DISPLAYCLUSTERCONFIGURATION_API FDisplayClusterConfigurationICVFX_GlobalChromakeySettings
+{
+	GENERATED_BODY()
+
+public:
+	FDisplayClusterConfigurationICVFX_GlobalChromakeySettings()
+		// Default chromakey color is (0,128,0)
+		: ChromakeyColor(0, 0.5f, 0)
+	{ }
+
+public:
+	/** Chromakey Color */
+	UPROPERTY(EditAnywhere, BlueprintReadWrite, Category = NDisplay)
+	FLinearColor ChromakeyColor;
+
+	/** Display Chromakey Markers to facilitate camera tracking in post production. */
+	UPROPERTY(EditAnywhere, BlueprintReadWrite, Category = NDisplay)
+	FDisplayClusterConfigurationICVFX_ChromakeyMarkers ChromakeyMarkers;
+};
+
 USTRUCT(Blueprintable)
 struct DISPLAYCLUSTERCONFIGURATION_API FDisplayClusterConfigurationICVFX_LightcardRenderSettings
 {
@@ -692,4 +734,8 @@ public:
 	/** OpenColorIO configuration for the Outer viewports. */
 	UPROPERTY(EditAnywhere, BlueprintReadWrite, Category = "OCIO", meta = (DisplayName = "Outer Viewports OCIO"))
 	FDisplayClusterConfigurationICVFX_ViewportOCIO ViewportOCIO;
+
+	/** Global chromakey settings that all ICVFX camera components can opt to use */
+	UPROPERTY(EditAnywhere, BlueprintReadWrite, Category = "Chromakey")
+	FDisplayClusterConfigurationICVFX_GlobalChromakeySettings GlobalChromakey;
 };
diff --git a/Engine/Plugins/Runtime/nDisplay/Source/DisplayClusterConfigurator/Private/Views/Details/DisplayClusterRootActorDetailsCustomization.cpp b/Engine/Plugins/Runtime/nDisplay/Source/DisplayClusterConfigurator/Private/Views/Details/DisplayClusterRootActorDetailsCustomization.cpp
--- a/Engine/Plugins/Runtime/nDisplay/Source/DisplayClusterConfigurator/Private/Views/Details/DisplayClusterRootActorDetailsCustomization.cpp
+++ b/Engine/Plugins/Runtime/nDisplay/Source/DisplayClusterConfigurator/Private/Views/Details/DisplayClusterRootActorDetailsCustomization.cpp
@@ -32,6 +32,7 @@ namespace DisplayClusterRootActorDetailsCustomizationUtils
 			DisplayClusterConfigurationStrings::categories::ICVFXCategory,
 			DisplayClusterConfigurationStrings::categories::ColorGradingCategory,
 			DisplayClusterConfigurationStrings::categories::OCIOCategory,
+			DisplayClusterConfigurationStrings::categories::ChromaKeyCategory,
 			DisplayClusterConfigurationStrings::categories::LightcardCategory,
 			DisplayClusterConfigurationStrings::categories::OverrideCategory,
 			TEXT("Rendering"),
diff --git a/Engine/Plugins/Runtime/nDisplay/Source/DisplayClusterShaders/Private/Shaders/DisplayClusterShadersWarpblend_ICVFX.cpp b/Engine/Plugins/Runtime/nDisplay/Source/DisplayClusterShaders/Private/Shaders/DisplayClusterShadersWarpblend_ICVFX.cpp
--- a/Engine/Plugins/Runtime/nDisplay/Source/DisplayClusterShaders/Private/Shaders/DisplayClusterShadersWarpblend_ICVFX.cpp
+++ b/Engine/Plugins/Runtime/nDisplay/Source/DisplayClusterShaders/Private/Shaders/DisplayClusterShadersWarpblend_ICVFX.cpp
@@ -808,7 +808,15 @@ public:
 
 		case EDisplayClusterShaderParametersICVFX_ChromakeySource::FrameColor:
 		{
-			RenderPassData.PSParameters.ChromakeyColor = Camera.ChromakeyColor;
+			if (RenderPass == EIcvfxPassRenderPass::InnerFrustumChromakeyOverlapIterator)
+			{
+				RenderPassData.PSParameters.ChromakeyColor = Camera.OverlapChromakeyColor;
+			}
+			else
+			{
+				RenderPassData.PSParameters.ChromakeyColor = Camera.ChromakeyColor;
+			}
+
 			RenderPassData.PSPermutationVector.Set<IcvfxShaderPermutation::FIcvfxShaderChromakeyFrameColor>(true);
 
 			return true;
@@ -831,19 +839,39 @@ public:
 
 		const FDisplayClusterShaderParameters_ICVFX::FCameraSettings& Camera = ICVFXParameters.Cameras[ActiveCameraIndex];
 
-		if (Camera.IsChromakeyMarkerUsed())
+		if (RenderPass == EIcvfxPassRenderPass::InnerFrustumChromakeyOverlapIterator)
 		{
-			RenderPassData.PSParameters.ChromakeyMarkerColor = Camera.ChromakeyMarkersColor;
+			if (Camera.IsOverlapChromakeyMarkerUsed())
+			{
+				RenderPassData.PSParameters.ChromakeyMarkerColor = Camera.OverlapChromakeyMarkersColor;
 
-			RenderPassData.PSParameters.ChromakeyMarkerTexture = Camera.ChromakeMarkerTextureRHI;
-			RenderPassData.PSParameters.ChromakeyMarkerSampler = TStaticSamplerState<SF_Trilinear, AM_Wrap, AM_Wrap, AM_Wrap>::GetRHI();
+				RenderPassData.PSParameters.ChromakeyMarkerTexture = Camera.OverlapChromakeyMarkerTextureRHI;
+				RenderPassData.PSParameters.ChromakeyMarkerSampler = TStaticSamplerState<SF_Trilinear, AM_Wrap, AM_Wrap, AM_Wrap>::GetRHI();
 
-			RenderPassData.PSParameters.ChromakeyMarkerScale    = Camera.ChromakeyMarkersScale;
-			RenderPassData.PSParameters.ChromakeyMarkerDistance = Camera.ChromakeyMarkersDistance;
-			RenderPassData.PSParameters.ChromakeyMarkerOffset   = FVector2f(Camera.ChromakeyMarkersOffset);
+				RenderPassData.PSParameters.ChromakeyMarkerScale = Camera.OverlapChromakeyMarkersScale;
+				RenderPassData.PSParameters.ChromakeyMarkerDistance = Camera.OverlapChromakeyMarkersDistance;
+				RenderPassData.PSParameters.ChromakeyMarkerOffset = FVector2f(Camera.OverlapChromakeyMarkersOffset);
 
-			RenderPassData.PSPermutationVector.Set<IcvfxShaderPermutation::FIcvfxShaderChromakeyMarker>(true);
-			return true;
+				RenderPassData.PSPermutationVector.Set<IcvfxShaderPermutation::FIcvfxShaderChromakeyMarker>(true);
+				return true;
+			}
+		}
+		else
+		{
+			if (Camera.IsChromakeyMarkerUsed())
+			{
+				RenderPassData.PSParameters.ChromakeyMarkerColor = Camera.ChromakeyMarkersColor;
+
+				RenderPassData.PSParameters.ChromakeyMarkerTexture = Camera.ChromakeMarkerTextureRHI;
+				RenderPassData.PSParameters.ChromakeyMarkerSampler = TStaticSamplerState<SF_Trilinear, AM_Wrap, AM_Wrap, AM_Wrap>::GetRHI();
+
+				RenderPassData.PSParameters.ChromakeyMarkerScale    = Camera.ChromakeyMarkersScale;
+				RenderPassData.PSParameters.ChromakeyMarkerDistance = Camera.ChromakeyMarkersDistance;
+				RenderPassData.PSParameters.ChromakeyMarkerOffset   = FVector2f(Camera.ChromakeyMarkersOffset);
+
+				RenderPassData.PSPermutationVector.Set<IcvfxShaderPermutation::FIcvfxShaderChromakeyMarker>(true);
+				return true;
+			}
 		}
 
 		return false;
diff --git a/Engine/Plugins/Runtime/nDisplay/Source/DisplayClusterShaders/Public/ShaderParameters/DisplayClusterShaderParameters_ICVFX.h b/Engine/Plugins/Runtime/nDisplay/Source/DisplayClusterShaders/Public/ShaderParameters/DisplayClusterShaderParameters_ICVFX.h
--- a/Engine/Plugins/Runtime/nDisplay/Source/DisplayClusterShaders/Public/ShaderParameters/DisplayClusterShaderParameters_ICVFX.h
+++ b/Engine/Plugins/Runtime/nDisplay/Source/DisplayClusterShaders/Public/ShaderParameters/DisplayClusterShaderParameters_ICVFX.h
@@ -189,6 +189,12 @@ public:
 			return ChromakeMarkerTextureRHI.IsValid();
 		}
 
+		/** Gets whether a valid texture is available to render chromakey markers for any overlapping ICVFX frustums  */
+		inline bool IsOverlapChromakeyMarkerUsed() const
+		{
+			return OverlapChromakeyMarkerTextureRHI.IsValid();	
+		}
+
 		inline void SetViewProjection(const FDisplayClusterShaderParametersICVFX_CameraViewProjection& InLocalSpaceViewProjection)
 		{
 			// Support icvfx stereo - update context from camera for each eye
@@ -219,7 +225,7 @@ public:
 		 * The function is moved here from the DisplayClusterMedia module, because it must also be updated when new settings are added.
 		 * 
 		 * @param InCameraSettings - source camera settings
-		 * @param bExcludeResources - disables resource copying
+		 * @param bIncludeResources - enables resource copying
 		 */
 		inline void SetCameraSettings(const FCameraSettings& InCameraSettings, const bool bIncludeResources)
 		{
@@ -228,6 +234,7 @@ public:
 				Resource = InCameraSettings.Resource;
 				Chromakey = InCameraSettings.Chromakey;
 				ChromakeMarkerTextureRHI = InCameraSettings.ChromakeMarkerTextureRHI;
+				OverlapChromakeyMarkerTextureRHI = InCameraSettings.OverlapChromakeyMarkerTextureRHI;
 			}
 
 			SoftEdge = InCameraSettings.SoftEdge;
@@ -242,12 +249,18 @@ public:
 
 			ChromakeySource = InCameraSettings.ChromakeySource;
 			ChromakeyColor = InCameraSettings.ChromakeyColor;
+			OverlapChromakeyColor = InCameraSettings.OverlapChromakeyColor;
 
 			ChromakeyMarkersColor = InCameraSettings.ChromakeyMarkersColor;
 			ChromakeyMarkersScale = InCameraSettings.ChromakeyMarkersScale;
 			ChromakeyMarkersDistance = InCameraSettings.ChromakeyMarkersDistance;
 			ChromakeyMarkersOffset = InCameraSettings.ChromakeyMarkersOffset;
 
+			OverlapChromakeyMarkersColor = InCameraSettings.OverlapChromakeyMarkersColor;
+			OverlapChromakeyMarkersScale = InCameraSettings.OverlapChromakeyMarkersScale;
+			OverlapChromakeyMarkersDistance = InCameraSettings.OverlapChromakeyMarkersDistance;
+			OverlapChromakeyMarkersOffset = InCameraSettings.OverlapChromakeyMarkersOffset;
+
 			RenderOrder = InCameraSettings.RenderOrder;
 		}
 
@@ -271,6 +284,9 @@ public:
 		FDisplayClusterShaderParametersICVFX_ViewportResource Chromakey;
 		FLinearColor ChromakeyColor = FLinearColor::Black;
 
+		/** The color to use when rendering chromakey for any regions of overlapping ICVFX frustums */
+		FLinearColor OverlapChromakeyColor = FLinearColor::Black;
+
 		// Chromakey markers settings:
 		FLinearColor ChromakeyMarkersColor = FLinearColor::Black;
 		float ChromakeyMarkersScale;
@@ -278,6 +294,21 @@ public:
 		FVector2D ChromakeyMarkersOffset;
 		FTextureRHIRef ChromakeMarkerTextureRHI;
 
+		/** The color of the chroamkey markers for any regions of overlapping ICVFX frustums */
+		FLinearColor OverlapChromakeyMarkersColor = FLinearColor::Black;
+
+		/** The scale of the chroamkey markers for any regions of overlapping ICVFX frustums */
+		float OverlapChromakeyMarkersScale;
+
+		/** The distance between the chroamkey markers for any regions of overlapping ICVFX frustums */
+		float OverlapChromakeyMarkersDistance;
+
+		/** The offset of the chroamkey markers for any regions of overlapping ICVFX frustums */
+		FVector2D OverlapChromakeyMarkersOffset;
+
+		/** The texture to use to render the chromakey markers for any regions of overlaping ICVFX frustums */
+		FTextureRHIRef OverlapChromakeyMarkerTextureRHI;
+
 		int32 RenderOrder = -1;
 	};
 
