Fixed mismatched copy constructor/assignment operator in FDisplayClusterConfigurationRectangle.

diff --git a/Engine/Plugins/Runtime/nDisplay/Source/DisplayClusterConfiguration/Public/DisplayClusterConfigurationTypes_Base.h b/Engine/Plugins/Runtime/nDisplay/Source/DisplayClusterConfiguration/Public/DisplayClusterConfigurationTypes_Base.h
--- a/Engine/Plugins/Runtime/nDisplay/Source/DisplayClusterConfiguration/Public/DisplayClusterConfigurationTypes_Base.h
+++ b/Engine/Plugins/Runtime/nDisplay/Source/DisplayClusterConfiguration/Public/DisplayClusterConfigurationTypes_Base.h
@@ -83,8 +83,6 @@ public:
 		: X(_X), Y(_Y), W(_W), H(_H)
 	{ }
 
-	FDisplayClusterConfigurationRectangle(const FDisplayClusterConfigurationRectangle&) = default;
-
 	FIntRect ToRect() const;
 
 public:
