nDisplay: Added a component that can be used to share the nDisplay preview between Unreal instances.

It shares using Shared Memory Media the viewport textures of the parent nDisplay Actor.

It should only be added to DisplayClusterRootActor instances, and only one component per instance.

The way it works is that the sender generates a unique name for each viewport and captures its texture by getting a pointer to it from the corresponding Preview Component.

The receiver will read it using the corresponding media source, and use the Texture Replace functionality in the nDisplay actor viewports to have them used and displayed.

Limited to WITH_EDITOR due to the current limitation of the preview itself.

diff --git a/Engine/Plugins/Runtime/nDisplay/Source/DisplayCluster/DisplayCluster.Build.cs b/Engine/Plugins/Runtime/nDisplay/Source/DisplayCluster/DisplayCluster.Build.cs
--- a/Engine/Plugins/Runtime/nDisplay/Source/DisplayCluster/DisplayCluster.Build.cs
+++ b/Engine/Plugins/Runtime/nDisplay/Source/DisplayCluster/DisplayCluster.Build.cs
@@ -30,6 +30,8 @@ public class DisplayCluster : ModuleRules
 				"InputCore",
 				"Json",
 				"JsonUtilities",
+				"MediaAssets",
+				"MediaIOCore",
 				"Networking",
 				"OpenColorIO",
 				"OpenCV",
@@ -38,6 +40,7 @@ public class DisplayCluster : ModuleRules
 				"Renderer",
 				"RenderCore",
 				"RHI",
+				"SharedMemoryMedia",
 				"Slate",
 				"SlateCore",
 				"Sockets",
diff --git a/Engine/Plugins/Runtime/nDisplay/Source/DisplayCluster/Private/Game/EngineClasses/Scene/DisplayClusterPreviewShareComponent.cpp b/Engine/Plugins/Runtime/nDisplay/Source/DisplayCluster/Private/Game/EngineClasses/Scene/DisplayClusterPreviewShareComponent.cpp
new file mode 100644
--- /dev/null
+++ b/Engine/Plugins/Runtime/nDisplay/Source/DisplayCluster/Private/Game/EngineClasses/Scene/DisplayClusterPreviewShareComponent.cpp
@@ -0,0 +1,694 @@
+// Copyright Epic Games, Inc. All Rights Reserved.
+
+#include "Components/DisplayClusterPreviewShareComponent.h"
+
+#include "Components/DisplayClusterPreviewComponent.h"
+#include "DisplayClusterRootActor.h"
+#include "Engine/TextureRenderTarget2D.h"
+#include "IDisplayCluster.h"
+#include "Game/IDisplayClusterGameManager.h"
+#include "MediaCapture.h"
+#include "Misc/DisplayClusterLog.h"
+#include "SharedMemoryMediaOutput.h"
+#include "SharedMemoryMediaSource.h"
+
+#if WITH_EDITOR
+
+namespace UE::DisplayClusterPreviewShare
+{
+	/** Retrieves the viewport configuration from the DCRA's current config using the given NodeId and ViewportId */
+	static UDisplayClusterConfigurationViewport* GetViewportFromDCRA(ADisplayClusterRootActor* RootActor, const FString& NodeId, const FString& ViewportId)
+	{
+		if (!RootActor)
+		{
+			return nullptr;
+		}
+
+		UDisplayClusterConfigurationData* ConfigData = RootActor->GetConfigData();
+
+		if (!ConfigData)
+		{
+			return nullptr;
+		}
+
+		const TObjectPtr<UDisplayClusterConfigurationClusterNode>* NodePtr = ConfigData->Cluster->Nodes.Find(NodeId);
+
+		if (!NodePtr)
+		{
+			return nullptr;
+		}
+
+		const UDisplayClusterConfigurationClusterNode* Node = *NodePtr;
+
+		if (!Node)
+		{
+			return nullptr;
+		}
+
+		const TObjectPtr<UDisplayClusterConfigurationViewport>* ViewportPtr = Node->Viewports.Find(ViewportId);
+
+		if (!ViewportPtr)
+		{
+			return nullptr;
+		}
+
+		return *ViewportPtr;
+	}
+}
+
+#endif // WITH_EDITOR
+
+UDisplayClusterPreviewShareComponent::UDisplayClusterPreviewShareComponent(const FObjectInitializer& ObjectInitializer)
+	: Super(ObjectInitializer)
+{
+#if WITH_EDITOR
+
+	if (!AllowedToShare())
+	{
+		return;
+	}
+
+	PrimaryComponentTick.bCanEverTick = true;
+	PrimaryComponentTick.bStartWithTickEnabled = false;
+	SetTickEnable(false);
+
+#endif // WITH_EDITOR
+}
+
+
+UDisplayClusterPreviewShareComponent::~UDisplayClusterPreviewShareComponent()
+{
+#if WITH_EDITOR
+
+	CloseAllMedia();
+
+	// Tell root actor that it doesn't need to keep up the preview enabled for us.
+	if (ADisplayClusterRootActor* RootActor = Cast<ADisplayClusterRootActor>(GetOwner()))
+	{
+		RootActor->RemovePreviewEnableOverride(reinterpret_cast<uint8*>(this));
+	}
+
+	// We try to leave the root actor as it originally was.
+	RestoreRootActorOriginalSettings();
+
+#endif // WITH_EDITOR
+}
+
+#if WITH_EDITOR // Bulk wrap with WITH_EDITOR until preview is supported in other modes.
+
+bool UDisplayClusterPreviewShareComponent::AllowedToShare() const
+{
+	// This component should be inactive if it is a CDO.
+	if (HasAnyFlags(RF_ArchetypeObject | RF_ClassDefaultObject))
+	{
+		return false;
+	}
+
+	// This component should be inactive if in a PreviewWorld
+
+	UWorld* World = GetWorld();
+
+	if (!World)
+	{
+		return false;
+	}
+
+	if (World->IsPreviewWorld())
+	{
+		return false;
+	}
+
+	// We don't allow sharing is the parent is the active root actor
+	{
+		IDisplayCluster& Display = IDisplayCluster::Get();
+		IDisplayClusterGameManager* GameMgr = Display.GetGameMgr();
+
+		if (GameMgr && GameMgr->GetRootActor() == GetOwner())
+		{
+			return false;
+		}
+	}
+
+	return true;
+}
+
+#if WITH_EDITOR
+void UDisplayClusterPreviewShareComponent::PostEditChangeProperty(FPropertyChangedEvent& PropertyChangedEvent)
+{
+	Super::PostEditChangeProperty(PropertyChangedEvent);
+
+	if (PropertyChangedEvent.GetPropertyName() == GET_MEMBER_NAME_CHECKED(UDisplayClusterPreviewShareComponent, Mode))
+	{
+		ModeChanged();
+	}
+	else if (PropertyChangedEvent.GetPropertyName() == GET_MEMBER_NAME_CHECKED(UDisplayClusterPreviewShareComponent, UniqueName))
+	{
+		// Remove spaces to reduce chances of the user not realizing that there is a mismatch with its counterpart.
+		UniqueName = UniqueName.TrimStartAndEnd();
+
+		// All the names are now invalid, so we need to close all media.
+		CloseAllMedia();
+	}
+}
+#endif // WITH_EDITOR
+
+void UDisplayClusterPreviewShareComponent::ModeChanged()
+{
+	// Close all media before restarting the sharing
+	CloseAllMedia();
+
+	// Restore original root actor settings (e.g. texture replace settings that we overwrote)
+	RestoreRootActorOriginalSettings();
+
+	ADisplayClusterRootActor* RootActor = Cast<ADisplayClusterRootActor>(GetOwner());
+
+	switch (Mode)
+	{
+	case EDisplayClusterPreviewShareMode::None:
+	{
+		// Tell root actor that it doesn't need to keep up the preview enabled for us.
+		if (RootActor)
+		{
+			RootActor->RemovePreviewEnableOverride(reinterpret_cast<uint8*>(this));
+		}
+
+		SetTickEnable(false);
+
+		break;
+	}
+	case EDisplayClusterPreviewShareMode::Send:
+	{
+		// If we're sending, we need to tell the root actor to generate the preview textures for us.
+
+		if (RootActor)
+		{
+			RootActor->AddPreviewEnableOverride(reinterpret_cast<uint8*>(this));
+		}
+
+		SetTickEnable(true);
+
+		break;
+	}
+	case EDisplayClusterPreviewShareMode::Receive:
+	{
+		if (RootActor)
+		{
+			// Tell root actor that it doesn't need to keep up the preview enabled for us.
+			RootActor->RemovePreviewEnableOverride(reinterpret_cast<uint8*>(this));
+		}
+
+		SetTickEnable(true);
+
+		break;
+	}
+	default:
+	{
+		break;
+	}
+	}
+}
+
+void UDisplayClusterPreviewShareComponent::SetTickEnable(const bool bEnable)
+{
+	bTickInEditor = bEnable;
+	PrimaryComponentTick.SetTickFunctionEnable(bEnable);
+}
+
+void UDisplayClusterPreviewShareComponent::RestoreRootActorOriginalSettings()
+{
+	using namespace UE::DisplayClusterPreviewShare;
+
+	ADisplayClusterRootActor* RootActor = Cast<ADisplayClusterRootActor>(GetOwner());
+	
+	if (!RootActor)
+	{
+		OriginalSourceTextures.Empty();
+		OriginalTextureReplaces.Empty();
+
+		return;
+	}
+
+	for (UActorComponent* ActorComponent : RootActor->K2_GetComponentsByClass(UDisplayClusterPreviewComponent::StaticClass()))
+	{
+		UDisplayClusterPreviewComponent* PreviewComponent = Cast<UDisplayClusterPreviewComponent>(ActorComponent);
+
+		if (!PreviewComponent)
+		{
+			continue;
+		}
+
+		UDisplayClusterConfigurationViewport* Viewport = GetViewportFromDCRA(RootActor, PreviewComponent->GetClusterNodeId(), PreviewComponent->GetViewportId());
+
+		if (!Viewport)
+		{
+			continue;
+		}
+
+		const FString ViewportKey = GenerateViewportKey(PreviewComponent->GetClusterNodeId(), PreviewComponent->GetViewportId());
+
+		if (bool* BoolPtr = OriginalTextureReplaces.Find(ViewportKey))
+		{
+			Viewport->RenderSettings.Replace.bAllowReplace = *BoolPtr;
+		}
+
+		if (TObjectPtr<UTexture>* TexturePtr = OriginalSourceTextures.Find(ViewportKey))
+		{
+			Viewport->RenderSettings.Replace.SourceTexture = *TexturePtr;
+		}
+	}
+
+	OriginalSourceTextures.Empty();
+	OriginalTextureReplaces.Empty();
+}
+
+void UDisplayClusterPreviewShareComponent::CloseAllMedia()
+{
+	// Stop all media
+
+	for (const TPair<FString, TObjectPtr<UMediaCapture>>& Pair : MediaCaptures)
+	{
+		UMediaCapture* MediaCapture = Pair.Value;
+
+		if (!IsValid(MediaCapture))
+		{
+			continue;
+		}
+
+		MediaCapture->StopCapture(false);
+	}
+
+	for (const TPair<FString, TObjectPtr<UMediaPlayer>>& Pair : MediaPlayers)
+	{
+		UMediaPlayer* MediaPlayer = Pair.Value;
+
+		if (!IsValid(MediaPlayer))
+		{
+			continue;
+		}
+
+		MediaPlayer->Close();
+	}
+
+	// Clear media
+
+	MediaCaptures.Empty();
+	MediaOutputs.Empty();
+
+	MediaSources.Empty();
+	MediaPlayers.Empty();
+	MediaTextures.Empty();
+}
+
+FString UDisplayClusterPreviewShareComponent::GenerateViewportKey(const FString& NodeName, const FString& ViewportName) const
+{
+	return FString::Printf(TEXT("%s_%s"),
+		*NodeName,
+		*ViewportName
+	);
+}
+
+FString UDisplayClusterPreviewShareComponent::GenerateMediaUniqueName(const FString& ActorName, const FString& UniqueViewportName) const
+{
+	return FString::Printf(TEXT("%s_%s"),
+		UniqueName.IsEmpty() ? *ActorName : *UniqueName,
+		*UniqueViewportName
+	);
+}
+
+
+void UDisplayClusterPreviewShareComponent::TickSend()
+{
+	ADisplayClusterRootActor* RootActor = Cast<ADisplayClusterRootActor>(GetOwner());
+
+	if (!RootActor)
+	{
+		CloseAllMedia();
+		return;
+	}
+
+	// Make sure we're up to date with the viewports
+
+	const TArray<UActorComponent*> ActorComponents = RootActor->K2_GetComponentsByClass(UDisplayClusterPreviewComponent::StaticClass());
+
+	// We will use this array to close and remove unused media captures.
+	TSet<FString> LeftoverViewportKeys;
+	MediaOutputs.GetKeys(LeftoverViewportKeys);
+
+	for (UActorComponent* ActorComponent : ActorComponents)
+	{
+		UDisplayClusterPreviewComponent* PreviewComponent = Cast<UDisplayClusterPreviewComponent>(ActorComponent);
+
+		if (!PreviewComponent)
+		{
+			continue;
+		}
+
+		// Get the texture to share
+		UTextureRenderTarget2D* PreviewTexture = PreviewComponent->GetRenderTargetTexture();
+
+		if (!PreviewTexture)
+		{
+			continue;
+		}
+
+		// Each viewport gets a unique name
+
+		const FString ViewportKey = GenerateViewportKey(
+			PreviewComponent->GetClusterNodeId(),
+			PreviewComponent->GetViewportId()
+		);
+
+		LeftoverViewportKeys.Remove(ViewportKey);
+
+		// Make sure the media output exists for that ViewportKey
+
+		USharedMemoryMediaOutput* MediaOutput = nullptr;
+
+		if (TObjectPtr<UMediaOutput>* MediaOutputPtr = MediaOutputs.Find(ViewportKey))
+		{
+			MediaOutput = Cast<USharedMemoryMediaOutput>(*MediaOutputPtr);
+		}
+
+		// If the media output for the given unique name does not exist, create it.
+		if (!MediaOutput)
+		{
+			// Instantiate the media output and give it the corresponding unique name of the viewport.
+			MediaOutput = NewObject<USharedMemoryMediaOutput>(GetTransientPackage(), USharedMemoryMediaOutput::StaticClass());
+
+			MediaOutput->UniqueName = GenerateMediaUniqueName(
+				RootActor->GetActorNameOrLabel(),
+				ViewportKey
+			);
+
+			// Create the associated media capture
+			UMediaCapture* MediaCapture = MediaOutput->CreateMediaCapture();
+			check(MediaCapture);
+
+			MediaCapture->SetMediaOutput(MediaOutput);
+
+			// Start the media capture right away
+
+			// Prepare the media capture options.
+			FMediaCaptureOptions MediaCaptureOptions;
+			MediaCaptureOptions.NumberOfFramesToCapture = -1;
+			MediaCaptureOptions.bAutoRestartOnSourceSizeChange = true;
+			MediaCaptureOptions.bSkipFrameWhenRunningExpensiveTasks = false;
+			MediaCaptureOptions.OverrunAction = EMediaCaptureOverrunAction::Skip;
+
+			const bool bCaptureStarted = MediaCapture->CaptureTextureRenderTarget2D(PreviewTexture, MediaCaptureOptions);
+
+			if (bCaptureStarted)
+			{
+				UE_LOG(LogDisplayClusterGame, Log, TEXT("Started media capture for viewport '%s'"), *MediaOutput->UniqueName);
+			}
+			else
+			{
+				UE_LOG(LogDisplayClusterGame, Warning, TEXT("Couldn't start media capture for viewport '%s'"), *MediaOutput->UniqueName);
+			}
+
+			MediaOutputs.Add(ViewportKey, MediaOutput);
+			MediaCaptures.Add(ViewportKey, MediaCapture);
+		}
+
+		// @todo: Update the render target if it has changed.
+	}
+
+	// Stop and remove unused media captures
+	for (const FString& ViewportKey : LeftoverViewportKeys)
+	{
+		MediaOutputs.Remove(ViewportKey);
+		UMediaCapture* MediaCapture = MediaCaptures.FindAndRemoveChecked(ViewportKey);
+		check(MediaCapture);
+		MediaCapture->StopCapture(false /* bAllowPendingFrameToBeProcess */);
+	}
+
+	// Restart stopped Captures (they may have stopped if e.g. the texture resolution changed)
+	{
+		TArray<FString> StalledCaptureViewportKeys;
+
+		for (const TPair<FString, TObjectPtr<UMediaCapture>>& Pair : MediaCaptures)
+		{
+			UMediaCapture* MediaCapture = Pair.Value;
+
+			if (!IsValid(MediaCapture))
+			{
+				StalledCaptureViewportKeys.Add(Pair.Key);
+				continue;
+			}
+
+			if (   MediaCapture->GetState() == EMediaCaptureState::Stopped 
+				|| MediaCapture->GetState() == EMediaCaptureState::Error)
+			{
+				MediaCapture->StopCapture(false /* bAllowPendingFrameToBeProcess */);
+				StalledCaptureViewportKeys.Add(Pair.Key);
+				continue;
+			}
+		}
+
+		for (const FString& ViewportKey : StalledCaptureViewportKeys)
+		{
+			MediaCaptures.Remove(ViewportKey);
+			MediaOutputs.Remove(ViewportKey);
+		}
+	}
+
+}
+
+void UDisplayClusterPreviewShareComponent::TickReceive()
+{
+	using namespace UE::DisplayClusterPreviewShare;
+
+	ADisplayClusterRootActor* RootActor = Cast<ADisplayClusterRootActor>(GetOwner());
+
+	if (!RootActor)
+	{
+		CloseAllMedia();
+		return;
+	}
+
+	// Make sure we're up to date with the viewports
+
+	const TArray<UActorComponent*> ActorComponents = RootActor->K2_GetComponentsByClass(UDisplayClusterPreviewComponent::StaticClass());
+
+	// We will use this array to close and remove unused media sources.
+	TSet<FString> LeftoverViewportKeys;
+	MediaSources.GetKeys(LeftoverViewportKeys);
+
+	for (UActorComponent* ActorComponent : ActorComponents)
+	{
+		UDisplayClusterPreviewComponent* PreviewComponent = Cast<UDisplayClusterPreviewComponent>(ActorComponent);
+
+		if (!PreviewComponent)
+		{
+			continue;
+		}
+
+		// We will ultimately update the viewport so let's make sure it exists.
+
+		UDisplayClusterConfigurationViewport* Viewport = GetViewportFromDCRA(RootActor, PreviewComponent->GetClusterNodeId(), PreviewComponent->GetViewportId());
+
+		if (!Viewport)
+		{
+			continue;
+		}
+
+		// Each viewport gets associated with a uniquely named shared texture.
+
+		const FString ViewportKey = GenerateViewportKey(
+			PreviewComponent->GetClusterNodeId(),
+			PreviewComponent->GetViewportId()
+		);
+
+		LeftoverViewportKeys.Remove(ViewportKey);
+
+		// Make sure the media source exists for that ViewportKey
+
+		USharedMemoryMediaSource* MediaSource = nullptr;
+
+		if (TObjectPtr<UMediaSource>* MediaOutputPtr = MediaSources.Find(ViewportKey))
+		{
+			MediaSource = Cast<USharedMemoryMediaSource>(*MediaOutputPtr);
+		}
+
+		// If the media output for the given unique name does not exist, create it.
+		if (!MediaSource)
+		{
+			// Instantiate the media source and give it the corresponding unique name of the viewport.
+			MediaSource = NewObject<USharedMemoryMediaSource>(GetTransientPackage(), USharedMemoryMediaSource::StaticClass());
+
+			MediaSource->UniqueName = GenerateMediaUniqueName(
+				RootActor->GetActorNameOrLabel(),
+				ViewportKey
+			);
+
+			// Note: Choosing Freerun for now but consider genlock to force the engines to run at the same rate, 
+			// which would make them immune to Windows GPU throttling of out of focus appliations.
+			MediaSource->Mode = ESharedMemoryMediaSourceMode::Freerun;
+
+			// Create the associated media player and texture
+
+			UMediaPlayer* MediaPlayer = NewObject<UMediaPlayer>();
+			if (!MediaPlayer)
+			{
+				UE_LOG(LogDisplayClusterGame, Error, TEXT("Failed to create MediaPlayer"));
+				continue;
+			}
+
+			MediaPlayer->SetLooping(false);
+			MediaPlayer->PlayOnOpen = false;
+
+			UMediaTexture* MediaTexture = NewObject<UMediaTexture>();
+			if (!MediaTexture)
+			{
+				UE_LOG(LogDisplayClusterGame, Error, TEXT("Failed to create MediaTexture"));
+				continue;
+			}
+
+			MediaTexture->NewStyleOutput = true;
+			MediaTexture->SetRenderMode(UMediaTexture::ERenderMode::Default); //@todo convert to just in time
+			MediaTexture->SetMediaPlayer(MediaPlayer);
+			MediaTexture->UpdateResource();
+
+			// Assign replacement texture
+
+			OriginalTextureReplaces.Add(ViewportKey, Viewport->RenderSettings.Replace.bAllowReplace);
+			OriginalSourceTextures.Add(ViewportKey, Viewport->RenderSettings.Replace.SourceTexture);
+
+			// Start the player right away
+			MediaPlayer->PlayOnOpen = true;
+			MediaPlayer->OpenSource(MediaSource);
+
+			// Add the new media objects to our map so that they don't get garbage collected.
+
+			MediaSources.Add(ViewportKey, MediaSource);
+			MediaPlayers.Add(ViewportKey, MediaPlayer);
+			MediaTextures.Add(ViewportKey, MediaTexture);
+		}
+
+		// Refresh the texture replace settings since reconstruction may overwrite it.
+		if (TObjectPtr<UMediaTexture>* MediaTexturePtr = MediaTextures.Find(ViewportKey))
+		{
+			Viewport->RenderSettings.Replace.bAllowReplace = true;
+			Viewport->RenderSettings.Replace.SourceTexture = *MediaTexturePtr;
+		}
+	}
+
+	// Stop and remove unused media sources
+	for (const FString& ViewportKey : LeftoverViewportKeys)
+	{
+		MediaSources.Remove(ViewportKey);
+
+		UMediaPlayer* MediaPlayer = MediaPlayers.FindAndRemoveChecked(ViewportKey);
+		check(MediaPlayer);
+
+		MediaPlayer->Close();
+
+		MediaTextures.Remove(ViewportKey);
+	}
+
+	// If there are invalid or closed players, start them from scratch
+	{
+		TArray<FString> StalledPlayerViewportKeys;
+
+		for (const TPair<FString, TObjectPtr<UMediaPlayer>>& Pair : MediaPlayers)
+		{
+			UMediaPlayer* MediaPlayer = Pair.Value;
+
+			if (!IsValid(MediaPlayer))
+			{
+				StalledPlayerViewportKeys.Add(Pair.Key);
+				continue;
+			}
+
+			if (MediaPlayer->IsClosed())
+			{
+				StalledPlayerViewportKeys.Add(Pair.Key);
+				continue;
+			}
+		}
+
+		for (const FString& ViewportKey : StalledPlayerViewportKeys)
+		{
+			MediaPlayers.Remove(ViewportKey);
+			MediaSources.Remove(ViewportKey);
+			MediaTextures.Remove(ViewportKey);
+		}
+	}
+
+}
+
+void UDisplayClusterPreviewShareComponent::TickComponent(float DeltaTime, ELevelTick TickType, FActorComponentTickFunction* ThisTickFunction)
+{
+	switch (Mode)
+	{
+	case EDisplayClusterPreviewShareMode::None:
+	{
+		break;
+	}
+	case EDisplayClusterPreviewShareMode::Send:
+	{
+		TickSend();
+		break;
+	}
+	case EDisplayClusterPreviewShareMode::Receive:
+	{
+		TickReceive();
+		break;
+	}
+	}
+}
+
+void UDisplayClusterPreviewShareComponent::DestroyComponent(bool bPromoteChildren)
+{
+	// Close all media to avoid keeping resources alive even though the component was destroyed (and possibly kept in undo buffer).
+	// This call also restores the dcra original settings.
+	SetMode(EDisplayClusterPreviewShareMode::None);
+
+	Super::DestroyComponent(bPromoteChildren);
+}
+
+#endif // WITH_EDITOR in bulk
+
+
+void UDisplayClusterPreviewShareComponent::SetMode(EDisplayClusterPreviewShareMode NewMode)
+{
+#if WITH_EDITORONLY_DATA
+
+	// Nothing to do if the mode is unchanged.
+	if (Mode == NewMode)
+	{
+		return;
+	}
+
+	Mode = NewMode;
+
+#if WITH_EDITOR
+
+	// We ignore the desired Mode if it is not allowed.
+	if (!AllowedToShare())
+	{
+		Mode = EDisplayClusterPreviewShareMode::None;
+	}
+
+	ModeChanged();
+
+#endif // WITH_EDITOR
+
+#endif // WITH_EDITORONLY_DATA
+
+}
+
+void UDisplayClusterPreviewShareComponent::SetUniqueName(const FString& NewUniqueName)
+{
+#if WITH_EDITORONLY_DATA
+	UniqueName = NewUniqueName.TrimStartAndEnd();
+#endif // WITH_EDITORONLY_DATA
+
+#if WITH_EDITOR
+
+	// All the names are now invalid, so we need to close all media.
+	CloseAllMedia();
+
+#endif // WITH_EDITOR
+}
+
diff --git a/Engine/Plugins/Runtime/nDisplay/Source/DisplayCluster/Private/Game/EngineClasses/Scene/DisplayClusterRootActor.cpp b/Engine/Plugins/Runtime/nDisplay/Source/DisplayCluster/Private/Game/EngineClasses/Scene/DisplayClusterRootActor.cpp
--- a/Engine/Plugins/Runtime/nDisplay/Source/DisplayCluster/Private/Game/EngineClasses/Scene/DisplayClusterRootActor.cpp
+++ b/Engine/Plugins/Runtime/nDisplay/Source/DisplayCluster/Private/Game/EngineClasses/Scene/DisplayClusterRootActor.cpp
@@ -1362,22 +1362,6 @@ bool ADisplayClusterRootActor::SetReplaceTextureFlagForAllViewports(bool bReplac
 				ViewportItem.Value->RenderSettings.Replace.bAllowReplace = bReplace;
 			}
 		}
-
-		for (const TPair<FString, TObjectPtr<UDisplayClusterConfigurationClusterNode>>& NodeItem : ConfigData->Cluster->Nodes)
-		{
-			if (!NodeItem.Value)
-			{
-				continue;
-			}
-
-			for (const TPair<FString, TObjectPtr<UDisplayClusterConfigurationViewport>>& ViewportItem : NodeItem.Value->Viewports)
-			{
-				if (ViewportItem.Value)
-				{
-					ViewportItem.Value->RenderSettings.Replace.bAllowReplace = bReplace;
-				}
-			}
-		}
 	}
 
 	return true;
diff --git a/Engine/Plugins/Runtime/nDisplay/Source/DisplayCluster/Public/Components/DisplayClusterPreviewShareComponent.h b/Engine/Plugins/Runtime/nDisplay/Source/DisplayCluster/Public/Components/DisplayClusterPreviewShareComponent.h
new file mode 100644
--- /dev/null
+++ b/Engine/Plugins/Runtime/nDisplay/Source/DisplayCluster/Public/Components/DisplayClusterPreviewShareComponent.h
@@ -0,0 +1,162 @@
+// Copyright Epic Games, Inc. All Rights Reserved.
+
+#pragma once
+
+#include "Components/ActorComponent.h"
+#include "Containers/Map.h"
+
+#include "DisplayClusterPreviewShareComponent.generated.h"
+
+
+class UMediaCapture;
+class UMediaOutput;
+class UMediaPlayer;
+class UMediaSource;
+class UMediaTexture;
+class UTexture;
+
+
+/** Available sharing modes */
+UENUM()
+enum class EDisplayClusterPreviewShareMode
+{
+	/** Sharing disabled */
+	None,
+
+	/** Sends the viewport textures for sharing */
+	Send,
+
+	/** Receives textures to replace the viewport textures with */
+	Receive,
+};
+
+
+/**
+ * nDisplay Viewport preview share component
+ * 
+ * It shares using Shared Memory Media the viewport textures of the parent nDisplay Actor.
+ * It should only be added to DisplayClusterRootActor instances, and only one component per instance.
+ * The way it works is that the sender generates a unique name for each viewport and captures its texture
+ * by getting a pointer to it from the corresponding Preview Component.
+ * The receiver will read it using the corresponding media source, and use the Texture Replace functionality
+ * in the nDisplay actor viewports to have them used and displayed.
+ */
+UCLASS(ClassGroup = (DisplayCluster), meta = (BlueprintSpawnableComponent), HideCategories = (Activation, Collision, Cooking))
+class DISPLAYCLUSTER_API UDisplayClusterPreviewShareComponent
+	: public UActorComponent
+{
+	GENERATED_BODY()
+
+public:
+
+	/** Constructor */
+	UDisplayClusterPreviewShareComponent(const FObjectInitializer& ObjectInitializer);
+
+	/** Destructor */
+	~UDisplayClusterPreviewShareComponent();
+
+#if WITH_EDITOR // Bulk wrap with WITH_EDITOR until preview is supported in other modes.
+
+	//~ UActorComponent interface begin
+	virtual void TickComponent(float DeltaTime, ELevelTick TickType, FActorComponentTickFunction* ThisTickFunction) override;
+	virtual void DestroyComponent(bool bPromoteChildren = false) override;
+	//~ UActorComponent interface end
+
+#endif // WITH_EDITOR
+
+	//~ UObject interface begin
+#if WITH_EDITOR
+	virtual void PostEditChangeProperty(FPropertyChangedEvent& PropertyChangedEvent) override;
+#endif // WITH_EDITOR
+	//~ UObject interface end
+
+	/** Sets the sharing mode */
+	UFUNCTION(BlueprintCallable, Category = "Sharing")
+	void SetMode(EDisplayClusterPreviewShareMode NewMode);
+
+	/** Sets the unique name, which should match between sender and receiver of viewport textures */
+	UFUNCTION(BlueprintCallable, Category = "Sharing")
+	void SetUniqueName(const FString& NewUniqueName);
+
+public:
+
+#if WITH_EDITORONLY_DATA
+
+	/** Current sharing mode of this component */
+	UPROPERTY(Transient, EditAnywhere, Setter=SetMode, BlueprintSetter=SetMode, Category=Sharing)
+	EDisplayClusterPreviewShareMode Mode = EDisplayClusterPreviewShareMode::None;
+
+	/** Current unique name of this component, which should match between sender and receiver of viewport textures */
+	UPROPERTY(EditAnywhere, Setter = SetUniqueName, BlueprintSetter = SetUniqueName, Category = Sharing)
+	FString UniqueName;
+
+#endif // WITH_EDITORONLY_DATA
+
+private:
+
+#if WITH_EDITOR
+
+	/** Closes all media related objects (i.e. media captures and media players) */
+	void CloseAllMedia();
+
+	/** True if this component is valid to be actively used for sharing. E.g. CDOs are not considered active. */
+	bool AllowedToShare() const;
+
+	/** Called when the sharing mode was changed, so that it can update its internal state accordingly */
+	void ModeChanged();
+
+	/** Generates a string id for the viewport, used as a key to store data about it */
+	FString GenerateViewportKey(const FString& ActorName, const FString& UniqueViewportName) const;
+
+	/** Generates a Unique Name for the media being shared */
+	FString GenerateMediaUniqueName(const FString& NodeName, const FString& ViewportName) const;
+
+	/** Logic that should run every tick when in Send mode */
+	void TickSend();
+
+	/** Logic that should run every tick when in Receive mode */
+	void TickReceive();
+
+	/** Restores settings in the nDisplay actor that were altered by this component to achieve its intended purpose */
+	void RestoreRootActorOriginalSettings();
+
+	/** Enables/Disables component ticking */
+	void SetTickEnable(const bool bEnable);
+
+#endif // WITH_EDITOR
+
+private:
+
+#if WITH_EDITORONLY_DATA
+
+	/** Media Outputs associated with the given viewport unique names */
+	UPROPERTY(Transient)
+	TMap<FString, TObjectPtr<UMediaOutput>> MediaOutputs;
+
+	/** Media Captures associated with the given viewport unique names */
+	UPROPERTY(Transient)
+	TMap<FString, TObjectPtr<UMediaCapture>> MediaCaptures;
+
+	/** Media Sources associated with the given viewport unique names */
+	UPROPERTY(Transient)
+	TMap<FString, TObjectPtr<UMediaSource>> MediaSources;
+
+	/** Media Players associated with the given viewport unique names */
+	UPROPERTY(Transient)
+	TMap<FString, TObjectPtr<UMediaPlayer>> MediaPlayers;
+
+	/** Media Textures associated with the given viewport unique names */
+	UPROPERTY(Transient)
+	TMap<FString, TObjectPtr<UMediaTexture>> MediaTextures;
+
+	/** Cache of original Texture Replace Source Textures associated with the given viewport unique names. Used when restoring the original state */
+	UPROPERTY(Transient)
+	TMap<FString, TObjectPtr<UTexture>> OriginalSourceTextures;
+
+	/** Cache of original Texture Replace enable boolean associated with the given viewport unique names. Used when restoring the original state */
+	UPROPERTY(Transient)
+	TMap<FString, bool> OriginalTextureReplaces;
+
+#endif // WITH_EDITORONLY_DATA
+
+};
diff --git a/Engine/Plugins/Runtime/nDisplay/Source/SharedMemoryMedia/Private/SharedMemoryMediaPlayer.cpp b/Engine/Plugins/Runtime/nDisplay/Source/SharedMemoryMedia/Private/SharedMemoryMediaPlayer.cpp
--- a/Engine/Plugins/Runtime/nDisplay/Source/SharedMemoryMedia/Private/SharedMemoryMediaPlayer.cpp
+++ b/Engine/Plugins/Runtime/nDisplay/Source/SharedMemoryMedia/Private/SharedMemoryMediaPlayer.cpp
@@ -82,8 +82,8 @@ bool FSharedMemoryMediaPlayer::Open(const FString& Url, const IMediaOptions* Opt
 		{
 			UE_LOG(LogSharedMemoryMedia, Error, TEXT("Unfortunately, SharedMemoryMedia doesn't support the current RHI type '%s'"),
 				*FSharedMemoryMediaPlatformFactory::GetRhiTypeString(RhiInterfaceType));
-
-			return false;
+			Close();
+			return true;
 		}
 	}
 
@@ -98,7 +98,8 @@ bool FSharedMemoryMediaPlayer::Open(const FString& Url, const IMediaOptions* Opt
 		if (!UniqueName.Len())
 		{
 			UE_LOG(LogSharedMemoryMedia, Error, TEXT("SharedMemoryMediaSource must have a UniqueName that is not empty, and should match the MediaOutput's UniqueName"));
-			return false;
+			Close();
+			return true;
 		}
 	}
 
@@ -126,9 +127,13 @@ bool FSharedMemoryMediaPlayer::Open(const FString& Url, const IMediaOptions* Opt
 	{
 		UE_LOG(LogSharedMemoryMedia, Error, TEXT("Only one SharedMemoryMediaPlayer the UniqueName '%s' can play at a time"), *UniqueName);
 		UniqueName.Reset(); // Reset the unique name to avoid unregistering when not open
-		return false;
+		Close();
+		return true;
 	}
 
+	// Set a playing state.
+	PlayerState = EMediaState::Playing;
+
 	return true;
 }
 
@@ -195,6 +200,10 @@ void FSharedMemoryMediaPlayer::Close()
 
 	// Allow this or other players with the given UniqueName to play
 	UnregisterUniqueName(UniqueName);
+
+	// Update our player state
+	PlayerState = EMediaState::Closed;
+
 }
 
 FGuid FSharedMemoryMediaPlayer::GetPlayerPluginGUID() const
@@ -217,6 +226,12 @@ void FSharedMemoryMediaPlayer::TickFetch(FTimespan DeltaTime, FTimespan Timecode
 {
 	Super::TickFetch(DeltaTime, Timecode);
 
+	// Don't populate anything if we're not playing
+	if (GetState() != EMediaState::Playing)
+	{
+		return;
+	}
+
 	// Make sure we have platform data
 	if (!PlatformData.IsValid())
 	{
@@ -411,7 +426,7 @@ void FSharedMemoryMediaPlayer::TickFetch(FTimespan DeltaTime, FTimespan Timecode
 		SharedMetadataPtr->Receiver.KeepAliveShiftRegister = ~0; // all ones
 	}
 
-	// Create a new media sample and populated with any needed data
+	// Create a new media sample and populate it with any needed data
 	{
 		Samples->CurrentSample = MakeShared<FSharedMemoryMediaSample>();
 		Samples->CurrentSample->Player = AsShared();
diff --git a/Engine/Plugins/Runtime/nDisplay/Source/SharedMemoryMedia/Private/SharedMemoryMediaPlayer.h b/Engine/Plugins/Runtime/nDisplay/Source/SharedMemoryMedia/Private/SharedMemoryMediaPlayer.h
--- a/Engine/Plugins/Runtime/nDisplay/Source/SharedMemoryMedia/Private/SharedMemoryMediaPlayer.h
+++ b/Engine/Plugins/Runtime/nDisplay/Source/SharedMemoryMedia/Private/SharedMemoryMediaPlayer.h
@@ -109,8 +109,7 @@ protected:
 
 	virtual EMediaState GetState() const override
 	{
-		// @todo fix this
-		return EMediaState::Playing;
+		return PlayerState;
 	}
 
 	virtual EMediaStatus GetStatus() const override
@@ -336,4 +335,7 @@ protected:
 
 	} ModeState;
 
+	/** Current state of the player */
+	EMediaState PlayerState = EMediaState::Stopped;
+
 };
