[nDisplay] Fixed crash in .dtor of UDisplayClusterPreviewShareComponent

diff --git a/Engine/Plugins/Runtime/nDisplay/Source/DisplayCluster/Private/Game/EngineClasses/Scene/DisplayClusterPreviewShareComponent.cpp b/Engine/Plugins/Runtime/nDisplay/Source/DisplayCluster/Private/Game/EngineClasses/Scene/DisplayClusterPreviewShareComponent.cpp
--- a/Engine/Plugins/Runtime/nDisplay/Source/DisplayCluster/Private/Game/EngineClasses/Scene/DisplayClusterPreviewShareComponent.cpp
+++ b/Engine/Plugins/Runtime/nDisplay/Source/DisplayCluster/Private/Game/EngineClasses/Scene/DisplayClusterPreviewShareComponent.cpp
@@ -75,25 +75,6 @@ UDisplayClusterPreviewShareComponent::UDisplayClusterPreviewShareComponent(const
 #endif // WITH_EDITOR
 }
 
-
-UDisplayClusterPreviewShareComponent::~UDisplayClusterPreviewShareComponent()
-{
-#if WITH_EDITOR
-
-	CloseAllMedia();
-
-	// Tell root actor that it doesn't need to keep up the preview enabled for us.
-	if (ADisplayClusterRootActor* RootActor = Cast<ADisplayClusterRootActor>(GetOwner()))
-	{
-		RootActor->RemovePreviewEnableOverride(reinterpret_cast<uint8*>(this));
-	}
-
-	// We try to leave the root actor as it originally was.
-	RestoreRootActorOriginalSettings();
-
-#endif // WITH_EDITOR
-}
-
 #if WITH_EDITOR // Bulk wrap with WITH_EDITOR until preview is supported in other modes.
 
 bool UDisplayClusterPreviewShareComponent::AllowedToShare() const
@@ -644,6 +625,21 @@ void UDisplayClusterPreviewShareComponent::DestroyComponent(bool bPromoteChildre
 	// This call also restores the dcra original settings.
 	SetMode(EDisplayClusterPreviewShareMode::None);
 
+#if WITH_EDITOR
+
+	CloseAllMedia();
+
+	// Tell root actor that it doesn't need to keep up the preview enabled for us.
+	if (ADisplayClusterRootActor* RootActor = Cast<ADisplayClusterRootActor>(GetOwner()))
+	{
+		RootActor->RemovePreviewEnableOverride(reinterpret_cast<uint8*>(this));
+	}
+
+	// We try to leave the root actor as it originally was.
+	RestoreRootActorOriginalSettings();
+
+#endif // WITH_EDITOR
+
 	Super::DestroyComponent(bPromoteChildren);
 }
 
diff --git a/Engine/Plugins/Runtime/nDisplay/Source/DisplayCluster/Public/Components/DisplayClusterPreviewShareComponent.h b/Engine/Plugins/Runtime/nDisplay/Source/DisplayCluster/Public/Components/DisplayClusterPreviewShareComponent.h
--- a/Engine/Plugins/Runtime/nDisplay/Source/DisplayCluster/Public/Components/DisplayClusterPreviewShareComponent.h
+++ b/Engine/Plugins/Runtime/nDisplay/Source/DisplayCluster/Public/Components/DisplayClusterPreviewShareComponent.h
@@ -52,9 +52,6 @@ public:
 	/** Constructor */
 	UDisplayClusterPreviewShareComponent(const FObjectInitializer& ObjectInitializer);
 
-	/** Destructor */
-	~UDisplayClusterPreviewShareComponent();
-
 #if WITH_EDITOR // Bulk wrap with WITH_EDITOR until preview is supported in other modes.
 
 	//~ UActorComponent interface begin
