nDisplay: Made adjustments to the ICVFX panel color grading drawer to reflect the new global chromakey settings on the root actor.

diff --git a/Engine/Plugins/Runtime/nDisplay/Source/DisplayClusterColorGrading/Private/DataModelGenerators/DisplayClusterColorGradingGenerator_RootActor.cpp b/Engine/Plugins/Runtime/nDisplay/Source/DisplayClusterColorGrading/Private/DataModelGenerators/DisplayClusterColorGradingGenerator_RootActor.cpp
--- a/Engine/Plugins/Runtime/nDisplay/Source/DisplayClusterColorGrading/Private/DataModelGenerators/DisplayClusterColorGradingGenerator_RootActor.cpp
+++ b/Engine/Plugins/Runtime/nDisplay/Source/DisplayClusterColorGrading/Private/DataModelGenerators/DisplayClusterColorGradingGenerator_RootActor.cpp
@@ -176,16 +176,12 @@ public:
 			DetailBuilder.HideCategory(Category);
 		}
 
-		// TransformCommon is a custom category that doesn't get returned by GetCategoryNames that also needs to be hidden
-		DetailBuilder.HideCategory(TEXT("TransformCommon"));
-
 		UClass* ConfigDataClass = UDisplayClusterConfigurationData::StaticClass();
 
 		IDetailCategoryBuilder& ColorGradingCategoryBuilder = DetailBuilder.EditCategory(TEXT("ColorGradingCategory"));
 		ColorGradingCategoryBuilder.AddProperty(DetailBuilder.GetProperty(GET_MEMBER_NAME_CHECKED(UDisplayClusterConfigurationData, StageSettings.EntireClusterColorGrading), ConfigDataClass));
 		ColorGradingCategoryBuilder.AddProperty(DetailBuilder.GetProperty(GET_MEMBER_NAME_CHECKED(UDisplayClusterConfigurationData, StageSettings.PerViewportColorGrading), ConfigDataClass));
 		ColorGradingCategoryBuilder.AddProperty(DetailBuilder.GetProperty(GET_MEMBER_NAME_CHECKED(UDisplayClusterConfigurationData, StageSettings.bEnableInnerFrustums), ConfigDataClass));
-		ColorGradingCategoryBuilder.AddProperty(DetailBuilder.GetProperty(GET_MEMBER_NAME_CHECKED(UDisplayClusterConfigurationData, StageSettings.ViewportOCIO.AllViewportsOCIOConfiguration.bIsEnabled), ConfigDataClass));
 
 		AddColorGradingDetailProperties(DetailBuilder);
 		AddDetailsPanelProperties(DetailBuilder);
@@ -259,18 +255,31 @@ private:
 
 	void AddDetailsPanelProperties(IDetailLayoutBuilder& DetailBuilder)
 	{
+		auto AddProperty = [&DetailBuilder](IDetailCategoryBuilder& Category, FName PropertyName, bool bExpandChildProperties = false)
+		{
+			TSharedRef<IPropertyHandle> PropertyHandle = DetailBuilder.GetProperty(PropertyName, ADisplayClusterRootActor::StaticClass());
+
+			if (bExpandChildProperties)
+			{
+				PropertyHandle->SetInstanceMetaData(TEXT("ShowOnlyInnerProperties"), TEXT("1"));
+			}
+
+			Category.AddProperty(PropertyHandle);
+		};
+
 		IDetailCategoryBuilder& ViewportsCategoryBuilder = DetailBuilder.EditCategory(TEXT("CustomViewportsCategory"), LOCTEXT("CustomViewportsCategoryLabel", "Viewports"));
-		ViewportsCategoryBuilder.AddProperty(TEXT("ViewportScreenPercentageMultiplierRef"), ADisplayClusterRootActor::StaticClass());
-		ViewportsCategoryBuilder.AddProperty(TEXT("FreezeRenderOuterViewportsRef"), ADisplayClusterRootActor::StaticClass());
+		AddProperty(ViewportsCategoryBuilder, TEXT("ViewportScreenPercentageMultiplierRef"));
+		AddProperty(ViewportsCategoryBuilder, TEXT("FreezeRenderOuterViewportsRef"));
 
 		IDetailCategoryBuilder& InnerFrustumCategoryBuilder = DetailBuilder.EditCategory(TEXT("CustomICVFXCategory"), LOCTEXT("CustomICVFXCategoryLabel", "In-Camera VFX"));
-		InnerFrustumCategoryBuilder.AddProperty(GET_MEMBER_NAME_CHECKED(ADisplayClusterRootActor, InnerFrustumPriority), ADisplayClusterRootActor::StaticClass());
+		AddProperty(InnerFrustumCategoryBuilder, TEXT("ShowInnerFrustumOverlapsRef"));
+		AddProperty(InnerFrustumCategoryBuilder, GET_MEMBER_NAME_CHECKED(ADisplayClusterRootActor, InnerFrustumPriority));
 
-		IDetailCategoryBuilder& AllViewportsOCIOCategoryBuilder = DetailBuilder.EditCategory(TEXT("CustomAllViewportsOCIOCategory"), LOCTEXT("CustomAllViewportsOCIOCategoryLabel", "All Viewports"));
-		AllViewportsOCIOCategoryBuilder.AddProperty(TEXT("AllViewportColorConfigurationRef"), ADisplayClusterRootActor::StaticClass());
+		IDetailCategoryBuilder& ViewportChromakeyCategoryBuilder = DetailBuilder.EditCategory(TEXT("CustomViewportChromakeyCategory"), LOCTEXT("CustomViewportChromakeyCategoryLabel", "Chromakey"));
+		AddProperty(ViewportChromakeyCategoryBuilder, TEXT("GlobalChromakeyColorRef"));
 
-		IDetailCategoryBuilder& PerViewportOCIOCategoryBuilder = DetailBuilder.EditCategory(TEXT("CustomPerViewportOCIOCategory"), LOCTEXT("CustomPerViewportOCIOCategoryLabel", "Per-Viewport"));
-		PerViewportOCIOCategoryBuilder.AddProperty(TEXT("PerViewportOCIOProfilesRef"), ADisplayClusterRootActor::StaticClass());
+		IDetailCategoryBuilder& ViewportChromakeyMarkersCategoryBuilder = DetailBuilder.EditCategory(TEXT("CustomViewportChromakeyMarkersCategory"), LOCTEXT("CustomViewportChromakeyMarkersCategoryLabel", "Chromakey Markers"));
+		AddProperty(ViewportChromakeyMarkersCategoryBuilder, TEXT("GlobalChromakeyMarkersRef"), true);
 	}
 
 private:
@@ -511,36 +520,30 @@ void FDisplayClusterColorGradingGenerator_RootActor::GenerateDataModel(IProperty
 	OutColorGradingDataModel.OnColorGradingGroupDeleted().AddSP(this, &FDisplayClusterColorGradingGenerator_RootActor::DeleteColorGradingGroup);
 	OutColorGradingDataModel.OnColorGradingGroupRenamed().AddSP(this, &FDisplayClusterColorGradingGenerator_RootActor::RenameColorGradingGroup);
 
-	FDisplayClusterColorGradingDataModel::FDetailsSection ViewportsDetailsSection;
-	ViewportsDetailsSection.DisplayName = LOCTEXT("ViewportsDetailsSectionLabel", "Viewports");
-	ViewportsDetailsSection.Categories.Add(TEXT("CustomViewportsCategory"));
-
-	OutColorGradingDataModel.DetailsSections.Add(ViewportsDetailsSection);
-
-	FDisplayClusterColorGradingDataModel::FDetailsSection InnerFrustumDetailsSection;
-	InnerFrustumDetailsSection.DisplayName = LOCTEXT("InnerFrustumDetailsSectionLabel", "Inner Frustum");
-	InnerFrustumDetailsSection.Categories.Add(TEXT("CustomICVFXCategory"));
-	InnerFrustumDetailsSection.EditConditionPropertyHandle = FindPropertyHandle(PropertyRowGenerator, CREATE_PROPERTY_PATH(UDisplayClusterConfigurationData, StageSettings.bEnableInnerFrustums));
-
-	OutColorGradingDataModel.DetailsSections.Add(InnerFrustumDetailsSection);
-
-	FDisplayClusterColorGradingDataModel::FDetailsSection OCIODetailsSection;
-	OCIODetailsSection.DisplayName = LOCTEXT("OCIODetailsSectionLabel", "OCIO");
-	OCIODetailsSection.EditConditionPropertyHandle = FindPropertyHandle(PropertyRowGenerator, CREATE_PROPERTY_PATH(UDisplayClusterConfigurationData, StageSettings.ViewportOCIO.AllViewportsOCIOConfiguration.bIsEnabled));
+	{
+		FDisplayClusterColorGradingDataModel::FDetailsSection InnerFrustumDetailsSection;
+		InnerFrustumDetailsSection.DisplayName = LOCTEXT("InnerFrustumDetailsSectionLabel", "Inner Frustum");
+		InnerFrustumDetailsSection.Categories.Add(TEXT("CustomViewportsCategory"));
+		InnerFrustumDetailsSection.Categories.Add(TEXT("CustomICVFXCategory"));
+		InnerFrustumDetailsSection.EditConditionPropertyHandle = FindPropertyHandle(PropertyRowGenerator, CREATE_PROPERTY_PATH(UDisplayClusterConfigurationData, StageSettings.bEnableInnerFrustums));
 
-	FDisplayClusterColorGradingDataModel::FDetailsSubsection AllViewportsOCIODetailsSubsection;
-	AllViewportsOCIODetailsSubsection.DisplayName = LOCTEXT("AllViewportsOCIOSubsectionLabel", "All Viewports");
-	AllViewportsOCIODetailsSubsection.Categories.Add(TEXT("CustomAllViewportsOCIOCategory"));
+		OutColorGradingDataModel.DetailsSections.Add(InnerFrustumDetailsSection);
+	}
 
-	OCIODetailsSection.Subsections.Add(AllViewportsOCIODetailsSubsection);
+	{
+		FDisplayClusterColorGradingDataModel::FDetailsSection TransformDetailsSection;
+		TransformDetailsSection.Categories.Add(TEXT("TransformCommon"));
 
-	FDisplayClusterColorGradingDataModel::FDetailsSubsection PerViewportOCIODetailsSubsection;
-	PerViewportOCIODetailsSubsection.DisplayName = LOCTEXT("PerViewportOCIOSubsectionLabel", "Per-Viewport");
-	PerViewportOCIODetailsSubsection.Categories.Add(TEXT("CustomPerViewportOCIOCategory"));
+		OutColorGradingDataModel.DetailsSections.Add(TransformDetailsSection);
+	}
 
-	OCIODetailsSection.Subsections.Add(PerViewportOCIODetailsSubsection);
+	{
+		FDisplayClusterColorGradingDataModel::FDetailsSection ChromakeyDetailsSection;
+		ChromakeyDetailsSection.Categories.Add(TEXT("CustomViewportChromakeyCategory"));
+		ChromakeyDetailsSection.Categories.Add(TEXT("CustomViewportChromakeyMarkersCategory"));
 
-	OutColorGradingDataModel.DetailsSections.Add(OCIODetailsSection);
+		OutColorGradingDataModel.DetailsSections.Add(ChromakeyDetailsSection);
+	}
 }
 
 FReply FDisplayClusterColorGradingGenerator_RootActor::AddColorGradingGroup()
@@ -809,8 +812,8 @@ public:
 		IDetailCategoryBuilder& CategoryBuilder = DetailBuilder.EditCategory(TEXT("ColorGradingCategory"));
 
 		CategoryBuilder.AddProperty(DetailBuilder.GetProperty(GET_MEMBER_NAME_CHECKED(UDisplayClusterICVFXCameraComponent, CameraSettings.bEnable)));
+		CategoryBuilder.AddProperty(DetailBuilder.GetProperty(GET_MEMBER_NAME_CHECKED(UDisplayClusterICVFXCameraComponent, CameraSettings.CustomFrustum.bEnable)));
 		CategoryBuilder.AddProperty(DetailBuilder.GetProperty(GET_MEMBER_NAME_CHECKED(UDisplayClusterICVFXCameraComponent, CameraSettings.Chromakey.bEnable)));
-		CategoryBuilder.AddProperty(DetailBuilder.GetProperty(GET_MEMBER_NAME_CHECKED(UDisplayClusterICVFXCameraComponent, CameraSettings.CameraOCIO.AllNodesOCIOConfiguration.bIsEnabled)));
 		CategoryBuilder.AddProperty(DetailBuilder.GetProperty(GET_MEMBER_NAME_CHECKED(UDisplayClusterICVFXCameraComponent, CameraSettings.AllNodesColorGrading)));
 		CategoryBuilder.AddProperty(DetailBuilder.GetProperty(GET_MEMBER_NAME_CHECKED(UDisplayClusterICVFXCameraComponent, CameraSettings.PerNodeColorGrading)));
 
@@ -913,19 +916,12 @@ private:
 		AddProperty(OverscanCategoryBuilder, TEXT("CustomFrustumRef"), true);
 
 		IDetailCategoryBuilder& ChromakeyCategoryBuilder = DetailBuilder.EditCategory(TEXT("CustomChromakeyCategory"), LOCTEXT("CustomChromakeyCategoryLabel", "Chromakey"));
+		AddProperty(ChromakeyCategoryBuilder, TEXT("ChromakeyTypeRef"));
+		AddProperty(ChromakeyCategoryBuilder, TEXT("ChromakeySettingsSourceRef"));
 		AddProperty(ChromakeyCategoryBuilder, TEXT("ChromakeyColorRef"));
 
 		IDetailCategoryBuilder& ChromakeyMarkersCategoryBuilder = DetailBuilder.EditCategory(TEXT("CustomChromakeyMarkersCategory"), LOCTEXT("CustomChromakeyMarkersCategoryLabel", "ChromakeyMarkers"));
 		AddProperty(ChromakeyMarkersCategoryBuilder, TEXT("ChromakeyMarkersRef"), true);
-
-		IDetailCategoryBuilder& ChromakeyCustomCategoryBuilder = DetailBuilder.EditCategory(TEXT("CustomChromakeyCustomCategory"), LOCTEXT("CustomChromakeyCustomCategoryLabel", "Custom Chromakey"));
-		AddProperty(ChromakeyCustomCategoryBuilder, TEXT("ChromakeyRenderTextureRef"), true);
-
-		IDetailCategoryBuilder& OCIOAllNodesCategoryBuilder = DetailBuilder.EditCategory(TEXT("CustomAllNodesOCIOCategory"), LOCTEXT("CustomAllNodesOCIOCategoryLabel", "All Nodes"));
-		AddProperty(OCIOAllNodesCategoryBuilder, TEXT("AllNodesOCIOConfigurationRef"));
-
-		IDetailCategoryBuilder& OCIOPerNodeCategoryBuilder = DetailBuilder.EditCategory(TEXT("CustomPerNodeOCIOCategory"), LOCTEXT("CustomPerNodeOCIOCategoryLabel", "Per-Node"));
-		AddProperty(OCIOPerNodeCategoryBuilder, TEXT("PerNodeOCIOProfilesRef"));
 	}
 
 private:
@@ -1170,62 +1166,28 @@ void FDisplayClusterColorGradingGenerator_ICVFXCamera::GenerateDataModel(IProper
 		FDisplayClusterColorGradingDataModel::FDetailsSection InnerFrustumDetailsSection;
 		InnerFrustumDetailsSection.DisplayName = LOCTEXT("InnerFrustumDetailsSectionLabel", "Inner Frustum");
 		InnerFrustumDetailsSection.EditConditionPropertyHandle = FindPropertyHandle(PropertyRowGenerator, CREATE_PROPERTY_PATH(UDisplayClusterICVFXCameraComponent, CameraSettings.bEnable));
+		InnerFrustumDetailsSection.Categories = { TEXT("CustomICVFXCategory"), TEXT("CustomSoftEdgeCategory"), TEXT("CustomBorderCategory") };
 
-		FDisplayClusterColorGradingDataModel::FDetailsSubsection ICVFXDetailsSubsection;
-		ICVFXDetailsSubsection.DisplayName = LOCTEXT("ICVFXSubsectionLabel", "ICVFX");
-		ICVFXDetailsSubsection.Categories = { TEXT("CustomICVFXCategory"), TEXT("CustomSoftEdgeCategory"), TEXT("CustomBorderCategory") };
-
-		InnerFrustumDetailsSection.Subsections.Add(ICVFXDetailsSubsection);
-
-		FDisplayClusterColorGradingDataModel::FDetailsSubsection OverscanDetailsSubsection;
-		OverscanDetailsSubsection.DisplayName = LOCTEXT("OverscanDetailsSubsectionLabel", "Overscan");
-		OverscanDetailsSubsection.Categories = { TEXT("CustomOverscanCategory") };
+		OutColorGradingDataModel.DetailsSections.Add(InnerFrustumDetailsSection);
+	}
 
-		InnerFrustumDetailsSection.Subsections.Add(OverscanDetailsSubsection);
+	{
+		FDisplayClusterColorGradingDataModel::FDetailsSection InnerFrustumOverscanDetailsSection;
+		InnerFrustumOverscanDetailsSection.DisplayName = LOCTEXT("InnerFrustumOverscanDetailsSectionLabel", "Inner Frustum Overscan");
+		InnerFrustumOverscanDetailsSection.EditConditionPropertyHandle = FindPropertyHandle(PropertyRowGenerator, CREATE_PROPERTY_PATH(UDisplayClusterICVFXCameraComponent, CameraSettings.CustomFrustum.bEnable));
+		InnerFrustumOverscanDetailsSection.Categories = { TEXT("CustomOverscanCategory") };
 
-		OutColorGradingDataModel.DetailsSections.Add(InnerFrustumDetailsSection);
+		OutColorGradingDataModel.DetailsSections.Add(InnerFrustumOverscanDetailsSection);
 	}
 
 	{
 		FDisplayClusterColorGradingDataModel::FDetailsSection ChromakeyDetailsSection;
 		ChromakeyDetailsSection.DisplayName = LOCTEXT("ChromakeyDetailsSectionLabel", "Chromakey");
 		ChromakeyDetailsSection.EditConditionPropertyHandle = FindPropertyHandle(PropertyRowGenerator, CREATE_PROPERTY_PATH(UDisplayClusterICVFXCameraComponent, CameraSettings.Chromakey.bEnable));
-
-		FDisplayClusterColorGradingDataModel::FDetailsSubsection ChromakeyMarkersDetailsSubsection;
-		ChromakeyMarkersDetailsSubsection.DisplayName = LOCTEXT("ChromakeyMarkersDetailsSubsectionLabel", "Markers");
-		ChromakeyMarkersDetailsSubsection.Categories = { TEXT("CustomChromakeyCategory"), TEXT("CustomChromakeyMarkersCategory") };
-
-		ChromakeyDetailsSection.Subsections.Add(ChromakeyMarkersDetailsSubsection);
-
-		FDisplayClusterColorGradingDataModel::FDetailsSubsection ChromakeyCustomDetailsSubsection;
-		ChromakeyCustomDetailsSubsection.DisplayName = LOCTEXT("ChromakeyCustomDetailsSubsectionLabel", "Custom");
-		ChromakeyCustomDetailsSubsection.Categories = { TEXT("CustomChromakeyCategory"), TEXT("CustomChromakeyCustomCategory") };
-
-		ChromakeyDetailsSection.Subsections.Add(ChromakeyCustomDetailsSubsection);
+		ChromakeyDetailsSection.Categories = { TEXT("CustomChromakeyCategory"), TEXT("CustomChromakeyMarkersCategory") };
 
 		OutColorGradingDataModel.DetailsSections.Add(ChromakeyDetailsSection);
 	}
-
-	{
-		FDisplayClusterColorGradingDataModel::FDetailsSection OCIODetailsSection;
-		OCIODetailsSection.DisplayName = LOCTEXT("OCIODetailsSectionLabel", "OCIO");
-		OCIODetailsSection.EditConditionPropertyHandle = FindPropertyHandle(PropertyRowGenerator,
-			CREATE_PROPERTY_PATH(UDisplayClusterICVFXCameraComponent, CameraSettings.CameraOCIO.AllNodesOCIOConfiguration.bIsEnabled));
-
-		FDisplayClusterColorGradingDataModel::FDetailsSubsection AllNodesOCIODetailsSubsection;
-		AllNodesOCIODetailsSubsection.DisplayName = LOCTEXT("AllNodesOCIODetailsSubsectionLabel", "All Nodes");
-		AllNodesOCIODetailsSubsection.Categories = { TEXT("CustomAllNodesOCIOCategory") };
-
-		OCIODetailsSection.Subsections.Add(AllNodesOCIODetailsSubsection);
-
-		FDisplayClusterColorGradingDataModel::FDetailsSubsection PerNodeOCIODetailsSubsection;
-		PerNodeOCIODetailsSubsection.DisplayName = LOCTEXT("PerNodeOCIODetailsSubsectionLabel", "Per-Node");
-		PerNodeOCIODetailsSubsection.Categories = { TEXT("CustomPerNodeOCIOCategory") };
-
-		OCIODetailsSection.Subsections.Add(PerNodeOCIODetailsSubsection);
-
-		OutColorGradingDataModel.DetailsSections.Add(OCIODetailsSection);
-	}
 }
 
 FReply FDisplayClusterColorGradingGenerator_ICVFXCamera::AddColorGradingGroup()
diff --git a/Engine/Plugins/Runtime/nDisplay/Source/DisplayClusterColorGrading/Private/DetailView/SDisplayClusterColorGradingDetailTreeRow.cpp b/Engine/Plugins/Runtime/nDisplay/Source/DisplayClusterColorGrading/Private/DetailView/SDisplayClusterColorGradingDetailTreeRow.cpp
--- a/Engine/Plugins/Runtime/nDisplay/Source/DisplayClusterColorGrading/Private/DetailView/SDisplayClusterColorGradingDetailTreeRow.cpp
+++ b/Engine/Plugins/Runtime/nDisplay/Source/DisplayClusterColorGrading/Private/DetailView/SDisplayClusterColorGradingDetailTreeRow.cpp
@@ -700,6 +700,12 @@ bool SDisplayClusterColorGradingDetailTreeRow::IsRowEnabled() const
 {
 	if (DetailTreeItem.IsValid())
 	{
+		// Categories are always enabled
+		if (DetailTreeItem.Pin()->IsCategory())
+		{
+			return true;
+		}
+
 		return DetailTreeItem.Pin()->IsPropertyEditingEnabled().Get(true) && 
 			WidgetRow.IsEnabledAttr.Get(true) && 
 			WidgetRow.EditConditionValue.Get(true);
diff --git a/Engine/Plugins/Runtime/nDisplay/Source/DisplayClusterColorGrading/Private/SDisplayClusterColorGradingDetailsPanel.cpp b/Engine/Plugins/Runtime/nDisplay/Source/DisplayClusterColorGrading/Private/SDisplayClusterColorGradingDetailsPanel.cpp
--- a/Engine/Plugins/Runtime/nDisplay/Source/DisplayClusterColorGrading/Private/SDisplayClusterColorGradingDetailsPanel.cpp
+++ b/Engine/Plugins/Runtime/nDisplay/Source/DisplayClusterColorGrading/Private/SDisplayClusterColorGradingDetailsPanel.cpp
@@ -103,6 +103,12 @@ private:
 	{
 		if (DetailsSection)
 		{
+			// If there is nothing to display in the title bar, don't display one
+			if (DetailsSection->DisplayName.IsEmpty() && !DetailsSection->EditConditionPropertyHandle.IsValid() && DetailsSection->Subsections.Num() == 0)
+			{
+				return SNullWidget::NullWidget;
+			}
+
 			TSharedRef<SHorizontalBox> HorizontalBox = SNew(SHorizontalBox);
 
 			if (DetailsSection->EditConditionPropertyHandle.IsValid() && DetailsSection->EditConditionPropertyHandle->IsValidHandle())
@@ -110,7 +116,7 @@ private:
 				HorizontalBox->AddSlot()
 				.AutoWidth()
 				.VAlign(VAlign_Center)
-				.Padding(2, 0, 4, 0)
+				.Padding(2, 0, 0, 0)
 				[
 					DetailsSection->EditConditionPropertyHandle->CreatePropertyValueWidget(false)
 				];
@@ -119,7 +125,7 @@ private:
 			HorizontalBox->AddSlot()
 			.FillWidth(1.0f)
 			.VAlign(VAlign_Center)
-			.Padding(2, 0, 0, 0)
+			.Padding(6, 0, 0, 0)
 			[
 				SNew(STextBlock)
 				.Text(DetailsSection->DisplayName)
@@ -136,7 +142,11 @@ private:
 				];
 			}
 
-			return HorizontalBox;
+			return SNew(SBox)
+				.HeightOverride(18)
+				[
+					HorizontalBox
+				];
 		}
 
 		return SNullWidget::NullWidget;
diff --git a/Engine/Plugins/Runtime/nDisplay/Source/DisplayClusterConfigurator/Private/Views/Details/DisplayClusterEditorPropertyReferenceTypeCustomization.cpp b/Engine/Plugins/Runtime/nDisplay/Source/DisplayClusterConfigurator/Private/Views/Details/DisplayClusterEditorPropertyReferenceTypeCustomization.cpp
--- a/Engine/Plugins/Runtime/nDisplay/Source/DisplayClusterConfigurator/Private/Views/Details/DisplayClusterEditorPropertyReferenceTypeCustomization.cpp
+++ b/Engine/Plugins/Runtime/nDisplay/Source/DisplayClusterConfigurator/Private/Views/Details/DisplayClusterEditorPropertyReferenceTypeCustomization.cpp
@@ -154,7 +154,10 @@ void FDisplayClusterEditorPropertyReferenceTypeCustomization::CustomizeChildren(
 					for (uint32 Index = 0; Index < NumChildren; ++Index)
 					{
 						TSharedPtr<IPropertyHandle> ChildHandle = ReferencedPropertyHandle->GetChildHandle(Index);
-						InChildBuilder.AddProperty(ChildHandle.ToSharedRef());
+						if (!ChildHandle->IsCustomized())
+						{
+							InChildBuilder.AddProperty(ChildHandle.ToSharedRef());
+						}
 					}
 				}
 				else
