nDisplay:  support debug draw, enabled via CVar DC.DebugDraw.

diff --git a/Engine/Plugins/Runtime/nDisplay/Source/DisplayCluster/Private/Game/EngineClasses/Basics/DisplayClusterViewportClient.cpp b/Engine/Plugins/Runtime/nDisplay/Source/DisplayCluster/Private/Game/EngineClasses/Basics/DisplayClusterViewportClient.cpp
--- a/Engine/Plugins/Runtime/nDisplay/Source/DisplayCluster/Private/Game/EngineClasses/Basics/DisplayClusterViewportClient.cpp
+++ b/Engine/Plugins/Runtime/nDisplay/Source/DisplayCluster/Private/Game/EngineClasses/Basics/DisplayClusterViewportClient.cpp
@@ -95,6 +95,14 @@ static FAutoConsoleVariableRef CVarDisplayClusterLumenPerView(
 	ECVF_RenderThreadSafe
 );
 
+int32 GDisplayClusterDebugDraw = 1;
+static FAutoConsoleVariableRef CVarDisplayClusterDebugDraw(
+	TEXT("DC.DebugDraw"),
+	GDisplayClusterDebugDraw,
+	TEXT("Enable debug draw for nDisplay views.  Debug draw features are separately enabled, and default to off, this just provides an additional global toggle."),
+	ECVF_RenderThreadSafe
+);
+
 struct FCompareViewFamilyBySizeAndGPU
 {
 	FORCEINLINE bool operator()(const FSceneViewFamilyContext& A, const FSceneViewFamilyContext& B) const
@@ -876,12 +884,6 @@ void UDisplayClusterViewportClient::Draw(FViewport* InViewport, FCanvas* SceneCa
 				}
 			}
 		}
-
-		for (FSceneViewFamilyContext* ViewFamilyContext : ViewFamilies)
-		{
-			delete ViewFamilyContext;
-		}
-		ViewFamilies.Empty();
 	}
 	else
 	{
@@ -960,6 +962,11 @@ void UDisplayClusterViewportClient::Draw(FViewport* InViewport, FCanvas* SceneCa
 #endif
 		}
 
+		if (GDisplayClusterDebugDraw && !ViewFamilies.IsEmpty())
+		{
+			UDebugDrawService::Draw(ViewFamilies.Last()->EngineShowFlags, InViewport, const_cast<FSceneView*>(ViewFamilies.Last()->Views[0]), DebugCanvas, DebugCanvasObject);
+		}
+
 		// Render the console absolutely last because developer input is was matter the most.
 		if (ViewportConsole)
 		{
@@ -967,6 +974,15 @@ void UDisplayClusterViewportClient::Draw(FViewport* InViewport, FCanvas* SceneCa
 		}
 	}
 
+	if (!ViewFamilies.IsEmpty())
+	{
+		for (FSceneViewFamilyContext* ViewFamilyContext : ViewFamilies)
+		{
+			delete ViewFamilyContext;
+		}
+		ViewFamilies.Empty();
+	}
+
 	OnEndDraw().Broadcast();
 }
 
