SharedMemoryMedia: Improved GPU traces to include the wait for pixels.

A side effect is that the previous commands recorded in the commandlist will be dispatched before the wait instead of after the wait.

diff --git a/Engine/Plugins/Runtime/nDisplay/Source/SharedMemoryMedia/Private/SharedMemoryMediaPlayer.cpp b/Engine/Plugins/Runtime/nDisplay/Source/SharedMemoryMedia/Private/SharedMemoryMediaPlayer.cpp
--- a/Engine/Plugins/Runtime/nDisplay/Source/SharedMemoryMedia/Private/SharedMemoryMediaPlayer.cpp
+++ b/Engine/Plugins/Runtime/nDisplay/Source/SharedMemoryMedia/Private/SharedMemoryMediaPlayer.cpp
@@ -23,6 +23,8 @@
 DECLARE_GPU_STAT_NAMED(SharedMemoryMedia_CopyToSampleCommon, TEXT("SharedMemoryMedia_CopyToSampleCommon"));
 DECLARE_GPU_STAT_NAMED(SharedMemoryMedia_UnlockTexture, TEXT("SharedMemoryMedia_UnlockTexture"));
 DECLARE_GPU_STAT_NAMED(SharedMemoryMedia_LockTexture, TEXT("SharedMemoryMedia_LockTexture"));
+DECLARE_GPU_STAT_NAMED(SharedMemoryMedia_WaitForPixels, TEXT("SharedMemoryMedia_WaitForPixels"));
+
 
 using namespace UE::SharedMemoryMedia;
 
@@ -747,54 +749,63 @@ void FSharedMemoryMediaPlayer::JustInTimeSampleRender()
 	// Enqueue a lambda that will wait for the cross gpu texture for this frame have the data populated.
 	// It will determine this by polling the sender's shared memory metadata.
 	// To avoid any hangs, it will give up after some time.
-	RHICmdList.EnqueueLambda(
-		[
-			FrameNumber = GFrameCounterRenderThread, 
-			SharedMemoryIdx, 
-			ExpectedFrameNumber, 
-			SharedMemoryData, 
-			this
-		](FRHICommandListImmediate& RHICmdList)
-		{
-			TRACE_CPUPROFILER_EVENT_SCOPE(SharedMemMediaWaitForSharedGpuTexture);
-
-			alignas(128) FSharedMemoryMediaFrameMetadata::FSender SenderMetadata;
+	{
+		SCOPED_GPU_STAT(RHICmdList, SharedMemoryMedia_WaitForPixels);
+		SCOPED_DRAW_EVENT(RHICmdList, SharedMemoryMedia_WaitForPixels);
+
+		// Since we are going to enqueue a lambda that can potentially sleep in the RHI thread if the pixels haven't arrived,
+		// we dispatch the existing commands (including the draw event start timing in the SCOPED_DRAW_EVENT above) before any potential sleep.
+		RHICmdList.ImmediateFlush(EImmediateFlushType::DispatchToRHIThread);
+
+		RHICmdList.EnqueueLambda(
+			[
+				FrameNumber = GFrameCounterRenderThread, 
+				SharedMemoryIdx, 
+				ExpectedFrameNumber, 
+				SharedMemoryData, 
+				this
+			](FRHICommandListImmediate& RHICmdList)
+			{
+				TRACE_CPUPROFILER_EVENT_SCOPE(SharedMemMediaWaitForSharedGpuTexture);
 
-			const double StartTimeSeconds = FPlatformTime::Seconds();
-			constexpr double TimeoutSeconds = 0.5;
+				alignas(128) FSharedMemoryMediaFrameMetadata::FSender SenderMetadata;
 
-			while (true)
-			{
-				FMemory::Memcpy(&SenderMetadata, SharedMemoryData, sizeof(FSharedMemoryMediaFrameMetadata::FSender));
+				const double StartTimeSeconds = FPlatformTime::Seconds();
+				constexpr double TimeoutSeconds = 0.5;
 
-				if (SenderMetadata.Magic == FSharedMemoryMediaFrameMetadata::MAGIC)
+				while (true)
 				{
-					if (SenderMetadata.FrameNumber >= ExpectedFrameNumber)
+					FMemory::Memcpy(&SenderMetadata, SharedMemoryData, sizeof(FSharedMemoryMediaFrameMetadata::FSender));
+
+					if (SenderMetadata.Magic == FSharedMemoryMediaFrameMetadata::MAGIC)
 					{
-						if (SenderMetadata.FrameNumber > ExpectedFrameNumber)
+						if (SenderMetadata.FrameNumber >= ExpectedFrameNumber)
 						{
-							UE_LOG(LogSharedMemoryMedia, Warning, TEXT("Using too recent frame. Expected %u and used %u"), 
-								ExpectedFrameNumber, SenderMetadata.FrameNumber);
-						}
+							if (SenderMetadata.FrameNumber > ExpectedFrameNumber)
+							{
+								UE_LOG(LogSharedMemoryMedia, Warning, TEXT("Using too recent frame. Expected %u and used %u"), 
+									ExpectedFrameNumber, SenderMetadata.FrameNumber);
+							}
 
-						break;
+							break;
+						}
 					}
-				}
 
-				FPlatformProcess::SleepNoStats(SpinWaitTimeSeconds);
+					FPlatformProcess::SleepNoStats(SpinWaitTimeSeconds);
 
-				if ((FPlatformTime::Seconds() - StartTimeSeconds) > TimeoutSeconds)
-				{
-					UE_LOG(LogSharedMemoryMedia, Warning, TEXT("FSharedMemoryMediaPlayer timed out waiting for ExpectedFrameNumber %u for frame %llu, only saw up to frame %u"),
-						ExpectedFrameNumber, FrameNumber, SenderMetadata.FrameNumber);
+					if ((FPlatformTime::Seconds() - StartTimeSeconds) > TimeoutSeconds)
+					{
+						UE_LOG(LogSharedMemoryMedia, Warning, TEXT("FSharedMemoryMediaPlayer timed out waiting for ExpectedFrameNumber %u for frame %llu, only saw up to frame %u"),
+							ExpectedFrameNumber, FrameNumber, SenderMetadata.FrameNumber);
 
-					// Note: It would be desirable to stop the copy texture from happening, but at this point it has already been enqueued.
+						// Note: It would be desirable to stop the copy texture from happening, but at this point it has already been enqueued.
 
-					break;
+						break;
+					}
 				}
 			}
-		}
-	);
+		);
+	}
 
 	// Copy to sample common texture
 
