nDisplay: Fix ICVFX operator panel not properly handling color picker transactions.

diff --git a/Engine/Plugins/Runtime/nDisplay/Source/DisplayClusterColorGrading/Private/SDisplayClusterColorGradingDetailsPanel.cpp b/Engine/Plugins/Runtime/nDisplay/Source/DisplayClusterColorGrading/Private/SDisplayClusterColorGradingDetailsPanel.cpp
--- a/Engine/Plugins/Runtime/nDisplay/Source/DisplayClusterColorGrading/Private/SDisplayClusterColorGradingDetailsPanel.cpp
+++ b/Engine/Plugins/Runtime/nDisplay/Source/DisplayClusterColorGrading/Private/SDisplayClusterColorGradingDetailsPanel.cpp
@@ -10,6 +10,7 @@
 #include "Modules/ModuleManager.h"
 #include "PropertyEditorModule.h"
 #include "Widgets/SBoxPanel.h"
+#include "Widgets/Colors/SColorPicker.h"
 #include "Widgets/Input/SCheckBox.h"
 #include "Widgets/Layout/SBox.h"
 #include "Widgets/Layout/SSplitter.h"
@@ -285,6 +286,23 @@ void SDisplayClusterColorGradingDetailsPanel::SetDrawerState(const FDisplayClust
 void SDisplayClusterColorGradingDetailsPanel::FillDetailsSections()
 {
 	TRACE_CPUPROFILER_EVENT_SCOPE(SDisplayClusterColorGradingDetailsPanel::FillDetailsSections);
+	
+	TSharedPtr<SColorPicker> ExistingColorPicker = GetColorPicker();
+	if (ExistingColorPicker.IsValid())
+	{
+		// Color picker has special handling -- FColorStructCustomization relies on listening to SColorPicker window closed
+		// events in order to fully commit changes and close out an open transaction. If we destroy the details section views
+		// before the window close event fires we will leave an open transaction and this will cause problems with undo history
+		// and transacting changes to MU.
+	
+		// Similar handling is done in SDetailsView. On PreSetObject, GetOptionalOwningDetailsView is checked and the picker
+		// closed if they do not own the picker, and then closed on PostSetObject if they do own it. Our implementation doesn't rely
+		// on PreSet/PostSet, and there is never a case where we own the color picker, so we should always close it.
+		
+		DestroyColorPicker();
+		ExistingColorPicker.Reset();
+	}
+	
 	if (ColorGradingDataModel.IsValid())
 	{
 		TArray<TWeakObjectPtr<UObject>> Objects = ColorGradingDataModel->GetObjects();
