[nDisplay] Automatically remove lightcards from parent DCRA list when destroyed

diff --git a/Engine/Plugins/Runtime/nDisplay/Source/DisplayCluster/Private/Game/EngineClasses/Scene/DisplayClusterLightCardActor.cpp b/Engine/Plugins/Runtime/nDisplay/Source/DisplayCluster/Private/Game/EngineClasses/Scene/DisplayClusterLightCardActor.cpp
--- a/Engine/Plugins/Runtime/nDisplay/Source/DisplayCluster/Private/Game/EngineClasses/Scene/DisplayClusterLightCardActor.cpp
+++ b/Engine/Plugins/Runtime/nDisplay/Source/DisplayCluster/Private/Game/EngineClasses/Scene/DisplayClusterLightCardActor.cpp
@@ -185,6 +185,11 @@ void ADisplayClusterLightCardActor::Tick(float DeltaSeconds)
 	}
 }
 
+void ADisplayClusterLightCardActor::Destroyed()
+{
+	RemoveFromRootActor();
+}
+
 #if WITH_EDITOR
 
 void ADisplayClusterLightCardActor::PostEditChangeProperty(FPropertyChangedEvent& PropertyChangedEvent)
@@ -760,10 +765,30 @@ void ADisplayClusterLightCardActor::RemoveFromRootActor()
 #if WITH_EDITOR
 	Modify();
 #endif
+
+
+	if (ADisplayClusterRootActor* RootActorOwner = GetRootActorOwner())
+	{
+		// Remove this from the list of LCs for the owning root actor
+		RemoveFromRootActorList(RootActorOwner);
+	}
+	else
+	{
+		// This may be a legacy LC, so we need to check if any root actors in the scene own it and remove it from their lists if so
+		if (const UWorld* World = GetWorld())
+		{
+			for (TActorIterator<ADisplayClusterRootActor> Iter(World); Iter; ++Iter)
+			{
+				RemoveFromRootActorList(*Iter);
+			}
+		}
+	}
+
 	if (StageActorComponent)
 	{
 		StageActorComponent->SetRootActor(nullptr);
 	}
+
 	WeakRootActorOwner.Reset();
 }
 
@@ -1009,10 +1034,30 @@ void ADisplayClusterLightCardActor::CleanUpComponentsForExtenders()
 	}
 }
 
+void ADisplayClusterLightCardActor::RemoveFromRootActorList(ADisplayClusterRootActor* RootActor)
+{
+	if (!RootActor)
+	{
+		return;
+	}
+
+	if (UDisplayClusterConfigurationData* CurrentData = RootActor->GetConfigData())
+	{
+		for (auto Iter = CurrentData->StageSettings.Lightcard.ShowOnlyList.Actors.CreateIterator(); Iter; ++Iter)
+		{
+			if (Iter->Get() == this)
+			{
+				Iter.RemoveCurrent();
+			}
+		}
+	}
+}
+
 #if WITH_EDITOR
 void ADisplayClusterLightCardActor::OnLevelActorDeleted(AActor* DeletedActor)
 {
-	if (DeletedActor && StageActorComponent && Cast<ADisplayClusterRootActor>(DeletedActor) == StageActorComponent->GetRootActor().Get())
+	ADisplayClusterRootActor* RootActorOwner = GetRootActorOwner();
+	if (DeletedActor && RootActorOwner && StageActorComponent && Cast<ADisplayClusterRootActor>(DeletedActor) == RootActorOwner)
 	{
 		RemoveFromRootActor();
 	}
diff --git a/Engine/Plugins/Runtime/nDisplay/Source/DisplayCluster/Public/DisplayClusterLightCardActor.h b/Engine/Plugins/Runtime/nDisplay/Source/DisplayCluster/Public/DisplayClusterLightCardActor.h
--- a/Engine/Plugins/Runtime/nDisplay/Source/DisplayCluster/Public/DisplayClusterLightCardActor.h
+++ b/Engine/Plugins/Runtime/nDisplay/Source/DisplayCluster/Public/DisplayClusterLightCardActor.h
@@ -75,6 +75,7 @@ public:
 	virtual void PostLoad() override;
 	virtual void OnConstruction(const FTransform& Transform) override;
 	virtual void Tick(float DeltaSeconds) override;
+	virtual void Destroyed() override;
 
 #if WITH_EDITOR
 	virtual void PostEditChangeProperty(FPropertyChangedEvent& PropertyChangedEvent) override;
@@ -302,6 +303,9 @@ protected:
 	/** Removes components that were added by IDisplayClusterLightCardActorExtender */
 	void CleanUpComponentsForExtenders();
 
+	/** Removes this actor from a root actor's ShowOnlyList.Actors, if possible. */
+	void RemoveFromRootActorList(ADisplayClusterRootActor* RootActor);
+
 #if WITH_EDITOR
 	/** Called when a level actor is deleted */
 	void OnLevelActorDeleted(AActor* DeletedActor);
