nDisplay: Updated cast (raw pointer to TObjectPtr) in details customization.

diff --git a/Engine/Plugins/Runtime/nDisplay/Source/DisplayClusterLightCardEditor/Private/DetailCustomizations/DisplayClusterLightCardActorDetails.cpp b/Engine/Plugins/Runtime/nDisplay/Source/DisplayClusterLightCardEditor/Private/DetailCustomizations/DisplayClusterLightCardActorDetails.cpp
--- a/Engine/Plugins/Runtime/nDisplay/Source/DisplayClusterLightCardEditor/Private/DetailCustomizations/DisplayClusterLightCardActorDetails.cpp
+++ b/Engine/Plugins/Runtime/nDisplay/Source/DisplayClusterLightCardEditor/Private/DetailCustomizations/DisplayClusterLightCardActorDetails.cpp
@@ -55,10 +55,10 @@ void FDisplayClusterLightCardActorDetails::CustomizeDetails(IDetailLayoutBuilder
 
 			for (const void* RawData : RawDatas)
 			{
-				TMap<FName, UActorComponent*>* ExtenderNameToComponentMapPtr = (TMap<FName, UActorComponent*>*)(RawData);
+				TMap<FName, TObjectPtr<UActorComponent>>* ExtenderNameToComponentMapPtr = (TMap<FName, TObjectPtr<UActorComponent>>*)(RawData);
 				check(ExtenderNameToComponentMapPtr);
 
-				for (const TTuple<FName, UActorComponent*>& ExtenderNameToComponentPair : *ExtenderNameToComponentMapPtr)
+				for (const TTuple<FName, TObjectPtr<UActorComponent>>& ExtenderNameToComponentPair : *ExtenderNameToComponentMapPtr)
 				{
 					if (!ensureMsgf(ExtenderNameToComponentPair.Value, TEXT("Trying to display component for Extender %s, but component is invalid."), *ExtenderNameToComponentPair.Key.ToString()))
 					{
