nDisplay: Fixed issue where level instance ICVFX camera component would become uneditable if its blueprint had been updated to the global chromakey change but it had not

diff --git a/Engine/Plugins/Runtime/nDisplay/Source/DisplayCluster/Private/Game/EngineClasses/Scene/DisplayClusterICVFXCameraComponent.cpp b/Engine/Plugins/Runtime/nDisplay/Source/DisplayCluster/Private/Game/EngineClasses/Scene/DisplayClusterICVFXCameraComponent.cpp
--- a/Engine/Plugins/Runtime/nDisplay/Source/DisplayCluster/Private/Game/EngineClasses/Scene/DisplayClusterICVFXCameraComponent.cpp
+++ b/Engine/Plugins/Runtime/nDisplay/Source/DisplayCluster/Private/Game/EngineClasses/Scene/DisplayClusterICVFXCameraComponent.cpp
@@ -33,13 +33,22 @@ void UDisplayClusterICVFXCameraComponent::PostLoad()
 	const int32 CustomVersion = GetLinkerCustomVersion(FDisplayClusterICVFXCameraCustomVersion::GUID);
 	if (CustomVersion < FDisplayClusterICVFXCameraCustomVersion::UpdateChromakeyConfig)
 	{
-		const bool bCustomChromakey = CameraSettings.Chromakey.ChromakeyRenderTexture.bEnable_DEPRECATED;
-		CameraSettings.Chromakey.ChromakeyType = bCustomChromakey ? 
-			EDisplayClusterConfigurationICVFX_ChromakeyType::CustomChromakey :
-			EDisplayClusterConfigurationICVFX_ChromakeyType::InnerFrustum;
+		const bool bHasCustomArchetype = GetArchetype() != StaticClass()->ClassDefaultObject;
+		const int32 ArchetypeVersion = GetArchetype()->GetLinkerCustomVersion(FDisplayClusterICVFXCameraCustomVersion::GUID);
 
-		// New ICVFX cameras default to the global chromakey settings, but for pre 5.3 cameras, the source must be set to the ICVFX camera
-		CameraSettings.Chromakey.ChromakeySettingsSource = EDisplayClusterConfigurationICVFX_ChromakeySettingsSource::ICVFXCamera;
+		// UE-184291: If this camera component has a user-defined archetype and that archetype has been updated already, do not
+		// attempt to update the component's properties; the new properties will already be set to the correct values from the
+		// archetype and overriding them to these "default" values can cause bad things to happen. 
+		if (!bHasCustomArchetype || ArchetypeVersion < FDisplayClusterICVFXCameraCustomVersion::UpdateChromakeyConfig)
+		{
+			const bool bCustomChromakey = CameraSettings.Chromakey.ChromakeyRenderTexture.bEnable_DEPRECATED;
+			CameraSettings.Chromakey.ChromakeyType = bCustomChromakey ? 
+				EDisplayClusterConfigurationICVFX_ChromakeyType::CustomChromakey :
+				EDisplayClusterConfigurationICVFX_ChromakeyType::InnerFrustum;
+
+			// New ICVFX cameras default to the global chromakey settings, but for pre 5.3 cameras, the source must be set to the ICVFX camera
+			CameraSettings.Chromakey.ChromakeySettingsSource = EDisplayClusterConfigurationICVFX_ChromakeySettingsSource::ICVFXCamera;
+		}
 	}
 	PRAGMA_ENABLE_DEPRECATION_WARNINGS
 }
