[nDisplay] Don't include LCs from other DCRAs when loading auto-populating a preview scene

diff --git a/Engine/Plugins/Runtime/nDisplay/Source/DisplayClusterScenePreview/Private/DisplayClusterScenePreviewModule.cpp b/Engine/Plugins/Runtime/nDisplay/Source/DisplayClusterScenePreview/Private/DisplayClusterScenePreviewModule.cpp
--- a/Engine/Plugins/Runtime/nDisplay/Source/DisplayClusterScenePreview/Private/DisplayClusterScenePreviewModule.cpp
+++ b/Engine/Plugins/Runtime/nDisplay/Source/DisplayClusterScenePreview/Private/DisplayClusterScenePreviewModule.cpp
@@ -484,7 +484,7 @@ void FDisplayClusterScenePreviewModule::AutoPopulateScene(FRendererConfig& Rende
 			{
 				for (const TWeakObjectPtr<AActor> WeakActor : TActorRange<AActor>(World))
 				{
-					if (WeakActor.IsValid() && WeakActor->Implements<UDisplayClusterStageActor>())
+					if (WeakActor.IsValid() && WeakActor->Implements<UDisplayClusterStageActor>() && !WeakActor->IsA<ADisplayClusterLightCardActor>())
 					{
 						Actors.Add(WeakActor.Get());
 					}
