Shared Memory Media Capture: Fixed memory leak.

diff --git a/Engine/Plugins/Runtime/nDisplay/Source/SharedMemoryMedia/Private/Platforms/Windows/SharedMemoryMediaPlatformWindowsD3D12.cpp b/Engine/Plugins/Runtime/nDisplay/Source/SharedMemoryMedia/Private/Platforms/Windows/SharedMemoryMediaPlatformWindowsD3D12.cpp
--- a/Engine/Plugins/Runtime/nDisplay/Source/SharedMemoryMedia/Private/Platforms/Windows/SharedMemoryMediaPlatformWindowsD3D12.cpp
+++ b/Engine/Plugins/Runtime/nDisplay/Source/SharedMemoryMedia/Private/Platforms/Windows/SharedMemoryMediaPlatformWindowsD3D12.cpp
@@ -77,6 +77,7 @@ FSharedMemoryMediaPlatformWindowsD3D12::FSharedMemoryMediaPlatformWindowsD3D12()
 	for (int32 Idx = 0; Idx < UE::SharedMemoryMedia::SenderNumBuffers; ++Idx)
 	{
 		SharedHandle[Idx] = INVALID_HANDLE_VALUE;
+		CommittedResource[Idx] = nullptr;
 	}
 }
 
@@ -86,11 +87,7 @@ FSharedMemoryMediaPlatformWindowsD3D12::~FSharedMemoryMediaPlatformWindowsD3D12(
 	// Close shared handles
 	for (int32 BufferIdx = 0; BufferIdx < UE::SharedMemoryMedia::SenderNumBuffers; ++BufferIdx)
 	{
-		if (SharedHandle[BufferIdx] != INVALID_HANDLE_VALUE)
-		{
-			::CloseHandle(SharedHandle[BufferIdx]);
-			SharedHandle[BufferIdx] = INVALID_HANDLE_VALUE;
-		}
+		ReleaseSharedCrossGpuTexture(BufferIdx);
 	}
 }
 
@@ -100,6 +97,7 @@ FTextureRHIRef FSharedMemoryMediaPlatformWindowsD3D12::CreateSharedCrossGpuTextu
 
 	check(BufferIdx < UE::SharedMemoryMedia::SenderNumBuffers);
 	check(SharedHandle[BufferIdx] == INVALID_HANDLE_VALUE);
+	check(CommittedResource[BufferIdx] == nullptr);
 
 	ID3D12Device* D3D12Device = static_cast<ID3D12Device*>(GDynamicRHI->RHIGetNativeDevice());
 	check(D3D12Device);
@@ -120,15 +118,13 @@ FTextureRHIRef FSharedMemoryMediaPlatformWindowsD3D12::CreateSharedCrossGpuTextu
 
 	CD3DX12_HEAP_PROPERTIES HeapProperties(D3D12_HEAP_TYPE_DEFAULT);
 
-	ID3D12Resource* SharedGpuResource = nullptr;
-
 	HRESULT HResult = D3D12Device->CreateCommittedResource(
 		&HeapProperties,
 		D3D12_HEAP_FLAG_SHARED | D3D12_HEAP_FLAG_SHARED_CROSS_ADAPTER,
 		&SharedCrossGpuTextureDesc,
 		D3D12_RESOURCE_STATE_COPY_DEST,
 		nullptr, // pOptimizedClearValue
-		IID_PPV_ARGS(&SharedGpuResource)
+		IID_PPV_ARGS(&CommittedResource[BufferIdx])
 	);
 
 	if (FAILED(HResult))
@@ -142,13 +138,13 @@ FTextureRHIRef FSharedMemoryMediaPlatformWindowsD3D12::CreateSharedCrossGpuTextu
 	}
 
 	check(HResult == S_OK);
-	check(SharedGpuResource);
+	check(CommittedResource[BufferIdx]);
 
 	const FString SharedGpuTextureName = FString::Printf(TEXT("Global\\%s"),
 		*Guid.ToString(EGuidFormats::DigitsWithHyphensInBraces));
 
 	HResult = D3D12Device->CreateSharedHandle(
-		SharedGpuResource,                    // pObject
+		CommittedResource[BufferIdx],         // pObject
 		nullptr,                              // pAttributes
 		GENERIC_ALL,                          // Access
 		*SharedGpuTextureName,                // Name
@@ -162,17 +158,26 @@ FTextureRHIRef FSharedMemoryMediaPlatformWindowsD3D12::CreateSharedCrossGpuTextu
 		Format,
 		TexCreate_Dynamic | TexCreate_DisableSRVCreation,
 		FClearValueBinding::None,
-		SharedGpuResource
+		CommittedResource[BufferIdx]
 	);
 }
 
 void FSharedMemoryMediaPlatformWindowsD3D12::ReleaseSharedCrossGpuTexture(uint32 BufferIdx)
 {
 	check(BufferIdx < UE::SharedMemoryMedia::SenderNumBuffers);
+	if (SharedHandle[BufferIdx] != INVALID_HANDLE_VALUE)
+	{
+		::CloseHandle(SharedHandle[BufferIdx]);
+		SharedHandle[BufferIdx] = INVALID_HANDLE_VALUE;
+	}
 
-	::CloseHandle(SharedHandle[BufferIdx]);
-
-	SharedHandle[BufferIdx] = INVALID_HANDLE_VALUE;
+	// We need to manually release committed resources because FD3D12Resource::~FD3D12Resource()
+	// that is created via RHICreateTexture2DFromResource will not release it unless it is a bBackBuffer.
+	if (CommittedResource[BufferIdx])
+	{
+		CommittedResource[BufferIdx]->Release();
+		CommittedResource[BufferIdx] = nullptr;
+	}
 }
 
 FTextureRHIRef FSharedMemoryMediaPlatformWindowsD3D12::OpenSharedCrossGpuTextureByGuid(const FGuid& Guid, FSharedMemoryMediaTextureDescription& OutTextureDescription)
diff --git a/Engine/Plugins/Runtime/nDisplay/Source/SharedMemoryMedia/Private/Platforms/Windows/SharedMemoryMediaPlatformWindowsD3D12.h b/Engine/Plugins/Runtime/nDisplay/Source/SharedMemoryMedia/Private/Platforms/Windows/SharedMemoryMediaPlatformWindowsD3D12.h
--- a/Engine/Plugins/Runtime/nDisplay/Source/SharedMemoryMedia/Private/Platforms/Windows/SharedMemoryMediaPlatformWindowsD3D12.h
+++ b/Engine/Plugins/Runtime/nDisplay/Source/SharedMemoryMedia/Private/Platforms/Windows/SharedMemoryMediaPlatformWindowsD3D12.h
@@ -18,6 +18,9 @@ class FSharedMemoryMediaPlatformWindowsD3D12 : public FSharedMemoryMediaPlatform
 {
 private:
 
+	/** Pointer to the D3D12 shared memory heap. */
+	ID3D12Resource* CommittedResource[UE::SharedMemoryMedia::SenderNumBuffers]{ 0 };
+
 	/** Handles to the shared handles of the cross gpu resources */
 	HANDLE SharedHandle[UE::SharedMemoryMedia::SenderNumBuffers]{ 0 };
 
