nDisplay:

-Create a manager module to facilitate the editor UI updates as the DDC commandlet's job progresses
-Open up a bridge from the DDC commandlet to the editor to display a progress bar in the editor when it's running if called from editor menu
-Get Platforms argument from Project settings

diff --git a/Engine/Plugins/Runtime/nDisplay/Source/DisplayClusterPreloadDerivedDataCache/DisplayClusterPreloadDerivedDataCache.build.cs b/Engine/Plugins/Runtime/nDisplay/Source/DisplayClusterPreloadDerivedDataCache/DisplayClusterPreloadDerivedDataCache.build.cs
new file mode 100644
--- /dev/null
+++ b/Engine/Plugins/Runtime/nDisplay/Source/DisplayClusterPreloadDerivedDataCache/DisplayClusterPreloadDerivedDataCache.build.cs
@@ -0,0 +1,46 @@
+// Copyright Epic Games, Inc. All Rights Reserved.
+
+using System.IO;
+using UnrealBuildTool;
+
+public class DisplayClusterPreloadDerivedDataCache : ModuleRules
+{
+	public DisplayClusterPreloadDerivedDataCache(ReadOnlyTargetRules Target) : base(Target)
+	{
+		PCHUsage = ModuleRules.PCHUsageMode.UseExplicitOrSharedPCHs;
+
+		PublicDependencyModuleNames.AddRange(
+			new string[]
+			{
+				"Core",
+				"OutputLog", 
+				"UnrealEd"
+			}
+		);
+			
+		PrivateDependencyModuleNames.AddRange(
+			new string[]
+			{
+				"AssetRegistry",
+				"AssetTools",
+				"CoreUObject",
+				"ContentBrowser",
+				"Engine",
+				"EditorStyle",
+				"EditorWidgets",
+				"InputCore",
+				"Kismet",
+				"Messaging",
+				"Networking",
+				"Projects",
+				"PropertyEditor",
+				"Slate",
+				"SlateCore",
+				"ToolMenus", 
+				"ToolWidgets",
+				"UnrealEd",
+				"WorkspaceMenuStructure"
+			}
+		);
+	}
+}
diff --git a/Engine/Plugins/Runtime/nDisplay/Source/DisplayClusterPreloadDerivedDataCache/Private/DisplayClusterPreloadDerivedDataCacheModule.cpp b/Engine/Plugins/Runtime/nDisplay/Source/DisplayClusterPreloadDerivedDataCache/Private/DisplayClusterPreloadDerivedDataCacheModule.cpp
new file mode 100644
--- /dev/null
+++ b/Engine/Plugins/Runtime/nDisplay/Source/DisplayClusterPreloadDerivedDataCache/Private/DisplayClusterPreloadDerivedDataCacheModule.cpp
@@ -0,0 +1,289 @@
+// Copyright Epic Games, Inc. All Rights Reserved.
+
+#include "DisplayClusterPreloadDerivedDataCacheModule.h"
+
+#include "DisplayClusterPreloadDerivedDataCacheLog.h"
+
+#include "Commandlets/DerivedDataCacheCommandlet.h"
+#include "Framework/Docking/TabManager.h"
+#include "Framework/Notifications/NotificationManager.h"
+#include "Interfaces/IProjectManager.h"
+#include "Internationalization/Regex.h"
+#include "Misc/App.h"
+#include "Misc/DataDrivenPlatformInfoRegistry.h"
+#include "Misc/ScopedSlowTask.h"
+#include "Modules/ModuleManager.h"
+#include "ToolMenu.h"
+#include "ToolMenus.h"
+#include "Widgets/Notifications/SNotificationList.h"
+
+#define LOCTEXT_NAMESPACE "FDisplayClusterPreloadDerivedDataCacheModule"
+
+FDisplayClusterPreloadDerivedDataCacheModule& FDisplayClusterPreloadDerivedDataCacheModule::Get()
+{
+	return FModuleManager::GetModuleChecked<FDisplayClusterPreloadDerivedDataCacheModule>("DisplayClusterPreloadDerivedDataCacheModule");
+}
+
+void FDisplayClusterPreloadDerivedDataCacheModule::StartupModule()
+{
+	FCoreDelegates::OnFEngineLoopInitComplete.AddRaw(this, &FDisplayClusterPreloadDerivedDataCacheModule::OnFEngineLoopInitComplete);
+}
+
+void FDisplayClusterPreloadDerivedDataCacheModule::OnFEngineLoopInitComplete()
+{
+	UToolMenu* Menu = UToolMenus::Get()->ExtendMenu("LevelEditor.MainMenu.Build");
+
+	FToolMenuSection& Section = Menu->FindOrAddSection("LevelEditorAutomation");
+		
+	Section.AddMenuEntry(
+		"PreloadDerivedDataCache",
+		LOCTEXT("PreloadDerivedDataCache", "Preload Derived Data Cache"),
+		LOCTEXT("PreloadDerivedDataCacheTooltip", "Precompile all shaders and preprocess Nanite and other heavy data ahead of time for all assets in this project for all checked platforms in Project Settings > Supported Platforms. \nThis is useful to avoid shader compilation and pop-in when opening a level or running PIE/Standalone. \nNote that this operation can take a very long time to complete and can take up a lot of disk space, depending on the size of your project."),
+		FSlateIcon(),
+		FUIAction(FExecuteAction::CreateRaw(this, &FDisplayClusterPreloadDerivedDataCacheModule::LaunchAndCommunicateWithProcess)));
+}
+
+void FDisplayClusterPreloadDerivedDataCacheModule::ShutdownModule()
+{
+	UToolMenus::UnregisterOwner(this);
+	FCoreDelegates::OnFEngineLoopInitComplete.RemoveAll(this);
+}
+
+void FDisplayClusterPreloadDerivedDataCacheModule::RegexParseForLoadingProgress(
+	const FText& LoadingPackagesFormat, const FText& UnknownAmount, FScopedSlowTask& SlowTask, int32& LoadingTotal, FRegexMatcher& LoadingProgressRegex)
+{
+	while (LoadingProgressRegex.FindNext())
+	{
+		SlowTask.CompletedWork = FCString::Atoi(*LoadingProgressRegex.GetCaptureGroup(1));
+		SlowTask.TotalAmountOfWork = LoadingTotal;
+
+		if (LoadingTotal > INDEX_NONE)
+		{
+			SlowTask.DefaultMessage = FText::Format(LoadingPackagesFormat, FText::AsNumber(SlowTask.CompletedWork), FText::AsNumber(SlowTask.TotalAmountOfWork));
+		}
+		else
+		{
+			SlowTask.DefaultMessage = FText::Format(LoadingPackagesFormat, FText::AsNumber(SlowTask.CompletedWork), UnknownAmount);
+		}
+	}
+}
+
+void FDisplayClusterPreloadDerivedDataCacheModule::RegexParseForCompilationProgress(
+	const FText& CompilingAssetsFormat, const FText& UnknownAmount, FScopedSlowTask& SlowTask,
+	int32& CompileTotal, int32& AssetsWaitingToCompile, FRegexMatcher& CompileProgressRegex)
+{
+	while (CompileProgressRegex.FindNext())
+	{
+		AssetsWaitingToCompile = FCString::Atoi(*CompileProgressRegex.GetCaptureGroup(1));
+		if (CompileTotal == INDEX_NONE)
+		{
+			CompileTotal = AssetsWaitingToCompile;
+			SlowTask.TotalAmountOfWork = CompileTotal;
+		}
+
+		if (const int32 Progress = CompileTotal - AssetsWaitingToCompile; Progress > -1)
+		{
+			SlowTask.CompletedWork = Progress;
+			SlowTask.DefaultMessage = FText::Format(CompilingAssetsFormat, FText::AsNumber(SlowTask.CompletedWork), FText::AsNumber(SlowTask.TotalAmountOfWork));
+		}
+		else
+		{
+			SlowTask.DefaultMessage = FText::Format(CompilingAssetsFormat, UnknownAmount, UnknownAmount);
+		}
+	}
+}
+
+void FDisplayClusterPreloadDerivedDataCacheModule::LaunchAndCommunicateWithProcess()
+{
+	int32 OutResult = 0;
+	bool bOutCancelled = false;
+
+	FProcHandle ProcessHandle;
+
+	const FText EnumeratingPackages = LOCTEXT("EnumeratingPackages", "Enumerating Packages...");
+	const FText LoadingPackagesFormat = LOCTEXT("LoadingPackagesFormat", "Loading Packages ({0}/{1})...");
+	const FText CompilingAssetsFormat = LOCTEXT("CompilingAssetsFormat", "Compiling Assets ({0}/{1})...");
+	const FText UnknownAmount = LOCTEXT("Unknown", "Unknown");
+	
+	FScopedSlowTask SlowTask(1.0f, EnumeratingPackages);
+	SlowTask.MakeDialog(true);
+
+	// Show a partially completed progress bar so the end user can see something is happening
+	SlowTask.CompletedWork = 1;
+	SlowTask.TotalAmountOfWork = 2;
+
+	void* ReadPipe = nullptr;
+	void* WritePipe = nullptr;
+	verify(FPlatformProcess::CreatePipe(ReadPipe, WritePipe));
+
+	const FString CurrentExecutableName = FPlatformProcess::ExecutablePath();
+	const FString ProjectPath = FPaths::SetExtension(
+		FPaths::Combine(FPaths::ProjectDir(), FApp::GetProjectName()),".uproject");
+	const FString Arguments = ProjectPath + " " + GetDdcCommandletParams();
+	
+	UE_LOG(LogDisplayClusterPreloadDerivedDataCache, Display, TEXT("Running commandlet: %s %s"), *CurrentExecutableName, *Arguments);
+
+	uint32 ProcessID;
+	const bool bLaunchDetached = true;
+	const bool bLaunchHidden = true;
+	const bool bLaunchReallyHidden = true;
+	ProcessHandle = FPlatformProcess::CreateProc(
+		*CurrentExecutableName, *Arguments, bLaunchDetached, bLaunchHidden, bLaunchReallyHidden, &ProcessID,
+		0, nullptr, WritePipe, ReadPipe);
+
+	int32 LoadingTotal = INDEX_NONE;
+	int32 CompileTotal = INDEX_NONE;
+	int32 AssetsWaitingToCompile = INDEX_NONE;
+	
+	const FRegexPattern LoadingTotalPattern(TEXT("Display:\\s([0-9]+)\\spackages to load from command line arguments"));
+	const FRegexPattern LoadingProgressPattern(TEXT("Display:\\sLoading\\s\\(([0-9]+)"));
+	const FRegexPattern CompileProgressPattern(TEXT("Display:\\sWaiting for\\s([0-9]+)"));
+	
+	while (FPlatformProcess::IsProcRunning(ProcessHandle))
+	{
+		if (SlowTask.ShouldCancel())
+		{
+			bOutCancelled = true;
+			FPlatformProcess::TerminateProc(ProcessHandle);
+			break;
+		}
+
+		const FString LogString = FPlatformProcess::ReadPipe(ReadPipe);
+		if (LogString.IsEmpty())
+		{
+			continue;
+		}
+		UE_LOG(LogDisplayClusterPreloadDerivedDataCache, Display, TEXT("Commandlet Output: %s"), *LogString);
+
+		// Total number of assets to load
+		if (LoadingTotal == INDEX_NONE)
+		{
+			FRegexMatcher LoadingTotalRegex(LoadingTotalPattern, *LogString);
+
+			while (LoadingTotalRegex.FindNext())
+			{
+				LoadingTotal = FCString::Atoi(*LoadingTotalRegex.GetCaptureGroup(1));
+			}
+		}
+
+		// Progress on asset loading
+		FRegexMatcher LoadingProgressRegex(LoadingProgressPattern, *LogString);
+		RegexParseForLoadingProgress(LoadingPackagesFormat, UnknownAmount, SlowTask, LoadingTotal, LoadingProgressRegex);
+
+		// Progress on asset compilation
+		FRegexMatcher CompileProgressRegex(CompileProgressPattern, *LogString);
+		RegexParseForCompilationProgress(CompilingAssetsFormat, UnknownAmount, SlowTask, CompileTotal, AssetsWaitingToCompile,
+		                                 CompileProgressRegex);
+
+		SlowTask.EnterProgressFrame(0);
+		FPlatformProcess::Sleep(0.1);
+	}
+	
+	FPlatformProcess::GetProcReturnCode(ProcessHandle, &OutResult);
+
+	FPlatformProcess::ClosePipe(ReadPipe, WritePipe);
+
+	CompleteCommandletAndShowNotification(OutResult, bOutCancelled, CurrentExecutableName, Arguments);
+}
+
+void FDisplayClusterPreloadDerivedDataCacheModule::CompleteCommandletAndShowNotification(
+	const int32 ResultCode, const bool bWasCancelled,
+	const FString& CurrentExecutableName, const FString& Arguments)
+{
+	const bool bFinishedWithFailures = ResultCode != 0;
+
+	FText NotificationText = LOCTEXT("ToastCommandletSuccess", "Commandlet Finished Without Error!");
+	SNotificationItem::ECompletionState CompletionState = SNotificationItem::CS_Success;
+	if (bWasCancelled)
+	{
+		UE_LOG(LogDisplayClusterPreloadDerivedDataCache, Display, TEXT("#### Commandlet Process Cancelled ####"));
+		NotificationText = LOCTEXT("ToastCommandletCancelled", "Commandlet Process Cancelled");
+		CompletionState = SNotificationItem::CS_None;
+	}
+	else if (bFinishedWithFailures)
+	{
+		UE_LOG(LogDisplayClusterPreloadDerivedDataCache, Error, TEXT("#### Commandlet Finished With Errors ####"));
+		UE_LOG(LogDisplayClusterPreloadDerivedDataCache, Error, TEXT("%s %s"), *CurrentExecutableName, *Arguments);
+		UE_LOG(LogDisplayClusterPreloadDerivedDataCache, Error, TEXT("Return Code: %i"), ResultCode);
+
+		NotificationText = FText::Format(LOCTEXT("ToastCommandletFailure", "Commandlet Finished With Error(s), Return Code: {0}"), FText::AsNumber(ResultCode));
+		CompletionState = SNotificationItem::CS_Fail;
+	}
+	else
+	{
+		UE_LOG(LogDisplayClusterPreloadDerivedDataCache, Display, TEXT("#### Commandlet Finished Without Error ####"));
+	}
+
+	FNotificationInfo NotificationInfo(NotificationText);
+	NotificationInfo.bFireAndForget = bWasCancelled;
+
+	// "close" button
+	FNotificationButtonInfo ButtonOK(LOCTEXT("ButtonOK", "OK"), FText(),
+		FSimpleDelegate::CreateLambda([this]()
+		{
+			if (Notification.IsValid())
+			{
+				Notification->ExpireAndFadeout();
+			}
+		}));
+
+	// Should be visible in all cases
+	ButtonOK.VisibilityOnFail = ButtonOK.VisibilityOnNone = ButtonOK.VisibilityOnPending = ButtonOK.VisibilityOnSuccess =EVisibility::Visible;
+	NotificationInfo.ButtonDetails.Add(ButtonOK);
+
+	if (bFinishedWithFailures)
+	{
+		NotificationInfo.Hyperlink = FSimpleDelegate::CreateLambda([]() { FGlobalTabmanager::Get()->TryInvokeTab(FName("OutputLog")); });
+		NotificationInfo.HyperlinkText = LOCTEXT("ShowOutputLogHyperlink", "Show Output Log");
+	}
+
+	// Close previous notification if it's still persistent
+	if (Notification.IsValid())
+	{
+		Notification->ExpireAndFadeout();
+	}
+
+	Notification = FSlateNotificationManager::Get().AddNotification( NotificationInfo );
+	if ( Notification.IsValid() )
+	{
+		Notification->SetCompletionState(CompletionState);
+	}
+}
+
+FString FDisplayClusterPreloadDerivedDataCacheModule::GetDdcCommandletParams() const
+{
+	return FString::Printf(TEXT("-run=DerivedDataCache %s -fill -DDC=CreateInstalledEnginePak"), *GetTargetPlatformParams());
+}
+
+FString FDisplayClusterPreloadDerivedDataCacheModule::GetTargetPlatformParams() const
+{
+	TArray<FString> SupportedPlatforms;
+	FProjectStatus ProjectStatus;
+	
+	if(IProjectManager::Get().QueryStatusForCurrentProject(ProjectStatus))
+	{
+		for (const FDataDrivenPlatformInfo* DDPI : FDataDrivenPlatformInfoRegistry::GetSortedPlatformInfos(EPlatformInfoType::TruePlatformsOnly))
+		{
+			if (ProjectStatus.IsTargetPlatformSupported(DDPI->IniPlatformName))
+			{
+				const FString PlatformName = DDPI->IniPlatformName.ToString();
+				SupportedPlatforms.Add(PlatformName);
+
+				// Not all platforms have an editor target but they may in the future.
+				// It doesn't hurt to add targets that don't exist as they are ignored.
+				SupportedPlatforms.Add(PlatformName + "Editor");
+			}
+		}
+	}
+
+	if (SupportedPlatforms.Num() > 0)
+	{
+		return FString::Printf(TEXT("-TargetPlatform=%s"), *FString::Join(SupportedPlatforms, TEXT("+")));
+	}
+	
+	return FString();
+}
+
+#undef LOCTEXT_NAMESPACE
+
+IMPLEMENT_MODULE(FDisplayClusterPreloadDerivedDataCacheModule, DisplayClusterPreloadDerivedDataCache)
diff --git a/Engine/Plugins/Runtime/nDisplay/Source/DisplayClusterPreloadDerivedDataCache/Public/DisplayClusterPreloadDerivedDataCacheLog.h b/Engine/Plugins/Runtime/nDisplay/Source/DisplayClusterPreloadDerivedDataCache/Public/DisplayClusterPreloadDerivedDataCacheLog.h
new file mode 100644
--- /dev/null
+++ b/Engine/Plugins/Runtime/nDisplay/Source/DisplayClusterPreloadDerivedDataCache/Public/DisplayClusterPreloadDerivedDataCacheLog.h
@@ -0,0 +1,7 @@
+// Copyright Epic Games, Inc. All Rights Reserved.
+
+#pragma once
+
+#include "Logging/LogMacros.h"
+
+DECLARE_LOG_CATEGORY_CLASS(LogDisplayClusterPreloadDerivedDataCache, Log, All);
diff --git a/Engine/Plugins/Runtime/nDisplay/Source/DisplayClusterPreloadDerivedDataCache/Public/DisplayClusterPreloadDerivedDataCacheModule.h b/Engine/Plugins/Runtime/nDisplay/Source/DisplayClusterPreloadDerivedDataCache/Public/DisplayClusterPreloadDerivedDataCacheModule.h
new file mode 100644
--- /dev/null
+++ b/Engine/Plugins/Runtime/nDisplay/Source/DisplayClusterPreloadDerivedDataCache/Public/DisplayClusterPreloadDerivedDataCacheModule.h
@@ -0,0 +1,40 @@
+// Copyright Epic Games, Inc. All Rights Reserved.
+
+#pragma once
+
+#include "Containers/UnrealString.h"
+#include "Modules/ModuleInterface.h"
+#include "Templates/SharedPointer.h"
+
+class FDisplayClusterPreloadDerivedDataCacheModule : public IModuleInterface
+{
+public:
+	
+	static FDisplayClusterPreloadDerivedDataCacheModule& Get();
+
+	//~ Begin IModuleInterface Interface
+	virtual void StartupModule() override;
+	virtual void ShutdownModule() override;
+	//~ End IModuleInterface Interface
+
+	void LaunchAndCommunicateWithProcess();
+
+protected:
+	
+	void RegexParseForLoadingProgress(const FText& LoadingPackagesFormat, const FText& UnknownAmount, struct FScopedSlowTask& SlowTask,
+									  int32& LoadingTotal, class FRegexMatcher& LoadingProgressRegex);
+	void RegexParseForCompilationProgress(const FText& CompilingAssetsFormat, const FText& UnknownAmount, FScopedSlowTask& SlowTask,
+	                                      int32& CompileTotal, int32& AssetsWaitingToCompile,
+	                                      FRegexMatcher& CompileProgressRegex);
+	void CompleteCommandletAndShowNotification(const int32 ResultCode, const bool bWasCancelled,
+	const FString& CurrentExecutableName, const FString& Arguments);
+
+	FString GetDdcCommandletParams() const;
+	FString GetTargetPlatformParams() const;
+
+	TSharedPtr<class SNotificationItem> Notification;
+
+private:
+	
+	void OnFEngineLoopInitComplete();
+};
diff --git a/Engine/Plugins/Runtime/nDisplay/nDisplay.uplugin b/Engine/Plugins/Runtime/nDisplay/nDisplay.uplugin
--- a/Engine/Plugins/Runtime/nDisplay/nDisplay.uplugin
+++ b/Engine/Plugins/Runtime/nDisplay/nDisplay.uplugin
@@ -28,6 +28,15 @@
 				"Linux"
 			]
 		},
+		{
+			"Name": "DisplayClusterPreloadDerivedDataCache",
+			"Type": "Editor",
+			"LoadingPhase": "Default",
+			"PlatformAllowList": [
+				"Win64",
+				"Linux"
+			]
+		},
 		{
 			"Name": "DisplayClusterEditor",
 			"Type": "Editor",
