- Fixing crash in ndisplay capture sync policy ini

diff --git a/Engine/Plugins/Runtime/nDisplay/Source/DisplayClusterMedia/Private/Synchronization/SyncPolicies/DisplayClusterMediaOutputSynchronizationPolicyEthernetBarrierBase.cpp b/Engine/Plugins/Runtime/nDisplay/Source/DisplayClusterMedia/Private/Synchronization/SyncPolicies/DisplayClusterMediaOutputSynchronizationPolicyEthernetBarrierBase.cpp
--- a/Engine/Plugins/Runtime/nDisplay/Source/DisplayClusterMedia/Private/Synchronization/SyncPolicies/DisplayClusterMediaOutputSynchronizationPolicyEthernetBarrierBase.cpp
+++ b/Engine/Plugins/Runtime/nDisplay/Source/DisplayClusterMedia/Private/Synchronization/SyncPolicies/DisplayClusterMediaOutputSynchronizationPolicyEthernetBarrierBase.cpp
@@ -255,16 +255,19 @@ void UDisplayClusterMediaOutputSynchronizationPolicyEthernetBarrierBase::Generat
 					for (const FDisplayClusterConfigurationMediaOutputGroup& MediaOutputGroup : MediaSettings.MediaOutputGroups)
 					{
 						// Pick the same sync policy only
-						if (MediaOutputGroup.OutputSyncPolicy->GetClass() == GetClass())
+						if (MediaOutputGroup.MediaOutput && MediaOutputGroup.OutputSyncPolicy)
 						{
-							for (const FString& NodeId : MediaOutputGroup.ClusterNodes.ItemNames)
+							if (MediaOutputGroup.OutputSyncPolicy->GetClass() == GetClass())
 							{
-								const FString ICVFXCaptureId = DisplayClusterMediaHelpers::MediaId::GenerateMediaId(
-									DisplayClusterMediaHelpers::MediaId::EMediaDeviceType::Output,
-									DisplayClusterMediaHelpers::MediaId::EMediaOwnerType::ICVFXCamera,
-									NodeId, RootActor->GetName(), ICVFXCameraComponent->GetName());
-
-								OutMarkers.Add(ICVFXCaptureId);
+								for (const FString& NodeId : MediaOutputGroup.ClusterNodes.ItemNames)
+								{
+									const FString ICVFXCaptureId = DisplayClusterMediaHelpers::MediaId::GenerateMediaId(
+										DisplayClusterMediaHelpers::MediaId::EMediaDeviceType::Output,
+										DisplayClusterMediaHelpers::MediaId::EMediaOwnerType::ICVFXCamera,
+										NodeId, RootActor->GetName(), ICVFXCameraComponent->GetName());
+
+									OutMarkers.Add(ICVFXCaptureId);
+								}
 							}
 						}
 					}
