nDisplay: Added optimization of frustum calculation for viewports.
(This is one of the CLs for refactoring the nDisplay Preview)

The frustum calculation is called from the IStereoRendering interface by the CalculateStereoViewOffset() and GetStereoProjectionMatrix() functions. These functions are called several times from LocalPlayer. Since calculating the frustum for some projection policies can take a long time, we have added saving the calculated values to the current frame for optimization.

diff --git a/Engine/Plugins/Runtime/nDisplay/Source/DisplayCluster/Private/Game/EngineClasses/Scene/DisplayClusterRootActorEditor.cpp b/Engine/Plugins/Runtime/nDisplay/Source/DisplayCluster/Private/Game/EngineClasses/Scene/DisplayClusterRootActorEditor.cpp
--- a/Engine/Plugins/Runtime/nDisplay/Source/DisplayCluster/Private/Game/EngineClasses/Scene/DisplayClusterRootActorEditor.cpp
+++ b/Engine/Plugins/Runtime/nDisplay/Source/DisplayCluster/Private/Game/EngineClasses/Scene/DisplayClusterRootActorEditor.cpp
@@ -643,7 +643,7 @@ void ADisplayClusterRootActor::ImplRenderPreviewFrustums_Editor()
 		bool bIsValidViewportContext = false;
 
 		// Preview rendered only in mono
-		if (Contexts.Num() == 1 && Contexts[0].bIsValidProjectionMatrix && Contexts[0].bIsValidViewLocation && Contexts[0].bIsValidViewRotation)
+		if (Contexts.Num() == 1 && EnumHasAllFlags(Contexts[0].ContextState, EDisplayClusterViewportContextState::HasCalculatedProjectionMatrix | EDisplayClusterViewportContextState::HasCalculatedViewPoint))
 		{
 			FrustumPreviewViewportContext.ProjectionMatrix = Contexts[0].ProjectionMatrix;
 			FrustumPreviewViewportContext.ViewLocation = Contexts[0].ViewLocation;
diff --git a/Engine/Plugins/Runtime/nDisplay/Source/DisplayCluster/Private/Render/Viewport/DisplayClusterViewport_Math.cpp b/Engine/Plugins/Runtime/nDisplay/Source/DisplayCluster/Private/Render/Viewport/DisplayClusterViewport_Math.cpp
--- a/Engine/Plugins/Runtime/nDisplay/Source/DisplayCluster/Private/Render/Viewport/DisplayClusterViewport_Math.cpp
+++ b/Engine/Plugins/Runtime/nDisplay/Source/DisplayCluster/Private/Render/Viewport/DisplayClusterViewport_Math.cpp
@@ -27,18 +27,42 @@
 
 bool FDisplayClusterViewport::CalculateView(const uint32 InContextNum, FVector& InOutViewLocation, FRotator& InOutViewRotation, const FVector& ViewOffset, const float WorldToMeters, const float NCP, const float FCP)
 {
-	if (ProjectionPolicy.IsValid() && ProjectionPolicy->CalculateView(this, InContextNum, InOutViewLocation, InOutViewRotation, ViewOffset, WorldToMeters, NCP, FCP))
+	if (Contexts.IsValidIndex(InContextNum))
 	{
-		// Store the view location/rotation
-		Contexts[InContextNum].ViewLocation = InOutViewLocation;
-		Contexts[InContextNum].bIsValidViewLocation = true;
-
-		Contexts[InContextNum].ViewRotation = InOutViewRotation;
-		Contexts[InContextNum].bIsValidViewRotation = true;
-
-		Contexts[InContextNum].WorldToMeters = WorldToMeters;
-
-		return true;
+		if (!EnumHasAnyFlags(Contexts[InContextNum].ContextState, EDisplayClusterViewportContextState::InvalidViewPoint))
+		{
+			// The function can be called several times per frame.
+			// Each time it must return the same values. For optimization purposes, after the first call this function
+			// stores the result in the context variables 'ViewLocation' and 'ViewRotation'.
+			// Finally, raises this flag for subsequent calls in the current frame.
+
+			if (EnumHasAnyFlags(Contexts[InContextNum].ContextState, EDisplayClusterViewportContextState::HasCalculatedViewPoint))
+			{
+				// Use calculated values
+				// Since this function can be called several times from LocalPlayer.cpp, the cached values are used on repeated calls.
+				// This should give a performance boost for 'mesh', 'mpcdi' projections with a large number of vertices in the geometry or large warp texture size.
+				InOutViewLocation = Contexts[InContextNum].ViewLocation;
+				InOutViewRotation = Contexts[InContextNum].ViewRotation;
+
+				return true;
+			}
+			else if (ProjectionPolicy.IsValid() && ProjectionPolicy->CalculateView(this, InContextNum, InOutViewLocation, InOutViewRotation, ViewOffset, WorldToMeters, NCP, FCP))
+			{
+				Contexts[InContextNum].WorldToMeters = WorldToMeters;
+
+				// Save the calculated values and update the state of the context
+				Contexts[InContextNum].ViewLocation = InOutViewLocation;
+				Contexts[InContextNum].ViewRotation = InOutViewRotation;
+				EnumAddFlags(Contexts[InContextNum].ContextState, EDisplayClusterViewportContextState::HasCalculatedViewPoint);
+
+				return true;
+			}
+			else
+			{
+				// ProjectionPolicy->CalculateView() returns false, this view is invalid
+				EnumAddFlags(Contexts[InContextNum].ContextState, EDisplayClusterViewportContextState::InvalidViewPoint);
+			}
+		}
 	}
 
 	return false;
@@ -46,18 +70,50 @@ bool FDisplayClusterViewport::CalculateView(const uint32 InContextNum, FVector&
 
 bool FDisplayClusterViewport::GetProjectionMatrix(const uint32 InContextNum, FMatrix& OutPrjMatrix)
 {
-	if (ProjectionPolicy.IsValid() && ProjectionPolicy->GetProjectionMatrix(this, InContextNum, OutPrjMatrix))
+	if (Contexts.IsValidIndex(InContextNum))
 	{
-		Contexts[InContextNum].ProjectionMatrix = OutPrjMatrix;
-		Contexts[InContextNum].bIsValidProjectionMatrix = true;
-
-		if (OverscanRendering.IsEnabled())
+		if (!EnumHasAnyFlags(Contexts[InContextNum].ContextState, EDisplayClusterViewportContextState::InvalidProjectionMatrix))
 		{
-			// use overscan proj matrix for rendering
-			OutPrjMatrix = Contexts[InContextNum].OverscanProjectionMatrix;
+			// The function can also be called several times per frame.
+			// stores the result in the context variables 'ProjectionMatrix' and 'OverscanProjectionMatrix'.
+			// Finally, raises this flag for subsequent calls in the current frame.
+
+			if (EnumHasAnyFlags(Contexts[InContextNum].ContextState, EDisplayClusterViewportContextState::HasCalculatedProjectionMatrix))
+			{
+				// use already calculated values
+				if (EnumHasAnyFlags(Contexts[InContextNum].ContextState, EDisplayClusterViewportContextState::HasCalculatedOverscanProjectionMatrix))
+				{
+					// use overscan proj matrix for rendering
+					OutPrjMatrix = Contexts[InContextNum].OverscanProjectionMatrix;
+				}
+				else
+				{
+					OutPrjMatrix = Contexts[InContextNum].ProjectionMatrix;
+				}
+
+				return true;
+			}
+			else if (ProjectionPolicy.IsValid() && ProjectionPolicy->GetProjectionMatrix(this, InContextNum, OutPrjMatrix))
+			{
+				// Save the calculated values and update the state of the context
+				Contexts[InContextNum].ProjectionMatrix = OutPrjMatrix;
+				EnumAddFlags(Contexts[InContextNum].ContextState, EDisplayClusterViewportContextState::HasCalculatedProjectionMatrix);
+
+				if (OverscanRendering.IsEnabled())
+				{
+					// use overscan proj matrix for rendering
+					OutPrjMatrix = Contexts[InContextNum].OverscanProjectionMatrix;
+					EnumAddFlags(Contexts[InContextNum].ContextState, EDisplayClusterViewportContextState::HasCalculatedOverscanProjectionMatrix);
+				}
+
+				return true;
+			}
+			else
+			{
+				// ProjectionPolicy->GetProjectionMatrix() returns false, this projection matrix is invalid
+				EnumAddFlags(Contexts[InContextNum].ContextState, EDisplayClusterViewportContextState::InvalidProjectionMatrix);
+			}
 		}
-
-		return true;
 	}
 
 	return false;
diff --git a/Engine/Plugins/Runtime/nDisplay/Source/DisplayCluster/Public/Render/Viewport/Containers/DisplayClusterViewport_Context.h b/Engine/Plugins/Runtime/nDisplay/Source/DisplayCluster/Public/Render/Viewport/Containers/DisplayClusterViewport_Context.h
--- a/Engine/Plugins/Runtime/nDisplay/Source/DisplayCluster/Public/Render/Viewport/Containers/DisplayClusterViewport_Context.h
+++ b/Engine/Plugins/Runtime/nDisplay/Source/DisplayCluster/Public/Render/Viewport/Containers/DisplayClusterViewport_Context.h
@@ -6,6 +6,36 @@
 #include "StereoRendering.h"
 #include "ShowFlags.h"
 
+/**
+ * Special flags for nDisplay Viewport context
+ */
+enum class EDisplayClusterViewportContextState : uint8
+{
+	None = 0,
+
+	// The FDisplayClusterViewport::CalculateView() function can be called several times per frame.
+	// Each time it must return the same values. For optimization purposes, after the first call this function
+	// stores the result in the context variables 'ViewLocation' and 'ViewRotation'.
+	// Finally, raises this flag for subsequent calls in the current frame.
+	HasCalculatedViewPoint = 1 << 0,
+
+	// Viewpoint is not valid for this viewport (cannot be calculated)
+	InvalidViewPoint = 1 << 1,
+
+	// The FDisplayClusterViewport::GetProjectionMatrix() function can also be called several times per frame.
+	// stores the result in the context variables 'ProjectionMatrix' and 'OverscanProjectionMatrix'.
+	// Finally, raises this flag for subsequent calls in the current frame.
+	HasCalculatedProjectionMatrix = 1 << 2,
+	HasCalculatedOverscanProjectionMatrix = 1 << 3,
+
+	// The projection matrix is not valid (cannot be calculated)
+	InvalidProjectionMatrix = 1 << 4,
+};
+ENUM_CLASS_FLAGS(EDisplayClusterViewportContextState);
+
+/**
+ * Viewport context with cahched data and states
+ */
 class FDisplayClusterViewport_Context
 {
 public:
@@ -21,16 +51,15 @@ public:
 	EStereoscopicPass StereoscopicPass;
 	int32             StereoViewIndex;
 
-	// Camera location and orientation
-	FVector  ViewLocation = FVector::ZeroVector;
-	bool bIsValidViewLocation = false;
+	// Special flags for this context
+	EDisplayClusterViewportContextState ContextState = EDisplayClusterViewportContextState::None;
 
+	// ViewPoint: Camera location and orientation
+	FVector  ViewLocation = FVector::ZeroVector;
 	FRotator ViewRotation = FRotator::ZeroRotator;
-	bool bIsValidViewRotation = false;
 
 	// Projection Matrix
 	FMatrix ProjectionMatrix = FMatrix::Identity;
-	bool bIsValidProjectionMatrix = false;
 
 	// Overscan Projection Matrix (internal use)
 	FMatrix OverscanProjectionMatrix = FMatrix::Identity;
