nDisplay: Fix for UDisplayClusterGameEngine::Init corrupting command line arguments during parsing.

FParse::QuotedString was inappropriate to use here, and led to common sequences such as \n and \t, or Switchboard randomly-generated octal escapes, in the dc_cfg file path getting mishandled.

This issue would only occur if the quoted argument (e.g. Windows username) also contained a space, because otherwise FCommandLine would strip the quotes, and FParse::Token would be used instead. FParse::Token also handles quoted strings, and so it appears the problematic branch was not actually necessary at all.

diff --git a/Engine/Plugins/Runtime/nDisplay/Source/DisplayCluster/Private/Game/EngineClasses/Basics/DisplayClusterGameEngine.cpp b/Engine/Plugins/Runtime/nDisplay/Source/DisplayCluster/Private/Game/EngineClasses/Basics/DisplayClusterGameEngine.cpp
--- a/Engine/Plugins/Runtime/nDisplay/Source/DisplayCluster/Private/Game/EngineClasses/Basics/DisplayClusterGameEngine.cpp
+++ b/Engine/Plugins/Runtime/nDisplay/Source/DisplayCluster/Private/Game/EngineClasses/Basics/DisplayClusterGameEngine.cpp
@@ -86,12 +86,6 @@ void UDisplayClusterGameEngine::Init(class IEngineLoop* InEngineLoop)
 			}
 
 			const TCHAR* TagValue = &CommandLine[TagPos + Tag.Len()];
-
-			if (*TagValue == TEXT('"'))
-			{
-				return FParse::QuotedString(TagValue, OutArgVal);
-			}
-
 			return FParse::Token(TagValue, OutArgVal, false);
 		};
 
