[nDisplay]: Added the ability to override ViewPoint.
- Added a new ViewExtension used to override ViewPoint
- Added new public functions to IDisplayClusterViewportManager and IDisplayClusterViewport
- Reworked the implementation of viewpoint-related functions in DCRenderDevice, DCPreview and DCMovieRenderPipeline.

diff --git a/Engine/Plugins/Runtime/nDisplay/Source/DisplayCluster/Private/Game/EngineClasses/Basics/DisplayClusterViewportClient.cpp b/Engine/Plugins/Runtime/nDisplay/Source/DisplayCluster/Private/Game/EngineClasses/Basics/DisplayClusterViewportClient.cpp
--- a/Engine/Plugins/Runtime/nDisplay/Source/DisplayCluster/Private/Game/EngineClasses/Basics/DisplayClusterViewportClient.cpp
+++ b/Engine/Plugins/Runtime/nDisplay/Source/DisplayCluster/Private/Game/EngineClasses/Basics/DisplayClusterViewportClient.cpp
@@ -551,7 +551,7 @@ void UDisplayClusterViewportClient::Draw(FViewport* InViewport, FCanvas* SceneCa
 				// Calculate the player's view information.
 				FVector		ViewLocation;
 				FRotator	ViewRotation;
-				FSceneView* View = LocalPlayer->CalcSceneView(&ViewFamily, ViewLocation, ViewRotation, InViewport, nullptr, ViewportContext.StereoViewIndex);
+				FSceneView* View = RenderFrame.ViewportManager->CalcSceneView(LocalPlayer, &ViewFamily, ViewLocation, ViewRotation, InViewport, nullptr, ViewportContext.StereoViewIndex);
 
 				if (View && !DCView.ShouldRenderSceneView())
 				{
diff --git a/Engine/Plugins/Runtime/nDisplay/Source/DisplayCluster/Private/Render/Device/DisplayClusterDeviceBase.cpp b/Engine/Plugins/Runtime/nDisplay/Source/DisplayCluster/Private/Render/Device/DisplayClusterDeviceBase.cpp
--- a/Engine/Plugins/Runtime/nDisplay/Source/DisplayCluster/Private/Render/Device/DisplayClusterDeviceBase.cpp
+++ b/Engine/Plugins/Runtime/nDisplay/Source/DisplayCluster/Private/Render/Device/DisplayClusterDeviceBase.cpp
@@ -239,116 +239,19 @@ void FDisplayClusterDeviceBase::CalculateStereoViewOffset(const int32 ViewIndex,
 	check(IsInGameThread());
 	check(WorldToMeters > 0.f);
 
-	if (ViewportManager == nullptr || ViewportManager->IsSceneOpened() == false)
-	{
-		return;
-	}
-	// ViewIndex == eSSE_MONOSCOPIC(-1) is a special case called for ISR culling math.
-	// Since nDisplay is not ISR compatible, we ignore this request. This won't be neccessary once
-	// we stop using nDisplay as a stereoscopic rendering device (IStereoRendering).
-	else if (ViewIndex < 0)
-	{
-		return;
-	}
-
 	uint32 ViewportContextNum = 0;
-	IDisplayClusterViewport* ViewportPtr = ViewportManager->FindViewport(ViewIndex, &ViewportContextNum);
+	IDisplayClusterViewport* ViewportPtr = ViewportManager ? ViewportManager->FindViewport(ViewIndex, &ViewportContextNum) : nullptr;
 	if (ViewportPtr == nullptr)
 	{
-		UE_LOG(LogDisplayClusterRender, Warning, TEXT("Viewport StereoViewIndex='%i' not found"), ViewIndex);
-		return;
-	}
-
-	if (!ViewportPtr->GetProjectionPolicy().IsValid())
-	{
-		// ignore viewports with uninitialized prj policy
 		return;
 	}
 
-	// Get root actor from viewport
-	ADisplayClusterRootActor* const RootActor = ViewportPtr->GetOwner().GetRootActor();
-	if (!RootActor)
-	{
-		UE_LOG(LogDisplayClusterRender, Warning, TEXT("No root actor found in game manager"));
-		return;
-	}
-
-	const TArray<FDisplayClusterViewport_Context>& ViewportContexts = ViewportPtr->GetContexts();
-	const FDisplayClusterViewport_Context& ViewportContext = ViewportContexts[ViewportContextNum];
-
-	UE_LOG(LogDisplayClusterRender, VeryVerbose, TEXT("OLD ViewLoc: %s, ViewRot: %s"), *ViewLocation.ToString(), *ViewRotation.ToString());
-	UE_LOG(LogDisplayClusterRender, VeryVerbose, TEXT("WorldToMeters: %f"), WorldToMeters);
-
-
-	// Get camera ID assigned to the viewport
-	const FString& CameraId = ViewportPtr->GetRenderSettings().CameraId;
-
-	// Get camera component assigned to the viewport (or default camera if nothing assigned)
-	UDisplayClusterCameraComponent* const ViewCamera = (CameraId.IsEmpty() ?
-		RootActor->GetDefaultCamera() :
-		RootActor->GetComponentByName<UDisplayClusterCameraComponent>(CameraId));
-
-	if (!ViewCamera)
-	{
-		UE_LOG(LogDisplayClusterRender, Warning, TEXT("No camera found for viewport '%s'"), *ViewportPtr->GetId());
-		return;
-	}
-
-	if (CameraId.Len() > 0)
-	{
-		UE_LOG(LogDisplayClusterRender, Verbose, TEXT("Viewport '%s' has assigned camera '%s'"), *ViewportPtr->GetId(), *CameraId);
-	}
-
-	// Get the actual camera settings
-	const float CfgEyeDist = ViewCamera->GetInterpupillaryDistance();
-	const bool  CfgEyeSwap = ViewCamera->GetSwapEyes();
-	const float CfgNCP     = GNearClippingPlane;
-	const EDisplayClusterEyeStereoOffset CfgEyeOffset = ViewCamera->GetStereoOffset();
-
-	// Calculate eye offset considering the world scale
-	const float EyeOffset = CfgEyeDist / 2.f;
-	const float EyeOffsetValues[] = { -EyeOffset, 0.f, EyeOffset };
-
-	// Decode current eye type
-	EDisplayClusterEyeType EyeType = EDisplayClusterEyeType::Mono;
-	if(ViewportContexts.Num() > 1)
-	{
-		// Support stereo:
-		EyeType = (ViewportContextNum == 0)? EDisplayClusterEyeType::StereoLeft : EDisplayClusterEyeType::StereoRight;
-	}
-
-	const int32 EyeIndex = (int32)EyeType;
-
-	float PassOffset = 0.f;
-	float PassOffsetSwap = 0.f;
-
-	if (EyeType == EDisplayClusterEyeType::Mono)
-	{
-		// For monoscopic camera let's check if the "force offset" feature is used
-		// * Force left (-1) ==> 0 left eye
-		// * Force right (1) ==> 2 right eye
-		// * Default (0) ==> 1 mono
-		const int32 EyeOffsetIdx = 
-			(CfgEyeOffset == EDisplayClusterEyeStereoOffset::None ? 0 :
-			(CfgEyeOffset == EDisplayClusterEyeStereoOffset::Left ? -1 : 1));
-
-		PassOffset = EyeOffsetValues[EyeOffsetIdx + 1];
-		// Eye swap is not available for monoscopic so just save the value
-		PassOffsetSwap = PassOffset;
-	}
-	else
-	{
-		// For stereo camera we can only swap eyes if required (no "force offset" allowed)
-		PassOffset = EyeOffsetValues[EyeIndex];
-		PassOffsetSwap = (CfgEyeSwap ? -PassOffset : PassOffset);
-	}
+	// The camera position has already been determined from the SetupViewPoint() function
+	// Obtaining the stereo eye offset for a given viewport
+	const float PassOffsetSwap = ViewportPtr->GetStereoEyeOffsetDistance(ViewportContextNum);
 
 	FVector ViewOffset = FVector::ZeroVector;
-	if (ViewCamera)
 	{
-		// View base location
-		ViewLocation = ViewCamera->GetComponentLocation();
-		ViewRotation = ViewCamera->GetComponentRotation();
 		// Apply computed offset to the view location
 		const FQuat EyeQuat = ViewRotation.Quaternion();
 		ViewOffset = EyeQuat.RotateVector(FVector(0.0f, PassOffsetSwap, 0.0f));
@@ -356,6 +259,7 @@ void FDisplayClusterDeviceBase::CalculateStereoViewOffset(const int32 ViewIndex,
 	}
 
 	// Perform view calculations on a policy side
+	const float CfgNCP = GNearClippingPlane;
 	if (ViewportPtr->CalculateView(ViewportContextNum, ViewLocation, ViewRotation, ViewOffset, WorldToMeters, CfgNCP, CfgNCP) == false)
 	{
 #if WITH_EDITOR
diff --git a/Engine/Plugins/Runtime/nDisplay/Source/DisplayCluster/Private/Render/Device/DisplayClusterDeviceBase.h b/Engine/Plugins/Runtime/nDisplay/Source/DisplayCluster/Private/Render/Device/DisplayClusterDeviceBase.h
--- a/Engine/Plugins/Runtime/nDisplay/Source/DisplayCluster/Private/Render/Device/DisplayClusterDeviceBase.h
+++ b/Engine/Plugins/Runtime/nDisplay/Source/DisplayCluster/Private/Render/Device/DisplayClusterDeviceBase.h
@@ -82,18 +82,6 @@ protected:
 	virtual void CalculateRenderTargetSize(const class FViewport& Viewport, uint32& InOutSizeX, uint32& InOutSizeY) override;
 	virtual bool NeedReAllocateViewportRenderTarget(const class FViewport& Viewport) override;
 
-protected:
-	//////////////////////////////////////////////////////////////////////////////////////////////
-	// FDisplayClusterDeviceBase
-	//////////////////////////////////////////////////////////////////////////////////////////////
-	enum EDisplayClusterEyeType
-	{
-		StereoLeft  = 0,
-		Mono        = 1,
-		StereoRight = 2,
-		COUNT
-	};
-
 protected:
 	// Factory method to instantiate an output presentation class implementation
 	virtual FDisplayClusterPresentationBase* CreatePresentationObject(FViewport* const Viewport, TSharedPtr<IDisplayClusterRenderSyncPolicy>& SyncPolicy) = 0;
diff --git a/Engine/Plugins/Runtime/nDisplay/Source/DisplayCluster/Private/Render/Viewport/DisplayClusterViewport.h b/Engine/Plugins/Runtime/nDisplay/Source/DisplayCluster/Private/Render/Viewport/DisplayClusterViewport.h
--- a/Engine/Plugins/Runtime/nDisplay/Source/DisplayCluster/Private/Render/Viewport/DisplayClusterViewport.h
+++ b/Engine/Plugins/Runtime/nDisplay/Source/DisplayCluster/Private/Render/Viewport/DisplayClusterViewport.h
@@ -88,6 +88,9 @@ public:
 	virtual bool    CalculateView(const uint32 InContextNum, FVector& InOutViewLocation, FRotator& InOutViewRotation, const FVector& ViewOffset, const float WorldToMeters, const float NCP, const float FCP) override;
 	virtual bool    GetProjectionMatrix(const uint32 InContextNum, FMatrix& OutPrjMatrix)  override;
 
+	virtual bool SetupViewPoint(FMinimalViewInfo& InOutViewInfo) override;
+	virtual float GetStereoEyeOffsetDistance(const uint32 InContextNum) override;
+
 	virtual const FDisplayClusterViewport_RenderSettingsICVFX& GetRenderSettingsICVFX() const override
 	{
 		check(IsInGameThread());
@@ -171,6 +174,8 @@ public:
 	bool HandleStartScene();
 	void HandleEndScene();
 
+	class UDisplayClusterCameraComponent* GetViewPointCamera() const;
+
 	void AddReferencedObjects(FReferenceCollector& Collector);
 
 	void ResetRuntimeParameters()
diff --git a/Engine/Plugins/Runtime/nDisplay/Source/DisplayCluster/Private/Render/Viewport/DisplayClusterViewportManager.cpp b/Engine/Plugins/Runtime/nDisplay/Source/DisplayCluster/Private/Render/Viewport/DisplayClusterViewportManager.cpp
--- a/Engine/Plugins/Runtime/nDisplay/Source/DisplayCluster/Private/Render/Viewport/DisplayClusterViewportManager.cpp
+++ b/Engine/Plugins/Runtime/nDisplay/Source/DisplayCluster/Private/Render/Viewport/DisplayClusterViewportManager.cpp
@@ -3,6 +3,7 @@
 #include "Render/Viewport/DisplayClusterViewportManager.h"
 #include "Render/Viewport/DisplayClusterViewportManagerProxy.h"
 #include "Render/Viewport/DisplayClusterViewportManagerViewExtension.h"
+#include "Render/Viewport/DisplayClusterViewportManagerViewPointExtension.h"
 
 #include "IDisplayCluster.h"
 #include "Render/IDisplayClusterRenderManager.h"
@@ -31,6 +32,7 @@
 #include "Misc/DisplayClusterLog.h"
 #include "LegacyScreenPercentageDriver.h"
 
+#include "Engine/LocalPlayer.h"
 #include "Engine/Console.h"
 
 int32 GDisplayClusterLightcardsAllowNanite = 0;
@@ -101,6 +103,12 @@ FDisplayClusterViewportManager::~FDisplayClusterViewportManager()
 		ViewportManagerViewExtension.Reset();
 	}
 
+	if (ViewportManagerViewPointExtension.IsValid())
+	{
+		ViewportManagerViewPointExtension->Release();
+		ViewportManagerViewPointExtension.Reset();
+	}
+
 	Configuration.Reset();
 
 	if (ViewportManagerProxy.IsValid())
@@ -510,6 +518,12 @@ bool FDisplayClusterViewportManager::BeginNewFrame(FViewport* InViewport, UWorld
 		ViewportManagerViewExtension = FSceneViewExtensions::NewExtension<FDisplayClusterViewportManagerViewExtension>(this);
 	}
 
+	// Create DC ViewPointExtension to handle special features
+	if (!ViewportManagerViewPointExtension.IsValid())
+	{
+		ViewportManagerViewPointExtension = FSceneViewExtensions::NewExtension<FDisplayClusterViewportManagerViewPointExtension>(this);
+	}
+
 	// Before new frame
 	if (PostProcessManager.IsValid())
 	{
@@ -986,6 +1000,14 @@ IDisplayClusterViewport* FDisplayClusterViewportManager::FindViewport(const int3
 {
 	check(IsInGameThread());
 
+	// ViewIndex == eSSE_MONOSCOPIC(-1) is a special case called for ISR culling math.
+	// Since nDisplay is not ISR compatible, we ignore this request. This won't be neccessary once
+	// we stop using nDisplay as a stereoscopic rendering device (IStereoRendering).
+	if (ViewIndex < 0)
+	{
+		return nullptr;
+	}
+
 	for (FDisplayClusterViewport* Viewport : Viewports)
 	{
 		if (Viewport && Viewport->FindContext(ViewIndex, OutContextNum))
@@ -1073,3 +1095,22 @@ void FDisplayClusterViewportManager::AddReferencedObjects(FReferenceCollector& C
 		}
 	}
 }
+
+FSceneView* FDisplayClusterViewportManager::CalcSceneView(ULocalPlayer* LocalPlayer, FSceneViewFamily* ViewFamily, FVector& OutViewLocation, FRotator& OutViewRotation, FViewport* InViewport, FViewElementDrawer* InViewDrawer, int32 InStereoViewIndex)
+{
+	check(ViewportManagerViewPointExtension.IsValid());
+	check(LocalPlayer);
+
+	// Set current view index for this VE:
+	// This function must be used because LocalPlayer::GetViewPoint() calls the ISceneViewExtension::SetupViewPoint() function from this VE.
+	// And at this point the VE must know the current StereoViewIndex value in order to understand which viewport will be used for this SetupViewPoint() call.
+	ViewportManagerViewPointExtension->SetCurrentStereoViewIndex(InStereoViewIndex);
+
+	// Calculate view inside LocalPlayer: VE used to override view point for this StereoViewIndex
+	FSceneView* View = LocalPlayer->CalcSceneView(ViewFamily, OutViewLocation, OutViewRotation, InViewport, InViewDrawer, InStereoViewIndex);
+
+	// Set INDEX_NONE when this VE no longer needs to be used.
+	ViewportManagerViewPointExtension->SetCurrentStereoViewIndex(INDEX_NONE);
+
+	return View;
+}
diff --git a/Engine/Plugins/Runtime/nDisplay/Source/DisplayCluster/Private/Render/Viewport/DisplayClusterViewportManager.h b/Engine/Plugins/Runtime/nDisplay/Source/DisplayCluster/Private/Render/Viewport/DisplayClusterViewportManager.h
--- a/Engine/Plugins/Runtime/nDisplay/Source/DisplayCluster/Private/Render/Viewport/DisplayClusterViewportManager.h
+++ b/Engine/Plugins/Runtime/nDisplay/Source/DisplayCluster/Private/Render/Viewport/DisplayClusterViewportManager.h
@@ -98,6 +98,8 @@ public:
 
 	virtual void MarkComponentGeometryDirty(const FName InComponentName = NAME_None) override;
 
+	virtual class FSceneView* CalcSceneView(class ULocalPlayer* LocalPlayer, class FSceneViewFamily* ViewFamily, FVector& OutViewLocation, FRotator& OutViewRotation, class FViewport* Viewport, class FViewElementDrawer* ViewDrawer, int32 StereoViewIndex) override;
+
 	virtual void AddReferencedObjects(FReferenceCollector& Collector) override;
 
 	// internal use only
@@ -197,4 +199,7 @@ private:
 
 	// Handle special features
 	TSharedPtr<class FDisplayClusterViewportManagerViewExtension, ESPMode::ThreadSafe> ViewportManagerViewExtension;
+
+	// This VE is used for LocalPlayer::GetViewPoint(). It calls the ISceneViewExtension::SetupViewPoint() function of this VE.
+	TSharedPtr<class FDisplayClusterViewportManagerViewPointExtension, ESPMode::ThreadSafe> ViewportManagerViewPointExtension;
 };
diff --git a/Engine/Plugins/Runtime/nDisplay/Source/DisplayCluster/Private/Render/Viewport/DisplayClusterViewportManagerProxy.cpp b/Engine/Plugins/Runtime/nDisplay/Source/DisplayCluster/Private/Render/Viewport/DisplayClusterViewportManagerProxy.cpp
--- a/Engine/Plugins/Runtime/nDisplay/Source/DisplayCluster/Private/Render/Viewport/DisplayClusterViewportManagerProxy.cpp
+++ b/Engine/Plugins/Runtime/nDisplay/Source/DisplayCluster/Private/Render/Viewport/DisplayClusterViewportManagerProxy.cpp
@@ -710,6 +710,14 @@ FDisplayClusterViewportProxy* FDisplayClusterViewportManagerProxy::ImplFindViewp
 FDisplayClusterViewportProxy* FDisplayClusterViewportManagerProxy::ImplFindViewport_RenderThread(const int32 StereoViewIndex, uint32* OutContextNum) const
 {
 	check(IsInRenderingThread());
+	
+	// ViewIndex == eSSE_MONOSCOPIC(-1) is a special case called for ISR culling math.
+	// Since nDisplay is not ISR compatible, we ignore this request. This won't be neccessary once
+	// we stop using nDisplay as a stereoscopic rendering device (IStereoRendering).
+	if (StereoViewIndex < 0)
+	{
+		return nullptr;
+	}
 
 	for (const TSharedPtr<FDisplayClusterViewportProxy, ESPMode::ThreadSafe>&  ViewportProxyIt : ViewportProxies)
 	{
diff --git a/Engine/Plugins/Runtime/nDisplay/Source/DisplayCluster/Private/Render/Viewport/DisplayClusterViewportManagerViewPointExtension.cpp b/Engine/Plugins/Runtime/nDisplay/Source/DisplayCluster/Private/Render/Viewport/DisplayClusterViewportManagerViewPointExtension.cpp
new file mode 100644
--- /dev/null
+++ b/Engine/Plugins/Runtime/nDisplay/Source/DisplayCluster/Private/Render/Viewport/DisplayClusterViewportManagerViewPointExtension.cpp
@@ -0,0 +1,39 @@
+// Copyright Epic Games, Inc. All Rights Reserved.
+
+#include "DisplayClusterViewportManagerViewPointExtension.h"
+
+#include "Render/Viewport/DisplayClusterViewportManager.h"
+#include "Render/Viewport/DisplayClusterViewport.h"
+
+///////////////////////////////////////////////////////////////////////////////////////
+// FDisplayClusterViewportManagerViewPointExtension
+///////////////////////////////////////////////////////////////////////////////////////
+FDisplayClusterViewportManagerViewPointExtension::FDisplayClusterViewportManagerViewPointExtension(const FAutoRegister& AutoRegister, const FDisplayClusterViewportManager* InViewportManager)
+	: FSceneViewExtensionBase(AutoRegister)
+	, ViewportManager(InViewportManager)
+{ }
+
+FDisplayClusterViewportManagerViewPointExtension::~FDisplayClusterViewportManagerViewPointExtension()
+{
+	ViewportManager = nullptr;
+}
+
+void FDisplayClusterViewportManagerViewPointExtension::Release()
+{
+	check(IsInGameThread());
+
+	ViewportManager = nullptr;
+}
+
+bool FDisplayClusterViewportManagerViewPointExtension::IsActiveThisFrame_Internal(const FSceneViewExtensionContext& Context) const
+{
+	return IsActive() && Context.IsStereoSupported();
+}
+
+void FDisplayClusterViewportManagerViewPointExtension::SetupViewPoint(APlayerController* Player, FMinimalViewInfo& InOutViewInfo)
+{
+	if (IDisplayClusterViewport* DCViewport = IsActive() ? ViewportManager->FindViewport(CurrentStereoViewIndex) : nullptr)
+	{
+		DCViewport->SetupViewPoint(InOutViewInfo);
+	}
+}
diff --git a/Engine/Plugins/Runtime/nDisplay/Source/DisplayCluster/Private/Render/Viewport/DisplayClusterViewportManagerViewPointExtension.h b/Engine/Plugins/Runtime/nDisplay/Source/DisplayCluster/Private/Render/Viewport/DisplayClusterViewportManagerViewPointExtension.h
new file mode 100644
--- /dev/null
+++ b/Engine/Plugins/Runtime/nDisplay/Source/DisplayCluster/Private/Render/Viewport/DisplayClusterViewportManagerViewPointExtension.h
@@ -0,0 +1,58 @@
+// Copyright Epic Games, Inc. All Rights Reserved.
+
+#pragma once
+#include "SceneViewExtension.h"
+
+class FDisplayClusterViewportManager;
+
+#define DISPLAYCLUSTER_SCENE_VIEWPOINT_EXTENSION_PRIORITY -1
+
+/**
+ * View extension applying an DC Viewport features
+ */
+class FDisplayClusterViewportManagerViewPointExtension : public FSceneViewExtensionBase
+{
+public:
+	FDisplayClusterViewportManagerViewPointExtension(const FAutoRegister& AutoRegister, const FDisplayClusterViewportManager* InViewportManager);
+	virtual ~FDisplayClusterViewportManagerViewPointExtension();
+
+public:
+	//~ Begin ISceneViewExtension interface
+	virtual void SetupViewFamily(FSceneViewFamily& InViewFamily) override { }
+	virtual void SetupView(FSceneViewFamily& InViewFamily, FSceneView& InView) override { }
+	virtual void BeginRenderViewFamily(FSceneViewFamily& InViewFamily) override { }
+
+	virtual void SetupViewPoint(APlayerController* Player, FMinimalViewInfo& InOutViewInfo) override;
+
+	virtual int32 GetPriority() const override { return DISPLAYCLUSTER_SCENE_VIEWPOINT_EXTENSION_PRIORITY; }
+
+	/** Release from game thread.*/
+	void Release();
+
+	/** Set current view index for this VE:
+	 * This function must be used because LocalPlayer::GetViewPoint() calls the ISceneViewExtension::SetupViewPoint() function from this VE.
+	 * And at this point the VE must know the current StereoViewIndex value in order to understand which viewport will be used for this SetupViewPoint() call.
+	 * Set INDEX_NONE when this VE no longer needs to be used.
+	 */
+	void SetCurrentStereoViewIndex(const int32 InStereoViewIndex)
+	{
+		CurrentStereoViewIndex = InStereoViewIndex;
+	}
+
+protected:
+	virtual bool IsActiveThisFrame_Internal(const FSceneViewExtensionContext& Context) const override;
+
+private:
+	/** True, if VE can be used at the moment. */
+	bool IsActive() const
+	{
+		return ViewportManager && CurrentStereoViewIndex != INDEX_NONE;
+	}
+
+private:
+	// Pointer to the viewport manager to which this VE belongs
+	const FDisplayClusterViewportManager* ViewportManager;
+
+	// Current StereoViewIndex for rendered viewport
+	int32 CurrentStereoViewIndex = INDEX_NONE;
+};
diff --git a/Engine/Plugins/Runtime/nDisplay/Source/DisplayCluster/Private/Render/Viewport/DisplayClusterViewport_EditorPreview.cpp b/Engine/Plugins/Runtime/nDisplay/Source/DisplayCluster/Private/Render/Viewport/DisplayClusterViewport_EditorPreview.cpp
--- a/Engine/Plugins/Runtime/nDisplay/Source/DisplayCluster/Private/Render/Viewport/DisplayClusterViewport_EditorPreview.cpp
+++ b/Engine/Plugins/Runtime/nDisplay/Source/DisplayCluster/Private/Render/Viewport/DisplayClusterViewport_EditorPreview.cpp
@@ -204,107 +204,22 @@ FSceneView* FDisplayClusterViewport::ImplCalcScenePreview(FSceneViewFamilyContex
 	return nullptr;
 }
 
-enum EDisplayClusterEyeType
-{
-	StereoLeft = 0,
-	Mono = 1,
-	StereoRight = 2,
-	COUNT
-};
-
 bool FDisplayClusterViewport::ImplPreview_CalculateStereoViewOffset(const uint32 InContextNum, FRotator& ViewRotation, const float WorldToMeters, FVector& ViewLocation)
 {
 	check(IsInGameThread());
 	check(WorldToMeters > 0.f);
 
-	if (!GetOwner().IsSceneOpened())
-	{
-		return false;
-	}
-
-	// Get root actor from viewport
-	ADisplayClusterRootActor* const RootActor = GetOwner().GetRootActor();
-	if (!RootActor)
+	// Obtaining the internal viewpoint for a given viewport with stereo eye offset distance.
+	FMinimalViewInfo ViewInfo;
+	if (!SetupViewPoint(ViewInfo))
 	{
-		// No root actor found in game manager.
 		return false;
 	}
 
-	const FDisplayClusterViewport_Context& ViewportContext = Contexts[InContextNum];
-
-	UE_LOG(LogDisplayClusterViewport, VeryVerbose, TEXT("OLD ViewLoc: %s, ViewRot: %s"), *ViewLocation.ToString(), *ViewRotation.ToString());
-	UE_LOG(LogDisplayClusterViewport, VeryVerbose, TEXT("WorldToMeters: %f"), WorldToMeters);
-
-
-	// Get camera ID assigned to the viewport
-	const FString& CameraId = RenderSettings.CameraId;
-
-	// Get camera component assigned to the viewport (or default camera if nothing assigned)
-	UDisplayClusterCameraComponent* const ViewCamera = (CameraId.IsEmpty() ?
-		RootActor->GetDefaultCamera() :
-		RootActor->GetComponentByName<UDisplayClusterCameraComponent>(CameraId));
-
-	if (ViewCamera)
-	{
-		// View base location
-		ViewLocation = ViewCamera->GetComponentLocation();
-		ViewRotation = ViewCamera->GetComponentRotation();
-	}
-	else
-	{
-		UE_LOG(LogDisplayClusterViewport, Warning, TEXT("No camera found for viewport '%s'"), *GetId());
-	}
-
-	if (CameraId.Len() > 0)
-	{
-		UE_LOG(LogDisplayClusterViewport, Verbose, TEXT("Viewport '%s' has assigned camera '%s'"), *GetId(), *CameraId);
-	}
-
-	// Get the actual camera settings
-	const float CfgEyeDist  = ViewCamera ? ViewCamera->GetInterpupillaryDistance() : 6.4f;
-	const bool  bCfgEyeSwap = ViewCamera ? ViewCamera->GetSwapEyes() : false;
-	const float CfgNCP = GNearClippingPlane;
-
-	const EDisplayClusterEyeStereoOffset CfgEyeOffset = ViewCamera ? ViewCamera->GetStereoOffset() : EDisplayClusterEyeStereoOffset::None;
-
-	// Calculate eye offset considering the world scale
-	const float EyeOffset = CfgEyeDist / 2.f;
-	const float EyeOffsetValues[] = { -EyeOffset, 0.f, EyeOffset };
-
-	// Decode current eye type	
-	// Decode current eye type
-	EDisplayClusterEyeType EyeType = EDisplayClusterEyeType::Mono;
-	if (Contexts.Num() > 1)
-	{
-		// Support stereo:
-		EyeType = (InContextNum == 0) ? EDisplayClusterEyeType::StereoLeft : EDisplayClusterEyeType::StereoRight;
-	}
-
-	const int32 EyeIndex = (int32)EyeType;
+	ViewLocation = ViewInfo.Location;
+	ViewRotation = ViewInfo.Rotation;
 
-	float PassOffset = 0.f;
-	float PassOffsetSwap = 0.f;
-
-	if (EyeType == EDisplayClusterEyeType::Mono)
-	{
-		// For monoscopic camera let's check if the "force offset" feature is used
-		// * Force left (-1) ==> 0 left eye
-		// * Force right (1) ==> 2 right eye
-		// * Default (0) ==> 1 mono
-		const int32 EyeOffsetIdx =
-			(CfgEyeOffset == EDisplayClusterEyeStereoOffset::None ? 0 :
-				(CfgEyeOffset == EDisplayClusterEyeStereoOffset::Left ? -1 : 1));
-
-		PassOffset = EyeOffsetValues[EyeOffsetIdx + 1];
-		// Eye swap is not available for monoscopic so just save the value
-		PassOffsetSwap = PassOffset;
-	}
-	else
-	{
-		// For stereo camera we can only swap eyes if required (no "force offset" allowed)
-		PassOffset = EyeOffsetValues[EyeIndex];
-		PassOffsetSwap = (bCfgEyeSwap ? -PassOffset : PassOffset);
-	}
+	const float PassOffsetSwap = GetStereoEyeOffsetDistance(InContextNum);
 
 	FVector ViewOffset = FVector::ZeroVector;
 
@@ -314,6 +229,7 @@ bool FDisplayClusterViewport::ImplPreview_CalculateStereoViewOffset(const uint32
 	ViewLocation += ViewOffset;
 
 	// Perform view calculations on a policy side
+	const float CfgNCP = GNearClippingPlane;
 	if (!CalculateView(InContextNum, ViewLocation, ViewRotation, ViewOffset, WorldToMeters, CfgNCP, CfgNCP))
 	{
 		if (!bProjectionPolicyCalculateViewWarningOnce)
diff --git a/Engine/Plugins/Runtime/nDisplay/Source/DisplayCluster/Private/Render/Viewport/DisplayClusterViewport_ViewPoint.cpp b/Engine/Plugins/Runtime/nDisplay/Source/DisplayCluster/Private/Render/Viewport/DisplayClusterViewport_ViewPoint.cpp
new file mode 100644
--- /dev/null
+++ b/Engine/Plugins/Runtime/nDisplay/Source/DisplayCluster/Private/Render/Viewport/DisplayClusterViewport_ViewPoint.cpp
@@ -0,0 +1,141 @@
+// Copyright Epic Games, Inc. All Rights Reserved.
+
+#include "Render/Viewport/DisplayClusterViewport.h"
+#include "Render/Viewport/DisplayClusterViewportManager.h"
+
+#include "Render/Projection/IDisplayClusterProjectionPolicy.h"
+
+#include "DisplayClusterRootActor.h"
+#include "Components/DisplayClusterCameraComponent.h"
+
+#include "SceneView.h"
+
+#include "Misc/DisplayClusterLog.h"
+
+namespace UE::DisplayClusterViewport
+{
+	enum EDisplayClusterEyeType : int32
+	{
+		StereoLeft = 0,
+		Mono = 1,
+		StereoRight = 2,
+		COUNT
+	};
+};
+using namespace UE::DisplayClusterViewport;
+
+///////////////////////////////////////////////////////////////////////////////////////
+//          FDisplayClusterViewport
+///////////////////////////////////////////////////////////////////////////////////////
+UDisplayClusterCameraComponent* FDisplayClusterViewport::GetViewPointCamera() const
+{
+	if (!ProjectionPolicy.IsValid())
+	{
+		// ignore viewports with uninitialized prj policy
+		return nullptr;
+	}
+
+	if (!GetOwner().IsSceneOpened())
+	{
+		return nullptr;
+	}
+
+	// Get root actor from viewport
+	ADisplayClusterRootActor* const RootActor = GetOwner().GetRootActor();
+	if (!RootActor)
+	{
+		UE_LOG(LogDisplayClusterViewport, Warning, TEXT("Viewport '%s' has no root actor found in game manager"), *GetId());
+
+		return nullptr;
+	}
+
+	// Get camera ID assigned to the viewport
+	const FString& CameraId = GetRenderSettings().CameraId;
+	if (CameraId.Len() > 0)
+	{
+		UE_LOG(LogDisplayClusterViewport, Verbose, TEXT("Viewport '%s' has assigned camera '%s'"), *GetId(), *CameraId);
+	}
+
+	// Get camera component assigned to the viewport (or default camera if nothing assigned)
+	UDisplayClusterCameraComponent* const ViewCamera = (CameraId.IsEmpty() ?
+		RootActor->GetDefaultCamera() :
+		RootActor->GetComponentByName<UDisplayClusterCameraComponent>(CameraId));
+
+	if (!ViewCamera)
+	{
+		UE_LOG(LogDisplayClusterViewport, Warning, TEXT("No camera found for viewport '%s'"), *GetId());
+	}
+
+	return ViewCamera;
+}
+
+bool FDisplayClusterViewport::SetupViewPoint(FMinimalViewInfo& InOutViewInfo)
+{
+	if (UDisplayClusterCameraComponent* ViewCamera = GetViewPointCamera())
+	{
+		// First try to get the viewpoint from the projection policy
+		if (!ProjectionPolicy->GetViewPoint(this, InOutViewInfo.Rotation, InOutViewInfo.Location))
+		{
+			// if unsuccessful, get the viewport from the viewport camera
+			InOutViewInfo.Location = ViewCamera->GetComponentLocation();
+			InOutViewInfo.Rotation = ViewCamera->GetComponentRotation();
+		}
+
+		return true;
+	}
+
+	return false;
+}
+
+float FDisplayClusterViewport::GetStereoEyeOffsetDistance(const uint32 InContextNum)
+{
+	float StereoEyeOffsetDistance = 0.f;
+
+	if (UDisplayClusterCameraComponent* ViewCamera = GetViewPointCamera())
+	{
+		// First try to get the eye offset distance from the projection policy
+		if (!ProjectionPolicy->GetStereoEyeOffsetDistance(this, InContextNum, StereoEyeOffsetDistance))
+		{
+			// if failed, get from the viewport camera
+
+			// Calculate eye offset considering the world scale
+			const float CfgEyeDist = ViewCamera->GetInterpupillaryDistance();
+			const float EyeOffset = CfgEyeDist / 2.f;
+			const float EyeOffsetValues[] = { -EyeOffset, 0.f, EyeOffset };
+
+			// Decode current eye type
+			const EDisplayClusterEyeType EyeType = (Contexts.Num() == 1)
+				? EDisplayClusterEyeType::Mono
+				: (InContextNum == 0) ? EDisplayClusterEyeType::StereoLeft : EDisplayClusterEyeType::StereoRight;
+
+			float PassOffset = 0.f;
+
+			if (EyeType == EDisplayClusterEyeType::Mono)
+			{
+				// For monoscopic camera let's check if the "force offset" feature is used
+				// * Force left (-1) ==> 0 left eye
+				// * Force right (1) ==> 2 right eye
+				// * Default (0) ==> 1 mono
+				const EDisplayClusterEyeStereoOffset CfgEyeOffset = ViewCamera->GetStereoOffset();
+				const int32 EyeOffsetIdx =
+					(CfgEyeOffset == EDisplayClusterEyeStereoOffset::None ? 0 :
+						(CfgEyeOffset == EDisplayClusterEyeStereoOffset::Left ? -1 : 1));
+
+				PassOffset = EyeOffsetValues[EyeOffsetIdx + 1];
+				// Eye swap is not available for monoscopic so just save the value
+				StereoEyeOffsetDistance = PassOffset;
+			}
+			else
+			{
+				// For stereo camera we can only swap eyes if required (no "force offset" allowed)
+				PassOffset = EyeOffsetValues[EyeType];
+
+				// Apply eye swap
+				const bool  CfgEyeSwap = ViewCamera->GetSwapEyes();
+				StereoEyeOffsetDistance = (CfgEyeSwap ? -PassOffset : PassOffset);
+			}
+		}
+	}
+
+	return StereoEyeOffsetDistance;
+}
diff --git a/Engine/Plugins/Runtime/nDisplay/Source/DisplayCluster/Public/Render/Projection/IDisplayClusterProjectionPolicy.h b/Engine/Plugins/Runtime/nDisplay/Source/DisplayCluster/Public/Render/Projection/IDisplayClusterProjectionPolicy.h
--- a/Engine/Plugins/Runtime/nDisplay/Source/DisplayCluster/Public/Render/Projection/IDisplayClusterProjectionPolicy.h
+++ b/Engine/Plugins/Runtime/nDisplay/Source/DisplayCluster/Public/Render/Projection/IDisplayClusterProjectionPolicy.h
@@ -109,6 +109,30 @@ public:
 	*/
 	virtual bool IsConfigurationChanged(const struct FDisplayClusterConfigurationProjection* InConfigurationProjectionPolicy) const = 0;
 
+	/** Get viewpoint for this projection policy
+	 *
+	 * @param OutViewRotation - viewpoint rotation
+	 * @param OutViewLocation - viewpoint location
+	 *
+	 * @return - true, if the viewpoint values have changed
+	 */
+	virtual bool GetViewPoint(IDisplayClusterViewport* InViewport, FRotator& InOutViewRotation, FVector& InOutViewLocation)
+	{
+		return false;
+	}
+
+	 /** Get the distance from the eye to the viewpoint location
+	 *
+	 * @param InContextNum - eye context of this viewport
+	 * @param OutStereoEyeOffsetDistance - new eye distance
+	 *
+	 * @return - true, if the offset distance of the stereo eye has changed
+	 */
+	virtual bool GetStereoEyeOffsetDistance(IDisplayClusterViewport* InViewport, const uint32 InContextNum, float& InOutStereoEyeOffsetDistance)
+	{
+		return false;
+	}
+
 	/**
 	* @param ViewIdx           - Index of view that is being processed for this viewport
 	* @param InOutViewLocation - (in/out) View location with ViewOffset (i.e. left eye pre-computed location)
diff --git a/Engine/Plugins/Runtime/nDisplay/Source/DisplayCluster/Public/Render/Viewport/IDisplayClusterViewport.h b/Engine/Plugins/Runtime/nDisplay/Source/DisplayCluster/Public/Render/Viewport/IDisplayClusterViewport.h
--- a/Engine/Plugins/Runtime/nDisplay/Source/DisplayCluster/Public/Render/Viewport/IDisplayClusterViewport.h
+++ b/Engine/Plugins/Runtime/nDisplay/Source/DisplayCluster/Public/Render/Viewport/IDisplayClusterViewport.h
@@ -35,6 +35,20 @@ public:
 	virtual bool CalculateView(const uint32 InContextNum, FVector& InOutViewLocation, FRotator& InOutViewRotation, const FVector& ViewOffset, const float WorldToMeters, const float NCP, const float FCP) = 0;
 	virtual bool GetProjectionMatrix(const uint32 InContextNum, FMatrix& OutPrjMatrix) = 0;
 
+	/** Setup viewpoint for this viewport
+	 *
+	 * @param InOutViewInfo - [in\out] viewinfo
+	 * 
+	 * @return - true if there is an internal viewpoint for the given viewport.
+	 */
+	virtual bool SetupViewPoint(struct FMinimalViewInfo& InOutViewInfo) = 0;
+
+	/** Get the distance from the eye to the viewpoint location.
+	 *
+	 * @param InContextNum - eye context of this viewport
+	 */
+	virtual float GetStereoEyeOffsetDistance(const uint32 InContextNum) = 0;
+
 	virtual const TSharedPtr<class IDisplayClusterProjectionPolicy, ESPMode::ThreadSafe>& GetProjectionPolicy() const = 0;
 
 	virtual const TArray<FDisplayClusterViewport_Context>& GetContexts() const = 0;
diff --git a/Engine/Plugins/Runtime/nDisplay/Source/DisplayCluster/Public/Render/Viewport/IDisplayClusterViewportManager.h b/Engine/Plugins/Runtime/nDisplay/Source/DisplayCluster/Public/Render/Viewport/IDisplayClusterViewportManager.h
--- a/Engine/Plugins/Runtime/nDisplay/Source/DisplayCluster/Public/Render/Viewport/IDisplayClusterViewportManager.h
+++ b/Engine/Plugins/Runtime/nDisplay/Source/DisplayCluster/Public/Render/Viewport/IDisplayClusterViewportManager.h
@@ -176,5 +176,19 @@ public:
 	* @param InComponentName - (optional) Unique component name
 	*/
 	virtual void MarkComponentGeometryDirty(const FName InComponentName = NAME_None) = 0;
+
+	/**
+	* Call LocalPlayer->CalcSceneView()
+	* [Game thread func]
+	*
+	* @param LocalPlayer - local player
+	* @param ViewFamily - output view struct
+	* @param OutViewLocation - output actor location
+	* @param OutViewRotation - output actor rotation
+	* @param Viewport - current client viewport
+	* @param ViewDrawer - optional drawing in the view
+	* @param StereoViewIndex - index of the view when using stereoscopy
+	*/
+	virtual class FSceneView* CalcSceneView(class ULocalPlayer* LocalPlayer, class FSceneViewFamily* ViewFamily, FVector& OutViewLocation, FRotator& OutViewRotation, class FViewport* Viewport, class FViewElementDrawer* ViewDrawer, int32 StereoViewIndex) = 0;
 };
 
diff --git a/Engine/Plugins/Runtime/nDisplay/Source/DisplayClusterMoviePipeline/Private/DisplayClusterMoviePipelineViewportPass.cpp b/Engine/Plugins/Runtime/nDisplay/Source/DisplayClusterMoviePipeline/Private/DisplayClusterMoviePipelineViewportPass.cpp
--- a/Engine/Plugins/Runtime/nDisplay/Source/DisplayClusterMoviePipeline/Private/DisplayClusterMoviePipelineViewportPass.cpp
+++ b/Engine/Plugins/Runtime/nDisplay/Source/DisplayClusterMoviePipeline/Private/DisplayClusterMoviePipelineViewportPass.cpp
@@ -550,90 +550,21 @@ IDisplayClusterViewport* UDisplayClusterMoviePipelineViewportPassBase::GetAndCal
 		FDisplayClusterRenderFrame RenderFrame;
 		if (ViewportManager->BeginNewFrame(GameViewportClient->Viewport, CurrentWorld, RenderFrame))
 		{
-			if (!DCViewport->GetProjectionPolicy().IsValid())
+			// Obtaining the internal viewpoint for a given viewport with stereo eye offset distance.
+			FMinimalViewInfo ViewInfo;
+			if (!DCViewport->SetupViewPoint(ViewInfo))
 			{
-				// ignore viewports with uninitialized prj policy
 				return nullptr;
 			}
 
-			const TArray<FDisplayClusterViewport_Context>& ViewportContexts = DCViewport->GetContexts();
-			check(ViewportContexts.IsValidIndex(InContextNum));
-
-			const FDisplayClusterViewport_Context& ViewportContext = ViewportContexts[InContextNum];
-
-			UDisplayClusterCameraComponent* ViewCamera = nullptr;
-			{
-				// Get camera component assigned to the viewport (or default camera if nothing assigned)
-				ViewCamera = (CameraId.IsEmpty() ?
-					DCRootActor->GetDefaultCamera() :
-					DCRootActor->GetComponentByName<UDisplayClusterCameraComponent>(CameraId));
-			}
-
-			if (ViewCamera == nullptr)
-			{
-				UE_LOG(LogMovieRenderPipeline, Error, TEXT("No camera found for viewport '%s'"), *DCViewport->GetId());
-				return nullptr;
-			}
+			OutView.ViewRotation = ViewInfo.Rotation;
+			OutView.ViewLocation = ViewInfo.Location;
+			
+			const float PassOffsetSwap = DCViewport->GetStereoEyeOffsetDistance(InContextNum);
 
 			FVector ViewOffset = FVector::ZeroVector;
 			const float CfgNCP = GNearClippingPlane;
 			{
-				// Get the actual camera settings
-				const float CfgEyeDist = ViewCamera->GetInterpupillaryDistance();
-				const bool  CfgEyeSwap = ViewCamera->GetSwapEyes();
-				const EDisplayClusterEyeStereoOffset CfgEyeOffset = ViewCamera->GetStereoOffset();
-
-				// Calculate eye offset considering the world scale
-				const float EyeOffset = CfgEyeDist / 2.f;
-				const float EyeOffsetValues[] = { -EyeOffset, 0.f, EyeOffset };
-
-				enum class EDisplayClusterEyeType : uint8
-				{
-					StereoLeft = 0,
-					Mono = 1,
-					StereoRight = 2,
-					COUNT
-				};
-
-				// Decode current eye type
-				const EDisplayClusterEyeType DCEyeType = (ViewportContexts.Num() > 1)
-					? ((InContextNum == 0) ? EDisplayClusterEyeType::StereoLeft : EDisplayClusterEyeType::StereoRight)
-					: EDisplayClusterEyeType::Mono;
-
-				const int32 EyeIndex = (uint8)DCEyeType;
-
-				float PassOffset = 0.f;
-				float PassOffsetSwap = 0.f;
-
-				switch (DCEyeType)
-				{
-				case EDisplayClusterEyeType::Mono:
-				{
-					// For monoscopic camera let's check if the "force offset" feature is used
-					// * Force left (-1) ==> 0 left eye
-					// * Force right (1) ==> 2 right eye
-					// * Default (0) ==> 1 mono
-					const int32 EyeOffsetIdx = (CfgEyeOffset == EDisplayClusterEyeStereoOffset::None ? 0 : (CfgEyeOffset == EDisplayClusterEyeStereoOffset::Left ? -1 : 1));
-
-					PassOffset = EyeOffsetValues[EyeOffsetIdx + 1];
-					// Eye swap is not available for monoscopic so just save the value
-					PassOffsetSwap = PassOffset;
-				}
-				break;
-
-				default:
-				{
-					// For stereo camera we can only swap eyes if required (no "force offset" allowed)
-					PassOffset = EyeOffsetValues[EyeIndex];
-					PassOffsetSwap = (CfgEyeSwap ? -PassOffset : PassOffset);
-				}
-				break;
-				}
-
-				// View base location
-				OutView.ViewLocation = ViewCamera->GetComponentLocation();
-				OutView.ViewRotation = ViewCamera->GetComponentRotation();
-
 				// Apply computed offset to the view location
 				const FQuat EyeQuat = OutView.ViewRotation.Quaternion();
 				ViewOffset = EyeQuat.RotateVector(FVector(0.0f, PassOffsetSwap, 0.0f));
diff --git a/Engine/Plugins/Runtime/nDisplay/Source/DisplayClusterProjection/Private/Policy/Camera/DisplayClusterProjectionCameraPolicy.cpp b/Engine/Plugins/Runtime/nDisplay/Source/DisplayClusterProjection/Private/Policy/Camera/DisplayClusterProjectionCameraPolicy.cpp
--- a/Engine/Plugins/Runtime/nDisplay/Source/DisplayClusterProjection/Private/Policy/Camera/DisplayClusterProjectionCameraPolicy.cpp
+++ b/Engine/Plugins/Runtime/nDisplay/Source/DisplayClusterProjection/Private/Policy/Camera/DisplayClusterProjectionCameraPolicy.cpp
@@ -82,43 +82,12 @@ bool FDisplayClusterProjectionCameraPolicy::CalculateView(IDisplayClusterViewpor
 {
 	check(IsInGameThread());
 
-	InOutViewLocation = FVector::ZeroVector;
-	InOutViewRotation = FRotator::ZeroRotator;
-	
+	// The camera position has already been determined from the GetViewPoint() function
+
 	// Save Z values
 	ZNear = NCP;
 	ZFar  = FCP;
 
-	// Use assigned camera
-	if (UCameraComponent* CameraComponent = GetCameraComponent())
-	{
-		if (UCineCameraComponent* CineCameraComponent = Cast<UCineCameraComponent>(CameraComponent))
-		{
-			if (CineCameraComponent->bOverride_CustomNearClippingPlane)
-			{
-				ZNear = CineCameraComponent->CustomNearClippingPlane;
-				ZFar = ZNear;
-			}
-		}
-
-		InOutViewLocation = CameraComponent->GetComponentLocation();
-		InOutViewRotation = CameraComponent->GetComponentRotation();
-	}
-	// Otherwise default UE camera is used
-	else
-	{
-		APlayerCameraManager* const CurPlayerCameraManager = GetCurPlayerCameraManager(InViewport);
-		if (CurPlayerCameraManager)
-		{
-			InOutViewLocation = CurPlayerCameraManager->GetCameraLocation();
-			InOutViewRotation = CurPlayerCameraManager->GetCameraRotation();
-		}
-	}
-
-	// Fix camera lens deffects (prototype)
-	InOutViewLocation += CameraSettings.FrustumOffset;
-	InOutViewRotation += CameraSettings.FrustumRotation;
-
 	return true;
 }
 
@@ -172,6 +141,51 @@ bool FDisplayClusterProjectionCameraPolicy::GetProjectionMatrix(IDisplayClusterV
 	return false;
 }
 
+bool FDisplayClusterProjectionCameraPolicy::GetViewPoint(IDisplayClusterViewport* InViewport, FRotator& InOutViewRotation, FVector& InOutViewLocation)
+{
+	InOutViewLocation = FVector::ZeroVector;
+	InOutViewRotation = FRotator::ZeroRotator;
+
+	// Use assigned camera
+	if (UCameraComponent* CameraComponent = GetCameraComponent())
+	{
+		if (UCineCameraComponent* CineCameraComponent = Cast<UCineCameraComponent>(CameraComponent))
+		{
+			if (CineCameraComponent->bOverride_CustomNearClippingPlane)
+			{
+				ZNear = CineCameraComponent->CustomNearClippingPlane;
+				ZFar = ZNear;
+			}
+		}
+
+		InOutViewLocation = CameraComponent->GetComponentLocation();
+		InOutViewRotation = CameraComponent->GetComponentRotation();
+	}
+	// Otherwise default UE camera is used
+	else
+	{
+		APlayerCameraManager* const CurPlayerCameraManager = GetCurPlayerCameraManager(InViewport);
+		if (CurPlayerCameraManager)
+		{
+			InOutViewLocation = CurPlayerCameraManager->GetCameraLocation();
+			InOutViewRotation = CurPlayerCameraManager->GetCameraRotation();
+		}
+	}
+
+	// Fix camera lens deffects (prototype)
+	InOutViewLocation += CameraSettings.FrustumOffset;
+	InOutViewRotation += CameraSettings.FrustumRotation;
+
+	return true;
+}
+
+bool FDisplayClusterProjectionCameraPolicy::GetStereoEyeOffsetDistance(IDisplayClusterViewport* InViewport, const uint32 InContextNum, float& InOutStereoEyeOffsetDistance)
+{
+	// Stereo is not supported by this policy
+	InOutStereoEyeOffsetDistance = 0.f;
+
+	return true;
+}
 
 //////////////////////////////////////////////////////////////////////////////////////////////
 // FDisplayClusterProjectionCameraPolicy
diff --git a/Engine/Plugins/Runtime/nDisplay/Source/DisplayClusterProjection/Private/Policy/Camera/DisplayClusterProjectionCameraPolicy.h b/Engine/Plugins/Runtime/nDisplay/Source/DisplayClusterProjection/Private/Policy/Camera/DisplayClusterProjectionCameraPolicy.h
--- a/Engine/Plugins/Runtime/nDisplay/Source/DisplayClusterProjection/Private/Policy/Camera/DisplayClusterProjectionCameraPolicy.h
+++ b/Engine/Plugins/Runtime/nDisplay/Source/DisplayClusterProjection/Private/Policy/Camera/DisplayClusterProjectionCameraPolicy.h
@@ -30,6 +30,9 @@ public:
 	virtual bool CalculateView(IDisplayClusterViewport* InViewport, const uint32 InContextNum, FVector& InOutViewLocation, FRotator& InOutViewRotation, const FVector& ViewOffset, const float WorldToMeters, const float NCP, const float FCP) override;
 	virtual bool GetProjectionMatrix(IDisplayClusterViewport* InViewport, const uint32 InContextNum, FMatrix& OutPrjMatrix) override;
 
+	virtual bool GetViewPoint(IDisplayClusterViewport* InViewport, FRotator& InOutViewRotation, FVector& InOutViewLocation) override;
+	virtual bool GetStereoEyeOffsetDistance(IDisplayClusterViewport* InViewport, const uint32 InContextNum, float& InOutStereoEyeOffsetDistance) override;
+
 	virtual bool ShouldUseSourceTextureWithMips() const override
 	{
 		return true;
