nDisplay: The base classes DCViewport and DCViewportManager are reworked as TSharedPtr.

- The new Preview interfaces will refer to objects of the base classes.
- Replace all direct pointers to base classes objects with TSharedPtr or TWeakPtr

diff --git a/Engine/Plugins/Runtime/nDisplay/Source/DisplayCluster/Private/Game/EngineClasses/Basics/DisplayClusterViewportClient.cpp b/Engine/Plugins/Runtime/nDisplay/Source/DisplayCluster/Private/Game/EngineClasses/Basics/DisplayClusterViewportClient.cpp
--- a/Engine/Plugins/Runtime/nDisplay/Source/DisplayCluster/Private/Game/EngineClasses/Basics/DisplayClusterViewportClient.cpp
+++ b/Engine/Plugins/Runtime/nDisplay/Source/DisplayCluster/Private/Game/EngineClasses/Basics/DisplayClusterViewportClient.cpp
@@ -48,6 +48,7 @@
 
 #include "Render/Viewport/IDisplayClusterViewportManager.h"
 #include "Render/Viewport/IDisplayClusterViewport.h"
+#include "Render/Viewport/RenderFrame/DisplayClusterRenderFrame.h"
 
 #include "Config/DisplayClusterConfigManager.h"
 #include "UObject/Package.h"
@@ -475,9 +476,9 @@ void UDisplayClusterViewportClient::Draw(FViewport* InViewport, FCanvas* SceneCa
 	// Handle special viewports game-thread logic at frame begin
 	DCRenderDevice->InitializeNewFrame();
 
-	for (FDisplayClusterRenderFrame::FFrameRenderTarget& DCRenderTarget : RenderFrame.RenderTargets)
+	for (FDisplayClusterRenderFrameTarget& DCRenderTarget : RenderFrame.RenderTargets)
 	{
-		for (FDisplayClusterRenderFrame::FFrameViewFamily& DCViewFamily : DCRenderTarget.ViewFamilies)
+		for (FDisplayClusterRenderFrameTargetViewFamily& DCViewFamily : DCRenderTarget.ViewFamilies)
 		{
 			// Create the view family for rendering the world scene to the viewport's render target
 			ViewFamilies.Add(new FSceneViewFamilyContext(RenderFrame.ViewportManager->CreateViewFamilyConstructionValues(
@@ -544,7 +545,7 @@ void UDisplayClusterViewportClient::Draw(FViewport* InViewport, FCanvas* SceneCa
 			FAudioDeviceHandle RetrievedAudioDevice = MyWorld->GetAudioDevice();
 			TArray<FSceneView*> Views;
 
-			for (FDisplayClusterRenderFrame::FFrameView& DCView : DCViewFamily.Views)
+			for (FDisplayClusterRenderFrameTargetView& DCView : DCViewFamily.Views)
 			{
 				const FDisplayClusterViewport_Context ViewportContext = DCView.Viewport->GetContexts()[DCView.ContextNum];
 
@@ -553,7 +554,7 @@ void UDisplayClusterViewportClient::Draw(FViewport* InViewport, FCanvas* SceneCa
 				FRotator	ViewRotation;
 				FSceneView* View = RenderFrame.ViewportManager->CalcSceneView(LocalPlayer, &ViewFamily, ViewLocation, ViewRotation, InViewport, nullptr, ViewportContext.StereoViewIndex);
 
-				if (View && !DCView.ShouldRenderSceneView())
+				if (View && !DCView.IsViewportContextCanBeRendered())
 				{
 					ViewFamily.Views.Remove(View);
 
@@ -1014,19 +1015,6 @@ bool UDisplayClusterViewportClient::Draw_PIE(FViewport* InViewport, FCanvas* Sce
 		return false;
 	}
 
-	//@todo add render mode select
-	EDisplayClusterRenderFrameMode RenderFrameMode = EDisplayClusterRenderFrameMode::Mono;
-	switch (RootActor->RenderMode)
-	{
-	case EDisplayClusterConfigurationRenderMode::SideBySide:
-		RenderFrameMode = EDisplayClusterRenderFrameMode::SideBySide;
-		break;
-	case EDisplayClusterConfigurationRenderMode::TopBottom:
-		RenderFrameMode = EDisplayClusterRenderFrameMode::TopBottom;
-		break;
-	default:
-		break;
-	}
 
 	//Get world for render
 	UWorld* const MyWorld = GetWorld();
@@ -1042,7 +1030,8 @@ bool UDisplayClusterViewportClient::Draw_PIE(FViewport* InViewport, FCanvas* Sce
 	PreviewSettings.bIsPIE = true;
 
 	// Update local node viewports (update\create\delete) and build new render frame
-	if (ViewportManager->UpdateConfiguration(RenderFrameMode, LocalNodeId, RootActor, &PreviewSettings) == false)
+	const EDisplayClusterRenderFrameMode PreviewRenderMode = RootActor->GetPreviewRenderMode();
+	if (ViewportManager->UpdateConfiguration(PreviewRenderMode, LocalNodeId, RootActor, &PreviewSettings) == false)
 	{
 		return false;
 	}
diff --git a/Engine/Plugins/Runtime/nDisplay/Source/DisplayCluster/Private/Game/EngineClasses/Scene/DisplayClusterPreviewComponent.cpp b/Engine/Plugins/Runtime/nDisplay/Source/DisplayCluster/Private/Game/EngineClasses/Scene/DisplayClusterPreviewComponent.cpp
--- a/Engine/Plugins/Runtime/nDisplay/Source/DisplayCluster/Private/Game/EngineClasses/Scene/DisplayClusterPreviewComponent.cpp
+++ b/Engine/Plugins/Runtime/nDisplay/Source/DisplayCluster/Private/Game/EngineClasses/Scene/DisplayClusterPreviewComponent.cpp
@@ -14,6 +14,7 @@
 #include "Render/Projection/IDisplayClusterProjectionPolicyFactory.h"
 
 #include "Render/Viewport/DisplayClusterViewport.h"
+#include "Render/Viewport/DisplayClusterViewportManager.h"
 #include "Render/Viewport/Containers/DisplayClusterViewportReadPixels.h"
 #include "Render/Viewport/RenderFrame/DisplayClusterRenderFrameSettings.h"
 #include "Render/Viewport/DisplayClusterViewportHelpers.h"
@@ -410,17 +411,16 @@ void UDisplayClusterPreviewComponent::UpdateRenderTargetImpl(T* InOutRenderTarge
 
 bool UDisplayClusterPreviewComponent::GetPreviewTextureSettings(FIntPoint& OutSize, EPixelFormat& OutTextureFormat, float& OutGamma, bool& bOutSRGB) const
 {
-	if (IDisplayClusterViewport* PublicViewport = GetCurrentViewport())
+	if (FDisplayClusterViewport* Viewport = static_cast<FDisplayClusterViewport*>(GetCurrentViewport()))
 	{
-		if (FDisplayClusterViewport* Viewport = static_cast<FDisplayClusterViewport*>(PublicViewport))
+		if (FDisplayClusterViewportManager* ViewportManager = Viewport->GetViewportManagerImpl())
 		{
 			// The viewport size is already capped for RenderSettings
-			const TArray<FDisplayClusterViewport_Context>& Contexts = PublicViewport->GetContexts();
-			if (Contexts.Num() > 0)
+			if (!Viewport->GetContexts().IsEmpty())
 			{
-				DisplayClusterViewportHelpers::GetPreviewRenderTargetDesc_Editor(Viewport->GetRenderFrameSettings(), OutTextureFormat, OutGamma, bOutSRGB);
+				DisplayClusterViewportHelpers::GetPreviewRenderTargetDesc_Editor(ViewportManager->GetRenderFrameSettings(), OutTextureFormat, OutGamma, bOutSRGB);
 
-				OutSize = Contexts[0].FrameTargetRect.Size();
+				OutSize = Viewport->GetContexts()[0].FrameTargetRect.Size();
 
 				check(OutSize.X > 0);
 				check(OutSize.Y > 0);
diff --git a/Engine/Plugins/Runtime/nDisplay/Source/DisplayCluster/Private/Game/EngineClasses/Scene/DisplayClusterRootActor.cpp b/Engine/Plugins/Runtime/nDisplay/Source/DisplayCluster/Private/Game/EngineClasses/Scene/DisplayClusterRootActor.cpp
--- a/Engine/Plugins/Runtime/nDisplay/Source/DisplayCluster/Private/Game/EngineClasses/Scene/DisplayClusterRootActor.cpp
+++ b/Engine/Plugins/Runtime/nDisplay/Source/DisplayCluster/Private/Game/EngineClasses/Scene/DisplayClusterRootActor.cpp
@@ -89,7 +89,7 @@ ADisplayClusterRootActor::ADisplayClusterRootActor(const FObjectInitializer& Obj
 	DefaultViewPoint->AttachToComponent(GetRootComponent(), FAttachmentTransformRules::KeepRelativeTransform);
 	DefaultViewPoint->SetRelativeLocation(FVector(0.f, 0.f, 50.f));
 
-	ViewportManager = MakeUnique<FDisplayClusterViewportManager>();
+	CreateViewportManagerImpl();
 
 	PrimaryActorTick.bCanEverTick = true;
 	PrimaryActorTick.TickGroup = ETickingGroup::TG_PostUpdateWork;
@@ -113,6 +113,20 @@ ADisplayClusterRootActor::~ADisplayClusterRootActor()
 #endif
 }
 
+IDisplayClusterViewportManager* ADisplayClusterRootActor::GetViewportManager() const
+{
+	return ViewportManager.Get();
+}
+
+void ADisplayClusterRootActor::CreateViewportManagerImpl()
+{
+	if (!ViewportManager.IsValid())
+	{
+		ViewportManager = MakeShared<FDisplayClusterViewportManager, ESPMode::ThreadSafe>();
+		ViewportManager->Initialize();
+	}
+}
+
 bool ADisplayClusterRootActor::IsRunningGameOrPIE() const
 {
 	if (!IsRunningGame())
@@ -142,6 +156,35 @@ const FDisplayClusterConfigurationRenderFrame& ADisplayClusterRootActor::GetRend
 	return CurrentConfigData->RenderFrameSettings;
 }
 
+EDisplayClusterRenderFrameMode ADisplayClusterRootActor::GetRenderMode() const
+{
+	if (ViewportManager.IsValid())
+	{
+		return ViewportManager->GetRenderMode();
+	}
+
+	return EDisplayClusterRenderFrameMode::Unknown;
+}
+
+EDisplayClusterRenderFrameMode ADisplayClusterRootActor::GetPreviewRenderMode() const
+{
+#if WITH_EDITOR
+	switch (RenderMode)
+	{
+	case EDisplayClusterConfigurationRenderMode::SideBySide:
+		return EDisplayClusterRenderFrameMode::SideBySide;
+
+	case EDisplayClusterConfigurationRenderMode::TopBottom:
+		return EDisplayClusterRenderFrameMode::TopBottom;
+
+	default:
+		break;
+	}
+#endif
+
+	return EDisplayClusterRenderFrameMode::Mono;
+}
+
 void ADisplayClusterRootActor::InitializeFromConfig(UDisplayClusterConfigurationData* ConfigData)
 {
 	if (ConfigData)
@@ -657,10 +700,7 @@ void ADisplayClusterRootActor::InitializeRootActor()
 		UpdateConfigDataInstance(GetDefaultConfigDataFromAsset());
 	}
 
-	if (ViewportManager.IsValid() == false)
-	{
-		ViewportManager = MakeUnique<FDisplayClusterViewportManager>();
-	}
+	CreateViewportManagerImpl();
 
 	StageGeometryComponent->Invalidate();
 
diff --git a/Engine/Plugins/Runtime/nDisplay/Source/DisplayCluster/Private/Game/EngineClasses/Scene/DisplayClusterRootActorEditor.cpp b/Engine/Plugins/Runtime/nDisplay/Source/DisplayCluster/Private/Game/EngineClasses/Scene/DisplayClusterRootActorEditor.cpp
--- a/Engine/Plugins/Runtime/nDisplay/Source/DisplayCluster/Private/Game/EngineClasses/Scene/DisplayClusterRootActorEditor.cpp
+++ b/Engine/Plugins/Runtime/nDisplay/Source/DisplayCluster/Private/Game/EngineClasses/Scene/DisplayClusterRootActorEditor.cpp
@@ -129,10 +129,7 @@ void ADisplayClusterRootActor::Tick_Editor(float DeltaSeconds)
 	if (IsPreviewEnabled())
 	{
 		// Restore ViewportManager
-		if (ViewportManager.IsValid() == false)
-		{
-			ViewportManager = MakeUnique<FDisplayClusterViewportManager>();
-		}
+		CreateViewportManagerImpl();
 
 		if (bDeferPreviewGeneration)
 		{
@@ -162,17 +159,10 @@ void ADisplayClusterRootActor::Tick_Editor(float DeltaSeconds)
 		ResetPreviewInternals_Editor();
 		if (ViewportManager.IsValid())
 		{
-			if (FDisplayClusterViewportManager* ViewportManagerPrivate = static_cast<FDisplayClusterViewportManager*>(ViewportManager.Get()))
+			if(ViewportManager->GetRenderMode() == EDisplayClusterRenderFrameMode::PreviewInScene)
 			{
-				switch (ViewportManagerPrivate->GetRenderFrameSettings().RenderMode)
-				{
-				case EDisplayClusterRenderFrameMode::PreviewInScene:
-					// Release viewport manager with resources immediatelly for preview in scene
-					ViewportManager.Reset();
-					break;
-				default:
-					break;
-				}
+				// Release viewport manager with resources immediatelly for preview in scene
+				ViewportManager.Reset();
 			}
 		}
 	}
@@ -624,9 +614,12 @@ void ADisplayClusterRootActor::ImplRenderPreviewFrustums_Editor()
 				if (CineCameraComponent && CineCameraComponent->IsICVFXEnabled())
 				{
 					// Iterate over rendered incamera viewports (whole cluster)
-					for (FDisplayClusterViewport* InCameraViewportIt : FDisplayClusterViewportConfigurationHelpers_ICVFX::PreviewGetRenderedInCameraViewports(*this, *CineCameraComponent))
+					for (const TSharedPtr<FDisplayClusterViewport, ESPMode::ThreadSafe>& InCameraViewportIt : FDisplayClusterViewportConfigurationHelpers_ICVFX::PreviewGetRenderedInCameraViewports(*this, *CineCameraComponent))
 					{
-						FrustumPreviewViewports.Add(InCameraViewportIt);
+						if (InCameraViewportIt.IsValid())
+						{
+							FrustumPreviewViewports.Add(InCameraViewportIt.Get());
+						}
 					}
 				}
 			}
diff --git a/Engine/Plugins/Runtime/nDisplay/Source/DisplayCluster/Private/Render/Device/DisplayClusterDeviceBase.cpp b/Engine/Plugins/Runtime/nDisplay/Source/DisplayCluster/Private/Render/Device/DisplayClusterDeviceBase.cpp
--- a/Engine/Plugins/Runtime/nDisplay/Source/DisplayCluster/Private/Render/Device/DisplayClusterDeviceBase.cpp
+++ b/Engine/Plugins/Runtime/nDisplay/Source/DisplayCluster/Private/Render/Device/DisplayClusterDeviceBase.cpp
@@ -339,7 +339,7 @@ void FDisplayClusterDeviceBase::InitializeNewFrame()
 		if (IDisplayClusterViewportManager* ViewportManagerPtr = RootActor->GetViewportManager())
 		{
 			// Begin use viewport manager for current frame
-			ViewportManager = ViewportManagerPtr;
+			ViewportManager = ViewportManagerPtr->ToSharedPtr();
 
 			// Initialize frame for render
 			ViewportManager->InitializeNewFrame();
diff --git a/Engine/Plugins/Runtime/nDisplay/Source/DisplayCluster/Private/Render/Device/DisplayClusterDeviceBase.h b/Engine/Plugins/Runtime/nDisplay/Source/DisplayCluster/Private/Render/Device/DisplayClusterDeviceBase.h
--- a/Engine/Plugins/Runtime/nDisplay/Source/DisplayCluster/Private/Render/Device/DisplayClusterDeviceBase.h
+++ b/Engine/Plugins/Runtime/nDisplay/Source/DisplayCluster/Private/Render/Device/DisplayClusterDeviceBase.h
@@ -93,7 +93,7 @@ protected:
 
 private:
 	// Runtime viewport manager api for game and render threads. Internal usage only
-	IDisplayClusterViewportManager* ViewportManager = nullptr;
+	TSharedPtr<IDisplayClusterViewportManager, ESPMode::ThreadSafe> ViewportManager;
 	TSharedPtr<FDisplayClusterViewportManagerProxy, ESPMode::ThreadSafe> ViewportManagerProxy;
 
 	EDisplayClusterRenderFrameMode RenderFrameMode = EDisplayClusterRenderFrameMode::Mono;
diff --git a/Engine/Plugins/Runtime/nDisplay/Source/DisplayCluster/Private/Render/Viewport/Configuration/DisplayClusterViewportConfiguration.cpp b/Engine/Plugins/Runtime/nDisplay/Source/DisplayCluster/Private/Render/Viewport/Configuration/DisplayClusterViewportConfiguration.cpp
--- a/Engine/Plugins/Runtime/nDisplay/Source/DisplayCluster/Private/Render/Viewport/Configuration/DisplayClusterViewportConfiguration.cpp
+++ b/Engine/Plugins/Runtime/nDisplay/Source/DisplayCluster/Private/Render/Viewport/Configuration/DisplayClusterViewportConfiguration.cpp
@@ -119,6 +119,13 @@ static FAutoConsoleVariableRef CVarDisplayClusterAlphaChannelCaptureMode(
 ///////////////////////////////////////////////////////////////////
 // FDisplayClusterViewportConfiguration
 ///////////////////////////////////////////////////////////////////
+FDisplayClusterViewportConfiguration::FDisplayClusterViewportConfiguration(FDisplayClusterViewportManager& InViewportManager)
+	: ViewportManager(InViewportManager.AsShared())
+{ }
+
+FDisplayClusterViewportConfiguration::~FDisplayClusterViewportConfiguration()
+{ }
+
 EDisplayClusterRenderFrameAlphaChannelCaptureMode FDisplayClusterViewportConfiguration::GetAlphaChannelCaptureMode() const
 {
 	ECVarDisplayClusterAlphaChannelCaptureMode AlphaChannelCaptureMode = (ECVarDisplayClusterAlphaChannelCaptureMode)FMath::Clamp(GDisplayClusterAlphaChannelCaptureMode, 0, (int32)ECVarDisplayClusterAlphaChannelCaptureMode::COUNT - 1);
@@ -195,9 +202,9 @@ bool FDisplayClusterViewportConfiguration::ImplUpdateConfiguration(EDisplayClust
 		const UDisplayClusterConfigurationData* ConfigurationData = RootActor->GetConfigData();
 		if (ConfigurationData)
 		{
-			FDisplayClusterViewportConfigurationBase ConfigurationBase(ViewportManager, *RootActor, *ConfigurationData);
+			FDisplayClusterViewportConfigurationBase ConfigurationBase(*ViewportManager, *RootActor, *ConfigurationData);
 			FDisplayClusterViewportConfigurationICVFX ConfigurationICVFX(*RootActor);
-			FDisplayClusterViewportConfigurationProjectionPolicy ConfigurationProjectionPolicy(ViewportManager, *RootActor, *ConfigurationData);
+			FDisplayClusterViewportConfigurationProjectionPolicy ConfigurationProjectionPolicy(*ViewportManager, *RootActor, *ConfigurationData);
 
 			ImplUpdateRenderFrameConfiguration(RootActor->GetRenderFrameSettings());
 
@@ -329,9 +336,12 @@ void FDisplayClusterViewportConfiguration::ImplUpdateConfigurationVisibility(ADi
 	TSet<FPrimitiveComponentId> RootActorHidePrimitivesList;
 	if (RootActor.GetHiddenInGamePrimitives(RootActorHidePrimitivesList))
 	{
-		for (FDisplayClusterViewport* ViewportIt : ViewportManager.ImplGetViewports())
+		for (const TSharedPtr<FDisplayClusterViewport, ESPMode::ThreadSafe>& ViewportIt : ViewportManager->ImplGetViewports())
 		{
-			ViewportIt->VisibilitySettings.SetRootActorHideList(RootActorHidePrimitivesList);
+			if (ViewportIt.IsValid())
+			{
+				ViewportIt->VisibilitySettings.SetRootActorHideList(RootActorHidePrimitivesList);
+			}
 		}
 	}
 }
@@ -339,8 +349,8 @@ void FDisplayClusterViewportConfiguration::ImplUpdateConfigurationVisibility(ADi
 void FDisplayClusterViewportConfiguration::ImplPostUpdateRenderFrameConfiguration()
 {
 	// Some frame postprocess require additional render targetable resources
-	RenderFrameSettings.bShouldUseAdditionalFrameTargetableResource = ViewportManager.ShouldUseAdditionalFrameTargetableResource();
-	RenderFrameSettings.bShouldUseFullSizeFrameTargetableResource = ViewportManager.ShouldUseFullSizeFrameTargetableResource();
+	RenderFrameSettings.bShouldUseAdditionalFrameTargetableResource = ViewportManager->ShouldUseAdditionalFrameTargetableResource();
+	RenderFrameSettings.bShouldUseFullSizeFrameTargetableResource = ViewportManager->ShouldUseFullSizeFrameTargetableResource();
 }
 
 void FDisplayClusterViewportConfiguration::ImplUpdateRenderFrameConfiguration(const FDisplayClusterConfigurationRenderFrame& InRenderFrameConfiguration)
diff --git a/Engine/Plugins/Runtime/nDisplay/Source/DisplayCluster/Private/Render/Viewport/Configuration/DisplayClusterViewportConfiguration.h b/Engine/Plugins/Runtime/nDisplay/Source/DisplayCluster/Private/Render/Viewport/Configuration/DisplayClusterViewportConfiguration.h
--- a/Engine/Plugins/Runtime/nDisplay/Source/DisplayCluster/Private/Render/Viewport/Configuration/DisplayClusterViewportConfiguration.h
+++ b/Engine/Plugins/Runtime/nDisplay/Source/DisplayCluster/Private/Render/Viewport/Configuration/DisplayClusterViewportConfiguration.h
@@ -17,12 +17,8 @@ struct FDisplayClusterConfigurationRenderFrame;
 class FDisplayClusterViewportConfiguration
 {
 public:
-	FDisplayClusterViewportConfiguration(FDisplayClusterViewportManager& InViewportManager)
-		: ViewportManager(InViewportManager)
-	{}
-
-	~FDisplayClusterViewportConfiguration()
-	{}
+	FDisplayClusterViewportConfiguration(FDisplayClusterViewportManager& InViewportManager);
+	~FDisplayClusterViewportConfiguration();
 
 public:
 	// Return true, if root actor ref changed
@@ -33,7 +29,7 @@ public:
 	{ 
 		check(IsInGameThread());
 
-		return RenderFrameSettings; 
+		return RenderFrameSettings;
 	}
 
 	bool UpdateConfiguration(EDisplayClusterRenderFrameMode InRenderMode, const FString& InClusterNodeId);
@@ -56,7 +52,8 @@ private:
 	EDisplayClusterRenderFrameAlphaChannelCaptureMode GetAlphaChannelCaptureMode() const;
 
 private:
-	FDisplayClusterViewportManager&    ViewportManager;
+	const TSharedRef<FDisplayClusterViewportManager> ViewportManager;
+
 	FDisplayClusterActorRef            RootActorRef;
 	FDisplayClusterRenderFrameSettings RenderFrameSettings;
 };
diff --git a/Engine/Plugins/Runtime/nDisplay/Source/DisplayCluster/Private/Render/Viewport/Configuration/DisplayClusterViewportConfigurationBase.cpp b/Engine/Plugins/Runtime/nDisplay/Source/DisplayCluster/Private/Render/Viewport/Configuration/DisplayClusterViewportConfigurationBase.cpp
--- a/Engine/Plugins/Runtime/nDisplay/Source/DisplayCluster/Private/Render/Viewport/Configuration/DisplayClusterViewportConfigurationBase.cpp
+++ b/Engine/Plugins/Runtime/nDisplay/Source/DisplayCluster/Private/Render/Viewport/Configuration/DisplayClusterViewportConfigurationBase.cpp
@@ -59,49 +59,58 @@ void FDisplayClusterViewportConfigurationBase::Update(const FString& ClusterNode
 		TArray<FString> ExistClusterNodesIDs;
 		ConfigurationData.Cluster->GetNodeIds(ExistClusterNodesIDs);
 
-		TArray<FDisplayClusterViewport*> UnusedViewports;
-		for (FDisplayClusterViewport* It : ViewportManager.ImplGetViewports())
+		TArray<TSharedPtr<FDisplayClusterViewport, ESPMode::ThreadSafe>> UnusedViewports;
+		for (const TSharedPtr<FDisplayClusterViewport, ESPMode::ThreadSafe>& ViewportIt : ViewportManager.ImplGetViewports())
 		{
-			// ignore ICVFX internal resources
-			if (!EnumHasAllFlags(It->GetRenderSettingsICVFX().RuntimeFlags, EDisplayClusterViewportRuntimeICVFXFlags::InternalResource))
+			if (ViewportIt.IsValid())
 			{
-				// Only viewports from cluster node in render
-				if (ClusterNodeId.IsEmpty() || It->GetClusterNodeId() == ClusterNodeId)
+				// ignore ICVFX internal resources
+				if (!EnumHasAllFlags(ViewportIt->GetRenderSettingsICVFX().RuntimeFlags, EDisplayClusterViewportRuntimeICVFXFlags::InternalResource))
 				{
-					if (!DesiredViewports.Contains(It->GetId()))
+					// Only viewports from cluster node in render
+					if (ClusterNodeId.IsEmpty() || ViewportIt->GetClusterNodeId() == ClusterNodeId)
 					{
-						UnusedViewports.Add(It);
+						if (!DesiredViewports.Contains(ViewportIt->GetId()))
+						{
+							UnusedViewports.Add(ViewportIt);
+						}
 					}
-				}
-				else
-				{
-					if(ExistClusterNodesIDs.Find(It->GetClusterNodeId()) == INDEX_NONE)
+					else
 					{
-						// also remove viewports for deleted cluster nodes
-						UnusedViewports.Add(It);
+						if (ExistClusterNodesIDs.Find(ViewportIt->GetClusterNodeId()) == INDEX_NONE)
+						{
+							// also remove viewports for deleted cluster nodes
+							UnusedViewports.Add(ViewportIt);
+						}
 					}
 				}
 			}
 		}
 
 		// Delete unused viewports
-		for (FDisplayClusterViewport* DeleteIt : UnusedViewports)
+		for (const TSharedPtr<FDisplayClusterViewport, ESPMode::ThreadSafe>& UnusedViewportIt : UnusedViewports)
 		{
-			ViewportManager.ImplDeleteViewport(DeleteIt);
+			if (UnusedViewportIt.IsValid())
+			{
+				ViewportManager.ImplDeleteViewport(UnusedViewportIt);
+			}
 		}
 	}
 
 	// Update and Create new viewports
 	for (TPair<FString, UDisplayClusterConfigurationViewport*>& CfgIt : DesiredViewports)
 	{
-		FDisplayClusterViewport* ExistViewport = ViewportManager.ImplFindViewport(CfgIt.Key);
-		if (ExistViewport)
+		if (const UDisplayClusterConfigurationViewport* ConfigurationViewport = CfgIt.Value)
 		{
-			FDisplayClusterViewportConfigurationBase::UpdateViewportConfiguration(ViewportManager, RootActor, ExistViewport, CfgIt.Value);
-		}
-		else
-		{
-			ViewportManager.CreateViewport(CfgIt.Key, CfgIt.Value);
+			const TSharedPtr<FDisplayClusterViewport, ESPMode::ThreadSafe> ExistViewport = ViewportManager.ImplFindViewport(CfgIt.Key);
+			if (ExistViewport.IsValid())
+			{
+				FDisplayClusterViewportConfigurationBase::UpdateViewportConfiguration(*ExistViewport, ViewportManager, RootActor, *ConfigurationViewport);
+			}
+			else
+			{
+				ViewportManager.CreateViewport(CfgIt.Key, *ConfigurationViewport);
+			}
 		}
 	}
 }
@@ -110,21 +119,24 @@ void FDisplayClusterViewportConfigurationBase::Update(const TArray<FString>& InV
 {
 	// Collect unused viewports and delete
 	{
-		TArray<FDisplayClusterViewport*> UnusedViewports;
-		for (FDisplayClusterViewport* It : ViewportManager.ImplGetViewports())
+		TArray<TSharedPtr<FDisplayClusterViewport, ESPMode::ThreadSafe>> UnusedViewports;
+		for (const TSharedPtr<FDisplayClusterViewport, ESPMode::ThreadSafe>& ViewportIt : ViewportManager.ImplGetViewports())
 		{
-			// ignore ICVFX internal resources
-			if (!EnumHasAllFlags(It->GetRenderSettingsICVFX().RuntimeFlags, EDisplayClusterViewportRuntimeICVFXFlags::InternalResource))
+			if (ViewportIt.IsValid())
 			{
-				if (InViewportNames.Find(It->GetId()) == INDEX_NONE)
+				// ignore ICVFX internal resources
+				if (!EnumHasAllFlags(ViewportIt->GetRenderSettingsICVFX().RuntimeFlags, EDisplayClusterViewportRuntimeICVFXFlags::InternalResource))
 				{
-					UnusedViewports.Add(It);
+					if (InViewportNames.Find(ViewportIt->GetId()) == INDEX_NONE)
+					{
+						UnusedViewports.Add(ViewportIt);
+					}
 				}
 			}
 		}
 
 		// Delete unused viewports
-		for (FDisplayClusterViewport* DeleteIt : UnusedViewports)
+		for (const TSharedPtr<FDisplayClusterViewport, ESPMode::ThreadSafe>& DeleteIt : UnusedViewports)
 		{
 			ViewportManager.ImplDeleteViewport(DeleteIt);
 		}
@@ -140,16 +152,19 @@ void FDisplayClusterViewportConfigurationBase::Update(const TArray<FString>& InV
 
 			for (const TPair<FString, TObjectPtr<UDisplayClusterConfigurationViewport>>& ViewportIt : ClusterNodeConfiguration->Viewports)
 			{
-				if (ViewportIt.Key.Len() && ViewportIt.Value && InViewportNames.Find(ViewportIt.Key) != INDEX_NONE)
+				if (const UDisplayClusterConfigurationViewport* ConfigurationViewport = ViewportIt.Value)
 				{
-					FDisplayClusterViewport* ExistViewport = ViewportManager.ImplFindViewport(ViewportIt.Key);
-					if (ExistViewport)
+					if (ViewportIt.Key.Len() && ViewportIt.Value && InViewportNames.Find(ViewportIt.Key) != INDEX_NONE)
 					{
-						FDisplayClusterViewportConfigurationBase::UpdateViewportConfiguration(ViewportManager, RootActor, ExistViewport, ViewportIt.Value);
-					}
-					else
-					{
-						ViewportManager.CreateViewport(ViewportIt.Key, ViewportIt.Value);
+						const TSharedPtr<FDisplayClusterViewport, ESPMode::ThreadSafe> ExistViewport = ViewportManager.ImplFindViewport(ViewportIt.Key);
+						if (ExistViewport.IsValid())
+						{
+							FDisplayClusterViewportConfigurationBase::UpdateViewportConfiguration(*ExistViewport, ViewportManager, RootActor, *ConfigurationViewport);
+						}
+						else
+						{
+							ViewportManager.CreateViewport(ViewportIt.Key, *ConfigurationViewport);
+						}
 					}
 				}
 			}
@@ -299,14 +314,12 @@ void FDisplayClusterViewportConfigurationBase::UpdateClusterNodePostProcess(cons
 }
 
 // Assign new configuration to this viewport <Runtime>
-bool FDisplayClusterViewportConfigurationBase::UpdateViewportConfiguration(FDisplayClusterViewportManager& ViewportManager, ADisplayClusterRootActor& RootActor, FDisplayClusterViewport* DesiredViewport, const UDisplayClusterConfigurationViewport* ConfigurationViewport)
+bool FDisplayClusterViewportConfigurationBase::UpdateViewportConfiguration(FDisplayClusterViewport& DstViewport, FDisplayClusterViewportManager& ViewportManager, ADisplayClusterRootActor& RootActor, const UDisplayClusterConfigurationViewport& ConfigurationViewport)
 {
 	check(IsInGameThread());
-	check(DesiredViewport);
-	check(ConfigurationViewport);
 
-	FDisplayClusterViewportConfigurationHelpers::UpdateBaseViewportSetting(*DesiredViewport, RootActor, *ConfigurationViewport);
-	FDisplayClusterViewportConfigurationHelpers::UpdateProjectionPolicy(*DesiredViewport, &(ConfigurationViewport->ProjectionPolicy));
+	FDisplayClusterViewportConfigurationHelpers::UpdateBaseViewportSetting(DstViewport, RootActor, ConfigurationViewport);
+	FDisplayClusterViewportConfigurationHelpers::UpdateProjectionPolicy(DstViewport, &(ConfigurationViewport.ProjectionPolicy));
 
 	return true;
 }
diff --git a/Engine/Plugins/Runtime/nDisplay/Source/DisplayCluster/Private/Render/Viewport/Configuration/DisplayClusterViewportConfigurationBase.h b/Engine/Plugins/Runtime/nDisplay/Source/DisplayCluster/Private/Render/Viewport/Configuration/DisplayClusterViewportConfigurationBase.h
--- a/Engine/Plugins/Runtime/nDisplay/Source/DisplayCluster/Private/Render/Viewport/Configuration/DisplayClusterViewportConfigurationBase.h
+++ b/Engine/Plugins/Runtime/nDisplay/Source/DisplayCluster/Private/Render/Viewport/Configuration/DisplayClusterViewportConfigurationBase.h
@@ -29,7 +29,7 @@ public:
 	void UpdateClusterNodePostProcess(const FString& ClusterNodeId, const FDisplayClusterRenderFrameSettings& InRenderFrameSettings);
 
 public:
-	static bool UpdateViewportConfiguration(FDisplayClusterViewportManager& ViewportManager, ADisplayClusterRootActor& RootActor, FDisplayClusterViewport* DesiredViewport, const UDisplayClusterConfigurationViewport* ConfigurationViewport);
+	static bool UpdateViewportConfiguration(FDisplayClusterViewport& DstViewport, FDisplayClusterViewportManager& ViewportManager, ADisplayClusterRootActor& RootActor, const UDisplayClusterConfigurationViewport& ConfigurationViewport);
 
 protected:
 	/**
diff --git a/Engine/Plugins/Runtime/nDisplay/Source/DisplayCluster/Private/Render/Viewport/Configuration/DisplayClusterViewportConfigurationHelpers.cpp b/Engine/Plugins/Runtime/nDisplay/Source/DisplayCluster/Private/Render/Viewport/Configuration/DisplayClusterViewportConfigurationHelpers.cpp
--- a/Engine/Plugins/Runtime/nDisplay/Source/DisplayCluster/Private/Render/Viewport/Configuration/DisplayClusterViewportConfigurationHelpers.cpp
+++ b/Engine/Plugins/Runtime/nDisplay/Source/DisplayCluster/Private/Render/Viewport/Configuration/DisplayClusterViewportConfigurationHelpers.cpp
@@ -53,7 +53,7 @@ void FDisplayClusterViewportConfigurationHelpers::UpdateViewportStereoMode(FDisp
 
 void FDisplayClusterViewportConfigurationHelpers::UpdateViewportSetting_OverlayRenderSettings(FDisplayClusterViewport& DstViewport, const FDisplayClusterConfigurationICVFX_OverlayAdvancedRenderSettings& InOverlaySettings)
 {
-	DstViewport.Owner.SetViewportBufferRatio(DstViewport, InOverlaySettings.BufferRatio);
+	DstViewport.SetViewportBufferRatio(InOverlaySettings.BufferRatio);
 	DstViewport.RenderSettings.RenderTargetRatio = InOverlaySettings.RenderTargetRatio;
 
 	DstViewport.RenderSettings.GPUIndex = InOverlaySettings.GPUIndex;
@@ -196,7 +196,7 @@ void FDisplayClusterViewportConfigurationHelpers::UpdateBaseViewportSetting(FDis
 	FDisplayClusterViewportConfigurationHelpers_Postprocess::UpdatePerViewportPostProcessSettings(DstViewport, RootActor);
 
 	{
-		DstViewport.Owner.SetViewportBufferRatio(DstViewport, InRenderSettings.BufferRatio);
+		DstViewport.SetViewportBufferRatio(InRenderSettings.BufferRatio);
 
 		UpdateViewportSetting_Overscan(DstViewport, InRenderSettings.Overscan);
 
@@ -335,7 +335,7 @@ void FDisplayClusterViewportConfigurationHelpers::UpdateProjectionPolicy(FDispla
 	{
 		if (!DstViewport.ProjectionPolicy.IsValid())
 		{
-			if (DstViewport.Owner.IsSceneOpened())
+			if (DstViewport.IsSceneOpened())
 			{
 				// Try initialize proj policy every tick (mesh deferred load, etc)
 				DstViewport.HandleStartScene();
diff --git a/Engine/Plugins/Runtime/nDisplay/Source/DisplayCluster/Private/Render/Viewport/Configuration/DisplayClusterViewportConfigurationHelpers_ICVFX.cpp b/Engine/Plugins/Runtime/nDisplay/Source/DisplayCluster/Private/Render/Viewport/Configuration/DisplayClusterViewportConfigurationHelpers_ICVFX.cpp
--- a/Engine/Plugins/Runtime/nDisplay/Source/DisplayCluster/Private/Render/Viewport/Configuration/DisplayClusterViewportConfigurationHelpers_ICVFX.cpp
+++ b/Engine/Plugins/Runtime/nDisplay/Source/DisplayCluster/Private/Render/Viewport/Configuration/DisplayClusterViewportConfigurationHelpers_ICVFX.cpp
@@ -92,17 +92,9 @@ namespace DisplayClusterViewportConfigurationHelpers_ICVFX_Impl
 	{
 		if (GDisplayClusterPreviewEnableReuseViewportInCluster > 0)
 		{
-			FDisplayClusterViewportManager* ViewportManager = FDisplayClusterViewportConfigurationHelpers_ICVFX::GetViewportManager(InRootActor);
-			if (ViewportManager)
+			if (InRootActor.GetRenderMode() == EDisplayClusterRenderFrameMode::PreviewInScene)
 			{
-				const FDisplayClusterRenderFrameSettings& RenderFrameSettings = ViewportManager->GetRenderFrameSettings();
-				switch (RenderFrameSettings.RenderMode)
-				{
-				case EDisplayClusterRenderFrameMode::PreviewInScene:
-					return true;
-				default:
-					break;
-				}
+				return true;
 			}
 		}
 
@@ -118,10 +110,11 @@ using namespace DisplayClusterViewportConfigurationHelpers_ICVFX_Impl;
 ////////////////////////////////////////////////////////////////////////////////
 FDisplayClusterViewport* FDisplayClusterViewportConfigurationHelpers_ICVFX::ImplFindViewport(ADisplayClusterRootActor& RootActor, const FString& InViewportId, const FString& InResourceId)
 {
-	FDisplayClusterViewportManager* ViewportManager = FDisplayClusterViewportConfigurationHelpers_ICVFX::GetViewportManager(RootActor);
-	if (ViewportManager != nullptr)
+	if (FDisplayClusterViewportManager* ViewportManager = FDisplayClusterViewportConfigurationHelpers_ICVFX::GetViewportManager(RootActor))
 	{
-		return ViewportManager->ImplFindViewport(ImplGetNameICVFX(ViewportManager->GetRenderFrameSettings().ClusterNodeId, InViewportId, InResourceId));
+		TSharedPtr<FDisplayClusterViewport, ESPMode::ThreadSafe> Viewport = ViewportManager->ImplFindViewport(ImplGetNameICVFX(ViewportManager->GetRenderFrameSettings().ClusterNodeId, InViewportId, InResourceId));
+
+		return Viewport.Get();
 	}
 
 	return nullptr;
@@ -129,8 +122,7 @@ FDisplayClusterViewport* FDisplayClusterViewportConfigurationHelpers_ICVFX::Impl
 
 static bool ImplCreateProjectionPolicy(ADisplayClusterRootActor& RootActor, const FString& InViewportId, const FString& InResourceId, bool bIsCameraProjection, TSharedPtr<IDisplayClusterProjectionPolicy, ESPMode::ThreadSafe>& OutProjPolicy)
 {
-	FDisplayClusterViewportManager* ViewportManager = FDisplayClusterViewportConfigurationHelpers_ICVFX::GetViewportManager(RootActor);
-	if (ViewportManager != nullptr)
+	if (FDisplayClusterViewportManager* ViewportManager = FDisplayClusterViewportConfigurationHelpers_ICVFX::GetViewportManager(RootActor))
 	{
 		const FString& ClusterNodeId = ViewportManager->GetRenderFrameSettings().ClusterNodeId;
 
@@ -160,12 +152,11 @@ FDisplayClusterViewport* FDisplayClusterViewportConfigurationHelpers_ICVFX::Impl
 {
 	check(InProjectionPolicy.IsValid());
 
-	FDisplayClusterViewportManager* ViewportManager = FDisplayClusterViewportConfigurationHelpers_ICVFX::GetViewportManager(RootActor);
-	if (ViewportManager != nullptr)
+	if (FDisplayClusterViewportManager* ViewportManager = FDisplayClusterViewportConfigurationHelpers_ICVFX::GetViewportManager(RootActor))
 	{
 		// Create viewport for new projection policy
-		FDisplayClusterViewport* NewViewport = ViewportManager->ImplCreateViewport(ImplGetNameICVFX(ViewportManager->GetRenderFrameSettings().ClusterNodeId, InViewportId, InResourceId), InProjectionPolicy);
-		if (NewViewport != nullptr)
+		TSharedPtr<FDisplayClusterViewport, ESPMode::ThreadSafe> NewViewport = ViewportManager->ImplCreateViewport(ImplGetNameICVFX(ViewportManager->GetRenderFrameSettings().ClusterNodeId, InViewportId, InResourceId), InProjectionPolicy);
+		if (NewViewport.IsValid())
 		{
 			// Mark as internal resource
 			EnumAddFlags(NewViewport->RenderSettingsICVFX.RuntimeFlags, EDisplayClusterViewportRuntimeICVFXFlags::InternalResource);
@@ -173,7 +164,7 @@ FDisplayClusterViewport* FDisplayClusterViewportConfigurationHelpers_ICVFX::Impl
 			// Dont show ICVFX composing viewports on frame target
 			NewViewport->RenderSettings.bVisible = false;
 
-			return NewViewport;
+			return NewViewport.Get();
 		}
 	}
 
@@ -231,44 +222,39 @@ FDisplayClusterViewport* FDisplayClusterViewportConfigurationHelpers_ICVFX::GetO
 
 void FDisplayClusterViewportConfigurationHelpers_ICVFX::ReuseUVLightCardViewportWithinClusterNode(FDisplayClusterViewport& InUVLightCardViewport)
 {
-	const TArray<FDisplayClusterViewport*>& Viewports = 
-#if WITH_EDITOR
-		InUVLightCardViewport.ImplGetOwner().ImplGetWholeClusterViewports_Editor();
-#else
-		InUVLightCardViewport.ImplGetOwner().ImplGetViewports();
-#endif
-
-	for (FDisplayClusterViewport* ViewportIt : Viewports)
+	if (FDisplayClusterViewportManager* ViewportManager = InUVLightCardViewport.GetViewportManagerImpl())
 	{
-		if(ViewportIt && ViewportIt != &InUVLightCardViewport && !ViewportIt->RenderSettings.IsViewportOverrided()
-		&& EnumHasAnyFlags(ViewportIt->RenderSettingsICVFX.RuntimeFlags, EDisplayClusterViewportRuntimeICVFXFlags::UVLightcard))
+		for (const TSharedPtr<FDisplayClusterViewport, ESPMode::ThreadSafe>& ViewportIt : ViewportManager->ImplGetWholeClusterViewports())
 		{
-			if (ViewportIt->IsOpenColorIOEquals(InUVLightCardViewport))
+			if (ViewportIt.IsValid() && ViewportIt != InUVLightCardViewport.AsShared() && !ViewportIt->RenderSettings.IsViewportOverrided()
+				&& EnumHasAnyFlags(ViewportIt->RenderSettingsICVFX.RuntimeFlags, EDisplayClusterViewportRuntimeICVFXFlags::UVLightcard))
 			{
-				// Reuse exist viewport:
-				InUVLightCardViewport.RenderSettings.SetViewportOverride(ViewportIt->GetId(), EDisplayClusterViewportOverrideMode::All);
-				break;
+				if (ViewportIt->IsOpenColorIOEquals(InUVLightCardViewport))
+				{
+					// Reuse exist viewport:
+					InUVLightCardViewport.RenderSettings.SetViewportOverride(ViewportIt->GetId(), EDisplayClusterViewportOverrideMode::All);
+					break;
+				}
 			}
 		}
 	}
 }
 
 #if WITH_EDITOR
-TArray<FDisplayClusterViewport*> FDisplayClusterViewportConfigurationHelpers_ICVFX::PreviewGetRenderedInCameraViewports(ADisplayClusterRootActor& InRootActor, UDisplayClusterICVFXCameraComponent& InCameraComponent, bool bGetChromakey)
+TArray<TSharedPtr<FDisplayClusterViewport, ESPMode::ThreadSafe>> FDisplayClusterViewportConfigurationHelpers_ICVFX::PreviewGetRenderedInCameraViewports(ADisplayClusterRootActor& InRootActor, UDisplayClusterICVFXCameraComponent& InCameraComponent, bool bGetChromakey)
 {
-	TArray<FDisplayClusterViewport*> OutViewports;
+	TArray<TSharedPtr<FDisplayClusterViewport, ESPMode::ThreadSafe>> OutViewports;
 
 	// Search for rendered camera viewport on other cluster nodes
-	FDisplayClusterViewportManager* ViewportManager = FDisplayClusterViewportConfigurationHelpers_ICVFX::GetViewportManager(InRootActor);
-	if (ViewportManager != nullptr)
+	if (FDisplayClusterViewportManager* ViewportManager = FDisplayClusterViewportConfigurationHelpers_ICVFX::GetViewportManager(InRootActor))
 	{
 		const FString ICVFXCameraId = InCameraComponent.GetCameraUniqueId();
 		const EDisplayClusterViewportRuntimeICVFXFlags InRuntimeFlagsMask = bGetChromakey ? EDisplayClusterViewportRuntimeICVFXFlags::Chromakey : EDisplayClusterViewportRuntimeICVFXFlags::InCamera;
 		const FString ViewportTypeId = bGetChromakey ? DisplayClusterViewportStrings::icvfx::chromakey : DisplayClusterViewportStrings::icvfx::camera;
 
-		for (FDisplayClusterViewport* ViewportIt : ViewportManager->ImplGetWholeClusterViewports_Editor())
+		for (const TSharedPtr<FDisplayClusterViewport, ESPMode::ThreadSafe>& ViewportIt : ViewportManager->ImplGetWholeClusterViewports())
 		{
-			if (ViewportIt != nullptr
+			if (ViewportIt.IsValid()
 				&& EnumHasAnyFlags(ViewportIt->RenderSettingsICVFX.RuntimeFlags, InRuntimeFlagsMask)
 				&& (ViewportIt->InputShaderResources.Num() > 0 && ViewportIt->InputShaderResources[0] != nullptr && ViewportIt->Contexts.Num() > 0)
 				&& (!ViewportIt->RenderSettings.IsViewportOverrided()))
@@ -293,19 +279,15 @@ void FDisplayClusterViewportConfigurationHelpers_ICVFX::PreviewReuseInnerFrustum
 		return;
 	}
 
-	const FDisplayClusterRenderFrameSettings& RenderFrameSettingsConstRef = InCameraViewport.GetRenderFrameSettings();
-	switch (RenderFrameSettingsConstRef.RenderMode)
+	// Only for preview mode
+	if (InCameraViewport.GetRenderMode() != EDisplayClusterRenderFrameMode::PreviewInScene)
 	{
-	case EDisplayClusterRenderFrameMode::PreviewInScene:
-		// Only for preview mode
-		break;
-	default:
 		return;
 	}
 
-	for (FDisplayClusterViewport* ViewportIt : PreviewGetRenderedInCameraViewports(InRootActor, InCameraComponent, false))
+	for (const TSharedPtr<FDisplayClusterViewport, ESPMode::ThreadSafe>& ViewportIt : PreviewGetRenderedInCameraViewports(InRootActor, InCameraComponent, false))
 	{
-		if (ViewportIt && ViewportIt != &InCameraViewport && ViewportIt->GetClusterNodeId() != InCameraViewport.GetClusterNodeId()
+		if (ViewportIt.IsValid() && ViewportIt != InCameraViewport.AsShared() && ViewportIt->GetClusterNodeId() != InCameraViewport.GetClusterNodeId()
 			&& FDisplayClusterViewportConfigurationHelpers_Postprocess::IsInnerFrustumViewportSettingsEqual_Editor(*ViewportIt, InCameraViewport, InCameraComponent))
 		{
 			EDisplayClusterViewportOverrideMode ViewportOverrideMode = ViewportIt->IsOpenColorIOEquals(InCameraViewport) ?
@@ -324,20 +306,16 @@ void FDisplayClusterViewportConfigurationHelpers_ICVFX::PreviewReuseChromakeyVie
 	{
 		return;
 	}
-
-	const FDisplayClusterRenderFrameSettings& RenderFrameSettingsConstRef = InChromakeyViewport.GetRenderFrameSettings();
-	switch (RenderFrameSettingsConstRef.RenderMode)
+	
+	// Only for preview mode
+	if (InChromakeyViewport.GetRenderMode() != EDisplayClusterRenderFrameMode::PreviewInScene)
 	{
-	case EDisplayClusterRenderFrameMode::PreviewInScene:
-		// Only for preview mode
-		break;
-	default:
 		return;
 	}
 
-	for (FDisplayClusterViewport* ViewportIt : PreviewGetRenderedInCameraViewports(InRootActor, InCameraComponent, true))
+	for (const TSharedPtr<FDisplayClusterViewport, ESPMode::ThreadSafe>& ViewportIt : PreviewGetRenderedInCameraViewports(InRootActor, InCameraComponent, true))
 	{
-		if (ViewportIt && ViewportIt!= &InChromakeyViewport  && ViewportIt->GetClusterNodeId() != InChromakeyViewport.GetClusterNodeId())
+		if (ViewportIt.IsValid() && ViewportIt != InChromakeyViewport.AsShared() && ViewportIt->GetClusterNodeId() != InChromakeyViewport.GetClusterNodeId())
 		{
 			// Chromakey support OCIO
 			EDisplayClusterViewportOverrideMode ViewportOverrideMode = ViewportIt->IsOpenColorIOEquals(InChromakeyViewport) ?
@@ -962,5 +940,5 @@ void FDisplayClusterViewportConfigurationHelpers_ICVFX::UpdateCameraViewportBuff
 		DstViewport.RenderSettings.RenderTargetAdaptRatio = 1.f;
 	}
 
-	DstViewport.Owner.SetViewportBufferRatio(DstViewport, BufferRatio);
+	DstViewport.SetViewportBufferRatio(BufferRatio);
 }
diff --git a/Engine/Plugins/Runtime/nDisplay/Source/DisplayCluster/Private/Render/Viewport/Configuration/DisplayClusterViewportConfigurationHelpers_ICVFX.h b/Engine/Plugins/Runtime/nDisplay/Source/DisplayCluster/Private/Render/Viewport/Configuration/DisplayClusterViewportConfigurationHelpers_ICVFX.h
--- a/Engine/Plugins/Runtime/nDisplay/Source/DisplayCluster/Private/Render/Viewport/Configuration/DisplayClusterViewportConfigurationHelpers_ICVFX.h
+++ b/Engine/Plugins/Runtime/nDisplay/Source/DisplayCluster/Private/Render/Viewport/Configuration/DisplayClusterViewportConfigurationHelpers_ICVFX.h
@@ -66,7 +66,7 @@ public:
 	static bool IsShouldUseUVLightcard(FDisplayClusterViewportManager& InViewportManager, const FDisplayClusterConfigurationICVFX_LightcardSettings& InLightcardSettings);
 
 #if WITH_EDITOR
-	static TArray<FDisplayClusterViewport*> PreviewGetRenderedInCameraViewports(ADisplayClusterRootActor& RootActor, UDisplayClusterICVFXCameraComponent& InCameraComponent, bool bGetChromakey = false);
+	static TArray<TSharedPtr<FDisplayClusterViewport, ESPMode::ThreadSafe>> PreviewGetRenderedInCameraViewports(ADisplayClusterRootActor& RootActor, UDisplayClusterICVFXCameraComponent& InCameraComponent, bool bGetChromakey = false);
 	static void PreviewReuseInnerFrustumViewportWithinClusterNodes_Editor(FDisplayClusterViewport& InCameraViewport, ADisplayClusterRootActor& RootActor, UDisplayClusterICVFXCameraComponent& InCameraComponent);
 	static void PreviewReuseChromakeyViewportWithinClusterNodes_Editor(FDisplayClusterViewport& InChromakeyViewport, ADisplayClusterRootActor& RootActor, UDisplayClusterICVFXCameraComponent& InCameraComponent);
 #endif
diff --git a/Engine/Plugins/Runtime/nDisplay/Source/DisplayCluster/Private/Render/Viewport/Configuration/DisplayClusterViewportConfigurationHelpers_Postprocess.cpp b/Engine/Plugins/Runtime/nDisplay/Source/DisplayCluster/Private/Render/Viewport/Configuration/DisplayClusterViewportConfigurationHelpers_Postprocess.cpp
--- a/Engine/Plugins/Runtime/nDisplay/Source/DisplayCluster/Private/Render/Viewport/Configuration/DisplayClusterViewportConfigurationHelpers_Postprocess.cpp
+++ b/Engine/Plugins/Runtime/nDisplay/Source/DisplayCluster/Private/Render/Viewport/Configuration/DisplayClusterViewportConfigurationHelpers_Postprocess.cpp
@@ -206,7 +206,6 @@ bool FDisplayClusterViewportConfigurationHelpers_Postprocess::IsInnerFrustumView
 bool FDisplayClusterViewportConfigurationHelpers_Postprocess::ImplUpdateInnerFrustumColorGrading(FDisplayClusterViewport& DstViewport, ADisplayClusterRootActor& RootActor, UDisplayClusterICVFXCameraComponent& InCameraComponent)
 {
 	const FDisplayClusterConfigurationICVFX_CameraSettings& CameraSettings = InCameraComponent.GetCameraSettingsICVFX();
-	const FDisplayClusterRenderFrameSettings& RenderFrameSettings = DstViewport.GetRenderFrameSettings();
 
 	// per node color grading first (it includes all nodes blending too)
 	const FString& ClusterNodeId = DstViewport.GetClusterNodeId();
diff --git a/Engine/Plugins/Runtime/nDisplay/Source/DisplayCluster/Private/Render/Viewport/Configuration/DisplayClusterViewportConfigurationHelpers_Visibility.cpp b/Engine/Plugins/Runtime/nDisplay/Source/DisplayCluster/Private/Render/Viewport/Configuration/DisplayClusterViewportConfigurationHelpers_Visibility.cpp
--- a/Engine/Plugins/Runtime/nDisplay/Source/DisplayCluster/Private/Render/Viewport/Configuration/DisplayClusterViewportConfigurationHelpers_Visibility.cpp
+++ b/Engine/Plugins/Runtime/nDisplay/Source/DisplayCluster/Private/Render/Viewport/Configuration/DisplayClusterViewportConfigurationHelpers_Visibility.cpp
@@ -107,7 +107,7 @@ void FDisplayClusterViewportConfigurationHelpers_Visibility::AppendHideList_ICVF
 	DstViewport.VisibilitySettings.AppendHideList(ActorLayers, AdditionalComponentsList);
 }
 
-void FDisplayClusterViewportConfigurationHelpers_Visibility::UpdateHideList_ICVFX(TArray<FDisplayClusterViewport*>& DstViewports, ADisplayClusterRootActor& InRootActor)
+void FDisplayClusterViewportConfigurationHelpers_Visibility::UpdateHideList_ICVFX(TArray<TSharedPtr<FDisplayClusterViewport, ESPMode::ThreadSafe>>& DstViewports, ADisplayClusterRootActor& InRootActor)
 {
 	if (DstViewports.Num() > 0)
 	{
@@ -139,9 +139,9 @@ void FDisplayClusterViewportConfigurationHelpers_Visibility::UpdateHideList_ICVF
 		ImplCollectVisibility(InRootActor, StageSettings.OuterViewportHideList, OuterActorLayerNames, OuterAdditionalComponentsList);
 
 		// Update hide list for all desired viewports:
-		for (FDisplayClusterViewport* ViewportIt : DstViewports)
+		for (TSharedPtr<FDisplayClusterViewport, ESPMode::ThreadSafe>& ViewportIt : DstViewports)
 		{
-			if (ViewportIt)
+			if (ViewportIt.IsValid())
 			{
 				ViewportIt->VisibilitySettings.UpdateConfiguration(EDisplayClusterViewport_VisibilityMode::Hide, ActorLayerNames, AdditionalComponentsList);
 
diff --git a/Engine/Plugins/Runtime/nDisplay/Source/DisplayCluster/Private/Render/Viewport/Configuration/DisplayClusterViewportConfigurationHelpers_Visibility.h b/Engine/Plugins/Runtime/nDisplay/Source/DisplayCluster/Private/Render/Viewport/Configuration/DisplayClusterViewportConfigurationHelpers_Visibility.h
--- a/Engine/Plugins/Runtime/nDisplay/Source/DisplayCluster/Private/Render/Viewport/Configuration/DisplayClusterViewportConfigurationHelpers_Visibility.h
+++ b/Engine/Plugins/Runtime/nDisplay/Source/DisplayCluster/Private/Render/Viewport/Configuration/DisplayClusterViewportConfigurationHelpers_Visibility.h
@@ -19,7 +19,7 @@ public:
 	static void UpdateShowOnlyList(FDisplayClusterViewport& DstViewport, ADisplayClusterRootActor& RootActor, const FDisplayClusterConfigurationICVFX_VisibilityList& InVisibilityList);
 
 	// Update hide lists for DstViewports (lightcards, chromakey, stage settings hide list, outer viewports hide list, etc)
-	static void UpdateHideList_ICVFX(TArray<FDisplayClusterViewport*>& DstViewports, ADisplayClusterRootActor& InRootActor);
+	static void UpdateHideList_ICVFX(TArray<TSharedPtr<FDisplayClusterViewport, ESPMode::ThreadSafe>>& DstViewports, ADisplayClusterRootActor& InRootActor);
 	
 	// Append to exist hide list (must be call after UpdateHideList_ICVFX)
 	static void AppendHideList_ICVFX(FDisplayClusterViewport& DstViewport, ADisplayClusterRootActor& InRootActor, const FDisplayClusterConfigurationICVFX_VisibilityList& InHideList);
diff --git a/Engine/Plugins/Runtime/nDisplay/Source/DisplayCluster/Private/Render/Viewport/Configuration/DisplayClusterViewportConfigurationICVFX.cpp b/Engine/Plugins/Runtime/nDisplay/Source/DisplayCluster/Private/Render/Viewport/Configuration/DisplayClusterViewportConfigurationICVFX.cpp
--- a/Engine/Plugins/Runtime/nDisplay/Source/DisplayCluster/Private/Render/Viewport/Configuration/DisplayClusterViewportConfigurationICVFX.cpp
+++ b/Engine/Plugins/Runtime/nDisplay/Source/DisplayCluster/Private/Render/Viewport/Configuration/DisplayClusterViewportConfigurationICVFX.cpp
@@ -204,14 +204,17 @@ public:
 			if (CameraConfiguration.GetShaderParametersCameraSettings(RootActor, CameraComponent, ShaderParametersCameraSettings))
 			{
 				// Add this camera data to all visible targets:
-				for (FDisplayClusterViewport* ViewportIt : VisibleTargets)
+				for (TSharedPtr<FDisplayClusterViewport, ESPMode::ThreadSafe>& ViewportIt : VisibleTargets)
 				{
-					ViewportIt->RenderSettingsICVFX.ICVFX.Cameras.Add(ShaderParametersCameraSettings);
-
-					// At lest one target must accept chromakey
-					if (!EnumHasAnyFlags(ViewportIt->RenderSettingsICVFX.Flags, EDisplayClusterViewportICVFXFlags::DisableChromakey))
+					if (ViewportIt.IsValid())
 					{
-						bAllowChromakey = true;
+						ViewportIt->RenderSettingsICVFX.ICVFX.Cameras.Add(ShaderParametersCameraSettings);
+
+						// At lest one target must accept chromakey
+						if (!EnumHasAnyFlags(ViewportIt->RenderSettingsICVFX.Flags, EDisplayClusterViewportICVFXFlags::DisableChromakey))
+						{
+							bAllowChromakey = true;
+						}
 					}
 				}
 			}
@@ -251,9 +254,9 @@ protected:
 		const FDisplayClusterConfigurationICVFX_ChromakeySettings& ChromakeySettings = CameraSettings.Chromakey;
 
 		// Assign the overlap chromakey to all supported targets, always done regardless of custom chromakey settings
-		for (FDisplayClusterViewport* TargetIt : VisibleTargets)
+		for (TSharedPtr<FDisplayClusterViewport, ESPMode::ThreadSafe>& TargetIt : VisibleTargets)
 		{
-			if (TargetIt)
+			if (TargetIt.IsValid())
 			{
 				CameraConfiguration.AssignOverlapChromakeyToTargetViewport(*TargetIt, RootActor);
 			}
@@ -282,9 +285,9 @@ protected:
 		}
 
 		// Assign this chromakey to all supported targets
-		for (FDisplayClusterViewport* TargetIt : VisibleTargets)
+		for (TSharedPtr<FDisplayClusterViewport, ESPMode::ThreadSafe>& TargetIt : VisibleTargets)
 		{
-			if (TargetIt)
+			if (TargetIt.IsValid())
 			{
 				CameraConfiguration.AssignChromakeyToTargetViewport(*TargetIt, RootActor, CameraComponent);
 			}
@@ -295,7 +298,7 @@ protected:
 
 public:
 	// List of targets for this camera
-	TArray<FDisplayClusterViewport*> VisibleTargets;
+	TArray<TSharedPtr<FDisplayClusterViewport, ESPMode::ThreadSafe>> VisibleTargets;
 
 private:
 	FDisplayClusterViewportConfigurationCameraViewport CameraConfiguration;
@@ -374,7 +377,7 @@ void FDisplayClusterViewportConfigurationICVFX::Update()
 			return;
 		}
 
-		TArray<FDisplayClusterViewport*> TargetViewports;
+		TArray<TSharedPtr<FDisplayClusterViewport, ESPMode::ThreadSafe>> TargetViewports;
 		const EDisplayClusterViewportICVFXFlags TargetViewportsFlags = ImplGetTargetViewports(*ViewportManager, TargetViewports);
 
 		// Find ICVFX target viewports
@@ -393,15 +396,15 @@ void FDisplayClusterViewportConfigurationICVFX::Update()
 			if (StageCameras.Num() > 0)
 			{
 				// Collect visible targets for cameras:
-				for (FDisplayClusterViewport* TargetIt : TargetViewports)
+				for (const TSharedPtr<FDisplayClusterViewport, ESPMode::ThreadSafe>& TargetIt : TargetViewports)
 				{
 					// Target viewpot must support camera render:
-					if (!EnumHasAnyFlags(TargetIt->RenderSettingsICVFX.Flags, EDisplayClusterViewportICVFXFlags::DisableCamera))
+					if (TargetIt.IsValid() && !EnumHasAnyFlags(TargetIt->RenderSettingsICVFX.Flags, EDisplayClusterViewportICVFXFlags::DisableCamera))
 					{
 						// Add this target to all cameras visible on it
 						for (FDisplayClusterViewportConfigurationCameraICVFX& CameraIt : StageCameras)
 						{
-						if (CameraIt.IsCameraProjectionVisibleOnViewport(TargetIt)
+						if (CameraIt.IsCameraProjectionVisibleOnViewport(TargetIt.Get())
 							&& !CameraIt.GetCameraSettings().HiddenICVFXViewports.ItemNames.Contains(TargetIt->GetId()))
 							{
 								CameraIt.VisibleTargets.Add(TargetIt);
@@ -429,10 +432,10 @@ void FDisplayClusterViewportConfigurationICVFX::Update()
 					const bool bUseLightCard = FDisplayClusterViewportConfigurationHelpers_ICVFX::IsShouldUseLightcard(StageSettings.Lightcard);
 					const bool bUseUVLightCard = FDisplayClusterViewportConfigurationHelpers_ICVFX::IsShouldUseUVLightcard(*ViewportManager, StageSettings.Lightcard);
 
-					for (FDisplayClusterViewport* TargetIt : TargetViewports)
+					for (const TSharedPtr<FDisplayClusterViewport, ESPMode::ThreadSafe>& TargetIt : TargetViewports)
 					{
 						// only for support targets
-						if (TargetIt && !EnumHasAnyFlags(TargetIt->RenderSettingsICVFX.Flags, EDisplayClusterViewportICVFXFlags::DisableLightcard))
+						if (TargetIt.IsValid() && !EnumHasAnyFlags(TargetIt->RenderSettingsICVFX.Flags, EDisplayClusterViewportICVFXFlags::DisableLightcard))
 						{
 							if (bUseLightCard)
 							{
@@ -448,9 +451,9 @@ void FDisplayClusterViewportConfigurationICVFX::Update()
 				}
 			}
 
-			for (FDisplayClusterViewport* TargetIt : TargetViewports)
+			for (const TSharedPtr<FDisplayClusterViewport, ESPMode::ThreadSafe>& TargetIt : TargetViewports)
 			{
-				if (TargetIt)
+				if (TargetIt.IsValid())
 				{
 					// Sort cameras by render order for each target
 					TargetIt->RenderSettingsICVFX.ICVFX.SortCamerasRenderOrder();
@@ -464,17 +467,20 @@ void FDisplayClusterViewportConfigurationICVFX::Update()
 			if (StageSettings.bFreezeRenderOuterViewports)
 			{
 				// Freeze rendering for outer viewports
-				for (FDisplayClusterViewport* TargetIt : TargetViewports)
+				for (const TSharedPtr<FDisplayClusterViewport, ESPMode::ThreadSafe>& TargetIt : TargetViewports)
 				{
-					TargetIt->RenderSettings.bFreezeRendering = true;
+					if (TargetIt.IsValid())
+					{
+						TargetIt->RenderSettings.bFreezeRendering = true;
+					}
 				}
 
 				// Freeze render for lightcards when outer viewports freezed
 				if (StageSettings.Lightcard.bIgnoreOuterViewportsFreezingForLightcards == false)
 				{
-					for (FDisplayClusterViewport* ViewportIt : ViewportManager->ImplGetViewports())
+					for (const TSharedPtr<FDisplayClusterViewport, ESPMode::ThreadSafe>& ViewportIt : ViewportManager->ImplGetViewports())
 					{
-						if (ViewportIt && EnumHasAnyFlags(ViewportIt->RenderSettingsICVFX.RuntimeFlags, EDisplayClusterViewportRuntimeICVFXFlags::Lightcard | EDisplayClusterViewportRuntimeICVFXFlags::UVLightcard))
+						if (ViewportIt.IsValid() && EnumHasAnyFlags(ViewportIt->RenderSettingsICVFX.RuntimeFlags, EDisplayClusterViewportRuntimeICVFXFlags::Lightcard | EDisplayClusterViewportRuntimeICVFXFlags::UVLightcard))
 						{
 							ViewportIt->RenderSettings.bFreezeRendering = true;
 						}
@@ -520,9 +526,12 @@ void FDisplayClusterViewportConfigurationICVFX::PostUpdate()
 		// Support alpha channel rendering for WarpBlend
 		if (GDisplayClusterEnableAlphaChannelRendering != 0)
 		{
-			for (FDisplayClusterViewport* ViewportIt : ViewportManager->ImplGetViewports())
+			for (const TSharedPtr<FDisplayClusterViewport, ESPMode::ThreadSafe>& ViewportIt : ViewportManager->ImplGetViewports())
 			{
-				ViewportIt->RenderSettings.bWarpBlendRenderAlphaChannel = true;
+				if (ViewportIt.IsValid())
+				{
+					ViewportIt->RenderSettings.bWarpBlendRenderAlphaChannel = true;
+				}
 			}
 		}
 	}
@@ -543,18 +552,21 @@ void FDisplayClusterViewportConfigurationICVFX::UpdateCameraHideList(FDisplayClu
 
 void FDisplayClusterViewportConfigurationICVFX::UpdateHideList(FDisplayClusterViewportManager& ViewportManager)
 {
-	TArray<FDisplayClusterViewport*> ICVFXViewports;
+	TArray<TSharedPtr<FDisplayClusterViewport, ESPMode::ThreadSafe>> ICVFXViewports;
 
 	// Collect viewports, that use ICVFX hide list
-	for (FDisplayClusterViewport* ViewportIt : ViewportManager.ImplGetViewports())
+	for (const TSharedPtr<FDisplayClusterViewport, ESPMode::ThreadSafe>& ViewportIt : ViewportManager.ImplGetViewports())
 	{
-		const bool bInternalResource = EnumHasAnyFlags(ViewportIt->RenderSettingsICVFX.RuntimeFlags, EDisplayClusterViewportRuntimeICVFXFlags::InternalResource);
-		const bool bIsInCamera = EnumHasAnyFlags(ViewportIt->RenderSettingsICVFX.RuntimeFlags, EDisplayClusterViewportRuntimeICVFXFlags::InCamera);
-		const bool bICVFX_Enable = EnumHasAnyFlags(ViewportIt->RenderSettingsICVFX.Flags, EDisplayClusterViewportICVFXFlags::Enable);
-
-		if ((bICVFX_Enable && !bInternalResource)  || (bInternalResource && bIsInCamera))
+		if (ViewportIt.IsValid())
 		{
-			ICVFXViewports.Add(ViewportIt);
+			const bool bInternalResource = EnumHasAnyFlags(ViewportIt->RenderSettingsICVFX.RuntimeFlags, EDisplayClusterViewportRuntimeICVFXFlags::InternalResource);
+			const bool bIsInCamera = EnumHasAnyFlags(ViewportIt->RenderSettingsICVFX.RuntimeFlags, EDisplayClusterViewportRuntimeICVFXFlags::InCamera);
+			const bool bICVFX_Enable = EnumHasAnyFlags(ViewportIt->RenderSettingsICVFX.Flags, EDisplayClusterViewportICVFXFlags::Enable);
+
+			if ((bICVFX_Enable && !bInternalResource) || (bInternalResource && bIsInCamera))
+			{
+				ICVFXViewports.Add(ViewportIt);
+			}
 		}
 	}
 
@@ -564,10 +576,10 @@ void FDisplayClusterViewportConfigurationICVFX::UpdateHideList(FDisplayClusterVi
 
 void FDisplayClusterViewportConfigurationICVFX::ImplBeginReallocateViewports(FDisplayClusterViewportManager& ViewportManager)
 {
-	for (FDisplayClusterViewport* ViewportIt : ViewportManager.ImplGetViewports())
+	for (const TSharedPtr<FDisplayClusterViewport, ESPMode::ThreadSafe>& ViewportIt : ViewportManager.ImplGetViewports())
 	{
 		// Runtime icvfx viewport support reallocate feature:
-		if (EnumHasAllFlags(ViewportIt->RenderSettingsICVFX.RuntimeFlags, EDisplayClusterViewportRuntimeICVFXFlags::InternalResource))
+		if (ViewportIt.IsValid() && EnumHasAllFlags(ViewportIt->RenderSettingsICVFX.RuntimeFlags, EDisplayClusterViewportRuntimeICVFXFlags::InternalResource))
 		{
 			// Mark all dynamic ICVFX viewports for delete
 			EnumAddFlags(ViewportIt->RenderSettingsICVFX.RuntimeFlags, EDisplayClusterViewportRuntimeICVFXFlags::Unused);
@@ -577,21 +589,24 @@ void FDisplayClusterViewportConfigurationICVFX::ImplBeginReallocateViewports(FDi
 
 void FDisplayClusterViewportConfigurationICVFX::ImplFinishReallocateViewports(FDisplayClusterViewportManager& ViewportManager)
 {
-	TArray<FDisplayClusterViewport*> UnusedViewports;
+	TArray<TSharedPtr<FDisplayClusterViewport, ESPMode::ThreadSafe>> UnusedViewports;
 
 	// Collect all unused viewports for remove
-	for (FDisplayClusterViewport* ViewportIt : ViewportManager.ImplGetViewports())
+	for (const TSharedPtr<FDisplayClusterViewport, ESPMode::ThreadSafe>& ViewportIt : ViewportManager.ImplGetViewports())
 	{
-		if (EnumHasAllFlags(ViewportIt->RenderSettingsICVFX.RuntimeFlags, EDisplayClusterViewportRuntimeICVFXFlags::Unused))
+		if (ViewportIt.IsValid() && EnumHasAllFlags(ViewportIt->RenderSettingsICVFX.RuntimeFlags, EDisplayClusterViewportRuntimeICVFXFlags::Unused))
 		{
 			UnusedViewports.Add(ViewportIt);
 		}
 	}
 
 	// Delete unused viewports:
-	for (FDisplayClusterViewport* RemoveViewportIt : UnusedViewports)
+	for (TSharedPtr<FDisplayClusterViewport, ESPMode::ThreadSafe>& RemoveViewportIt : UnusedViewports)
 	{
-		ViewportManager.DeleteViewport(RemoveViewportIt->GetId());
+		if (RemoveViewportIt.IsValid())
+		{
+			ViewportManager.DeleteViewport(RemoveViewportIt->GetId());
+		}
 	}
 
 	UnusedViewports.Empty();
@@ -621,15 +636,15 @@ void FDisplayClusterViewportConfigurationICVFX::ImplGetCameras()
 	}
 }
 
-EDisplayClusterViewportICVFXFlags FDisplayClusterViewportConfigurationICVFX::ImplGetTargetViewports(FDisplayClusterViewportManager& ViewportManager, TArray<class FDisplayClusterViewport*>& OutTargets)
+EDisplayClusterViewportICVFXFlags FDisplayClusterViewportConfigurationICVFX::ImplGetTargetViewports(FDisplayClusterViewportManager& ViewportManager, TArray<TSharedPtr<FDisplayClusterViewport, ESPMode::ThreadSafe>>& OutTargets)
 {
 	EDisplayClusterViewportICVFXFlags InvFlags = EDisplayClusterViewportICVFXFlags::None;
 
 	// Collect invertet disable flags from all target viewports
-	for (FDisplayClusterViewport* ViewportIt : ViewportManager.ImplGetViewports())
+	for (const TSharedPtr<FDisplayClusterViewport, ESPMode::ThreadSafe>& ViewportIt : ViewportManager.ImplGetViewports())
 	{
 		// Process only external viewports:
-		if (!EnumHasAnyFlags(ViewportIt->RenderSettingsICVFX.RuntimeFlags, EDisplayClusterViewportRuntimeICVFXFlags::InternalResource))
+		if (ViewportIt.IsValid() && !EnumHasAnyFlags(ViewportIt->RenderSettingsICVFX.RuntimeFlags, EDisplayClusterViewportRuntimeICVFXFlags::InternalResource))
 		{
 			//Raise new projection target if possible
 			if (ViewportIt->RenderSettings.bEnable && EnumHasAnyFlags(ViewportIt->RenderSettingsICVFX.Flags, EDisplayClusterViewportICVFXFlags::Enable))
diff --git a/Engine/Plugins/Runtime/nDisplay/Source/DisplayCluster/Private/Render/Viewport/Configuration/DisplayClusterViewportConfigurationICVFX.h b/Engine/Plugins/Runtime/nDisplay/Source/DisplayCluster/Private/Render/Viewport/Configuration/DisplayClusterViewportConfigurationICVFX.h
--- a/Engine/Plugins/Runtime/nDisplay/Source/DisplayCluster/Private/Render/Viewport/Configuration/DisplayClusterViewportConfigurationICVFX.h
+++ b/Engine/Plugins/Runtime/nDisplay/Source/DisplayCluster/Private/Render/Viewport/Configuration/DisplayClusterViewportConfigurationICVFX.h
@@ -31,7 +31,7 @@ private:
 	void ImplFinishReallocateViewports(FDisplayClusterViewportManager& ViewportManager);
 
 	void ImplGetCameras();
-	EDisplayClusterViewportICVFXFlags ImplGetTargetViewports(FDisplayClusterViewportManager& ViewportManager, TArray<FDisplayClusterViewport*>& OutTargets);
+	EDisplayClusterViewportICVFXFlags ImplGetTargetViewports(FDisplayClusterViewportManager& ViewportManager, TArray<TSharedPtr<FDisplayClusterViewport, ESPMode::ThreadSafe>>& OutTargets);
 
 	/** Create a LightCard viewport for the  BaseViewport. */
 	bool CreateLightcardViewport(FDisplayClusterViewport& BaseViewport);
diff --git a/Engine/Plugins/Runtime/nDisplay/Source/DisplayCluster/Private/Render/Viewport/Configuration/DisplayClusterViewportConfigurationProjectionPolicy.cpp b/Engine/Plugins/Runtime/nDisplay/Source/DisplayCluster/Private/Render/Viewport/Configuration/DisplayClusterViewportConfigurationProjectionPolicy.cpp
--- a/Engine/Plugins/Runtime/nDisplay/Source/DisplayCluster/Private/Render/Viewport/Configuration/DisplayClusterViewportConfigurationProjectionPolicy.cpp
+++ b/Engine/Plugins/Runtime/nDisplay/Source/DisplayCluster/Private/Render/Viewport/Configuration/DisplayClusterViewportConfigurationProjectionPolicy.cpp
@@ -28,10 +28,10 @@
 
 void FDisplayClusterViewportConfigurationProjectionPolicy::Update()
 {
-	for (FDisplayClusterViewport* ViewportIt : ViewportManager.ImplGetViewports())
+	for (const TSharedPtr<FDisplayClusterViewport, ESPMode::ThreadSafe>& ViewportIt : ViewportManager.ImplGetViewports())
 	{
 		// ignore ICVFX internal resources
-		if (ViewportIt && !EnumHasAnyFlags(ViewportIt->GetRenderSettingsICVFX().RuntimeFlags, EDisplayClusterViewportRuntimeICVFXFlags::InternalResource))
+		if (ViewportIt.IsValid() && !EnumHasAnyFlags(ViewportIt->GetRenderSettingsICVFX().RuntimeFlags, EDisplayClusterViewportRuntimeICVFXFlags::InternalResource))
 		{
 			// Support advanced logic for 'camera' projection policy
 			if (ViewportIt->ProjectionPolicy.IsValid() && ViewportIt->ProjectionPolicy->GetType().Compare(DisplayClusterProjectionStrings::projection::Camera) == 0)
diff --git a/Engine/Plugins/Runtime/nDisplay/Source/DisplayCluster/Private/Render/Viewport/Containers/DisplayClusterViewportProxyData.cpp b/Engine/Plugins/Runtime/nDisplay/Source/DisplayCluster/Private/Render/Viewport/Containers/DisplayClusterViewportProxyData.cpp
--- a/Engine/Plugins/Runtime/nDisplay/Source/DisplayCluster/Private/Render/Viewport/Containers/DisplayClusterViewportProxyData.cpp
+++ b/Engine/Plugins/Runtime/nDisplay/Source/DisplayCluster/Private/Render/Viewport/Containers/DisplayClusterViewportProxyData.cpp
@@ -8,10 +8,10 @@
 #include "Render/Viewport/DisplayClusterViewportProxy.h"
 
 
-FDisplayClusterViewportProxyData::FDisplayClusterViewportProxyData(const FDisplayClusterViewport* SrcViewport)
+FDisplayClusterViewportProxyData::FDisplayClusterViewportProxyData(const TSharedPtr<FDisplayClusterViewport, ESPMode::ThreadSafe>& SrcViewport)
 {
 	check(IsInGameThread());
-	check(SrcViewport);
+	check(SrcViewport.IsValid());
 
 	DstViewportProxy = SrcViewport->ViewportProxy;
 
diff --git a/Engine/Plugins/Runtime/nDisplay/Source/DisplayCluster/Private/Render/Viewport/Containers/DisplayClusterViewportProxyData.h b/Engine/Plugins/Runtime/nDisplay/Source/DisplayCluster/Private/Render/Viewport/Containers/DisplayClusterViewportProxyData.h
--- a/Engine/Plugins/Runtime/nDisplay/Source/DisplayCluster/Private/Render/Viewport/Containers/DisplayClusterViewportProxyData.h
+++ b/Engine/Plugins/Runtime/nDisplay/Source/DisplayCluster/Private/Render/Viewport/Containers/DisplayClusterViewportProxyData.h
@@ -15,6 +15,7 @@ class FDisplayClusterViewportProxy;
 class FDisplayClusterViewport_Context;
 class IDisplayClusterRender_MeshComponent;
 class FDisplayClusterViewport_OpenColorIO;
+class FSceneViewStateReference;
 
 //
 // Container for data exchange game->render threads
@@ -22,7 +23,7 @@ class FDisplayClusterViewport_OpenColorIO;
 class FDisplayClusterViewportProxyData
 {
 public:
-	FDisplayClusterViewportProxyData(const FDisplayClusterViewport* SrcViewport);
+	FDisplayClusterViewportProxyData(const TSharedPtr<FDisplayClusterViewport, ESPMode::ThreadSafe>& SrcViewport);
 	~FDisplayClusterViewportProxyData() = default;
 
 	void UpdateProxy_RenderThread() const;
diff --git a/Engine/Plugins/Runtime/nDisplay/Source/DisplayCluster/Private/Render/Viewport/Containers/DisplayClusterViewportRemap.cpp b/Engine/Plugins/Runtime/nDisplay/Source/DisplayCluster/Private/Render/Viewport/Containers/DisplayClusterViewportRemap.cpp
--- a/Engine/Plugins/Runtime/nDisplay/Source/DisplayCluster/Private/Render/Viewport/Containers/DisplayClusterViewportRemap.cpp
+++ b/Engine/Plugins/Runtime/nDisplay/Source/DisplayCluster/Private/Render/Viewport/Containers/DisplayClusterViewportRemap.cpp
@@ -64,18 +64,11 @@ bool FDisplayClusterViewportRemap::UpdateConfiguration(const FDisplayClusterView
 
 void FDisplayClusterViewportRemap::Update(const FDisplayClusterViewport& InViewport, const FIntPoint& InRenderFrameSize)
 {
-#if WITH_EDITOR
-	const FDisplayClusterRenderFrameSettings& RenderFrameSettings = InViewport.Owner.GetRenderFrameSettings();
-
-	switch(RenderFrameSettings.RenderMode)
+	// Preview in editor not support this feature
+	if (InViewport.GetRenderMode() == EDisplayClusterRenderFrameMode::PreviewInScene)
 	{
-	case EDisplayClusterRenderFrameMode::PreviewInScene:
-		// Preview in editor not support this feature
 		return;
-	default:
-		break;
 	}
-#endif
 
 	const FDisplayClusterViewport_RenderSettings& RenderSettings = InViewport.GetRenderSettings();
 
diff --git a/Engine/Plugins/Runtime/nDisplay/Source/DisplayCluster/Private/Render/Viewport/DisplayClusterViewport.cpp b/Engine/Plugins/Runtime/nDisplay/Source/DisplayCluster/Private/Render/Viewport/DisplayClusterViewport.cpp
--- a/Engine/Plugins/Runtime/nDisplay/Source/DisplayCluster/Private/Render/Viewport/DisplayClusterViewport.cpp
+++ b/Engine/Plugins/Runtime/nDisplay/Source/DisplayCluster/Private/Render/Viewport/DisplayClusterViewport.cpp
@@ -27,6 +27,8 @@
 
 #include "DisplayClusterSceneViewExtensions.h"
 
+#include "LegacyScreenPercentageDriver.h"
+
 #include "Misc/DisplayClusterLog.h"
 
 int32 GDisplayClusterMultiGPUEnable = 1;
@@ -44,33 +46,40 @@ FDisplayClusterViewport::FDisplayClusterViewport(FDisplayClusterViewportManager&
 	: UninitializedProjectionPolicy(InProjectionPolicy)
 	, ViewportId(InViewportId)
 	, ClusterNodeId(InClusterNodeId)
-	, Owner(InOwner)
+	, ViewportManagerWeakPtr(InOwner.AsShared())
+	, ViewportManagerProxyWeakPtr(InOwner.GetViewportManagerProxy())
 {
 	check(!ClusterNodeId.IsEmpty());
 	check(!ViewportId.IsEmpty());
 	check(UninitializedProjectionPolicy.IsValid());
 
-	// Create scene proxy pair with on game thread. Outside, in ViewportManager added to proxy array on render thread
-	ViewportProxy = MakeShared<FDisplayClusterViewportProxy, ESPMode::ThreadSafe>(*this);
+	if (ViewportManagerProxyWeakPtr.IsValid())
+	{
+		// Create scene proxy pair with on game thread. Outside, in ViewportManager added to proxy array on render thread
+		ViewportProxy = MakeShared<FDisplayClusterViewportProxy, ESPMode::ThreadSafe>(*this);
 
-	// Add viewport proxy on renderthread
-	ENQUEUE_RENDER_COMMAND(CreateDisplayClusterViewportProxy)(
-		[ViewportManagerProxy = Owner.GetViewportManagerProxy(), ViewportProxy = ViewportProxy](FRHICommandListImmediate& RHICmdList)
-		{
-			ViewportManagerProxy->CreateViewport_RenderThread(ViewportProxy);
-		}
-	);
+		// Add viewport proxy on renderthread
+		ENQUEUE_RENDER_COMMAND(CreateDisplayClusterViewportProxy)(
+			[ViewportManagerProxy = ViewportManagerProxyWeakPtr.Pin(), ViewportProxy = ViewportProxy](FRHICommandListImmediate& RHICmdList)
+			{
+				ViewportManagerProxy->CreateViewport_RenderThread(ViewportProxy);
+			}
+		);
+	}
 }
 
 FDisplayClusterViewport::~FDisplayClusterViewport()
 {
-	// Remove viewport proxy on render_thread
-	ENQUEUE_RENDER_COMMAND(DeleteDisplayClusterViewportProxy)(
-		[ViewportManagerProxy = Owner.GetViewportManagerProxy(), ViewportProxy = ViewportProxy](FRHICommandListImmediate& RHICmdList)
-		{
-			ViewportManagerProxy->DeleteViewport_RenderThread(ViewportProxy);
-		}
-	);
+	if (ViewportManagerProxyWeakPtr.IsValid())
+	{
+		// Remove viewport proxy on render_thread
+		ENQUEUE_RENDER_COMMAND(DeleteDisplayClusterViewportProxy)(
+			[ViewportManagerProxy = ViewportManagerProxyWeakPtr.Pin(), ViewportProxy = ViewportProxy](FRHICommandListImmediate& RHICmdList)
+			{
+				ViewportManagerProxy->DeleteViewport_RenderThread(ViewportProxy);
+			}
+		);
+	}
 
 	ViewportProxy.Reset();
 	OpenColorIO.Reset();
@@ -82,18 +91,98 @@ FDisplayClusterViewport::~FDisplayClusterViewport()
 	ProjectionPolicy.Reset();
 	UninitializedProjectionPolicy.Reset();
 
-	// Reset RTT size after viewport delete
-	Owner.ResetSceneRenderTargetSize();
+	if (FDisplayClusterViewportManager* ViewportManager = GetViewportManagerImpl())
+	{
+		// Reset RTT size after viewport delete
+		ViewportManager->ResetSceneRenderTargetSize();
+	}
+}
+
+IDisplayClusterViewportManager* FDisplayClusterViewport::GetViewportManager() const
+{
+	return GetViewportManagerImpl();
+}
+
+FDisplayClusterViewportManager* FDisplayClusterViewport::GetViewportManagerImpl() const
+{
+	return GetViewportManagerRefImpl().Get();
+}
+
+FDisplayClusterViewportManagerProxy* FDisplayClusterViewport::GetViewportManagerProxyImpl() const
+{
+	return GetViewportManagerProxyRefImpl().Get();
+}
+
+TSharedPtr<FDisplayClusterViewportManager, ESPMode::ThreadSafe> FDisplayClusterViewport::GetViewportManagerRefImpl() const
+{
+	return ViewportManagerWeakPtr.IsValid() ? ViewportManagerWeakPtr.Pin(): nullptr;
+}
+
+TSharedPtr<FDisplayClusterViewportManagerProxy, ESPMode::ThreadSafe> FDisplayClusterViewport::GetViewportManagerProxyRefImpl() const
+{
+	return ViewportManagerProxyWeakPtr.IsValid() ? ViewportManagerProxyWeakPtr.Pin() : nullptr;
+}
+
+ADisplayClusterRootActor* FDisplayClusterViewport::GetRootActor() const
+{
+	if (FDisplayClusterViewportManager* ViewportManager = GetViewportManagerImpl())
+	{
+		return ViewportManager->GetRootActor();
+	}
+
+	return nullptr;
+}
+
+UWorld* FDisplayClusterViewport::GetCurrentWorld() const
+{
+	if (FDisplayClusterViewportManager* ViewportManager = GetViewportManagerImpl())
+	{
+		return ViewportManager->GetCurrentWorld();
+	}
+
+	return nullptr;
+}
+
+bool FDisplayClusterViewport::IsCurrentWorldHasAnyType(const EWorldType::Type InWorldType1, const EWorldType::Type InWorldType2, const EWorldType::Type InWorldType3) const
+{
+	if (UWorld* CurrentWorld = GetCurrentWorld())
+	{
+		return (CurrentWorld->WorldType == InWorldType1 && InWorldType1 != EWorldType::None)
+			|| (CurrentWorld->WorldType == InWorldType2 && InWorldType2 != EWorldType::None)
+			|| (CurrentWorld->WorldType == InWorldType3 && InWorldType3 != EWorldType::None);
+	}
+
+	return false;
 }
 
-IDisplayClusterViewportManager& FDisplayClusterViewport::GetOwner() const
+const FDisplayClusterRenderFrameSettings* FDisplayClusterViewport::GetRenderFrameSettings() const
 {
-	return Owner;
+	if (FDisplayClusterViewportManager* ViewportManager = GetViewportManagerImpl())
+	{
+		return &ViewportManager->GetRenderFrameSettings();
+	}
+
+	return nullptr;
 }
 
-const FDisplayClusterRenderFrameSettings& FDisplayClusterViewport::GetRenderFrameSettings() const
+EDisplayClusterRenderFrameMode FDisplayClusterViewport::GetRenderMode() const
 {
-	return Owner.GetRenderFrameSettings();
+	if (FDisplayClusterViewportManager* ViewportManager = GetViewportManagerImpl())
+	{
+		return ViewportManager->GetRenderMode();
+	}
+
+	return EDisplayClusterRenderFrameMode::Unknown;
+}
+
+bool FDisplayClusterViewport::IsSceneOpened() const
+{
+	if (FDisplayClusterViewportManager* ViewportManager = GetViewportManagerImpl())
+	{
+		return ViewportManager->IsSceneOpened();
+	}
+
+	return false;
 }
 
 bool FDisplayClusterViewport::IsOpenColorIOEquals(const FDisplayClusterViewport& InViewport) const
@@ -120,6 +209,13 @@ bool FDisplayClusterViewport::IsOpenColorIOEquals(const FDisplayClusterViewport&
 
 const TArray<FSceneViewExtensionRef> FDisplayClusterViewport::GatherActiveExtensions(FViewport* InViewport) const
 {
+	TSharedPtr<FDisplayClusterViewportManager, ESPMode::ThreadSafe> ViewportManager = GetViewportManagerRefImpl();
+	if (!ViewportManager.IsValid())
+	{
+		// No viewport manager = no extension found.
+		return TArray<FSceneViewExtensionRef>();
+	}
+
 	// Use VE from engine for default render and MRQ:
 	switch (RenderSettings.CaptureMode)
 	{
@@ -127,15 +223,14 @@ const TArray<FSceneViewExtensionRef> FDisplayClusterViewport::GatherActiveExtens
 	case EDisplayClusterViewportCaptureMode::MoviePipeline:
 		if (InViewport)
 		{
-			FDisplayClusterSceneViewExtensionContext ViewExtensionContext(InViewport, &Owner, GetId());
+			FDisplayClusterSceneViewExtensionContext ViewExtensionContext(InViewport, ViewportManager, GetId());
 			return GEngine->ViewExtensions->GatherActiveExtensions(ViewExtensionContext);
 		}
 		else
 		{
-			UWorld* CurrentWorld = Owner.GetCurrentWorld();
-			if (CurrentWorld)
+			if (UWorld* CurrentWorld = ViewportManager->GetCurrentWorld())
 			{
-				FDisplayClusterSceneViewExtensionContext ViewExtensionContext(CurrentWorld->Scene, &Owner, GetId());
+				FDisplayClusterSceneViewExtensionContext ViewExtensionContext(CurrentWorld->Scene, ViewportManager, GetId());
 				return GEngine->ViewExtensions->GatherActiveExtensions(ViewExtensionContext);
 			}
 		}
@@ -157,7 +252,7 @@ const TArray<FSceneViewExtensionRef> FDisplayClusterViewport::GatherActiveExtens
 	case EDisplayClusterViewportCaptureMode::Lightcard:
 	{
 		// Chromakey and LightCard use only nDisplay VE for callback purposes (preserve alpha channel, etc).
-		TSharedPtr<FDisplayClusterViewportManagerViewExtension, ESPMode::ThreadSafe> ViewportManagerViewExtension = Owner.GetViewportManagerViewExtension();
+		TSharedPtr<FDisplayClusterViewportManagerViewExtension, ESPMode::ThreadSafe> ViewportManagerViewExtension = ViewportManager->GetViewportManagerViewExtension();
 		if (ViewportManagerViewExtension.IsValid())
 		{
 			OutCustomExtensions.Add(ViewportManagerViewExtension->AsShared());
@@ -178,7 +273,7 @@ const TArray<FSceneViewExtensionRef> FDisplayClusterViewport::GatherActiveExtens
 bool FDisplayClusterViewport::HandleStartScene()
 {
 	bool bResult = true;
-	if (Owner.IsSceneOpened())
+	if (IsSceneOpened())
 	{
 		if (UninitializedProjectionPolicy.IsValid())
 		{
@@ -263,6 +358,21 @@ bool FDisplayClusterViewport::ShouldUseFullSizeFrameTargetableResource() const
 	return false;
 }
 
+void FDisplayClusterViewport::SetViewportBufferRatio(const float InBufferRatio)
+{
+	const float BufferRatio = FLegacyScreenPercentageDriver::GetCVarResolutionFraction() * InBufferRatio;
+	if (RenderSettings.BufferRatio > BufferRatio)
+	{
+		if (FDisplayClusterViewportManager* ViewportManager = GetViewportManagerImpl())
+		{
+			// Reset scene RTT when buffer ratio changed down
+			ViewportManager->ResetSceneRenderTargetSize();
+		}
+	}
+
+	RenderSettings.BufferRatio = BufferRatio;
+}
+
 void FDisplayClusterViewport::SetupSceneView(uint32 ContextNum, class UWorld* World, FSceneViewFamily& InOutViewFamily, FSceneView& InOutView) const
 {
 	check(IsInGameThread());
@@ -438,7 +548,19 @@ bool FDisplayClusterViewport::UpdateFrameContexts(const uint32 InStereoViewIndex
 {
 	check(IsInGameThread());
 
-	const uint32 FrameTargetsAmount = Owner.GetViewPerViewportAmount();
+	FDisplayClusterViewportManager* ViewportManager = GetViewportManagerImpl();
+	if (!ViewportManager)
+	{
+		return false;
+	}
+
+	const uint32 FrameTargetsAmount = ViewportManager->GetViewPerViewportAmount();
+	if (FrameTargetsAmount == 0)
+	{
+		ResetFrameContexts();
+
+		return false;
+	}
 
 	FIntRect     DesiredFrameTargetRect = RenderSettings.Rect;
 	{
@@ -615,7 +737,7 @@ bool FDisplayClusterViewport::UpdateFrameContexts(const uint32 InStereoViewIndex
 		// Use the UVLightCard viewport only when this type of lightcards has been defined
 		bool bUseUVLightCardViewport = false;
 
-		TSharedPtr<FDisplayClusterViewportLightCardManager, ESPMode::ThreadSafe> LightCardManager = Owner.GetLightCardManager();
+		TSharedPtr<FDisplayClusterViewportLightCardManager, ESPMode::ThreadSafe> LightCardManager = ViewportManager->GetLightCardManager();
 		if (LightCardManager.IsValid() && LightCardManager->IsUVLightCardEnabled())
 		{
 			// Custom viewport size from LC Manager
@@ -733,25 +855,3 @@ bool FDisplayClusterViewport::UpdateFrameContexts(const uint32 InStereoViewIndex
 
 	return true;
 }
-
-// Currently, some data like EngineShowFlags and EngineGamma get updated on PostRenderViewFamily.
-// Since the viewports that have media input assigned never gets rendered in normal way, the PostRenderViewFamily
-// callback never gets called, therefore the data mentioned above never gets updated. This workaround initializes
-// those settings for all viewports. The viewports with no media input assigned will override the data
-// in PostRenderViewFamily like it was previously so nothing should be broken.
-void FDisplayClusterViewport::UpdateMediaDependencies(FViewport* InViewport)
-{
-	if (InViewport && InViewport->GetClient())
-	{
-		if (const FEngineShowFlags* const EngineShowFlags = InViewport->GetClient()->GetEngineShowFlags())
-		{
-			const float DisplayGamma = GEngine ? GEngine->DisplayGamma : 2.2f;
-
-			for (FDisplayClusterViewport_Context& Context : Contexts)
-			{
-				Context.RenderThreadData.EngineDisplayGamma = DisplayGamma;
-				Context.RenderThreadData.EngineShowFlags = *EngineShowFlags;
-			}
-		}
-	}
-}
diff --git a/Engine/Plugins/Runtime/nDisplay/Source/DisplayCluster/Private/Render/Viewport/DisplayClusterViewport.h b/Engine/Plugins/Runtime/nDisplay/Source/DisplayCluster/Private/Render/Viewport/DisplayClusterViewport.h
--- a/Engine/Plugins/Runtime/nDisplay/Source/DisplayCluster/Private/Render/Viewport/DisplayClusterViewport.h
+++ b/Engine/Plugins/Runtime/nDisplay/Source/DisplayCluster/Private/Render/Viewport/DisplayClusterViewport.h
@@ -17,9 +17,12 @@
 #include "SceneViewExtension.h"
 #include "SceneViewExtensionContext.h"
 
+#include "Templates/SharedPointer.h"
+
 class FDisplayClusterViewportRenderTargetResource;
 class FDisplayClusterViewportTextureResource;
 class FDisplayClusterViewportManager;
+class FDisplayClusterViewportManagerProxy;
 class FDisplayClusterRenderTargetManager;
 class FDisplayClusterRenderFrameManager;
 class FDisplayClusterViewportProxyData;
@@ -42,6 +45,7 @@ struct FDisplayClusterViewportConfigurationProjectionPolicy;
 
 class FDisplayClusterViewport
 	: public IDisplayClusterViewport
+	, public TSharedFromThis<FDisplayClusterViewport, ESPMode::ThreadSafe>
 {
 public:
 	FDisplayClusterViewport(FDisplayClusterViewportManager& Owner, const FString& ClusterNodeId, const FString& ViewportId, const TSharedPtr<IDisplayClusterProjectionPolicy, ESPMode::ThreadSafe>& InProjectionPolicy);
@@ -52,6 +56,18 @@ public:
 	//////////////////////////////////////////////////////
 	/// IDisplayClusterViewport
 	//////////////////////////////////////////////////////
+	virtual TSharedPtr<IDisplayClusterViewport, ESPMode::ThreadSafe> ToSharedPtr() override
+	{
+		return AsShared();
+	}
+
+	virtual TSharedPtr<const IDisplayClusterViewport, ESPMode::ThreadSafe> ToSharedPtr() const override
+	{
+		return AsShared();
+	}
+
+	virtual EDisplayClusterRenderFrameMode GetRenderMode() const override;
+
 	virtual FString GetId() const override
 	{ 
 		check(IsInGameThread());
@@ -122,35 +138,44 @@ public:
 	}
 
 	// Setup scene view for rendering specified Context
-	virtual void SetupSceneView(uint32 ContextNum, class UWorld* World, FSceneViewFamily& InViewFamily, FSceneView& InView) const override;
-
-	virtual IDisplayClusterViewportManager& GetOwner() const override;
+	virtual void SetupSceneView(uint32 ContextNum, class UWorld* World, class FSceneViewFamily& InViewFamily, FSceneView& InView) const override;
 
-	FDisplayClusterViewportManager& ImplGetOwner() const
-	{
-		return Owner;
-	}
+	virtual class IDisplayClusterViewportManager* GetViewportManager() const override;
+	virtual class ADisplayClusterRootActor* GetRootActor() const override;
+	virtual class UWorld* GetCurrentWorld() const override;
+	virtual bool IsSceneOpened() const override;
 
-	const FDisplayClusterRenderFrameSettings& GetRenderFrameSettings() const;
+	// Return true, if current world type equal to InWorldType
+	virtual bool IsCurrentWorldHasAnyType(const EWorldType::Type InWorldType1, const EWorldType::Type InWorldType2 = EWorldType::None, const EWorldType::Type InWorldType3 = EWorldType::None) const override;
 
 	//////////////////////////////////////////////////////
 	/// ~IDisplayClusterViewport
 	//////////////////////////////////////////////////////
 
+	TSharedPtr<FDisplayClusterViewportManager, ESPMode::ThreadSafe> GetViewportManagerRefImpl() const;
+	TSharedPtr<FDisplayClusterViewportManagerProxy, ESPMode::ThreadSafe> GetViewportManagerProxyRefImpl() const;
+
+	FDisplayClusterViewportManager* GetViewportManagerImpl() const;
+	FDisplayClusterViewportManagerProxy* GetViewportManagerProxyImpl() const;
+
 #if WITH_EDITOR
 	FSceneView* ImplCalcScenePreview(class FSceneViewFamilyContext& InOutViewFamily, uint32 ContextNum);
 	bool    ImplPreview_CalculateStereoViewOffset(const uint32 InContextNum, FRotator& ViewRotation, const float WorldToMeters, FVector& ViewLocation);
 	FMatrix ImplPreview_GetStereoProjectionMatrix(const uint32 InContextNum);
 
 	bool GetPreviewPixels(TSharedPtr<class FDisplayClusterViewportReadPixelsData, ESPMode::ThreadSafe>& OutPixelsData) const;
-
 #endif //WITH_EDITOR
 
+
 	// Get from logic request for additional targetable resource
 	bool ShouldUseAdditionalTargetableResource() const;
 	bool ShouldUseAdditionalFrameTargetableResource() const;
 	bool ShouldUseFullSizeFrameTargetableResource() const;
 
+	void SetViewportBufferRatio(const float InBufferRatio);
+
+	const FDisplayClusterRenderFrameSettings* GetRenderFrameSettings() const;
+
 	inline bool FindContext(const int32 ViewIndex, uint32* OutContextNum)
 	{
 		check(IsInGameThread());
@@ -202,9 +227,6 @@ public:
 	 */
 	bool UpdateFrameContexts(const uint32 InStereoViewIndex, const FDisplayClusterRenderFrameSettings& InFrameSettings);
 
-	/* Update media dependent data */
-	void UpdateMediaDependencies(class FViewport* InViewport);
-
 	/** Reset viewport contexts and resources. */
 	void ResetFrameContexts();
 
@@ -304,7 +326,9 @@ protected:
 	TArray<FDisplayClusterViewportTextureResource*> AdditionalTargetableResources;
 	TArray<FDisplayClusterViewportTextureResource*> MipsShaderResources;
 
-	FDisplayClusterViewportManager& Owner;
+	// viewport owners
+	TWeakPtr<FDisplayClusterViewportManager, ESPMode::ThreadSafe> ViewportManagerWeakPtr;
+	TWeakPtr<FDisplayClusterViewportManagerProxy, ESPMode::ThreadSafe> ViewportManagerProxyWeakPtr;
 
 private:
 	bool bProjectionPolicyCalculateViewWarningOnce = false;
diff --git a/Engine/Plugins/Runtime/nDisplay/Source/DisplayCluster/Private/Render/Viewport/DisplayClusterViewportManager.cpp b/Engine/Plugins/Runtime/nDisplay/Source/DisplayCluster/Private/Render/Viewport/DisplayClusterViewportManager.cpp
--- a/Engine/Plugins/Runtime/nDisplay/Source/DisplayCluster/Private/Render/Viewport/DisplayClusterViewportManager.cpp
+++ b/Engine/Plugins/Runtime/nDisplay/Source/DisplayCluster/Private/Render/Viewport/DisplayClusterViewportManager.cpp
@@ -30,7 +30,6 @@
 #include "DisplayClusterRootActor.h"
 
 #include "Misc/DisplayClusterLog.h"
-#include "LegacyScreenPercentageDriver.h"
 
 #include "Engine/LocalPlayer.h"
 #include "Engine/Console.h"
@@ -57,32 +56,29 @@ static FAutoConsoleVariableRef CVarDisplayClusterChromaKeyAllowNanite(
 
 FDisplayClusterViewportManager::FDisplayClusterViewportManager()
 {
-	Configuration      = MakeUnique<FDisplayClusterViewportConfiguration>(*this);
-	RenderFrameManager = MakeUnique<FDisplayClusterRenderFrameManager>();
-
 	ViewportManagerProxy = MakeShared<FDisplayClusterViewportManagerProxy, ESPMode::ThreadSafe>();
 
+	// Always reset RTT when root actor re-created
+	ResetSceneRenderTargetSize();
+}
+
+void FDisplayClusterViewportManager::Initialize()
+{
+	Configuration = MakeUnique<FDisplayClusterViewportConfiguration>(*this);
+	PostProcessManager = MakeShared<FDisplayClusterViewportPostProcessManager, ESPMode::ThreadSafe>(*this);
+
+	RenderFrameManager = MakeUnique<FDisplayClusterRenderFrameManager>();
 	RenderTargetManager = MakeShared<FDisplayClusterRenderTargetManager, ESPMode::ThreadSafe>(ViewportManagerProxy.Get());
-	PostProcessManager  = MakeShared<FDisplayClusterViewportPostProcessManager, ESPMode::ThreadSafe>(*this);
+
 	LightCardManager = MakeShared<FDisplayClusterViewportLightCardManager, ESPMode::ThreadSafe>(*this);
 
 	// initialize proxy
 	ViewportManagerProxy->Initialize(*this);
-
-	// Always reset RTT when root actor re-created
-	ResetSceneRenderTargetSize();
 }
 
 FDisplayClusterViewportManager::~FDisplayClusterViewportManager()
 {
 	// Remove viewports
-	for (FDisplayClusterViewport* ViewportIt : Viewports)
-	{
-		if (ViewportIt)
-		{
-			delete ViewportIt;
-		}
-	}
 	Viewports.Reset();
 	ClusterNodeViewports.Reset();
 
@@ -95,19 +91,8 @@ FDisplayClusterViewportManager::~FDisplayClusterViewportManager()
 		LightCardManager.Reset();
 	}
 
-	if (ViewportManagerViewExtension.IsValid())
-	{
-		// Force release of VE data since this TSharedPtr<> may also be held by other resources at this time.
-		// So Reset() doesn't actually release it right away in some situations.
-		ViewportManagerViewExtension->Release();
-		ViewportManagerViewExtension.Reset();
-	}
-
-	if (ViewportManagerViewPointExtension.IsValid())
-	{
-		ViewportManagerViewPointExtension->Release();
-		ViewportManagerViewPointExtension.Reset();
-	}
+	ViewportManagerViewExtension.Reset();
+	ViewportManagerViewPointExtension.Reset();
 
 	Configuration.Reset();
 
@@ -164,11 +149,11 @@ void FDisplayClusterViewportManager::StartScene(UWorld* InWorld)
 
 	CurrentWorldRef = TWeakObjectPtr<UWorld>(InWorld);
 
-	for (FDisplayClusterViewport* Viewport : Viewports)
+	for (TSharedPtr<FDisplayClusterViewport, ESPMode::ThreadSafe>& ViewportIt : Viewports)
 	{
-		if (Viewport)
+		if (ViewportIt.IsValid())
 		{
-			Viewport->HandleStartScene();
+			ViewportIt->HandleStartScene();
 		}
 	}
 
@@ -187,11 +172,11 @@ void FDisplayClusterViewportManager::EndScene()
 {
 	check(IsInGameThread());
 
-	for (FDisplayClusterViewport* Viewport : Viewports)
+	for (TSharedPtr<FDisplayClusterViewport, ESPMode::ThreadSafe>& ViewportIt : Viewports)
 	{
-		if (Viewport)
+		if (ViewportIt.IsValid())
 		{
-			Viewport->HandleEndScene();
+			ViewportIt->HandleEndScene();
 		}
 	}
 
@@ -218,28 +203,16 @@ void FDisplayClusterViewportManager::ResetScene()
 		LightCardManager->HandleStartScene();
 	}
 
-	for (FDisplayClusterViewport* Viewport : Viewports)
+	for (TSharedPtr<FDisplayClusterViewport, ESPMode::ThreadSafe>& ViewportIt : Viewports)
 	{
-		if (Viewport)
+		if (ViewportIt.IsValid())
 		{
-			Viewport->HandleEndScene();
-			Viewport->HandleStartScene();
+			ViewportIt->HandleEndScene();
+			ViewportIt->HandleStartScene();
 		}
 	}
 }
 
-void FDisplayClusterViewportManager::SetViewportBufferRatio(FDisplayClusterViewport& DstViewport, float InBufferRatio)
-{
-	const float BufferRatio = FLegacyScreenPercentageDriver::GetCVarResolutionFraction() * InBufferRatio;
-	if (DstViewport.RenderSettings.BufferRatio > BufferRatio)
-	{
-		// Reset scene RTT when buffer ratio changed down
-		ResetSceneRenderTargetSize();
-	}
-
-	DstViewport.RenderSettings.BufferRatio = BufferRatio;
-}
-
 void FDisplayClusterViewportManager::HandleViewportRTTChanges(const TArray<FDisplayClusterViewport_Context>& PrevContexts, const TArray<FDisplayClusterViewport_Context>& Contexts)
 {
 	// Support for resetting RTT size when viewport size is changed
@@ -314,9 +287,9 @@ bool FDisplayClusterViewportManager::ShouldUseFullSizeFrameTargetableResource()
 		return true;
 	}
 
-	for (const FDisplayClusterViewport* Viewport : ClusterNodeViewports)
+	for (const TSharedPtr<FDisplayClusterViewport, ESPMode::ThreadSafe>& Viewport : ClusterNodeViewports)
 	{
-		if (Viewport && Viewport->ShouldUseFullSizeFrameTargetableResource())
+		if (Viewport.IsValid() && Viewport->ShouldUseFullSizeFrameTargetableResource())
 		{
 			return true;
 		}
@@ -334,9 +307,9 @@ bool FDisplayClusterViewportManager::ShouldUseAdditionalFrameTargetableResource(
 		return true;
 	}
 
-	for (const FDisplayClusterViewport* Viewport : ClusterNodeViewports)
+	for (const TSharedPtr<FDisplayClusterViewport, ESPMode::ThreadSafe>& Viewport : ClusterNodeViewports)
 	{
-		if (Viewport && Viewport->ShouldUseAdditionalFrameTargetableResource())
+		if (Viewport.IsValid() && Viewport->ShouldUseAdditionalFrameTargetableResource())
 		{
 			return true;
 		}
@@ -412,6 +385,13 @@ const FDisplayClusterRenderFrameSettings& FDisplayClusterViewportManager::GetRen
 	return Configuration->GetRenderFrameSettings();
 }
 
+EDisplayClusterRenderFrameMode FDisplayClusterViewportManager::GetRenderMode() const
+{
+	check(IsInGameThread());
+
+	return Configuration->GetRenderFrameSettings().RenderMode;
+}
+
 void FDisplayClusterViewportManager::ImplUpdateClusterNodeViewports(const EDisplayClusterRenderFrameMode InRenderMode, const FString& InClusterNodeId)
 {
 	ClusterNodeViewports.Empty();
@@ -422,11 +402,11 @@ void FDisplayClusterViewportManager::ImplUpdateClusterNodeViewports(const EDispl
 		{
 			// Per-viewport preview render required special render order for viewports
 			// Sort viewports: move linked and overridden viewports to last
-			TArray<FDisplayClusterViewport*> LinkedViewports;
-			TArray<FDisplayClusterViewport*> OverriddenViewports;
-			for (FDisplayClusterViewport* Viewport : Viewports)
+			TArray<TSharedPtr<FDisplayClusterViewport, ESPMode::ThreadSafe>> LinkedViewports;
+			TArray<TSharedPtr<FDisplayClusterViewport, ESPMode::ThreadSafe>> OverriddenViewports;
+			for (TSharedPtr<FDisplayClusterViewport, ESPMode::ThreadSafe>& Viewport : Viewports)
 			{
-				if (Viewport && (InClusterNodeId.IsEmpty() || Viewport->GetClusterNodeId() == InClusterNodeId))
+				if (Viewport.IsValid() && (InClusterNodeId.IsEmpty() || Viewport->GetClusterNodeId() == InClusterNodeId))
 				{
 					if (Viewport->RenderSettings.IsViewportOverrided())
 					{
@@ -456,29 +436,37 @@ void FDisplayClusterViewportManager::ImplUpdateClusterNodeViewports(const EDispl
 
 int32 FDisplayClusterViewportManager::GetViewPerViewportAmount() const
 {
-	switch (GetRenderFrameSettings().RenderMode)
+	switch (GetRenderMode())
 	{
 	case EDisplayClusterRenderFrameMode::Mono:
 	case EDisplayClusterRenderFrameMode::PreviewInScene:
 		return 1;
-		break;
+
+	case EDisplayClusterRenderFrameMode::Stereo:
+	case EDisplayClusterRenderFrameMode::SideBySide:
+	case EDisplayClusterRenderFrameMode::TopBottom:
+		return 2;
+
 	default:
 		break;
 	}
 
-	return 2;
+	return 0;
 }
 
-int32 FDisplayClusterViewportManager::FindFirstViewportStereoViewIndex(FDisplayClusterViewport* InViewport) const
+int32 FDisplayClusterViewportManager::FindFirstViewportStereoViewIndex(const TSharedPtr<FDisplayClusterViewport, ESPMode::ThreadSafe>& InViewport) const
 {
-	const int32 ViewportIndex = Viewports.Find(InViewport);
-	if (ViewportIndex != INDEX_NONE)
+	const int32 ViewPerViewportAmount = GetViewPerViewportAmount();
+	if (InViewport.IsValid() && ViewPerViewportAmount > 0)
 	{
-		const int32 ViewPerViewportAmount = GetViewPerViewportAmount();
-		// Begin from 1, because INDEX_NONE use ViewState[0] in LocalPlayer.cpp:786
-		const int32 FirstViewportStereoViewIndex = (ViewportIndex * ViewPerViewportAmount) + 1;
+		const int32 ViewportIndex = Viewports.Find(InViewport);
+		if (ViewportIndex != INDEX_NONE)
+		{
+			// Begin from 1, because INDEX_NONE use ViewState[0] in LocalPlayer.cpp:786
+			const int32 FirstViewportStereoViewIndex = (ViewportIndex * ViewPerViewportAmount) + 1;
 
-		return FirstViewportStereoViewIndex;
+			return FirstViewportStereoViewIndex;
+		}
 	}
 
 	return INDEX_NONE;
@@ -488,7 +476,7 @@ bool FDisplayClusterViewportManager::BeginNewFrame(FViewport* InViewport, UWorld
 {
 	check(IsInGameThread());
 
-	OutRenderFrame.ViewportManager = this;
+	OutRenderFrame.ViewportManager = this->AsShared();
 
 	// Handle world runtime update
 	UWorld* CurrentWorld = GetCurrentWorld();
@@ -533,29 +521,58 @@ bool FDisplayClusterViewportManager::BeginNewFrame(FViewport* InViewport, UWorld
 	const FDisplayClusterRenderFrameSettings& RenderFrameSettings = GetRenderFrameSettings();
 
 	// Initialize viewports from new render settings, and create new contexts, reset prev frame resources
-	for (FDisplayClusterViewport* Viewport : ClusterNodeViewports)
+	for (TSharedPtr<FDisplayClusterViewport, ESPMode::ThreadSafe>& ViewportIt : ClusterNodeViewports)
 	{
-		const int32 FirstViewportStereoViewIndex = FindFirstViewportStereoViewIndex(Viewport);
+		if (!ViewportIt.IsValid())
+		{
+			continue;
+		}
+
+		const int32 FirstViewportStereoViewIndex = FindFirstViewportStereoViewIndex(ViewportIt);
 		if (FirstViewportStereoViewIndex != INDEX_NONE)
 		{
 			// Save orig viewport contexts
 			TArray<FDisplayClusterViewport_Context> PrevContexts;
-			PrevContexts.Append(Viewport->GetContexts());
-			Viewport->UpdateFrameContexts(FirstViewportStereoViewIndex, RenderFrameSettings);
+			PrevContexts.Append(ViewportIt->GetContexts());
+			ViewportIt->UpdateFrameContexts(FirstViewportStereoViewIndex, RenderFrameSettings);
 
-			HandleViewportRTTChanges(PrevContexts, Viewport->GetContexts());
+			HandleViewportRTTChanges(PrevContexts, ViewportIt->GetContexts());
 		}
 		else
 		{
 			// do not render this viewport
-			Viewport->ResetFrameContexts();
+			ViewportIt->ResetFrameContexts();
 		}
 	}
 
-	/* Some media dependent data can (and needs to) be initialized in advance */
-	for (FDisplayClusterViewport* const Viewport : ClusterNodeViewports)
-	{
-		Viewport->UpdateMediaDependencies(InViewport);
+	// Preliminary initialization of rendering thread data for viewport contexts
+	// Currently, some data like EngineShowFlags and EngineGamma get updated on PostRenderViewFamily.
+	// Since the viewports that have media input assigned never gets rendered in normal way, the PostRenderViewFamily
+	// callback never gets called, therefore the data mentioned above never gets updated. This workaround initializes
+	// those settings for all viewports. The viewports with no media input assigned will override the data
+	// in PostRenderViewFamily like it was previously so nothing should be broken.
+	const float DefaultDisplayGamma = GEngine ? GEngine->DisplayGamma : 2.2f;
+	const FEngineShowFlags* const DefaultEngineShowFlags = (InViewport && InViewport->GetClient()) ? InViewport->GetClient()->GetEngineShowFlags() : nullptr;
+	for (TSharedPtr<FDisplayClusterViewport, ESPMode::ThreadSafe>& ViewportIt : ClusterNodeViewports)
+	{
+		if (ViewportIt.IsValid())
+		{
+			for (FDisplayClusterViewport_Context& ContextIt : ViewportIt->Contexts)
+			{
+				// Get Context Display gamma
+				const FDisplayClusterViewportRenderTargetResource* ContextRTT = ViewportIt->RenderTargets.IsValidIndex(ContextIt.ContextNum) ? ViewportIt->RenderTargets[ContextIt.ContextNum] : nullptr;
+				const float ViewportDisplayGamma = ContextRTT ? ContextRTT->GetDisplayGamma() : DefaultDisplayGamma;
+
+				ContextIt.RenderThreadData.EngineDisplayGamma = ViewportDisplayGamma;
+
+
+				if (DefaultEngineShowFlags)
+				{
+					// Todo: DefaultEngineShowFlags is not the correct value, since each viewport uses its own Engine flags.
+					ContextIt.RenderThreadData.EngineShowFlags = *DefaultEngineShowFlags;
+				}
+			}
+		}
 	}
 
 	// Handle scene RTT resize
@@ -575,9 +592,9 @@ bool FDisplayClusterViewportManager::BeginNewFrame(FViewport* InViewport, UWorld
 
 	const FIntPoint RenderFrameSize = OutRenderFrame.FrameRect.Size();
 
-	for (FDisplayClusterViewport* Viewport : ClusterNodeViewports)
+	for (TSharedPtr<FDisplayClusterViewport, ESPMode::ThreadSafe>& Viewport : ClusterNodeViewports)
 	{
-		if (Viewport)
+		if (Viewport.IsValid())
 		{
 			// Update ViewportRemap geometry
 			Viewport->ViewportRemap.Update(*Viewport, RenderFrameSize);
@@ -585,7 +602,7 @@ bool FDisplayClusterViewportManager::BeginNewFrame(FViewport* InViewport, UWorld
 	}
 
 	// Update desired views number
-	UpdateDesiredNumberOfViews(OutRenderFrame);
+	OutRenderFrame.DesiredNumberOfViews = (Viewports.Num() * GetViewPerViewportAmount()) + 1;
 
 #if WITH_EDITOR
 	// Get preview resources from root actor
@@ -600,47 +617,14 @@ bool FDisplayClusterViewportManager::BeginNewFrame(FViewport* InViewport, UWorld
 	return true;
 }
 
-void FDisplayClusterViewportManager::UpdateDesiredNumberOfViews(FDisplayClusterRenderFrame& InOutRenderFrame)
-{
-	// Total views ammount
-	const int32 ViewPerViewportAmount = GetViewPerViewportAmount();
-	const int32 DesiredNumberOfViews = (Viewports.Num() * ViewPerViewportAmount) + 1;
-
-	InOutRenderFrame.DesiredNumberOfViews = DesiredNumberOfViews;
-
-	// Total view on current cluster node
-	InOutRenderFrame.ViewportsAmount = 0;
-
-	for (FDisplayClusterRenderFrame::FFrameRenderTarget& RenderTargetIt : InOutRenderFrame.RenderTargets)
-	{
-		for (FDisplayClusterRenderFrame::FFrameViewFamily& ViewFamilyIt : RenderTargetIt.ViewFamilies)
-		{
-			ViewFamilyIt.NumViewsForRender = 0;
-
-			for (FDisplayClusterRenderFrame::FFrameView& ViewIt : ViewFamilyIt.Views)
-			{
-				InOutRenderFrame.ViewportsAmount++;
-
-				if (ViewIt.bDisableRender == false && ViewIt.Viewport != nullptr)
-				{
-					ViewFamilyIt.NumViewsForRender++;
-
-					// Get StereoViewIndex for this viewport
-					const int32 StereoViewIndex = ViewIt.Viewport->GetContexts()[ViewIt.ContextNum].StereoViewIndex;
-				}
-			}
-		}
-	}
-}
-
 void FDisplayClusterViewportManager::InitializeNewFrame()
 {
 	check(IsInGameThread());
 
 	// Before render
-	for (FDisplayClusterViewport* Viewport : ClusterNodeViewports)
+	for (TSharedPtr<FDisplayClusterViewport, ESPMode::ThreadSafe>& Viewport : ClusterNodeViewports)
 	{
-		if (Viewport)
+		if (Viewport.IsValid())
 		{
 			// Initialize OCIO resources
 			if (Viewport->OpenColorIO.IsValid())
@@ -680,9 +664,9 @@ void FDisplayClusterViewportManager::FinalizeNewFrame()
 	// [2] The viewport math has been updated, we need to send it to the proxy again.
 	ViewportManagerProxy->ImplUpdateViewports(ClusterNodeViewports);
 
-	for (FDisplayClusterViewport* Viewport : ClusterNodeViewports)
+	for (TSharedPtr<FDisplayClusterViewport, ESPMode::ThreadSafe>& Viewport : ClusterNodeViewports)
 	{
-		if (Viewport)
+		if (Viewport.IsValid())
 		{
 			// When all viewports processed, we remove all single frame custom postprocess
 			Viewport->CustomPostProcessSettings.FinalizeFrame();
@@ -690,16 +674,16 @@ void FDisplayClusterViewportManager::FinalizeNewFrame()
 			// Update projection policy proxy data
 			if (Viewport->ProjectionPolicy.IsValid())
 			{
-				Viewport->ProjectionPolicy->UpdateProxyData(Viewport);
+				Viewport->ProjectionPolicy->UpdateProxyData(Viewport.Get());
 			}
 
 		}
 	}
 
 	// Finish update settings
-	for (FDisplayClusterViewport* Viewport : ClusterNodeViewports)
+	for (TSharedPtr<FDisplayClusterViewport, ESPMode::ThreadSafe>& Viewport : ClusterNodeViewports)
 	{
-		if (Viewport)
+		if (Viewport.IsValid())
 		{
 			Viewport->RenderSettings.FinishUpdateSettings();
 		}
@@ -707,7 +691,7 @@ void FDisplayClusterViewportManager::FinalizeNewFrame()
 }
 
 FSceneViewFamily::ConstructionValues FDisplayClusterViewportManager::CreateViewFamilyConstructionValues(
-	const FDisplayClusterRenderFrame::FFrameRenderTarget& InFrameTarget,
+	const FDisplayClusterRenderFrameTarget& InFrameTarget,
 	FSceneInterface* InScene,
 	FEngineShowFlags InEngineShowFlags,
 	const bool bInAdditionalViewFamily) const
@@ -829,7 +813,7 @@ bool FDisplayClusterViewportManager::ShouldRenderFinalColor() const
 	return true;
 }
 
-void FDisplayClusterViewportManager::ConfigureViewFamily(const FDisplayClusterRenderFrame::FFrameRenderTarget& InFrameTarget, const FDisplayClusterRenderFrame::FFrameViewFamily& InFrameViewFamily, FSceneViewFamilyContext& ViewFamily)
+void FDisplayClusterViewportManager::ConfigureViewFamily(const FDisplayClusterRenderFrameTarget& InFrameTarget, const FDisplayClusterRenderFrameTargetViewFamily& InFrameViewFamily, FSceneViewFamilyContext& ViewFamily)
 {
 	// Note: EngineShowFlags should have already been configured in CreateViewFamilyConstructionValues.
 	ViewFamily.SceneCaptureCompositeMode = ESceneCaptureCompositeMode::SCCM_Overwrite;
@@ -859,10 +843,9 @@ void FDisplayClusterViewportManager::RenderFrame(FViewport* InViewport)
 	ViewportManagerProxy->ImplRenderFrame(InViewport);
 }
 
-bool FDisplayClusterViewportManager::CreateViewport(const FString& InViewportId, const class UDisplayClusterConfigurationViewport* ConfigurationViewport)
+bool FDisplayClusterViewportManager::CreateViewport(const FString& InViewportId, const class UDisplayClusterConfigurationViewport& ConfigurationViewport)
 {
 	check(IsInGameThread());
-	check(ConfigurationViewport != nullptr);
 
 	// Check viewport ID
 	if (InViewportId.IsEmpty())
@@ -875,63 +858,78 @@ bool FDisplayClusterViewportManager::CreateViewport(const FString& InViewportId,
 	if (FindViewport(InViewportId) != nullptr)
 	{
 		UE_LOG(LogDisplayClusterViewport, Warning, TEXT("Viewport '%s' already exists"), *InViewportId);
+
+		return false;
+	}
+
+	ADisplayClusterRootActor* RootActorPtr = GetRootActor();
+	if (!RootActorPtr)
+	{
+		UE_LOG(LogDisplayClusterViewport, Error, TEXT("Can't create viewport '%s': DCRA not defined"), *InViewportId);
+
 		return false;
 	}
 
 	// Create projection policy for viewport
-	const TSharedPtr<IDisplayClusterProjectionPolicy, ESPMode::ThreadSafe> NewProjectionPolicy = CreateProjectionPolicy(InViewportId, &ConfigurationViewport->ProjectionPolicy);
+	const TSharedPtr<IDisplayClusterProjectionPolicy, ESPMode::ThreadSafe> NewProjectionPolicy = CreateProjectionPolicy(InViewportId, &ConfigurationViewport.ProjectionPolicy);
 	if (NewProjectionPolicy.IsValid())
 	{
 		// Create viewport for new projection policy
-		FDisplayClusterViewport* NewViewport = ImplCreateViewport(InViewportId, NewProjectionPolicy);
-		if (NewViewport != nullptr)
+		TSharedPtr<FDisplayClusterViewport, ESPMode::ThreadSafe> NewViewport = ImplCreateViewport(InViewportId, NewProjectionPolicy);
+		if (NewViewport.IsValid())
 		{
-			ADisplayClusterRootActor* RootActorPtr = GetRootActor();
-			if (RootActorPtr)
-			{
-				FDisplayClusterViewportConfigurationBase::UpdateViewportConfiguration(*this, *RootActorPtr, NewViewport, ConfigurationViewport);
-				return true;
-			}
+			FDisplayClusterViewportConfigurationBase::UpdateViewportConfiguration(*NewViewport , *this, *RootActorPtr, ConfigurationViewport);
+
+			return true;
 		}
 	}
 
 	UE_LOG(LogDisplayClusterViewport, Error, TEXT("Viewports '%s' not created."), *InViewportId);
+
 	return false;
 }
 
-FDisplayClusterViewport* FDisplayClusterViewportManager::ImplFindViewport(const FString& ViewportId) const
+IDisplayClusterViewport* FDisplayClusterViewportManager::FindViewport(const FString& InViewportId) const
+{
+	TSharedPtr<FDisplayClusterViewport, ESPMode::ThreadSafe> ExistViewport = ImplFindViewport(InViewportId);
+	return ExistViewport.IsValid() ? ExistViewport.Get() : nullptr;
+}
+
+TSharedPtr<FDisplayClusterViewport, ESPMode::ThreadSafe> FDisplayClusterViewportManager::ImplFindViewport(const FString& ViewportId) const
 {
 	check(IsInGameThread());
 
 	// Ok, we have a request for a particular viewport. Let's find it.
-	FDisplayClusterViewport* const* DesiredViewport = Viewports.FindByPredicate([ViewportId](const FDisplayClusterViewport* ItemViewport)
+	TSharedPtr<FDisplayClusterViewport, ESPMode::ThreadSafe> const* DesiredViewport = Viewports.FindByPredicate([ViewportId](const TSharedPtr<FDisplayClusterViewport, ESPMode::ThreadSafe>& ItemViewport)
 	{
-		return ViewportId.Equals(ItemViewport->ViewportId, ESearchCase::IgnoreCase);
+		return ItemViewport.IsValid() && ViewportId.Equals(ItemViewport->ViewportId, ESearchCase::IgnoreCase);
 	});
 
-	return (DesiredViewport != nullptr) ? *DesiredViewport : nullptr;
+	return (DesiredViewport && DesiredViewport->IsValid()) ? *DesiredViewport : nullptr;
 }
 
 IDisplayClusterViewport* FDisplayClusterViewportManager::CreateViewport(const FString& ViewportId, const TSharedPtr<IDisplayClusterProjectionPolicy, ESPMode::ThreadSafe>& InProjectionPolicy)
 {
 	check(IsInGameThread());
 
-	FDisplayClusterViewport* ExistViewport = ImplFindViewport(ViewportId);
-	if (ExistViewport != nullptr)
+	TSharedPtr<FDisplayClusterViewport, ESPMode::ThreadSafe> ExistViewport = ImplFindViewport(ViewportId);
+	if (ExistViewport.IsValid())
 	{
 		//add error log: Viewport with name '%s' already exist
 		return nullptr;
 	}
 
-	return ImplCreateViewport(ViewportId, InProjectionPolicy);
+	TSharedPtr<FDisplayClusterViewport, ESPMode::ThreadSafe> NewViewport = ImplCreateViewport(ViewportId, InProjectionPolicy);
+
+	return NewViewport.IsValid() ? NewViewport.Get() : nullptr;
 }
 
 bool FDisplayClusterViewportManager::DeleteViewport(const FString& ViewportId)
 {
 	check(IsInGameThread());
 
-	FDisplayClusterViewport* ExistViewport = ImplFindViewport(ViewportId);
-	if (ExistViewport != nullptr)
+	TSharedPtr<FDisplayClusterViewport, ESPMode::ThreadSafe> ExistViewport = ImplFindViewport(ViewportId);
+	if (ExistViewport.IsValid())
 	{
 		ImplDeleteViewport(ExistViewport);
 
@@ -941,7 +939,7 @@ bool FDisplayClusterViewportManager::DeleteViewport(const FString& ViewportId)
 	return false;
 }
 
-FDisplayClusterViewport* FDisplayClusterViewportManager::ImplCreateViewport(const FString& ViewportId, const TSharedPtr<IDisplayClusterProjectionPolicy, ESPMode::ThreadSafe>& InProjectionPolicy)
+TSharedPtr<FDisplayClusterViewport, ESPMode::ThreadSafe> FDisplayClusterViewportManager::ImplCreateViewport(const FString& ViewportId, const TSharedPtr<IDisplayClusterProjectionPolicy, ESPMode::ThreadSafe>& InProjectionPolicy)
 {
 	check(IsInGameThread());
 
@@ -951,21 +949,25 @@ FDisplayClusterViewport* FDisplayClusterViewportManager::ImplCreateViewport(cons
 	const FString& ClusterNodeId = GetRenderFrameSettings().ClusterNodeId;
 	check(!ClusterNodeId.IsEmpty());
 
-	// Create viewport
-	FDisplayClusterViewport* NewViewport = new FDisplayClusterViewport(*this, ClusterNodeId, ViewportId, InProjectionPolicy);
-
-	// Add viewport on gamethread
-	Viewports.Add(NewViewport);
-	ClusterNodeViewports.Add(NewViewport);
+	// Create a new viewport
+	TSharedPtr<FDisplayClusterViewport, ESPMode::ThreadSafe> NewViewport = MakeShared<FDisplayClusterViewport>(*this, ClusterNodeId, ViewportId, InProjectionPolicy);
+	if (NewViewport.IsValid())
+	{
+		// Add viewport on gamethread
+		Viewports.Add(NewViewport);
+		ClusterNodeViewports.Add(NewViewport);
 
-	// Handle start scene for viewport
-	NewViewport->HandleStartScene();
+		// Handle start scene for viewport
+		NewViewport->HandleStartScene();
+	}
 
 	return NewViewport;
 }
 
-void FDisplayClusterViewportManager::ImplDeleteViewport(FDisplayClusterViewport* ExistViewport)
+void FDisplayClusterViewportManager::ImplDeleteViewport(const TSharedPtr<FDisplayClusterViewport, ESPMode::ThreadSafe>& ExistViewport)
 {
+	check(ExistViewport.IsValid());
+
 	// Handle projection policy event
 	ExistViewport->ProjectionPolicy.Reset();
 	ExistViewport->UninitializedProjectionPolicy.Reset();
@@ -990,8 +992,6 @@ void FDisplayClusterViewportManager::ImplDeleteViewport(FDisplayClusterViewport*
 		}
 	}
 
-	delete ExistViewport;
-
 	// Reset RTT size after viewport delete
 	ResetSceneRenderTargetSize();
 }
@@ -1008,11 +1008,11 @@ IDisplayClusterViewport* FDisplayClusterViewportManager::FindViewport(const int3
 		return nullptr;
 	}
 
-	for (FDisplayClusterViewport* Viewport : Viewports)
+	for (const TSharedPtr<FDisplayClusterViewport, ESPMode::ThreadSafe>& ViewportIt : Viewports)
 	{
-		if (Viewport && Viewport->FindContext(ViewIndex, OutContextNum))
+		if (ViewportIt.IsValid() && ViewportIt->FindContext(ViewIndex, OutContextNum))
 		{
-			return Viewport;
+			return ViewportIt.Get();
 		}
 	}
 
@@ -1057,9 +1057,9 @@ void FDisplayClusterViewportManager::MarkComponentGeometryDirty(const FName InCo
 	check(IsInGameThread());
 
 	// 1. Update all ProceduralMeshComponent references for projection policies
-	for (IDisplayClusterViewport* ViewportIt : GetViewports())
+	for (const TSharedPtr<FDisplayClusterViewport, ESPMode::ThreadSafe>& ViewportIt : ImplGetViewports())
 	{
-		if (ViewportIt != nullptr)
+		if (ViewportIt.IsValid())
 		{
 			const TSharedPtr<IDisplayClusterProjectionPolicy, ESPMode::ThreadSafe>& ProjectionPolicy = ViewportIt->GetProjectionPolicy();
 			if (ProjectionPolicy.IsValid())
@@ -1087,9 +1087,9 @@ void FDisplayClusterViewportManager::MarkComponentGeometryDirty(const FName InCo
 
 void FDisplayClusterViewportManager::AddReferencedObjects(FReferenceCollector& Collector)
 {
-	for (FDisplayClusterViewport* ViewportIt : Viewports)
+	for (TSharedPtr<FDisplayClusterViewport, ESPMode::ThreadSafe>& ViewportIt : Viewports)
 	{
-		if (ViewportIt != nullptr)
+		if (ViewportIt.IsValid())
 		{
 			ViewportIt->AddReferencedObjects(Collector);
 		}
diff --git a/Engine/Plugins/Runtime/nDisplay/Source/DisplayCluster/Private/Render/Viewport/DisplayClusterViewportManager.h b/Engine/Plugins/Runtime/nDisplay/Source/DisplayCluster/Private/Render/Viewport/DisplayClusterViewportManager.h
--- a/Engine/Plugins/Runtime/nDisplay/Source/DisplayCluster/Private/Render/Viewport/DisplayClusterViewportManager.h
+++ b/Engine/Plugins/Runtime/nDisplay/Source/DisplayCluster/Private/Render/Viewport/DisplayClusterViewportManager.h
@@ -5,19 +5,18 @@
 #include "CoreMinimal.h"
 
 #include "SceneTypes.h"
+#include "Templates/SharedPointer.h"
 
 #include "Misc/DisplayClusterObjectRef.h"
 
 #include "Render/Viewport/IDisplayClusterViewportManager.h"
 #include "Render/Viewport/DisplayClusterViewport.h"
-
 #include "Render/Viewport/Containers/DisplayClusterViewport_Enums.h"
 
 class FDisplayClusterViewportConfiguration;
 class FDisplayClusterRenderTargetManager;
 class FDisplayClusterViewportPostProcessManager;
 class FDisplayClusterRenderFrameManager;
-class FDisplayClusterRenderFrame; 
 class FDisplayClusterViewportManagerProxy;
 class FDisplayClusterViewportLightCardManager;
 class IDisplayClusterProjectionPolicy;
@@ -30,93 +29,101 @@ class FViewport;
 
 class FDisplayClusterViewportManager
 	: public IDisplayClusterViewportManager
+	, public TSharedFromThis<FDisplayClusterViewportManager, ESPMode::ThreadSafe>
 {
 public:
 	FDisplayClusterViewportManager();
 	virtual ~FDisplayClusterViewportManager();
 
 public:
+	/** Game thread funcs */
+	void StartScene(UWorld* World);
+	void EndScene();
+	void ResetScene();
+
+	void Initialize();
+
+	//~ Begin IDisplayClusterViewportManager
+	virtual TSharedPtr<IDisplayClusterViewportManager, ESPMode::ThreadSafe> ToSharedPtr() override
+	{
+		return AsShared();
+	}
+
+	virtual TSharedPtr<const IDisplayClusterViewportManager, ESPMode::ThreadSafe> ToSharedPtr() const override
+	{
+		return AsShared();
+	}
+
+	virtual EDisplayClusterRenderFrameMode GetRenderMode() const override;
+
 	virtual const IDisplayClusterViewportManagerProxy* GetProxy() const override;
 	virtual       IDisplayClusterViewportManagerProxy* GetProxy() override;
 
 	virtual UWorld*                   GetCurrentWorld() const override;
 	virtual ADisplayClusterRootActor* GetRootActor() const override;
 
-	/** Game thread funcs */
-	void StartScene(UWorld* World);
-	void EndScene();
-	void ResetScene();
-
 	virtual bool IsSceneOpened() const override;
 
 	virtual bool UpdateConfiguration(EDisplayClusterRenderFrameMode InRenderMode, const FString& InClusterNodeId, class ADisplayClusterRootActor* InRootActorPtr, const FDisplayClusterPreviewSettings* InPreviewSettings = nullptr) override;
 	virtual bool UpdateCustomConfiguration(EDisplayClusterRenderFrameMode InRenderMode, const TArray<FString>& InViewportNames, class ADisplayClusterRootActor* InRootActorPtr) override;
 
 	virtual bool BeginNewFrame(FViewport* InViewport, UWorld* InWorld, FDisplayClusterRenderFrame& OutRenderFrame) override;
-
 	virtual void InitializeNewFrame() override;
 	virtual void FinalizeNewFrame() override;
 
 	virtual FSceneViewFamily::ConstructionValues CreateViewFamilyConstructionValues(
-		const FDisplayClusterRenderFrame::FFrameRenderTarget& InFrameTarget,
+		const FDisplayClusterRenderFrameTarget& InFrameTarget,
 		FSceneInterface* InScene,
 		FEngineShowFlags InEngineShowFlags,
 		const bool bInAdditionalViewFamily
 	) const override;
 
-	virtual void ConfigureViewFamily(const FDisplayClusterRenderFrame::FFrameRenderTarget& InFrameTarget, const FDisplayClusterRenderFrame::FFrameViewFamily& InFrameViewFamily, FSceneViewFamilyContext& InOutViewFamily) override;
+	virtual void ConfigureViewFamily(const FDisplayClusterRenderFrameTarget& InFrameTarget, const FDisplayClusterRenderFrameTargetViewFamily& InFrameViewFamily, FSceneViewFamilyContext& InOutViewFamily) override;
 	
 	virtual void RenderFrame(FViewport* InViewport) override;
 
 #if WITH_EDITOR
-	virtual bool RenderInEditor(class FDisplayClusterRenderFrame& InRenderFrame, FViewport* InViewport, const uint32 InFirstViewportNum, const int32 InViewportsAmount, int32& OutViewportsAmount, bool& bOutFrameRendered) override;
+	virtual bool RenderInEditor(FDisplayClusterRenderFrame& InRenderFrame, FViewport* InViewport, const uint32 InFirstViewportNum, const int32 InViewportsAmount, int32& OutViewportsAmount, bool& bOutFrameRendered) override;
 	
 	void ImplUpdatePreviewRTTResources();
-
-	const TArray<FDisplayClusterViewport*>& ImplGetWholeClusterViewports_Editor() const
-	{
-		check(IsInGameThread());
-
-		return Viewports;
-	}
-
-	// Current world is a tools window (configurator, other tools)
-	bool IsEditorPreviewWorld() const;
 #endif
 
-	virtual IDisplayClusterViewport* FindViewport(const FString& InViewportId) const override
-	{
-		return ImplFindViewport(InViewportId);
-	}
-
+	virtual IDisplayClusterViewport* FindViewport(const FString& InViewportId) const override;
 	virtual IDisplayClusterViewport* FindViewport(const int32 ViewIndex, uint32* OutContextNum = nullptr) const override;
 
-	virtual const TArrayView<IDisplayClusterViewport*> GetViewports() const override
+	virtual const TArrayView<TSharedPtr<IDisplayClusterViewport, ESPMode::ThreadSafe>> GetViewports() const override
 	{
-		return TArrayView<IDisplayClusterViewport*>((IDisplayClusterViewport**)(ClusterNodeViewports.GetData()), ClusterNodeViewports.Num());
+		return TArrayView<TSharedPtr<IDisplayClusterViewport, ESPMode::ThreadSafe>>((TSharedPtr<IDisplayClusterViewport, ESPMode::ThreadSafe>*)(ClusterNodeViewports.GetData()), ClusterNodeViewports.Num());
 	}
 
 	virtual void MarkComponentGeometryDirty(const FName InComponentName = NAME_None) override;
-
-	virtual class FSceneView* CalcSceneView(class ULocalPlayer* LocalPlayer, class FSceneViewFamily* ViewFamily, FVector& OutViewLocation, FRotator& OutViewRotation, class FViewport* Viewport, class FViewElementDrawer* ViewDrawer, int32 StereoViewIndex) override;
-
 	virtual void AddReferencedObjects(FReferenceCollector& Collector) override;
 
-	// internal use only
-	bool CreateViewport(const FString& ViewportId, const class UDisplayClusterConfigurationViewport* ConfigurationViewport);
-	IDisplayClusterViewport* CreateViewport(const FString& ViewportId, const TSharedPtr<IDisplayClusterProjectionPolicy, ESPMode::ThreadSafe>& InProjectionPolicy);
-	bool                     DeleteViewport(const FString& ViewportId);
+	virtual class FSceneView* CalcSceneView(class ULocalPlayer* LocalPlayer, class FSceneViewFamily* ViewFamily, FVector& OutViewLocation, FRotator& OutViewRotation, class FViewport* Viewport, class FViewElementDrawer* ViewDrawer, int32 StereoViewIndex) override;
 
-	FDisplayClusterViewport* ImplCreateViewport(const FString& ViewportId, const TSharedPtr<IDisplayClusterProjectionPolicy, ESPMode::ThreadSafe>& InProjectionPolicy);
-	void ImplDeleteViewport(FDisplayClusterViewport* Viewport);
+	//~~ End IDisplayClusterViewportManager
 
-	const TArray<FDisplayClusterViewport*>& ImplGetViewports() const
+	const TArray<TSharedPtr<FDisplayClusterViewport, ESPMode::ThreadSafe>>& ImplGetViewports() const
 	{
 		check(IsInGameThread());
 		return ClusterNodeViewports;
 	}
 
-	FDisplayClusterViewport* ImplFindViewport(const FString& InViewportId) const;
+	const TArray<TSharedPtr<FDisplayClusterViewport, ESPMode::ThreadSafe>>& ImplGetWholeClusterViewports() const
+	{
+		check(IsInGameThread());
+
+		return Viewports;
+	}
+
+	// internal use only
+	bool CreateViewport(const FString& ViewportId, const class UDisplayClusterConfigurationViewport& ConfigurationViewport);
+	IDisplayClusterViewport* CreateViewport(const FString& ViewportId, const TSharedPtr<IDisplayClusterProjectionPolicy, ESPMode::ThreadSafe>& InProjectionPolicy);
+	bool                     DeleteViewport(const FString& ViewportId);
+
+	TSharedPtr<FDisplayClusterViewport, ESPMode::ThreadSafe> ImplCreateViewport(const FString& ViewportId, const TSharedPtr<IDisplayClusterProjectionPolicy, ESPMode::ThreadSafe>& InProjectionPolicy);
+	TSharedPtr<FDisplayClusterViewport, ESPMode::ThreadSafe> ImplFindViewport(const FString& InViewportId) const;
+	void ImplDeleteViewport(const TSharedPtr<FDisplayClusterViewport, ESPMode::ThreadSafe>& InExistViewport);
 
 	static TSharedPtr<IDisplayClusterProjectionPolicy, ESPMode::ThreadSafe> CreateProjectionPolicy(const FString& InViewportId, const FDisplayClusterConfigurationProjection* InConfigurationProjectionPolicy);
 
@@ -133,7 +140,7 @@ public:
 	int32 GetViewPerViewportAmount() const;
 
 	/** Return initial StereoViewIndex for the input viewport. */
-	int32 FindFirstViewportStereoViewIndex(FDisplayClusterViewport* InViewport) const;
+	int32 FindFirstViewportStereoViewIndex(const TSharedPtr<FDisplayClusterViewport, ESPMode::ThreadSafe>& InViewport) const;
 
 	/** Returns true if the final color is renderable (to support preview in scene). */
 	bool ShouldRenderFinalColor() const;
@@ -153,11 +160,9 @@ public:
 	bool ShouldUseAdditionalFrameTargetableResource() const;
 	bool ShouldUseFullSizeFrameTargetableResource() const;
 
-	void SetViewportBufferRatio(FDisplayClusterViewport& DstViewport, float InBufferRatio);
 	void ResetSceneRenderTargetSize();
 
 private:
-	void UpdateDesiredNumberOfViews(FDisplayClusterRenderFrame& InOutRenderFrame);
 	void UpdateSceneRenderTargetSize();
 	void HandleViewportRTTChanges(const TArray<FDisplayClusterViewport_Context>& PrevContexts, const TArray<FDisplayClusterViewport_Context>& Contexts);
 	void ImplUpdateClusterNodeViewports(const EDisplayClusterRenderFrameMode InRenderMode, const FString& InClusterNodeId);
@@ -176,8 +181,8 @@ public:
 private:
 	TUniquePtr<FDisplayClusterRenderFrameManager>  RenderFrameManager;
 
-	TArray<FDisplayClusterViewport*> Viewports;
-	TArray<FDisplayClusterViewport*> ClusterNodeViewports;
+	TArray<TSharedPtr<FDisplayClusterViewport, ESPMode::ThreadSafe>> Viewports;
+	TArray<TSharedPtr<FDisplayClusterViewport, ESPMode::ThreadSafe>> ClusterNodeViewports;
 
 	/** Render thread proxy manager. Deleted on render thread */
 	TSharedPtr<FDisplayClusterViewportManagerProxy, ESPMode::ThreadSafe> ViewportManagerProxy;
diff --git a/Engine/Plugins/Runtime/nDisplay/Source/DisplayCluster/Private/Render/Viewport/DisplayClusterViewportManagerProxy.cpp b/Engine/Plugins/Runtime/nDisplay/Source/DisplayCluster/Private/Render/Viewport/DisplayClusterViewportManagerProxy.cpp
--- a/Engine/Plugins/Runtime/nDisplay/Source/DisplayCluster/Private/Render/Viewport/DisplayClusterViewportManagerProxy.cpp
+++ b/Engine/Plugins/Runtime/nDisplay/Source/DisplayCluster/Private/Render/Viewport/DisplayClusterViewportManagerProxy.cpp
@@ -78,6 +78,11 @@ FDisplayClusterViewportManagerProxy::FDisplayClusterViewportManagerProxy()
 FDisplayClusterViewportManagerProxy::~FDisplayClusterViewportManagerProxy()
 { }
 
+EDisplayClusterRenderFrameMode FDisplayClusterViewportManagerProxy::GetRenderMode() const
+{
+	return RenderFrameSettings.RenderMode;
+}
+
 void FDisplayClusterViewportManagerProxy::Release_RenderThread()
 {
 	check(IsInRenderingThread());
@@ -201,14 +206,17 @@ void FDisplayClusterViewportManagerProxy::ImplUpdateRenderFrameSettings(const FD
 
 DECLARE_GPU_STAT_NAMED(nDisplay_ViewportManager_UpdateViewports, TEXT("nDisplay ViewportManager::UpdateViewports"));
 
-void FDisplayClusterViewportManagerProxy::ImplUpdateViewports(const TArray<FDisplayClusterViewport*>& InViewports)
+void FDisplayClusterViewportManagerProxy::ImplUpdateViewports(const TArray<TSharedPtr<FDisplayClusterViewport, ESPMode::ThreadSafe>>& InViewports)
 {
 	check(IsInGameThread());
 
 	TArray<FDisplayClusterViewportProxyData*> ViewportProxiesData;
-	for (FDisplayClusterViewport* ViewportIt : InViewports)
+	for (const TSharedPtr<FDisplayClusterViewport, ESPMode::ThreadSafe>& ViewportIt : InViewports)
 	{
-		ViewportProxiesData.Add(new FDisplayClusterViewportProxyData(ViewportIt));
+		if (ViewportIt.IsValid())
+		{
+			ViewportProxiesData.Add(new FDisplayClusterViewportProxyData(ViewportIt));
+		}
 	}
 
 	// Send viewports settings to renderthread
@@ -314,7 +322,7 @@ void FDisplayClusterViewportManagerProxy::ImplRenderFrame(FViewport* InViewport)
 				// For quadbuf stereo copy only left eye, right copy from OutputFrameTarget
 				//@todo Copy QuadBuf_LeftEye/(mono,sbs,tp) to separate rtt, before UI and debug rendering
 				//@todo QuadBuf_LeftEye copied latter, before present
-				switch (ViewportManagerProxy->RenderFrameSettings.RenderMode)
+				switch (ViewportManagerProxy->GetRenderMode())
 				{
 				case EDisplayClusterRenderFrameMode::SideBySide:
 				case EDisplayClusterRenderFrameMode::TopBottom:
diff --git a/Engine/Plugins/Runtime/nDisplay/Source/DisplayCluster/Private/Render/Viewport/DisplayClusterViewportManagerProxy.h b/Engine/Plugins/Runtime/nDisplay/Source/DisplayCluster/Private/Render/Viewport/DisplayClusterViewportManagerProxy.h
--- a/Engine/Plugins/Runtime/nDisplay/Source/DisplayCluster/Private/Render/Viewport/DisplayClusterViewportManagerProxy.h
+++ b/Engine/Plugins/Runtime/nDisplay/Source/DisplayCluster/Private/Render/Viewport/DisplayClusterViewportManagerProxy.h
@@ -33,6 +33,18 @@ public:
 
 public:
 	//~IDisplayClusterViewportManagerProxy
+	virtual TSharedPtr<IDisplayClusterViewportManagerProxy, ESPMode::ThreadSafe> ToSharedPtr() override
+	{
+		return AsShared();
+	}
+
+	virtual TSharedPtr<const IDisplayClusterViewportManagerProxy, ESPMode::ThreadSafe> ToSharedPtr() const override
+	{
+		return AsShared();
+	}
+
+	virtual EDisplayClusterRenderFrameMode GetRenderMode() const override;
+
 	virtual IDisplayClusterViewportProxy* FindViewport_RenderThread(const FString& InViewportId) const override
 	{
 		return ImplFindViewport_RenderThread(InViewportId);
@@ -99,7 +111,7 @@ public:
 		const TSharedPtr<FDisplayClusterViewportManagerViewExtension, ESPMode::ThreadSafe>& InViewportManagerViewExtension);
 
 	/** Copy Viewports data to Proxies from game to render thread. */
-	void ImplUpdateViewports(const TArray<FDisplayClusterViewport*>& InViewports);
+	void ImplUpdateViewports(const TArray<TSharedPtr<FDisplayClusterViewport, ESPMode::ThreadSafe>>& InViewports);
 
 	/** Rendering the final frame of nDisplay. Called after all viewports have been rendered in RTT. */
 	void ImplRenderFrame(FViewport* InViewport);
diff --git a/Engine/Plugins/Runtime/nDisplay/Source/DisplayCluster/Private/Render/Viewport/DisplayClusterViewportManagerViewExtension.cpp b/Engine/Plugins/Runtime/nDisplay/Source/DisplayCluster/Private/Render/Viewport/DisplayClusterViewportManagerViewExtension.cpp
--- a/Engine/Plugins/Runtime/nDisplay/Source/DisplayCluster/Private/Render/Viewport/DisplayClusterViewportManagerViewExtension.cpp
+++ b/Engine/Plugins/Runtime/nDisplay/Source/DisplayCluster/Private/Render/Viewport/DisplayClusterViewportManagerViewExtension.cpp
@@ -24,41 +24,37 @@ using namespace DisplayClusterViewportManagerViewExtensionHelpers;
 ///////////////////////////////////////////////////////////////////////////////////////
 FDisplayClusterViewportManagerViewExtension::FDisplayClusterViewportManagerViewExtension(const FAutoRegister& AutoRegister, const FDisplayClusterViewportManager* InViewportManager)
 	: FSceneViewExtensionBase(AutoRegister)
-	, ViewportManager(InViewportManager)
-	, ViewportManagerProxy(InViewportManager->GetViewportManagerProxy())
+	, ViewportManagerWeakPtr(InViewportManager->AsWeak())
+	, ViewportManagerProxyWeakPtr(InViewportManager->GetViewportManagerProxy()->AsWeak())
 {
 	RegisterCallbacks();
 }
 
 FDisplayClusterViewportManagerViewExtension::~FDisplayClusterViewportManagerViewExtension()
 {
-	ViewportManagerProxy.Reset();
+	ViewportManagerWeakPtr.Reset();
+	ViewportManagerProxyWeakPtr.Reset();
+
 	UnregisterCallbacks();
 }
 
 bool FDisplayClusterViewportManagerViewExtension::IsActive() const
 {
-	return ViewportManagerProxy.IsValid();
-}
-
-void FDisplayClusterViewportManagerViewExtension::Release()
-{
-	check(IsInGameThread());
-
-	ViewportManagerProxy.Reset();
+	return ViewportManagerProxyWeakPtr.IsValid() && ViewportManagerProxyWeakPtr.Pin().IsValid();
 }
 
 void FDisplayClusterViewportManagerViewExtension::Release_RenderThread()
 {
 	check(IsInRenderingThread());
 
-	ViewportManagerProxy.Reset();
+	ViewportManagerWeakPtr.Reset();
+	ViewportManagerProxyWeakPtr.Reset();
+
 	ViewportProxies.Empty();
 
 	UnregisterCallbacks();
 }
 
-
 void FDisplayClusterViewportManagerViewExtension::PreRenderView_RenderThread(FRDGBuilder& GraphBuilder, FSceneView& InView)
 {
 	if (InView.bIsViewInfo)
@@ -80,10 +76,13 @@ void FDisplayClusterViewportManagerViewExtension::SubscribeToPostProcessingPass(
 	case EPostProcessingPass::SSRInput:
 		for (const FViewportProxy& ViewportIt : ViewportProxies)
 		{
-			if (ViewportIt.ViewportProxy.IsValid() && ViewportIt.ViewportProxy.Pin()->ShouldUsePostProcessPassAfterSSRInput())
+			if (FDisplayClusterViewportProxy* ViewportProxy = ViewportIt.GetViewportProxy())
 			{
-				InOutPassCallbacks.Add(FAfterPassCallbackDelegate::CreateRaw(this, &FDisplayClusterViewportManagerViewExtension::PostProcessPassAfterSSRInput_RenderThread));
-				break;
+				if (ViewportProxy->ShouldUsePostProcessPassAfterSSRInput())
+				{
+					InOutPassCallbacks.Add(FAfterPassCallbackDelegate::CreateRaw(this, &FDisplayClusterViewportManagerViewExtension::PostProcessPassAfterSSRInput_RenderThread));
+					break;
+				}
 			}
 		}
 		break;
@@ -91,10 +90,13 @@ void FDisplayClusterViewportManagerViewExtension::SubscribeToPostProcessingPass(
 	case EPostProcessingPass::FXAA:
 		for (const FViewportProxy& ViewportIt : ViewportProxies)
 		{
-			if (ViewportIt.ViewportProxy.IsValid() && ViewportIt.ViewportProxy.Pin()->ShouldUsePostProcessPassAfterFXAA())
+			if (FDisplayClusterViewportProxy* ViewportProxy = ViewportIt.GetViewportProxy())
 			{
-				InOutPassCallbacks.Add(FAfterPassCallbackDelegate::CreateRaw(this, &FDisplayClusterViewportManagerViewExtension::PostProcessPassAfterFXAA_RenderThread));
-				break;
+				if (ViewportProxy->ShouldUsePostProcessPassAfterFXAA())
+				{
+					InOutPassCallbacks.Add(FAfterPassCallbackDelegate::CreateRaw(this, &FDisplayClusterViewportManagerViewExtension::PostProcessPassAfterFXAA_RenderThread));
+					break;
+				}
 			}
 		}
 		break;
@@ -102,10 +104,13 @@ void FDisplayClusterViewportManagerViewExtension::SubscribeToPostProcessingPass(
 	case EPostProcessingPass::Tonemap:
 		for (const FViewportProxy& ViewportIt : ViewportProxies)
 		{
-			if (ViewportIt.ViewportProxy.IsValid() && ViewportIt.ViewportProxy.Pin()->ShouldUsePostProcessPassTonemap())
+			if (FDisplayClusterViewportProxy* ViewportProxy = ViewportIt.GetViewportProxy())
 			{
-				InOutPassCallbacks.Add(FAfterPassCallbackDelegate::CreateRaw(this, &FDisplayClusterViewportManagerViewExtension::PostProcessPassAfterTonemap_RenderThread));
-				break;
+				if (ViewportProxy->ShouldUsePostProcessPassTonemap())
+				{
+					InOutPassCallbacks.Add(FAfterPassCallbackDelegate::CreateRaw(this, &FDisplayClusterViewportManagerViewExtension::PostProcessPassAfterTonemap_RenderThread));
+					break;
+				}
 			}
 		}
 		break;
@@ -122,12 +127,15 @@ FScreenPassTexture FDisplayClusterViewportManagerViewExtension::PostProcessPassA
 		return ReturnUntouchedSceneColorForPostProcessing(Inputs);
 	}
 
-	uint32 ContextNum = 0;
-	if (FDisplayClusterViewportProxy* ViewportProxyPtr = ViewportManagerProxy.Pin()->ImplFindViewport_RenderThread(View.StereoViewIndex, &ContextNum))
+	if (const FDisplayClusterViewportManagerProxy* ViewportManagerProxy = GetViewportManagerProxy())
 	{
-		if (ViewportProxyPtr->ShouldUsePostProcessPassAfterFXAA())
+		uint32 ContextNum = 0;
+		if (FDisplayClusterViewportProxy* ViewportProxyPtr = ViewportManagerProxy->ImplFindViewport_RenderThread(View.StereoViewIndex, &ContextNum))
 		{
-			return ViewportProxyPtr->OnPostProcessPassAfterFXAA_RenderThread(GraphBuilder, View, Inputs, ContextNum);
+			if (ViewportProxyPtr->ShouldUsePostProcessPassAfterFXAA())
+			{
+				return ViewportProxyPtr->OnPostProcessPassAfterFXAA_RenderThread(GraphBuilder, View, Inputs, ContextNum);
+			}
 		}
 	}
 
@@ -141,12 +149,15 @@ FScreenPassTexture FDisplayClusterViewportManagerViewExtension::PostProcessPassA
 		return ReturnUntouchedSceneColorForPostProcessing(Inputs);
 	}
 
-	uint32 ContextNum = 0;
-	if (FDisplayClusterViewportProxy* ViewportProxyPtr = ViewportManagerProxy.Pin()->ImplFindViewport_RenderThread(View.StereoViewIndex, &ContextNum))
+	if (const FDisplayClusterViewportManagerProxy* ViewportManagerProxy = GetViewportManagerProxy())
 	{
-		if (ViewportProxyPtr->ShouldUsePostProcessPassAfterSSRInput())
+		uint32 ContextNum = 0;
+		if (FDisplayClusterViewportProxy* ViewportProxyPtr = ViewportManagerProxy->ImplFindViewport_RenderThread(View.StereoViewIndex, &ContextNum))
 		{
-			return ViewportProxyPtr->OnPostProcessPassAfterSSRInput_RenderThread(GraphBuilder, View, Inputs, ContextNum);
+			if (ViewportProxyPtr->ShouldUsePostProcessPassAfterSSRInput())
+			{
+				return ViewportProxyPtr->OnPostProcessPassAfterSSRInput_RenderThread(GraphBuilder, View, Inputs, ContextNum);
+			}
 		}
 	}
 
@@ -160,12 +171,15 @@ FScreenPassTexture FDisplayClusterViewportManagerViewExtension::PostProcessPassA
 		return ReturnUntouchedSceneColorForPostProcessing(Inputs);
 	}
 
-	uint32 ContextNum = 0;
-	if (FDisplayClusterViewportProxy* ViewportProxyPtr = ViewportManagerProxy.Pin()->ImplFindViewport_RenderThread(View.StereoViewIndex, &ContextNum))
+	if (const FDisplayClusterViewportManagerProxy* ViewportManagerProxy = GetViewportManagerProxy())
 	{
-		if (ViewportProxyPtr->ShouldUsePostProcessPassTonemap())
+		uint32 ContextNum = 0;
+		if (FDisplayClusterViewportProxy* ViewportProxyPtr = ViewportManagerProxy->ImplFindViewport_RenderThread(View.StereoViewIndex, &ContextNum))
 		{
-			return ViewportProxyPtr->OnPostProcessPassAfterTonemap_RenderThread(GraphBuilder, View, Inputs, ContextNum);
+			if (ViewportProxyPtr->ShouldUsePostProcessPassTonemap())
+			{
+				return ViewportProxyPtr->OnPostProcessPassAfterTonemap_RenderThread(GraphBuilder, View, Inputs, ContextNum);
+			}
 		}
 	}
 
@@ -200,7 +214,7 @@ bool FDisplayClusterViewportManagerViewExtension::IsActiveThisFrame_Internal(con
 	if (Context.IsA(MoveTempIfPossible(DCViewExtensionContext)))
 	{
 		const FDisplayClusterSceneViewExtensionContext& DisplayContext = static_cast<const FDisplayClusterSceneViewExtensionContext&>(Context);
-		if (DisplayContext.ViewportManager == ViewportManager)
+		if (DisplayContext.ViewportManager.Get() == GetViewportManager())
 		{
 			// Apply only for DC viewports
 			return true;
@@ -226,14 +240,17 @@ void FDisplayClusterViewportManagerViewExtension::PreRenderViewFamily_RenderThre
 		{
 			FViewportProxy NewViewportProxy;
 
-			if (IDisplayClusterViewportProxy* ViewportProxyPtr = ViewportManagerProxy.Pin()->FindViewport_RenderThread(SceneView->StereoViewIndex, &NewViewportProxy.ViewportProxyContext.ContextNum))
+			if (const FDisplayClusterViewportManagerProxy* ViewportManagerProxy = GetViewportManagerProxy())
 			{
-				NewViewportProxy.ViewportProxyContext.ViewFamilyProfileDescription = InViewFamily.ProfileDescription;
+				if (FDisplayClusterViewportProxy* ViewportProxyPtr = static_cast<FDisplayClusterViewportProxy*>(ViewportManagerProxy->FindViewport_RenderThread(SceneView->StereoViewIndex, &NewViewportProxy.ViewportProxyContext.ContextNum)))
+				{
+					NewViewportProxy.ViewportProxyContext.ViewFamilyProfileDescription = InViewFamily.ProfileDescription;
 
-				NewViewportProxy.ViewportProxy = (static_cast<FDisplayClusterViewportProxy*>(ViewportProxyPtr))->AsShared();
-				NewViewportProxy.ViewIndex = ViewIndex;
+					NewViewportProxy.ViewportProxyWeakPtr = ViewportProxyPtr->AsWeak();
+					NewViewportProxy.ViewIndex = ViewIndex;
 
-				ViewportProxies.Add(NewViewportProxy);
+					ViewportProxies.Add(NewViewportProxy);
+				}
 			}
 		}
 	}
@@ -250,9 +267,9 @@ void FDisplayClusterViewportManagerViewExtension::PostRenderViewFamily_RenderThr
 	{
 		if (const FSceneView* SceneView = InViewFamily.Views[ViewportIt.ViewIndex])
 		{
-			if (ViewportIt.ViewportProxy.IsValid())
+			if (FDisplayClusterViewportProxy* ViewportProxy = ViewportIt.GetViewportProxy())
 			{
-				ViewportIt.ViewportProxy.Pin()->OnPostRenderViewFamily_RenderThread(GraphBuilder, InViewFamily, *SceneView, ViewportIt.ViewportProxyContext);
+				ViewportProxy->OnPostRenderViewFamily_RenderThread(GraphBuilder, InViewFamily, *SceneView, ViewportIt.ViewportProxyContext);
 			}
 		}
 	}
@@ -270,9 +287,9 @@ void FDisplayClusterViewportManagerViewExtension::OnResolvedSceneColor_RenderThr
 
 	for (FViewportProxy& ViewportIt : ViewportProxies)
 	{
-		if (ViewportIt.ViewportProxy.IsValid())
+		if (FDisplayClusterViewportProxy* ViewportProxy = ViewportIt.GetViewportProxy())
 		{
-			ViewportIt.ViewportProxy.Pin()->OnResolvedSceneColor_RenderThread(GraphBuilder, SceneTextures, ViewportIt.ViewportProxyContext);
+			ViewportProxy->OnResolvedSceneColor_RenderThread(GraphBuilder, SceneTextures, ViewportIt.ViewportProxyContext);
 		}
 	}
 }
diff --git a/Engine/Plugins/Runtime/nDisplay/Source/DisplayCluster/Private/Render/Viewport/DisplayClusterViewportManagerViewExtension.h b/Engine/Plugins/Runtime/nDisplay/Source/DisplayCluster/Private/Render/Viewport/DisplayClusterViewportManagerViewExtension.h
--- a/Engine/Plugins/Runtime/nDisplay/Source/DisplayCluster/Private/Render/Viewport/DisplayClusterViewportManagerViewExtension.h
+++ b/Engine/Plugins/Runtime/nDisplay/Source/DisplayCluster/Private/Render/Viewport/DisplayClusterViewportManagerViewExtension.h
@@ -53,9 +53,6 @@ public:
 	 */
 	static FScreenPassTexture ReturnUntouchedSceneColorForPostProcessing(const FPostProcessMaterialInputs& InOutInputs);
 
-	/** Release from game thread.*/
-	void Release();
-
 	/** Release from render thread.*/
 	void Release_RenderThread();
 
@@ -108,17 +105,33 @@ private:
 	/** Return true if the VE is safe to use. */
 	bool IsActive() const;
 
+protected:
+	inline const FDisplayClusterViewportManager* GetViewportManager() const
+	{
+		return ViewportManagerWeakPtr.IsValid() ? ViewportManagerWeakPtr.Pin().Get() : nullptr;
+	}
+
+	inline const FDisplayClusterViewportManagerProxy* GetViewportManagerProxy() const
+	{
+		return ViewportManagerProxyWeakPtr.IsValid() ? ViewportManagerProxyWeakPtr.Pin().Get() : nullptr;
+	}
+
 private:
 	FDelegateHandle ResolvedSceneColorCallbackHandle;
 
 private:
-	const FDisplayClusterViewportManager* ViewportManager;
-	TWeakPtr<FDisplayClusterViewportManagerProxy, ESPMode::ThreadSafe> ViewportManagerProxy;
+	TWeakPtr<const FDisplayClusterViewportManager, ESPMode::ThreadSafe> ViewportManagerWeakPtr;
+	TWeakPtr<const FDisplayClusterViewportManagerProxy, ESPMode::ThreadSafe> ViewportManagerProxyWeakPtr;
 
 	struct FViewportProxy
 	{
+		inline FDisplayClusterViewportProxy* GetViewportProxy() const
+		{
+			return ViewportProxyWeakPtr.IsValid() ? ViewportProxyWeakPtr.Pin().Get() : nullptr;
+		}
+
 		// The ViewportProxy must exist until this object is removed.
-		TWeakPtr<FDisplayClusterViewportProxy, ESPMode::ThreadSafe> ViewportProxy;
+		TWeakPtr<FDisplayClusterViewportProxy, ESPMode::ThreadSafe> ViewportProxyWeakPtr;
 
 		// The index of view in viewfamily collection
 		int32 ViewIndex = 0;
diff --git a/Engine/Plugins/Runtime/nDisplay/Source/DisplayCluster/Private/Render/Viewport/DisplayClusterViewportManagerViewPointExtension.cpp b/Engine/Plugins/Runtime/nDisplay/Source/DisplayCluster/Private/Render/Viewport/DisplayClusterViewportManagerViewPointExtension.cpp
--- a/Engine/Plugins/Runtime/nDisplay/Source/DisplayCluster/Private/Render/Viewport/DisplayClusterViewportManagerViewPointExtension.cpp
+++ b/Engine/Plugins/Runtime/nDisplay/Source/DisplayCluster/Private/Render/Viewport/DisplayClusterViewportManagerViewPointExtension.cpp
@@ -10,19 +10,12 @@
 ///////////////////////////////////////////////////////////////////////////////////////
 FDisplayClusterViewportManagerViewPointExtension::FDisplayClusterViewportManagerViewPointExtension(const FAutoRegister& AutoRegister, const FDisplayClusterViewportManager* InViewportManager)
 	: FSceneViewExtensionBase(AutoRegister)
-	, ViewportManager(InViewportManager)
+	, ViewportManagerWeakPtr(InViewportManager->AsShared())
 { }
 
 FDisplayClusterViewportManagerViewPointExtension::~FDisplayClusterViewportManagerViewPointExtension()
 {
-	ViewportManager = nullptr;
-}
-
-void FDisplayClusterViewportManagerViewPointExtension::Release()
-{
-	check(IsInGameThread());
-
-	ViewportManager = nullptr;
+	ViewportManagerWeakPtr.Reset();
 }
 
 bool FDisplayClusterViewportManagerViewPointExtension::IsActiveThisFrame_Internal(const FSceneViewExtensionContext& Context) const
@@ -32,7 +25,7 @@ bool FDisplayClusterViewportManagerViewPointExtension::IsActiveThisFrame_Interna
 
 void FDisplayClusterViewportManagerViewPointExtension::SetupViewPoint(APlayerController* Player, FMinimalViewInfo& InOutViewInfo)
 {
-	if (IDisplayClusterViewport* DCViewport = IsActive() ? ViewportManager->FindViewport(CurrentStereoViewIndex) : nullptr)
+	if (IDisplayClusterViewport* DCViewport = IsActive() ? GetViewportManager()->FindViewport(CurrentStereoViewIndex) : nullptr)
 	{
 		DCViewport->SetupViewPoint(InOutViewInfo);
 	}
diff --git a/Engine/Plugins/Runtime/nDisplay/Source/DisplayCluster/Private/Render/Viewport/DisplayClusterViewportManagerViewPointExtension.h b/Engine/Plugins/Runtime/nDisplay/Source/DisplayCluster/Private/Render/Viewport/DisplayClusterViewportManagerViewPointExtension.h
--- a/Engine/Plugins/Runtime/nDisplay/Source/DisplayCluster/Private/Render/Viewport/DisplayClusterViewportManagerViewPointExtension.h
+++ b/Engine/Plugins/Runtime/nDisplay/Source/DisplayCluster/Private/Render/Viewport/DisplayClusterViewportManagerViewPointExtension.h
@@ -26,9 +26,6 @@ public:
 
 	virtual int32 GetPriority() const override { return DISPLAYCLUSTER_SCENE_VIEWPOINT_EXTENSION_PRIORITY; }
 
-	/** Release from game thread.*/
-	void Release();
-
 	/** Set current view index for this VE:
 	 * This function must be used because LocalPlayer::GetViewPoint() calls the ISceneViewExtension::SetupViewPoint() function from this VE.
 	 * And at this point the VE must know the current StereoViewIndex value in order to understand which viewport will be used for this SetupViewPoint() call.
@@ -43,15 +40,20 @@ protected:
 	virtual bool IsActiveThisFrame_Internal(const FSceneViewExtensionContext& Context) const override;
 
 private:
+	/** Get viewport manager ptr. */
+	inline const FDisplayClusterViewportManager* GetViewportManager() const
+	{
+		return ViewportManagerWeakPtr.IsValid() ? ViewportManagerWeakPtr.Pin().Get() : nullptr;
+	}
+
 	/** True, if VE can be used at the moment. */
 	bool IsActive() const
 	{
-		return ViewportManager && CurrentStereoViewIndex != INDEX_NONE;
+		return GetViewportManager() != nullptr && CurrentStereoViewIndex != INDEX_NONE;
 	}
 
 private:
-	// Pointer to the viewport manager to which this VE belongs
-	const FDisplayClusterViewportManager* ViewportManager;
+	TWeakPtr<const FDisplayClusterViewportManager, ESPMode::ThreadSafe> ViewportManagerWeakPtr;
 
 	// Current StereoViewIndex for rendered viewport
 	int32 CurrentStereoViewIndex = INDEX_NONE;
diff --git a/Engine/Plugins/Runtime/nDisplay/Source/DisplayCluster/Private/Render/Viewport/DisplayClusterViewportManager_EditorPreview.cpp b/Engine/Plugins/Runtime/nDisplay/Source/DisplayCluster/Private/Render/Viewport/DisplayClusterViewportManager_EditorPreview.cpp
--- a/Engine/Plugins/Runtime/nDisplay/Source/DisplayCluster/Private/Render/Viewport/DisplayClusterViewportManager_EditorPreview.cpp
+++ b/Engine/Plugins/Runtime/nDisplay/Source/DisplayCluster/Private/Render/Viewport/DisplayClusterViewportManager_EditorPreview.cpp
@@ -32,33 +32,13 @@
 ///////////////////////////////////////////////////////////////////////////////////////
 // FDisplayClusterViewportManager
 ///////////////////////////////////////////////////////////////////////////////////////
-bool FDisplayClusterViewportManager::IsEditorPreviewWorld() const
-{
-	if (UWorld* CurrentWorld = GetCurrentWorld())
-	{
-		switch (CurrentWorld->WorldType)
-		{
-		case EWorldType::EditorPreview:
-			return true;
-
-		default:
-			break;
-		}
-	}
-
-	return false;
-}
-
 void FDisplayClusterViewportManager::ImplUpdatePreviewRTTResources()
 {
 	check(IsInGameThread());
 
-	switch (GetRenderFrameSettings().RenderMode)
+	// Only for preview modes:
+	if(GetRenderMode() != EDisplayClusterRenderFrameMode::PreviewInScene)
 	{
-	case EDisplayClusterRenderFrameMode::PreviewInScene:
-		// Only for preview modes:
-		break;
-	default:
 		return;
 	}
 
@@ -78,10 +58,10 @@ void FDisplayClusterViewportManager::ImplUpdatePreviewRTTResources()
 	PreviewRenderTargetableTextures.Reserve(Viewports.Num());
 
 	// Collect visible and enabled viewports for preview rendering:
-	for (FDisplayClusterViewport* const ViewportIt : Viewports)
+	for (const TSharedPtr<FDisplayClusterViewport, ESPMode::ThreadSafe>& ViewportIt : Viewports)
 	{
 		// update only current cluster node
-		if ((ClusterNodeId.IsEmpty() || ViewportIt->GetClusterNodeId() == ClusterNodeId))
+		if (ViewportIt.IsValid() && (ClusterNodeId.IsEmpty() || ViewportIt->GetClusterNodeId() == ClusterNodeId))
 		{
 			if (ViewportIt->RenderSettings.bEnable && ViewportIt->RenderSettings.bVisible)
 			{
@@ -97,8 +77,8 @@ void FDisplayClusterViewportManager::ImplUpdatePreviewRTTResources()
 	// Configure preview RTT to viwports:
 	for (int32 ViewportIndex = 0; ViewportIndex < PreviewViewportNames.Num(); ViewportIndex++)
 	{
-		FDisplayClusterViewport* DesiredViewport = ImplFindViewport(PreviewViewportNames[ViewportIndex]);
-		if (DesiredViewport != nullptr)
+		TSharedPtr<FDisplayClusterViewport, ESPMode::ThreadSafe> DesiredViewport = ImplFindViewport(PreviewViewportNames[ViewportIndex]);
+		if (DesiredViewport.IsValid())
 		{
 			FTextureRHIRef& PreviewRTT = PreviewRenderTargetableTextures[ViewportIndex];
 			if (PreviewRTT.IsValid())
@@ -115,7 +95,7 @@ void FDisplayClusterViewportManager::ImplUpdatePreviewRTTResources()
 	}
 }
 
-bool FDisplayClusterViewportManager::RenderInEditor(class FDisplayClusterRenderFrame& InRenderFrame, FViewport* InViewport, const uint32 InFirstViewportNum, const int32 InViewportsAmount, int32& OutViewportsAmount, bool& bOutFrameRendered)
+bool FDisplayClusterViewportManager::RenderInEditor(FDisplayClusterRenderFrame& InRenderFrame, FViewport* InViewport, const uint32 InFirstViewportNum, const int32 InViewportsAmount, int32& OutViewportsAmount, bool& bOutFrameRendered)
 {
 	bOutFrameRendered = false;
 	OutViewportsAmount = 0;
@@ -142,7 +122,7 @@ bool FDisplayClusterViewportManager::RenderInEditor(class FDisplayClusterRenderF
 	int32 ViewportIndex = 0;
 	bool bViewportsRenderPassDone = false;
 
-	for (FDisplayClusterRenderFrame::FFrameRenderTarget& RenderTargetIt : InRenderFrame.RenderTargets)
+	for (FDisplayClusterRenderFrameTarget& RenderTargetIt : InRenderFrame.RenderTargets)
 	{
 		if (bViewportsRenderPassDone)
 		{
@@ -152,7 +132,7 @@ bool FDisplayClusterViewportManager::RenderInEditor(class FDisplayClusterRenderF
 		// Special flag, allow clear RTT surface only for first family
 		bool bAdditionalViewFamily = false;
 
-		for (FDisplayClusterRenderFrame::FFrameViewFamily& ViewFamiliesIt : RenderTargetIt.ViewFamilies)
+		for (FDisplayClusterRenderFrameTargetViewFamily& ViewFamiliesIt : RenderTargetIt.ViewFamilies)
 		{
 			if (bViewportsRenderPassDone)
 			{
@@ -188,23 +168,27 @@ bool FDisplayClusterViewportManager::RenderInEditor(class FDisplayClusterRenderF
 					}
 				}
 				
-				for (FDisplayClusterRenderFrame::FFrameView& ViewIt : ViewFamiliesIt.Views)
+				for (FDisplayClusterRenderFrameTargetView& ViewIt : ViewFamiliesIt.Views)
 				{
 					bool bViewportAlreadyRendered = ViewportIndex < (int32)InFirstViewportNum;
 					ViewportIndex++;
 
 					if (!bViewportAlreadyRendered)
 					{
-						FDisplayClusterViewport* ViewportPtr = static_cast<FDisplayClusterViewport*>(ViewIt.Viewport);
-						check(ViewportPtr != nullptr);
-						check(ViewIt.ContextNum < (uint32)ViewportPtr->Contexts.Num());
+						FDisplayClusterViewport* ViewportPtr = static_cast<FDisplayClusterViewport*>(ViewIt.Viewport.Get());
+						if(!ViewportPtr)
+						{
+							continue;
+						}
+
+						check(ViewportPtr->GetContexts().IsValidIndex(ViewIt.ContextNum));
 
 						// Calculate the player's view information.
 						FVector  ViewLocation;
 						FRotator ViewRotation;
 						FSceneView* View = ViewportPtr->ImplCalcScenePreview(ViewFamily, ViewIt.ContextNum);
 
-						if (View != nullptr && ViewIt.ShouldRenderSceneView() == false)
+						if (View != nullptr && !ViewIt.IsViewportContextCanBeRendered())
 						{
 							ViewFamily.Views.Remove(View);
 
@@ -218,7 +202,7 @@ bool FDisplayClusterViewportManager::RenderInEditor(class FDisplayClusterRenderF
 							OutViewportsAmount++;
 
 							// Apply viewport context settings to view (crossGPU, visibility, etc)
-							ViewIt.Viewport->SetupSceneView(ViewIt.ContextNum, PreviewScene->GetWorld(), ViewFamily, *View);
+							ViewportPtr->SetupSceneView(ViewIt.ContextNum, PreviewScene->GetWorld(), ViewFamily, *View);
 
 							// The allowed number of viewports is limited by the "InViewportAmount" value.
 							if (InViewportsAmount > 0 && OutViewportsAmount >= InViewportsAmount)
diff --git a/Engine/Plugins/Runtime/nDisplay/Source/DisplayCluster/Private/Render/Viewport/DisplayClusterViewportProxy.cpp b/Engine/Plugins/Runtime/nDisplay/Source/DisplayCluster/Private/Render/Viewport/DisplayClusterViewportProxy.cpp
--- a/Engine/Plugins/Runtime/nDisplay/Source/DisplayCluster/Private/Render/Viewport/DisplayClusterViewportProxy.cpp
+++ b/Engine/Plugins/Runtime/nDisplay/Source/DisplayCluster/Private/Render/Viewport/DisplayClusterViewportProxy.cpp
@@ -190,7 +190,7 @@ FDisplayClusterViewportProxy::FDisplayClusterViewportProxy(const FDisplayCluster
 	, ClusterNodeId(RenderViewport.ClusterNodeId)
 	, RenderSettings(RenderViewport.RenderSettings)
 	, ProjectionPolicy(RenderViewport.UninitializedProjectionPolicy)
-	, Owner(RenderViewport.ImplGetOwner().GetViewportManagerProxy().ToSharedRef())
+	, ViewportManagerProxyWeakRef(RenderViewport.GetViewportManagerProxyRefImpl())
 	, ShadersAPI(IDisplayClusterShaders::Get())
 {
 	check(ProjectionPolicy.IsValid());
@@ -200,11 +200,39 @@ FDisplayClusterViewportProxy::~FDisplayClusterViewportProxy()
 {
 }
 
-const IDisplayClusterViewportManagerProxy& FDisplayClusterViewportProxy::GetOwner_RenderThread() const
+const FDisplayClusterRenderFrameSettings* FDisplayClusterViewportProxy::GetRenderFrameSettings_RenderThread() const
 {
-	check(IsInRenderingThread());
+	if (FDisplayClusterViewportManagerProxy* ViewportManagerProxy = GetViewportManagerProxyImpl_RenderThread())
+	{
+		return &ViewportManagerProxy->GetRenderFrameSettings_RenderThread();
+	}
+
+	return nullptr;
+}
+
+const IDisplayClusterViewportManagerProxy* FDisplayClusterViewportProxy::GetViewportManagerProxy_RenderThread() const
+{
+	return GetViewportManagerProxyRefImpl_RenderThread().Get();
+}
+
+FDisplayClusterViewportManagerProxy* FDisplayClusterViewportProxy::GetViewportManagerProxyImpl_RenderThread() const
+{
+	return GetViewportManagerProxyRefImpl_RenderThread().Get();
+}
+
+TSharedPtr<FDisplayClusterViewportManagerProxy, ESPMode::ThreadSafe> FDisplayClusterViewportProxy::GetViewportManagerProxyRefImpl_RenderThread() const
+{
+	return ViewportManagerProxyWeakRef.IsValid() ? ViewportManagerProxyWeakRef.Pin() : nullptr;
+}
+
+EDisplayClusterRenderFrameMode FDisplayClusterViewportProxy::GetRenderMode() const
+{
+	if (FDisplayClusterViewportManagerProxy* ViewportManagerProxy = GetViewportManagerProxyImpl_RenderThread())
+	{
+		return ViewportManagerProxy->GetRenderMode();
+	}
 
-	return Owner.Get();
+	return EDisplayClusterRenderFrameMode::Unknown;
 }
 
 void FDisplayClusterViewportProxy::HandleResourceDelete_RenderThread(FDisplayClusterViewportResource* InDeletedResourcePtr)
@@ -276,9 +304,12 @@ const FDisplayClusterViewportProxy& FDisplayClusterViewportProxy::GetRenderingVi
 	{
 	case EDisplayClusterViewportOverrideMode::All:
 	case EDisplayClusterViewportOverrideMode::InernalRTT:
-		if (FDisplayClusterViewportProxy const* OverrideViewportProxy = Owner->ImplFindViewport_RenderThread(RenderSettings.ViewportOverrideId))
+		if (FDisplayClusterViewportManagerProxy* ViewportManagerProxy = GetViewportManagerProxyImpl_RenderThread())
 		{
-			return *OverrideViewportProxy;
+			if (FDisplayClusterViewportProxy const* OverrideViewportProxy = ViewportManagerProxy->ImplFindViewport_RenderThread(RenderSettings.ViewportOverrideId))
+			{
+				return *OverrideViewportProxy;
+			}
 		}
 
 		break;
@@ -351,14 +382,17 @@ bool FDisplayClusterViewportProxy::ImplGetResources_RenderThread(const EDisplayC
 				bResult = true;
 				
 				// Get resources from external UV LightCard manager
-				TSharedPtr<FDisplayClusterViewportLightCardManagerProxy, ESPMode::ThreadSafe> LightCardManager = Owner->GetLightCardManagerProxy_RenderThread();
-				if (LightCardManager.IsValid())
+				if (FDisplayClusterViewportManagerProxy* ViewportManagerProxy = GetViewportManagerProxyImpl_RenderThread())
 				{
-					if (FRHITexture* UVLightCardRHIResource = LightCardManager->GetUVLightCardRHIResource_RenderThread())
+					TSharedPtr<FDisplayClusterViewportLightCardManagerProxy, ESPMode::ThreadSafe> LightCardManager = ViewportManagerProxy->GetLightCardManagerProxy_RenderThread();
+					if (LightCardManager.IsValid())
 					{
-						for (int32 ContextIndex = 0; ContextIndex < Contexts.Num(); ContextIndex++)
+						if (FRHITexture* UVLightCardRHIResource = LightCardManager->GetUVLightCardRHIResource_RenderThread())
 						{
-							OutResources.Add(UVLightCardRHIResource);
+							for (int32 ContextIndex = 0; ContextIndex < Contexts.Num(); ContextIndex++)
+							{
+								OutResources.Add(UVLightCardRHIResource);
+							}
 						}
 					}
 				}
@@ -450,16 +484,11 @@ void FDisplayClusterViewportProxy::PostResolveViewport_RenderThread(FRHICommandL
 
 void FDisplayClusterViewportProxy::ImplViewportRemap_RenderThread(FRHICommandListImmediate& RHICmdList) const
 {
-#if WITH_EDITOR
-	switch(Owner->GetRenderFrameSettings_RenderThread().RenderMode)
+	// Preview in editor not support this feature
+	if (GetRenderMode() == EDisplayClusterRenderFrameMode::PreviewInScene)
 	{
-	case EDisplayClusterRenderFrameMode::PreviewInScene:
-		// Preview in editor not support this feature
 		return;
-	default:
-		break;
 	}
-#endif
 
 	if (RemapMesh.IsValid())
 	{
@@ -500,7 +529,8 @@ void FDisplayClusterViewportProxy::ImplPreviewReadPixels_RenderThread(FRHIComman
 	// Now try to read until success:
 	if (bPreviewReadPixels && Contexts.Num() > 0 && OutputPreviewTargetableResource.IsValid())
 	{
-		check(Owner->GetRenderFrameSettings_RenderThread().RenderMode == EDisplayClusterRenderFrameMode::PreviewInScene);
+		const FDisplayClusterRenderFrameSettings* RenderFrameSettings = GetRenderFrameSettings_RenderThread();
+		check(RenderFrameSettings && RenderFrameSettings->RenderMode == EDisplayClusterRenderFrameMode::PreviewInScene);
 
 		// We should synchronize thread for preview read
 		FScopeLock Lock(&PreviewPixelsCSGuard);
@@ -558,13 +588,9 @@ EDisplayClusterViewportResourceType FDisplayClusterViewportProxy::GetOutputResou
 	check(IsInRenderingThread());
 
 #if WITH_EDITOR
-	switch(Owner->GetRenderFrameSettings_RenderThread().RenderMode)
+	if(GetRenderMode() == EDisplayClusterRenderFrameMode::PreviewInScene)
 	{
-	case EDisplayClusterRenderFrameMode::PreviewInScene:
 		return EDisplayClusterViewportResourceType::OutputPreviewTargetableResource;
-
-	default:
-		break;
 	}
 #endif
 
@@ -1032,7 +1058,10 @@ bool FDisplayClusterViewportProxy::ShouldUsePostProcessPassAfterSSRInput() const
 {	
 	if (ShouldUseAlphaChannel_RenderThread())
 	{
-		return Owner->GetRenderFrameSettings_RenderThread().AlphaChannelCaptureMode == EDisplayClusterRenderFrameAlphaChannelCaptureMode::ThroughTonemapper;
+		if (const FDisplayClusterRenderFrameSettings* RenderFrameSettings = GetRenderFrameSettings_RenderThread())
+		{
+			return RenderFrameSettings->AlphaChannelCaptureMode == EDisplayClusterRenderFrameAlphaChannelCaptureMode::ThroughTonemapper;
+		}
 	}
 
 	return false;
@@ -1042,7 +1071,10 @@ bool FDisplayClusterViewportProxy::ShouldUsePostProcessPassAfterFXAA() const
 {
 	if (ShouldUseAlphaChannel_RenderThread())
 	{
-		return Owner->GetRenderFrameSettings_RenderThread().AlphaChannelCaptureMode == EDisplayClusterRenderFrameAlphaChannelCaptureMode::ThroughTonemapper;
+		if (const FDisplayClusterRenderFrameSettings* RenderFrameSettings = GetRenderFrameSettings_RenderThread())
+		{
+			return RenderFrameSettings->AlphaChannelCaptureMode == EDisplayClusterRenderFrameAlphaChannelCaptureMode::ThroughTonemapper;
+		}
 	}
 
 	return false;
@@ -1063,20 +1095,23 @@ void FDisplayClusterViewportProxy::OnResolvedSceneColor_RenderThread(FRDGBuilder
 	const uint32 InContextNum = InProxyContext.ContextNum;
 	if (ShouldUseAlphaChannel_RenderThread())
 	{
-		switch (Owner->GetRenderFrameSettings_RenderThread().AlphaChannelCaptureMode)
-		{
-		case EDisplayClusterRenderFrameAlphaChannelCaptureMode::FXAA:
-		case EDisplayClusterRenderFrameAlphaChannelCaptureMode::Copy:
-		case EDisplayClusterRenderFrameAlphaChannelCaptureMode::CopyAA:
+		if (const FDisplayClusterRenderFrameSettings* RenderFrameSettings = GetRenderFrameSettings_RenderThread())
 		{
-			const FIntRect SrcRect = GetFinalContextRect(EDisplayClusterViewportResourceType::InternalRenderTargetResource, Contexts[InContextNum].RenderTargetRect);
-			// Copy alpha channel from 'SceneTextures.Color.Resolve' to 'InputShaderResource'
-			CopyResource_RenderThread(GraphBuilder, EDisplayClusterTextureCopyMode::Alpha, InContextNum, SceneTextures.Color.Resolve, SrcRect, EDisplayClusterViewportResourceType::InputShaderResource);
-		}
-		break;
-
-		default:
+			switch (RenderFrameSettings->AlphaChannelCaptureMode)
+			{
+			case EDisplayClusterRenderFrameAlphaChannelCaptureMode::FXAA:
+			case EDisplayClusterRenderFrameAlphaChannelCaptureMode::Copy:
+			case EDisplayClusterRenderFrameAlphaChannelCaptureMode::CopyAA:
+			{
+				const FIntRect SrcRect = GetFinalContextRect(EDisplayClusterViewportResourceType::InternalRenderTargetResource, Contexts[InContextNum].RenderTargetRect);
+				// Copy alpha channel from 'SceneTextures.Color.Resolve' to 'InputShaderResource'
+				CopyResource_RenderThread(GraphBuilder, EDisplayClusterTextureCopyMode::Alpha, InContextNum, SceneTextures.Color.Resolve, SrcRect, EDisplayClusterViewportResourceType::InputShaderResource);
+			}
 			break;
+
+			default:
+				break;
+			}
 		}
 	}
 }
@@ -1151,9 +1186,10 @@ void FDisplayClusterViewportProxy::OnPostRenderViewFamily_RenderThread(FRDGBuild
 		}
 	}
 
-	if (ShouldUseAlphaChannel_RenderThread())
+	const FDisplayClusterRenderFrameSettings* RenderFrameSettings = GetRenderFrameSettings_RenderThread();
+	if (ShouldUseAlphaChannel_RenderThread() && RenderFrameSettings)
 	{
-		switch (Owner->GetRenderFrameSettings_RenderThread().AlphaChannelCaptureMode)
+		switch (RenderFrameSettings->AlphaChannelCaptureMode)
 		{
 		case EDisplayClusterRenderFrameAlphaChannelCaptureMode::FXAA:
 		{
diff --git a/Engine/Plugins/Runtime/nDisplay/Source/DisplayCluster/Private/Render/Viewport/DisplayClusterViewportProxy.h b/Engine/Plugins/Runtime/nDisplay/Source/DisplayCluster/Private/Render/Viewport/DisplayClusterViewportProxy.h
--- a/Engine/Plugins/Runtime/nDisplay/Source/DisplayCluster/Private/Render/Viewport/DisplayClusterViewportProxy.h
+++ b/Engine/Plugins/Runtime/nDisplay/Source/DisplayCluster/Private/Render/Viewport/DisplayClusterViewportProxy.h
@@ -69,6 +69,18 @@ public:
 	///////////////////////////////
 	// IDisplayClusterViewportProxy
 	///////////////////////////////
+	virtual TSharedPtr<IDisplayClusterViewportProxy, ESPMode::ThreadSafe> ToSharedPtr() override
+	{
+		return AsShared();
+	}
+
+	virtual TSharedPtr<const IDisplayClusterViewportProxy, ESPMode::ThreadSafe> ToSharedPtr() const override
+	{
+		return AsShared();
+	}
+
+	virtual EDisplayClusterRenderFrameMode GetRenderMode() const override;
+
 	virtual FString GetId() const override
 	{
 		check(IsInRenderingThread());
@@ -133,7 +145,10 @@ public:
 
 	virtual EDisplayClusterViewportResourceType GetOutputResourceType_RenderThread() const override;
 
-	virtual const IDisplayClusterViewportManagerProxy& GetOwner_RenderThread() const override;
+	virtual const class IDisplayClusterViewportManagerProxy* GetViewportManagerProxy_RenderThread() const override;
+
+	virtual const FDisplayClusterRenderFrameSettings* GetRenderFrameSettings_RenderThread() const override;
+
 	///////////////////////////////
 	// ~IDisplayClusterViewportProxy
 	///////////////////////////////
@@ -308,6 +323,9 @@ private:
 	/** Returns the OCIO rendering type for the given viewport. */
 	EDisplayClusterViewportOpenColorIOMode GetOpenColorIOMode() const;
 
+	FDisplayClusterViewportManagerProxy* GetViewportManagerProxyImpl_RenderThread() const;
+	TSharedPtr<FDisplayClusterViewportManagerProxy, ESPMode::ThreadSafe> GetViewportManagerProxyRefImpl_RenderThread() const;
+
 protected:
 	friend FDisplayClusterViewportProxyData;
 	friend FDisplayClusterViewportManagerProxy;
@@ -358,7 +376,8 @@ protected:
 	TArray<FDisplayClusterViewportTextureResource*> AdditionalTargetableResources;
 	TArray<FDisplayClusterViewportTextureResource*> MipsShaderResources;
 
-	const TSharedRef<FDisplayClusterViewportManagerProxy, ESPMode::ThreadSafe> Owner;
+	TWeakPtr<FDisplayClusterViewportManagerProxy, ESPMode::ThreadSafe> ViewportManagerProxyWeakRef;
+
 	IDisplayClusterShaders& ShadersAPI;
 };
 
diff --git a/Engine/Plugins/Runtime/nDisplay/Source/DisplayCluster/Private/Render/Viewport/DisplayClusterViewport_EditorPreview.cpp b/Engine/Plugins/Runtime/nDisplay/Source/DisplayCluster/Private/Render/Viewport/DisplayClusterViewport_EditorPreview.cpp
--- a/Engine/Plugins/Runtime/nDisplay/Source/DisplayCluster/Private/Render/Viewport/DisplayClusterViewport_EditorPreview.cpp
+++ b/Engine/Plugins/Runtime/nDisplay/Source/DisplayCluster/Private/Render/Viewport/DisplayClusterViewport_EditorPreview.cpp
@@ -65,7 +65,7 @@ void FDisplayClusterViewport::CleanupViewState()
 
 FSceneViewStateInterface* FDisplayClusterViewport::GetViewState(uint32 ViewIndex)
 {
-	if (GDisplayClusterPreviewEnableViewState == 0 || (GDisplayClusterPreviewEnableConfiguratorViewState == 0 && Owner.IsEditorPreviewWorld()))
+	if (GDisplayClusterPreviewEnableViewState == 0 || (GDisplayClusterPreviewEnableConfiguratorViewState == 0 && IsCurrentWorldHasAnyType(EWorldType::EditorPreview)))
 	{
 		// Disable ViewState
 		ViewStates.Empty();
@@ -86,7 +86,7 @@ FSceneViewStateInterface* FDisplayClusterViewport::GetViewState(uint32 ViewIndex
 
 	if (ViewStates[ViewIndex]->GetReference() == NULL)
 	{
-		const UWorld* CurrentWorld = Owner.GetCurrentWorld();
+		const UWorld* CurrentWorld = GetCurrentWorld();
 		const ERHIFeatureLevel::Type FeatureLevel = CurrentWorld ? CurrentWorld->GetFeatureLevel() : GMaxRHIFeatureLevel;
 
 		ViewStates[ViewIndex]->Allocate(FeatureLevel);
@@ -147,9 +147,12 @@ FSceneView* FDisplayClusterViewport::ImplCalcScenePreview(FSceneViewFamilyContex
 
 		ViewInitOptions.BackgroundColor = FLinearColor::Black;
 
-		if (Owner.GetRenderFrameSettings().bPreviewEnablePostProcess == false)
+		if (const FDisplayClusterRenderFrameSettings* RenderFrameSettings = GetRenderFrameSettings())
 		{
-			ViewInitOptions.OverlayColor = FLinearColor::Black;
+			if (RenderFrameSettings->bPreviewEnablePostProcess == false)
+			{
+				ViewInitOptions.OverlayColor = FLinearColor::Black;
+			}
 		}
 
 		ViewInitOptions.bSceneCaptureUsesRayTracing = false;
@@ -253,7 +256,7 @@ FMatrix FDisplayClusterViewport::ImplPreview_GetStereoProjectionMatrix(const uin
 
 	FMatrix PrjMatrix = FMatrix::Identity;
 
-	if (GetOwner().IsSceneOpened())
+	if (IsSceneOpened())
 	{
 		if (GetProjectionMatrix(InContextNum, PrjMatrix) == false)
 		{
diff --git a/Engine/Plugins/Runtime/nDisplay/Source/DisplayCluster/Private/Render/Viewport/DisplayClusterViewport_ViewPoint.cpp b/Engine/Plugins/Runtime/nDisplay/Source/DisplayCluster/Private/Render/Viewport/DisplayClusterViewport_ViewPoint.cpp
--- a/Engine/Plugins/Runtime/nDisplay/Source/DisplayCluster/Private/Render/Viewport/DisplayClusterViewport_ViewPoint.cpp
+++ b/Engine/Plugins/Runtime/nDisplay/Source/DisplayCluster/Private/Render/Viewport/DisplayClusterViewport_ViewPoint.cpp
@@ -29,23 +29,23 @@ using namespace UE::DisplayClusterViewport;
 ///////////////////////////////////////////////////////////////////////////////////////
 UDisplayClusterCameraComponent* FDisplayClusterViewport::GetViewPointCamera() const
 {
-	if (!ProjectionPolicy.IsValid())
+	ADisplayClusterRootActor* RootActor = GetRootActor();
+	if (!RootActor)
 	{
-		// ignore viewports with uninitialized prj policy
+		UE_LOG(LogDisplayClusterViewport, Warning, TEXT("Viewport '%s' has no root actor found in game manager"), *GetId());
+
 		return nullptr;
 	}
 
-	if (!GetOwner().IsSceneOpened())
+	if (!ProjectionPolicy.IsValid())
 	{
+		// ignore viewports with uninitialized prj policy
 		return nullptr;
 	}
 
-	// Get root actor from viewport
-	ADisplayClusterRootActor* const RootActor = GetOwner().GetRootActor();
-	if (!RootActor)
+	FDisplayClusterViewportManager* ViewportManager = GetViewportManagerImpl();
+	if (!ViewportManager || !ViewportManager->IsSceneOpened())
 	{
-		UE_LOG(LogDisplayClusterViewport, Warning, TEXT("Viewport '%s' has no root actor found in game manager"), *GetId());
-
 		return nullptr;
 	}
 
@@ -57,16 +57,16 @@ UDisplayClusterCameraComponent* FDisplayClusterViewport::GetViewPointCamera() co
 	}
 
 	// Get camera component assigned to the viewport (or default camera if nothing assigned)
-	UDisplayClusterCameraComponent* const ViewCamera = (CameraId.IsEmpty() ?
+	if (UDisplayClusterCameraComponent* const ViewCamera = (CameraId.IsEmpty() ?
 		RootActor->GetDefaultCamera() :
-		RootActor->GetComponentByName<UDisplayClusterCameraComponent>(CameraId));
-
-	if (!ViewCamera)
+		RootActor->GetComponentByName<UDisplayClusterCameraComponent>(CameraId)))
 	{
-		UE_LOG(LogDisplayClusterViewport, Warning, TEXT("No camera found for viewport '%s'"), *GetId());
+		return ViewCamera;
 	}
 
-	return ViewCamera;
+	UE_LOG(LogDisplayClusterViewport, Warning, TEXT("No camera found for viewport '%s'"), *GetId());
+
+	return nullptr;
 }
 
 bool FDisplayClusterViewport::SetupViewPoint(FMinimalViewInfo& InOutViewInfo)
diff --git a/Engine/Plugins/Runtime/nDisplay/Source/DisplayCluster/Private/Render/Viewport/LightCard/DisplayClusterViewportLightCardManager.cpp b/Engine/Plugins/Runtime/nDisplay/Source/DisplayCluster/Private/Render/Viewport/LightCard/DisplayClusterViewportLightCardManager.cpp
--- a/Engine/Plugins/Runtime/nDisplay/Source/DisplayCluster/Private/Render/Viewport/LightCard/DisplayClusterViewportLightCardManager.cpp
+++ b/Engine/Plugins/Runtime/nDisplay/Source/DisplayCluster/Private/Render/Viewport/LightCard/DisplayClusterViewportLightCardManager.cpp
@@ -26,7 +26,7 @@ static TAutoConsoleVariable<int32> CVarUVLightCardTextureSize(
 // FDisplayClusterViewportLightCardManager
 ///////////////////////////////////////////////////////////////////////////////////////////////
 FDisplayClusterViewportLightCardManager::FDisplayClusterViewportLightCardManager(FDisplayClusterViewportManager& InViewportManager)
-	: ViewportManager(InViewportManager)
+	: ViewportManager(InViewportManager.AsShared())
 {
 	LightCardManagerProxy = MakeShared<FDisplayClusterViewportLightCardManagerProxy, ESPMode::ThreadSafe>();
 }
@@ -92,7 +92,7 @@ void FDisplayClusterViewportLightCardManager::UpdateUVLightCardData()
 	/** The list of UV light card actors that are referenced by the root actor */
 	TArray<ADisplayClusterLightCardActor*> UVLightCardActors;
 
-	if (ADisplayClusterRootActor* RootActorPtr = ViewportManager.GetRootActor())
+	if (ADisplayClusterRootActor* RootActorPtr = ViewportManager->GetRootActor())
 	{
 		TSet<ADisplayClusterLightCardActor*> LightCards;
 		UDisplayClusterBlueprintLib::FindLightCardsForRootActor(RootActorPtr, LightCards);
@@ -165,7 +165,7 @@ void FDisplayClusterViewportLightCardManager::UpdateUVLightCardResource()
 void FDisplayClusterViewportLightCardManager::RenderUVLightCard()
 {
 	// Render UV LightCard:
-	UWorld* World = ViewportManager.GetCurrentWorld();
+	UWorld* World = ViewportManager->GetCurrentWorld();
 	if (IsUVLightCardEnabled() && World && LightCardManagerProxy.IsValid())
 	{
 		UpdateUVLightCardResource();
@@ -174,7 +174,7 @@ void FDisplayClusterViewportLightCardManager::RenderUVLightCard()
 		{
 			FDisplayClusterShaderParameters_UVLightCards UVLightCardParameters;
 			UVLightCardParameters.ProjectionPlaneSize = ADisplayClusterLightCardActor::UVPlaneDefaultSize;
-			UVLightCardParameters.bRenderFinalColor = ViewportManager.ShouldRenderFinalColor();
+			UVLightCardParameters.bRenderFinalColor = ViewportManager->ShouldRenderFinalColor();
 
 			// Store any components that were invisible but forced to be visible so they can be set back to invisible after the render
 			TArray<UPrimitiveComponent*> ComponentsToUnload;
diff --git a/Engine/Plugins/Runtime/nDisplay/Source/DisplayCluster/Private/Render/Viewport/LightCard/DisplayClusterViewportLightCardManager.h b/Engine/Plugins/Runtime/nDisplay/Source/DisplayCluster/Private/Render/Viewport/LightCard/DisplayClusterViewportLightCardManager.h
--- a/Engine/Plugins/Runtime/nDisplay/Source/DisplayCluster/Private/Render/Viewport/LightCard/DisplayClusterViewportLightCardManager.h
+++ b/Engine/Plugins/Runtime/nDisplay/Source/DisplayCluster/Private/Render/Viewport/LightCard/DisplayClusterViewportLightCardManager.h
@@ -75,9 +75,8 @@ private:
 
 private:
 	/** A reference to the owning viewport manager */
-	FDisplayClusterViewportManager& ViewportManager;
+	const TSharedRef<FDisplayClusterViewportManager, ESPMode::ThreadSafe> ViewportManager;
 
-private:
 	/** A list of primitive components that have been added to the preview scene for rendering in the current frame */
 	TArray<UPrimitiveComponent*> UVLightCardPrimitiveComponents;
 
diff --git a/Engine/Plugins/Runtime/nDisplay/Source/DisplayCluster/Private/Render/Viewport/Postprocess/DisplayClusterViewportPostProcessManager.cpp b/Engine/Plugins/Runtime/nDisplay/Source/DisplayCluster/Private/Render/Viewport/Postprocess/DisplayClusterViewportPostProcessManager.cpp
--- a/Engine/Plugins/Runtime/nDisplay/Source/DisplayCluster/Private/Render/Viewport/Postprocess/DisplayClusterViewportPostProcessManager.cpp
+++ b/Engine/Plugins/Runtime/nDisplay/Source/DisplayCluster/Private/Render/Viewport/Postprocess/DisplayClusterViewportPostProcessManager.cpp
@@ -62,7 +62,7 @@ static TAutoConsoleVariable<int32> CVarPostprocessFrameAfterWarpBlend(
 // FDisplayClusterViewportPostProcessManager
 //----------------------------------------------------------------------------------------------------------
 FDisplayClusterViewportPostProcessManager::FDisplayClusterViewportPostProcessManager(FDisplayClusterViewportManager& InViewportManager)
-	: ViewportManager(InViewportManager)
+	: ViewportManager(InViewportManager.AsShared())
 {
 	OutputRemap = MakeShared<FDisplayClusterViewportPostProcessOutputRemap, ESPMode::ThreadSafe>();
 }
@@ -72,6 +72,11 @@ FDisplayClusterViewportPostProcessManager::~FDisplayClusterViewportPostProcessMa
 	Release();
 }
 
+FDisplayClusterViewportManager* FDisplayClusterViewportPostProcessManager::GetViewportManager() const
+{
+	return &(ViewportManager.Get());
+}
+
 void FDisplayClusterViewportPostProcessManager::Release()
 {
 	if (IsInGameThread())
@@ -91,13 +96,13 @@ bool FDisplayClusterViewportPostProcessManager::HandleStartScene()
 {
 	check(IsInGameThread());;
 
-	if (ViewportManager.IsSceneOpened())
+	if (ViewportManager->IsSceneOpened())
 	{
 		for (const TSharedPtr<IDisplayClusterPostProcess, ESPMode::ThreadSafe>& It : Postprocess)
 		{
 			if (It.IsValid())
 			{
-				It->HandleStartScene(&ViewportManager);
+				It->HandleStartScene(GetViewportManager());
 			}
 		}
 	}
@@ -113,7 +118,7 @@ void FDisplayClusterViewportPostProcessManager::HandleEndScene()
 	{
 		if (It.IsValid())
 		{
-			It->HandleEndScene(&ViewportManager);
+			It->HandleEndScene(GetViewportManager());
 		}
 	}
 }
@@ -157,9 +162,9 @@ bool FDisplayClusterViewportPostProcessManager::CreatePostprocess(const FString&
 		{
 			UE_LOG(LogDisplayClusterRender, Log, TEXT("PostProcess '%s', type '%s' : Created"), *InConfigurationPostprocess->Type, *InPostprocessId);
 
-			if (ViewportManager.IsSceneOpened())
+			if (ViewportManager->IsSceneOpened())
 			{
-				PostProcessInstance->HandleStartScene(&ViewportManager);
+				PostProcessInstance->HandleStartScene(GetViewportManager());
 			}
 
 			for (int32 Order = 0; Order < Postprocess.Num(); Order++)
@@ -338,7 +343,7 @@ void FDisplayClusterViewportPostProcessManager::HandleSetupNewFrame()
 	{
 		for (const TSharedPtr<IDisplayClusterPostProcess, ESPMode::ThreadSafe >& It : Postprocess)
 		{
-			It->HandleSetupNewFrame(&ViewportManager);
+			It->HandleSetupNewFrame(GetViewportManager());
 		}
 	}
 }
@@ -349,7 +354,7 @@ void FDisplayClusterViewportPostProcessManager::HandleBeginNewFrame(FDisplayClus
 	{
 		for (const TSharedPtr<IDisplayClusterPostProcess, ESPMode::ThreadSafe >& It : Postprocess)
 		{
-			It->HandleBeginNewFrame(&ViewportManager, InOutRenderFrame);
+			It->HandleBeginNewFrame(GetViewportManager(), InOutRenderFrame);
 		}
 	}
 }
diff --git a/Engine/Plugins/Runtime/nDisplay/Source/DisplayCluster/Private/Render/Viewport/Postprocess/DisplayClusterViewportPostProcessManager.h b/Engine/Plugins/Runtime/nDisplay/Source/DisplayCluster/Private/Render/Viewport/Postprocess/DisplayClusterViewportPostProcessManager.h
--- a/Engine/Plugins/Runtime/nDisplay/Source/DisplayCluster/Private/Render/Viewport/Postprocess/DisplayClusterViewportPostProcessManager.h
+++ b/Engine/Plugins/Runtime/nDisplay/Source/DisplayCluster/Private/Render/Viewport/Postprocess/DisplayClusterViewportPostProcessManager.h
@@ -71,6 +71,7 @@ public:
 
 private:
 	TSharedPtr<IDisplayClusterPostProcess, ESPMode::ThreadSafe> ImplFindPostProcess(const FString& InPostprocessId) const;
+	FDisplayClusterViewportManager* GetViewportManager() const;
 
 protected:
 	void ImplPerformPostProcessViewBeforeWarpBlend_RenderThread(FRHICommandListImmediate& RHICmdList, const FDisplayClusterViewportManagerProxy* InViewportManagerProxy) const;
@@ -83,9 +84,10 @@ protected:
 	TArray<TSharedPtr<IDisplayClusterPostProcess, ESPMode::ThreadSafe>> Postprocess;
 
 private:
+	const TSharedRef<FDisplayClusterViewportManager, ESPMode::ThreadSafe> ViewportManager;
+
 	// Render thread instances
 	TArray<TSharedPtr<IDisplayClusterPostProcess, ESPMode::ThreadSafe>> PostprocessProxy;
 
-	FDisplayClusterViewportManager& ViewportManager;
 	TSharedPtr<FDisplayClusterViewportPostProcessOutputRemap, ESPMode::ThreadSafe> OutputRemap;
 };
diff --git a/Engine/Plugins/Runtime/nDisplay/Source/DisplayCluster/Private/Render/Viewport/RenderFrame/DisplayClusterRenderFrameManager.cpp b/Engine/Plugins/Runtime/nDisplay/Source/DisplayCluster/Private/Render/Viewport/RenderFrame/DisplayClusterRenderFrameManager.cpp
--- a/Engine/Plugins/Runtime/nDisplay/Source/DisplayCluster/Private/Render/Viewport/RenderFrame/DisplayClusterRenderFrameManager.cpp
+++ b/Engine/Plugins/Runtime/nDisplay/Source/DisplayCluster/Private/Render/Viewport/RenderFrame/DisplayClusterRenderFrameManager.cpp
@@ -13,7 +13,7 @@
 ///////////////////////////////////////////////////////////////
 // FDisplayClusterRenderTargetFrame
 ///////////////////////////////////////////////////////////////
-bool FDisplayClusterRenderFrameManager::BuildRenderFrame(FViewport* InViewport, const FDisplayClusterRenderFrameSettings& InRenderFrameSettings, const TArray<FDisplayClusterViewport*>& InViewports, FDisplayClusterRenderFrame& OutRenderFrame)
+bool FDisplayClusterRenderFrameManager::BuildRenderFrame(FViewport* InViewport, const FDisplayClusterRenderFrameSettings& InRenderFrameSettings, const TArray<TSharedPtr<FDisplayClusterViewport, ESPMode::ThreadSafe>>& InViewports, FDisplayClusterRenderFrame& OutRenderFrame)
 {
 
 	switch (InRenderFrameSettings.RenderMode)
@@ -21,6 +21,7 @@ bool FDisplayClusterRenderFrameManager::BuildRenderFrame(FViewport* InViewport,
 	case EDisplayClusterRenderFrameMode::PreviewInScene:
 		// Dont use render frame for preview
 		break;
+
 	default:
 		if (!FindFrameTargetRect(InViewport, InViewports, InRenderFrameSettings, OutRenderFrame.FrameRect))
 		{
@@ -34,19 +35,22 @@ bool FDisplayClusterRenderFrameManager::BuildRenderFrame(FViewport* InViewport,
 
 	// Sort viewports, childs after parents
 	//@todo save this order inside logic
-	TArray<FDisplayClusterViewport*> SortedViewports;
-	TArray<FDisplayClusterViewport*> ChildsViewports;
+	TArray<TSharedPtr<FDisplayClusterViewport, ESPMode::ThreadSafe>> SortedViewports;
+	TArray<TSharedPtr<FDisplayClusterViewport, ESPMode::ThreadSafe>> ChildsViewports;
 
 	// First add root viewports
-	for (FDisplayClusterViewport* Viewport : InViewports)
+	for (const TSharedPtr<FDisplayClusterViewport, ESPMode::ThreadSafe>& Viewport : InViewports)
 	{
-		if (Viewport->RenderSettings.GetParentViewportId().IsEmpty())
-		{
-			SortedViewports.Add(Viewport);
-		}
-		else
+		if (Viewport.IsValid())
 		{
-			ChildsViewports.Add(Viewport);
+			if (Viewport->RenderSettings.GetParentViewportId().IsEmpty())
+			{
+				SortedViewports.Add(Viewport);
+			}
+			else
+			{
+				ChildsViewports.Add(Viewport);
+			}
 		}
 	}
 
@@ -59,16 +63,16 @@ bool FDisplayClusterRenderFrameManager::BuildRenderFrame(FViewport* InViewport,
 	// have media capture assigned, so there won't be any problems.
 	{
 		// Find all viewports being captured
-		TArray<FDisplayClusterViewport*> ViewportsBeingCaptured = SortedViewports.FilterByPredicate([](const FDisplayClusterViewport* Viewport)
+		TArray<TSharedPtr<FDisplayClusterViewport, ESPMode::ThreadSafe>> ViewportsBeingCaptured = SortedViewports.FilterByPredicate([](const TSharedPtr<FDisplayClusterViewport, ESPMode::ThreadSafe>& Viewport)
 		{
-			return Viewport->GetRenderSettings().bIsBeingCaptured;
+			return Viewport.IsValid() && Viewport->GetRenderSettings().bIsBeingCaptured;
 		});
 
 		// Put them all in the beginning of the list
 		const int32 CapturingViewportsAmount = ViewportsBeingCaptured.Num();
 		if (CapturingViewportsAmount > 0 && CapturingViewportsAmount != SortedViewports.Num())
 		{
-			for (FDisplayClusterViewport* ViewportBeingCaptured : ViewportsBeingCaptured)
+			for (TSharedPtr<FDisplayClusterViewport, ESPMode::ThreadSafe>& ViewportBeingCaptured : ViewportsBeingCaptured)
 			{
 				SortedViewports.Remove(ViewportBeingCaptured);
 				SortedViewports.Insert(ViewportBeingCaptured, 0);
@@ -112,30 +116,30 @@ bool FDisplayClusterRenderFrameManager::BuildRenderFrame(FViewport* InViewport,
 	return bResult;
 }
 
-bool FDisplayClusterRenderFrameManager::BuildSimpleFrame(FViewport* InViewport, const FDisplayClusterRenderFrameSettings& InRenderFrameSettings, const TArray<FDisplayClusterViewport*>& InViewports, FDisplayClusterRenderFrame& OutRenderFrame)
+bool FDisplayClusterRenderFrameManager::BuildSimpleFrame(FViewport* InViewport, const FDisplayClusterRenderFrameSettings& InRenderFrameSettings, const TArray<TSharedPtr<FDisplayClusterViewport, ESPMode::ThreadSafe>>& InViewports, FDisplayClusterRenderFrame& OutRenderFrame)
 {
-	for (FDisplayClusterViewport* ViewportIt : InViewports)
+	for (const TSharedPtr<FDisplayClusterViewport, ESPMode::ThreadSafe>& ViewportIt : InViewports)
 	{
-		if (ViewportIt)
+		if (ViewportIt.IsValid())
 		{
 			for (const FDisplayClusterViewport_Context& ContextIt : ViewportIt->GetContexts())
 			{
-				FDisplayClusterRenderFrame::FFrameView FrameView;
+				FDisplayClusterRenderFrameTargetView FrameView;
 				{
 					FrameView.ContextNum = ContextIt.ContextNum;
-					FrameView.Viewport = ViewportIt;
+					FrameView.Viewport = ViewportIt.ToSharedRef();
 					FrameView.bDisableRender = ContextIt.bDisableRender || ViewportIt->RenderSettings.bSkipSceneRenderingButLeaveResourcesAvailable;
 					FrameView.bFreezeRendering = ViewportIt->RenderSettings.bFreezeRendering;
 				}
 
-				FDisplayClusterRenderFrame::FFrameViewFamily FrameViewFamily;
+				FDisplayClusterRenderFrameTargetViewFamily FrameViewFamily;
 				{
 					FrameViewFamily.Views.Add(FrameView);
 					FrameViewFamily.CustomBufferRatio = ContextIt.CustomBufferRatio;
 					FrameViewFamily.ViewExtensions = ViewportIt->GatherActiveExtensions(InViewport);
 				}
 
-				FDisplayClusterRenderFrame::FFrameRenderTarget FrameRenderTarget;
+				FDisplayClusterRenderFrameTarget FrameRenderTarget;
 				{
 					// Simple frame use unique RTT  for each viewport, so disable RTT when viewport rendering disabled
 					FrameRenderTarget.bShouldUseRenderTarget = ViewportIt->RenderTargets.Num() > 0;
@@ -154,7 +158,7 @@ bool FDisplayClusterRenderFrameManager::BuildSimpleFrame(FViewport* InViewport,
 	return true;
 }
 
-bool FDisplayClusterRenderFrameManager::FindFrameTargetRect(FViewport* InViewport, const TArray<FDisplayClusterViewport*>& InOutViewports, const FDisplayClusterRenderFrameSettings& InRenderFrameSettings, FIntRect& OutFrameTargetRect) const
+bool FDisplayClusterRenderFrameManager::FindFrameTargetRect(FViewport* InViewport, const TArray<TSharedPtr<FDisplayClusterViewport, ESPMode::ThreadSafe>>& InOutViewports, const FDisplayClusterRenderFrameSettings& InRenderFrameSettings, FIntRect& OutFrameTargetRect) const
 {
 	// Calculate Backbuffer frame
 	bool bIsUsed = false;
@@ -167,9 +171,9 @@ bool FDisplayClusterRenderFrameManager::FindFrameTargetRect(FViewport* InViewpor
 	}
 
 	// Optimize frame target RTT
-	for (const FDisplayClusterViewport* ViewportIt : InOutViewports)
+	for (const TSharedPtr<FDisplayClusterViewport, ESPMode::ThreadSafe>& ViewportIt : InOutViewports)
 	{
-		if (ViewportIt && ViewportIt->RenderSettings.bEnable && ViewportIt->RenderSettings.bVisible)
+		if (ViewportIt.IsValid() && ViewportIt->RenderSettings.bEnable && ViewportIt->RenderSettings.bVisible)
 		{
 			for (const auto& ContextIt : ViewportIt->GetContexts())
 			{
diff --git a/Engine/Plugins/Runtime/nDisplay/Source/DisplayCluster/Private/Render/Viewport/RenderFrame/DisplayClusterRenderFrameManager.h b/Engine/Plugins/Runtime/nDisplay/Source/DisplayCluster/Private/Render/Viewport/RenderFrame/DisplayClusterRenderFrameManager.h
--- a/Engine/Plugins/Runtime/nDisplay/Source/DisplayCluster/Private/Render/Viewport/RenderFrame/DisplayClusterRenderFrameManager.h
+++ b/Engine/Plugins/Runtime/nDisplay/Source/DisplayCluster/Private/Render/Viewport/RenderFrame/DisplayClusterRenderFrameManager.h
@@ -6,18 +6,18 @@
 
 class FViewport;
 class FDisplayClusterViewport;
+class FViewport;
 class FDisplayClusterRenderFrame;
 struct FDisplayClusterRenderFrameSettings;
-class FViewport;
 
 class FDisplayClusterRenderFrameManager
 {
 public:
-	bool BuildRenderFrame(FViewport* InViewport, const FDisplayClusterRenderFrameSettings& InRenderFrameSettings, const TArray<FDisplayClusterViewport*>& InViewports, FDisplayClusterRenderFrame& OutRenderFrame);
+	bool BuildRenderFrame(FViewport* InViewport, const FDisplayClusterRenderFrameSettings& InRenderFrameSettings, const TArray<TSharedPtr<FDisplayClusterViewport, ESPMode::ThreadSafe>>& InViewports, FDisplayClusterRenderFrame& OutRenderFrame);
 
 private:
-	bool FindFrameTargetRect(FViewport* InViewport, const TArray<FDisplayClusterViewport*>& InViewports, const FDisplayClusterRenderFrameSettings& InRenderFrameSettings, FIntRect& OutFrameTargetRect) const;
-	bool BuildSimpleFrame(FViewport* InViewport, const FDisplayClusterRenderFrameSettings& InRenderFrameSettings, const TArray<FDisplayClusterViewport*>& InViewports, FDisplayClusterRenderFrame& OutRenderFrame);
+	bool FindFrameTargetRect(FViewport* InViewport, const TArray<TSharedPtr<FDisplayClusterViewport, ESPMode::ThreadSafe>>& InViewports, const FDisplayClusterRenderFrameSettings& InRenderFrameSettings, FIntRect& OutFrameTargetRect) const;
+	bool BuildSimpleFrame(FViewport* InViewport, const FDisplayClusterRenderFrameSettings& InRenderFrameSettings, const TArray<TSharedPtr<FDisplayClusterViewport, ESPMode::ThreadSafe>>& InViewports, FDisplayClusterRenderFrame& OutRenderFrame);
 };
 
 
diff --git a/Engine/Plugins/Runtime/nDisplay/Source/DisplayCluster/Private/Render/Viewport/RenderFrame/DisplayClusterRenderFrameSettings.h b/Engine/Plugins/Runtime/nDisplay/Source/DisplayCluster/Private/Render/Viewport/RenderFrame/DisplayClusterRenderFrameSettings.h
--- a/Engine/Plugins/Runtime/nDisplay/Source/DisplayCluster/Private/Render/Viewport/RenderFrame/DisplayClusterRenderFrameSettings.h
+++ b/Engine/Plugins/Runtime/nDisplay/Source/DisplayCluster/Private/Render/Viewport/RenderFrame/DisplayClusterRenderFrameSettings.h
@@ -11,7 +11,7 @@
 struct FDisplayClusterRenderFrameSettings
 {
 	// customize mono\stereo render modes
-	EDisplayClusterRenderFrameMode RenderMode = EDisplayClusterRenderFrameMode::Mono;
+	EDisplayClusterRenderFrameMode RenderMode = EDisplayClusterRenderFrameMode::Unknown;
 
 	// nDisplay has its own implementation of cross-GPU transfer.
 	struct FCrossGPUTransfer
diff --git a/Engine/Plugins/Runtime/nDisplay/Source/DisplayCluster/Private/Render/Viewport/RenderTarget/DisplayClusterRenderTargetManager.cpp b/Engine/Plugins/Runtime/nDisplay/Source/DisplayCluster/Private/Render/Viewport/RenderTarget/DisplayClusterRenderTargetManager.cpp
--- a/Engine/Plugins/Runtime/nDisplay/Source/DisplayCluster/Private/Render/Viewport/RenderTarget/DisplayClusterRenderTargetManager.cpp
+++ b/Engine/Plugins/Runtime/nDisplay/Source/DisplayCluster/Private/Render/Viewport/RenderTarget/DisplayClusterRenderTargetManager.cpp
@@ -75,9 +75,9 @@ namespace DisplayClusterRenderTargetManager
 		return OutPixelFormat;
 	}
 
-	static void ImplViewportTextureResourceLogging(const FDisplayClusterViewport* InViewport, const uint32 ContextNum, const FString& ResourceId, const FDisplayClusterViewportTextureResource* InTextureResource)
+	static void ImplViewportTextureResourceLogging(const TSharedPtr<FDisplayClusterViewport, ESPMode::ThreadSafe>& InViewport, const uint32 ContextNum, const FString& ResourceId, const FDisplayClusterViewportTextureResource* InTextureResource)
 	{
-		if (InTextureResource != nullptr && InTextureResource->GetViewportResourceState(EDisplayClusterViewportResourceState::Initialized) == false)
+		if (InViewport.IsValid() && InTextureResource != nullptr && InTextureResource->GetViewportResourceState(EDisplayClusterViewportResourceState::Initialized) == false)
 		{
 			// Log: New resource created
 			UE_LOG(LogDisplayClusterViewport, Verbose, TEXT("Created new %s resource (%dx%d) for viewport '%s':%d"), *ResourceId, InTextureResource->GetSizeX(), InTextureResource->GetSizeY(), *InViewport->GetId(), ContextNum);
@@ -105,14 +105,14 @@ void FDisplayClusterRenderTargetManager::Release()
 	ResourcesPool->Release();
 }
 
-bool FDisplayClusterRenderTargetManager::AllocateRenderFrameResources(FViewport* InViewport, const FDisplayClusterRenderFrameSettings& InRenderFrameSettings, const TArray<FDisplayClusterViewport*>& InViewports, FDisplayClusterRenderFrame& InOutRenderFrame)
+bool FDisplayClusterRenderTargetManager::AllocateRenderFrameResources(FViewport* InViewport, const FDisplayClusterRenderFrameSettings& InRenderFrameSettings, const TArray<TSharedPtr<FDisplayClusterViewport, ESPMode::ThreadSafe>>& InViewports, FDisplayClusterRenderFrame& InOutRenderFrame)
 {
 	bool bResult = true;
 
 	if(ResourcesPool->BeginReallocateResources(InRenderFrameSettings, InViewport))
 	{
 		// ReAllocate Render targets for all viewports
-		for (FDisplayClusterRenderFrame::FFrameRenderTarget& FrameRenderTargetIt : InOutRenderFrame.RenderTargets)
+		for (FDisplayClusterRenderFrameTarget& FrameRenderTargetIt : InOutRenderFrame.RenderTargets)
 		{
 			if (FrameRenderTargetIt.bShouldUseRenderTarget)
 			{
@@ -124,26 +124,22 @@ bool FDisplayClusterRenderTargetManager::AllocateRenderFrameResources(FViewport*
 					FrameRenderTargetIt.RenderTargetPtr = NewResource;
 
 					// Assign for all views in render target families
-					for (FDisplayClusterRenderFrame::FFrameViewFamily& ViewFamily : FrameRenderTargetIt.ViewFamilies)
+					for (FDisplayClusterRenderFrameTargetViewFamily& ViewFamily : FrameRenderTargetIt.ViewFamilies)
 					{
-						for (FDisplayClusterRenderFrame::FFrameView& ViewIt : ViewFamily.Views)
+						for (FDisplayClusterRenderFrameTargetView& ViewIt : ViewFamily.Views)
 						{
-							if (ViewIt.Viewport != nullptr)
+							if (FDisplayClusterViewport* ViewportPtr = static_cast<FDisplayClusterViewport*>(ViewIt.Viewport.Get()))
 							{
-								FDisplayClusterViewport* ViewportPtr = static_cast<FDisplayClusterViewport*>(ViewIt.Viewport);
-								if (ViewportPtr)
-								{
-									// Array already resized in function FDisplayClusterViewport::UpdateFrameContexts() with RenderTargets.AddZeroed(ViewportContextAmount);
-									check(ViewIt.ContextNum < (uint32)ViewportPtr->RenderTargets.Num());
-									check(ViewportPtr->Contexts[ViewIt.ContextNum].bDisableRender == false);
+								// Array already resized in function FDisplayClusterViewport::UpdateFrameContexts() with RenderTargets.AddZeroed(ViewportContextAmount);
+								check(ViewportPtr->RenderTargets.IsValidIndex(ViewIt.ContextNum));
+								check(!ViewportPtr->Contexts[ViewIt.ContextNum].bDisableRender);
 
-									ViewportPtr->RenderTargets[ViewIt.ContextNum] = NewResource;
+								ViewportPtr->RenderTargets[ViewIt.ContextNum] = NewResource;
 
-									if (NewResource->GetViewportResourceState(EDisplayClusterViewportResourceState::Initialized) == false)
-									{
-										// Log: New resource created
-										UE_LOG(LogDisplayClusterViewport, Verbose, TEXT("Created new ViewportRenderTarget resource %08X (%dx%d) for viewport '%s'"), NewResource, NewResource->GetSizeX(), NewResource->GetSizeY(), *ViewportPtr->GetId());
-									}
+								if (NewResource->GetViewportResourceState(EDisplayClusterViewportResourceState::Initialized) == false)
+								{
+									// Log: New resource created
+									UE_LOG(LogDisplayClusterViewport, Verbose, TEXT("Created new ViewportRenderTarget resource %08X (%dx%d) for viewport '%s'"), NewResource, NewResource->GetSizeX(), NewResource->GetSizeY(), *ViewportPtr->GetId());
 								}
 							}
 						}
@@ -153,9 +149,9 @@ bool FDisplayClusterRenderTargetManager::AllocateRenderFrameResources(FViewport*
 		}
 
 		// Allocate viewport internal resources
-		for (FDisplayClusterViewport* ViewportIt : InViewports)
+		for (const TSharedPtr<FDisplayClusterViewport, ESPMode::ThreadSafe>& ViewportIt : InViewports)
 		{
-			if (ViewportIt && ViewportIt->RenderSettings.bFreezeRendering == false)
+			if (ViewportIt.IsValid() && ViewportIt->RenderSettings.bFreezeRendering == false)
 			{
 				EPixelFormat CustomFormat = ImplGetCustomFormat(ViewportIt->RenderSettings.CaptureMode);
 
@@ -278,27 +274,24 @@ bool FDisplayClusterRenderTargetManager::AllocateFrameTargets(const FDisplayClus
 	}
 
 	// Assign frame resources for all visible viewports
-	for (FDisplayClusterRenderFrame::FFrameRenderTarget& RenderTargetIt : InOutRenderFrame.RenderTargets)
+	for (FDisplayClusterRenderFrameTarget& RenderTargetIt : InOutRenderFrame.RenderTargets)
 	{
-		for (FDisplayClusterRenderFrame::FFrameViewFamily& ViewFamilieIt : RenderTargetIt.ViewFamilies)
+		for (FDisplayClusterRenderFrameTargetViewFamily& ViewFamilieIt : RenderTargetIt.ViewFamilies)
 		{
-			for (FDisplayClusterRenderFrame::FFrameView& ViewIt : ViewFamilieIt.Views)
+			for (FDisplayClusterRenderFrameTargetView& ViewIt : ViewFamilieIt.Views)
 			{
-				if (ViewIt.Viewport != nullptr)
+				FDisplayClusterViewport* ViewportPtr = static_cast<FDisplayClusterViewport*>(ViewIt.Viewport.Get());
+				if (ViewportPtr && ViewportPtr->RenderSettings.bVisible)
 				{
-					FDisplayClusterViewport* ViewportPtr = static_cast<FDisplayClusterViewport*>(ViewIt.Viewport);
-					if (ViewportPtr && ViewportPtr->RenderSettings.bVisible)
-					{
-						ViewportPtr->OutputFrameTargetableResources = NewFrameTargetResources;
-						ViewportPtr->AdditionalFrameTargetableResources = NewAdditionalFrameTargetableResources;
+					ViewportPtr->OutputFrameTargetableResources = NewFrameTargetResources;
+					ViewportPtr->AdditionalFrameTargetableResources = NewAdditionalFrameTargetableResources;
 
-						// Adjust viewports frame rects. This offset saved in 'BackbufferFrameOffset'
-						if (ViewIt.ContextNum < (uint32)ViewportPtr->Contexts.Num())
-						{
-							FDisplayClusterViewport_Context& Context = ViewportPtr->Contexts[ViewIt.ContextNum];
-							Context.FrameTargetRect.Min -= InOutRenderFrame.FrameRect.Min;
-							Context.FrameTargetRect.Max -= InOutRenderFrame.FrameRect.Min;
-						}
+					// Adjust viewports frame rects. This offset saved in 'BackbufferFrameOffset'
+					if (ViewportPtr->Contexts.IsValidIndex(ViewIt.ContextNum))
+					{
+						FDisplayClusterViewport_Context& Context = ViewportPtr->Contexts[ViewIt.ContextNum];
+						Context.FrameTargetRect.Min -= InOutRenderFrame.FrameRect.Min;
+						Context.FrameTargetRect.Max -= InOutRenderFrame.FrameRect.Min;
 					}
 				}
 			}
diff --git a/Engine/Plugins/Runtime/nDisplay/Source/DisplayCluster/Private/Render/Viewport/RenderTarget/DisplayClusterRenderTargetManager.h b/Engine/Plugins/Runtime/nDisplay/Source/DisplayCluster/Private/Render/Viewport/RenderTarget/DisplayClusterRenderTargetManager.h
--- a/Engine/Plugins/Runtime/nDisplay/Source/DisplayCluster/Private/Render/Viewport/RenderTarget/DisplayClusterRenderTargetManager.h
+++ b/Engine/Plugins/Runtime/nDisplay/Source/DisplayCluster/Private/Render/Viewport/RenderTarget/DisplayClusterRenderTargetManager.h
@@ -6,10 +6,9 @@
 
 class FViewport;
 class FDisplayClusterViewport;
-class FDisplayClusterRenderFrame;
 class FDisplayClusterRenderTargetResourcesPool;
 class FDisplayClusterViewportManagerProxy;
-
+class FDisplayClusterRenderFrame;
 struct FDisplayClusterRenderFrameSettings;
 
 class FDisplayClusterRenderTargetManager
@@ -21,7 +20,7 @@ public:
 	void Release();
 
 public:
-	bool AllocateRenderFrameResources(FViewport* InViewport, const FDisplayClusterRenderFrameSettings& InRenderFrameSettings, const TArray<FDisplayClusterViewport*>& InViewports, FDisplayClusterRenderFrame& InOutRenderFrame);
+	bool AllocateRenderFrameResources(FViewport* InViewport, const FDisplayClusterRenderFrameSettings& InRenderFrameSettings, const TArray<TSharedPtr<FDisplayClusterViewport, ESPMode::ThreadSafe>>& InViewports, FDisplayClusterRenderFrame& InOutRenderFrame);
 
 protected:
 	bool AllocateFrameTargets(const FDisplayClusterRenderFrameSettings& InRenderFrameSettings, const FIntPoint& InViewportSize, FDisplayClusterRenderFrame& InOutRenderFrame);
diff --git a/Engine/Plugins/Runtime/nDisplay/Source/DisplayCluster/Private/Render/Viewport/RenderTarget/DisplayClusterRenderTargetResourcesPool.cpp b/Engine/Plugins/Runtime/nDisplay/Source/DisplayCluster/Private/Render/Viewport/RenderTarget/DisplayClusterRenderTargetResourcesPool.cpp
--- a/Engine/Plugins/Runtime/nDisplay/Source/DisplayCluster/Private/Render/Viewport/RenderTarget/DisplayClusterRenderTargetResourcesPool.cpp
+++ b/Engine/Plugins/Runtime/nDisplay/Source/DisplayCluster/Private/Render/Viewport/RenderTarget/DisplayClusterRenderTargetResourcesPool.cpp
@@ -79,7 +79,7 @@ using namespace DisplayClusterRenderTargetResourcesPool;
 /// FDisplayClusterRenderTargetResourcesPool
 ////////////////////////////////////////////////////////////////////////////////////////////////////////////
 FDisplayClusterRenderTargetResourcesPool::FDisplayClusterRenderTargetResourcesPool(FDisplayClusterViewportManagerProxy* InViewportManagerProxy)
-	: ViewportManagerProxy(InViewportManagerProxy)
+	: ViewportManagerProxyWeakPtr(InViewportManagerProxy ? InViewportManagerProxy->AsWeak() : nullptr)
 { }
 
 FDisplayClusterRenderTargetResourcesPool::~FDisplayClusterRenderTargetResourcesPool()
@@ -168,7 +168,7 @@ void FDisplayClusterRenderTargetResourcesPool::ImplFinishReallocateResources(TAr
 		}
 
 		// Send to released resources to rendering thread
-		ImplReleaseViewportResourcesRHI(ViewportManagerProxy, TArrayView<FDisplayClusterViewportResource*>((FDisplayClusterViewportResource**)(UnusedResources.GetData()), UnusedResources.Num()));
+		ImplReleaseViewportResourcesRHI(GetViewportManagerProxy(), TArrayView<FDisplayClusterViewportResource*>((FDisplayClusterViewportResource**)(UnusedResources.GetData()), UnusedResources.Num()));
 	}
 }
 
@@ -176,7 +176,7 @@ template <typename TViewportResourceType>
 void FDisplayClusterRenderTargetResourcesPool::ImplReleaseResources(TArray<TViewportResourceType*>& InOutViewportResources)
 {
 	// Send all resources to rendering thread
-	ImplReleaseViewportResourcesRHI(ViewportManagerProxy, TArray<FDisplayClusterViewportResource*>((FDisplayClusterViewportResource**)(InOutViewportResources.GetData()), InOutViewportResources.Num()));
+	ImplReleaseViewportResourcesRHI(GetViewportManagerProxy(), TArray<FDisplayClusterViewportResource*>((FDisplayClusterViewportResource**)(InOutViewportResources.GetData()), InOutViewportResources.Num()));
 
 	// Reset refs on game thread too
 	InOutViewportResources.Empty();
diff --git a/Engine/Plugins/Runtime/nDisplay/Source/DisplayCluster/Private/Render/Viewport/RenderTarget/DisplayClusterRenderTargetResourcesPool.h b/Engine/Plugins/Runtime/nDisplay/Source/DisplayCluster/Private/Render/Viewport/RenderTarget/DisplayClusterRenderTargetResourcesPool.h
--- a/Engine/Plugins/Runtime/nDisplay/Source/DisplayCluster/Private/Render/Viewport/RenderTarget/DisplayClusterRenderTargetResourcesPool.h
+++ b/Engine/Plugins/Runtime/nDisplay/Source/DisplayCluster/Private/Render/Viewport/RenderTarget/DisplayClusterRenderTargetResourcesPool.h
@@ -23,6 +23,12 @@ public:
 	FDisplayClusterViewportRenderTargetResource* AllocateRenderTargetResource(const FIntPoint& InSize, EPixelFormat CustomPixelFormat);
 	FDisplayClusterViewportTextureResource*      AllocateTextureResource(const FIntPoint& InSize, bool bIsRenderTargetable, EPixelFormat CustomPixelFormat, int32 NumMips = 1);
 
+protected:
+	inline FDisplayClusterViewportManagerProxy* GetViewportManagerProxy() const
+	{
+		return ViewportManagerProxyWeakPtr.IsValid() ? ViewportManagerProxyWeakPtr.Pin().Get() : nullptr;
+	}
+
 private:
 	template <typename TViewportResourceType>
 	void ImplBeginReallocateResources(TArray<TViewportResourceType*>& InOutViewportResources);
@@ -44,5 +50,5 @@ private:
 	TArray<FDisplayClusterViewportRenderTargetResource*> RenderTargetResources;
 	TArray<FDisplayClusterViewportTextureResource*>      TextureResources;
 
-	FDisplayClusterViewportManagerProxy* ViewportManagerProxy = nullptr;
+	TWeakPtr<FDisplayClusterViewportManagerProxy, ESPMode::ThreadSafe> ViewportManagerProxyWeakPtr;
 };
diff --git a/Engine/Plugins/Runtime/nDisplay/Source/DisplayCluster/Public/DisplayClusterRootActor.h b/Engine/Plugins/Runtime/nDisplay/Source/DisplayCluster/Public/DisplayClusterRootActor.h
--- a/Engine/Plugins/Runtime/nDisplay/Source/DisplayCluster/Public/DisplayClusterRootActor.h
+++ b/Engine/Plugins/Runtime/nDisplay/Source/DisplayCluster/Public/DisplayClusterRootActor.h
@@ -100,6 +100,14 @@ public:
 	const FDisplayClusterConfigurationICVFX_StageSettings& GetStageSettings() const;
 	const FDisplayClusterConfigurationRenderFrame& GetRenderFrameSettings() const;
 
+	/** Returns the current rendering mode of this DCRA (not a value from configuration).
+	 * This value can be overridden from DCRenderDevice or other rendering subsystems (e.g. Preview).
+	 */
+	EDisplayClusterRenderFrameMode GetRenderMode() const;
+
+	/** Returns the preview rendering mode of this DCRA. */
+	EDisplayClusterRenderFrameMode GetPreviewRenderMode() const;
+
 protected:
 	//////////////////////////////////////////////////////////////////////////////////////////////
 	// AActor
@@ -188,19 +196,21 @@ public:
 	void SetPreviewEnablePostProcess(const bool bNewPreviewEnablePostProcess);
 
 public:
-	IDisplayClusterViewportManager* GetViewportManager() const
-	{
-		return ViewportManager.IsValid() ? ViewportManager.Get() : nullptr;
-	}
+	/** Get ViewportManager API. */
+	IDisplayClusterViewportManager* GetViewportManager() const;
 	
 	static FName GetCurrentConfigDataMemberName()
 	{
 		return GET_MEMBER_NAME_CHECKED(ADisplayClusterRootActor, CurrentConfigData);
 	}
 
-protected:
-	// Unique viewport manager for this configuration
-	TUniquePtr<IDisplayClusterViewportManager> ViewportManager;
+private:
+	/** Create new viewport manager instance if not exists. */
+	void CreateViewportManagerImpl();
+
+private:
+	// Viewport manager for this configuration
+	TSharedPtr<class FDisplayClusterViewportManager, ESPMode::ThreadSafe> ViewportManager;
 
 //////////////////////////////////////////////////////////////////////////////////////////////
 // Details Panel Property Referencers
diff --git a/Engine/Plugins/Runtime/nDisplay/Source/DisplayCluster/Public/DisplayClusterSceneViewExtensions.h b/Engine/Plugins/Runtime/nDisplay/Source/DisplayCluster/Public/DisplayClusterSceneViewExtensions.h
--- a/Engine/Plugins/Runtime/nDisplay/Source/DisplayCluster/Public/DisplayClusterSceneViewExtensions.h
+++ b/Engine/Plugins/Runtime/nDisplay/Source/DisplayCluster/Public/DisplayClusterSceneViewExtensions.h
@@ -21,22 +21,23 @@ private:
 	}
 
 public:
-
 	// The id of the nDisplay viewport being rendered.
 	const FString ViewportId;
-	const IDisplayClusterViewportManager* ViewportManager;
+
+	// Reference to viewport manager
+	TSharedPtr<IDisplayClusterViewportManager, ESPMode::ThreadSafe> ViewportManager;
 
 	FDisplayClusterSceneViewExtensionContext()
 		: FSceneViewExtensionContext()
 	{ }
 
-	FDisplayClusterSceneViewExtensionContext(FViewport* InViewport, IDisplayClusterViewportManager* InViewportManager, const FString& InViewportId)
+	FDisplayClusterSceneViewExtensionContext(FViewport* InViewport, const TSharedPtr<IDisplayClusterViewportManager, ESPMode::ThreadSafe>& InViewportManager, const FString& InViewportId)
 		: FSceneViewExtensionContext(InViewport)
 		, ViewportId(InViewportId)
 		, ViewportManager(InViewportManager)
 	{ }
 
-	FDisplayClusterSceneViewExtensionContext(FSceneInterface* InScene, IDisplayClusterViewportManager* InViewportManager, const FString& InViewportId)
+	FDisplayClusterSceneViewExtensionContext(FSceneInterface* InScene, const TSharedPtr<IDisplayClusterViewportManager, ESPMode::ThreadSafe>& InViewportManager, const FString& InViewportId)
 		: FSceneViewExtensionContext(InScene)
 		, ViewportId(InViewportId)
 		, ViewportManager(InViewportManager)
diff --git a/Engine/Plugins/Runtime/nDisplay/Source/DisplayCluster/Public/Render/Viewport/IDisplayClusterViewport.h b/Engine/Plugins/Runtime/nDisplay/Source/DisplayCluster/Public/Render/Viewport/IDisplayClusterViewport.h
--- a/Engine/Plugins/Runtime/nDisplay/Source/DisplayCluster/Public/Render/Viewport/IDisplayClusterViewport.h
+++ b/Engine/Plugins/Runtime/nDisplay/Source/DisplayCluster/Public/Render/Viewport/IDisplayClusterViewport.h
@@ -5,17 +5,17 @@
 #include "Render/Viewport/IDisplayClusterViewport_CustomPostProcessSettings.h"
 
 #include "Render/Viewport/Containers/DisplayClusterViewport_Context.h"
-
 #include "Render/Viewport/Containers/DisplayClusterViewport_RenderSettings.h"
 #include "Render/Viewport/Containers/DisplayClusterViewport_RenderSettingsICVFX.h"
 #include "Render/Viewport/Containers/DisplayClusterViewport_PostRenderSettings.h"
-
 #include "Render/Viewport/Containers/DisplayClusterViewport_OverscanSettings.h"
 
-class FSceneViewFamily;
+#include "Render/Viewport/RenderFrame/DisplayClusterRenderFrameEnums.h"
+
+#include "Engine/EngineTypes.h"
 
 /**
- * Rendering viewport (sub-region of the main viewport)
+ * nDisplay: Viewport (interface for GameThread)
  */
 class DISPLAYCLUSTER_API IDisplayClusterViewport
 {
@@ -23,6 +23,10 @@ public:
 	virtual ~IDisplayClusterViewport() = default;
 
 public:
+	/** Get TSharedPtr from self. */
+	virtual TSharedPtr<IDisplayClusterViewport, ESPMode::ThreadSafe> ToSharedPtr() = 0;
+	virtual TSharedPtr<const IDisplayClusterViewport, ESPMode::ThreadSafe> ToSharedPtr() const = 0;
+
 	virtual FString GetId() const = 0;
 	virtual FString GetClusterNodeId() const = 0;
 
@@ -57,9 +61,25 @@ public:
 	virtual const IDisplayClusterViewport_CustomPostProcessSettings& GetViewport_CustomPostProcessSettings() const = 0;
 
 	// Setup scene view for rendering specified Context
-	virtual void SetupSceneView(uint32 ContextNum, class UWorld* World, FSceneViewFamily& InViewFamily, FSceneView& InView) const = 0;
+	virtual void SetupSceneView(uint32 ContextNum, class UWorld* World, class FSceneViewFamily& InViewFamily, FSceneView& InView) const = 0;
+
+	/** Return the viewport manager that owns this viewport */
+	virtual class IDisplayClusterViewportManager* GetViewportManager() const = 0;
+
+	/** Return the DCRA that owns this viewport */
+	virtual class ADisplayClusterRootActor* GetRootActor() const = 0;
+
+	/** Return current render mode. */
+	virtual EDisplayClusterRenderFrameMode GetRenderMode() const = 0;
+
+	/** Return current world. */
+	virtual class UWorld* GetCurrentWorld() const = 0;
+
+	/** Returns true if the scene is open now (The current world is assigned and DCRA has already initialized for it). */
+	virtual bool IsSceneOpened() const = 0;
 
-	virtual class IDisplayClusterViewportManager& GetOwner() const = 0;
+	/** Returns true if the current world type is equal to one of the input types. */
+	virtual bool IsCurrentWorldHasAnyType(const EWorldType::Type InWorldType1, const EWorldType::Type InWorldType2 = EWorldType::None, const EWorldType::Type InWorldType3 = EWorldType::None) const = 0;
 
 	virtual void SetRenderSettings(const FDisplayClusterViewport_RenderSettings& InRenderSettings) = 0;
 	virtual void SetContexts(TArray<FDisplayClusterViewport_Context>& InContexts) = 0;
diff --git a/Engine/Plugins/Runtime/nDisplay/Source/DisplayCluster/Public/Render/Viewport/IDisplayClusterViewportManager.h b/Engine/Plugins/Runtime/nDisplay/Source/DisplayCluster/Public/Render/Viewport/IDisplayClusterViewportManager.h
--- a/Engine/Plugins/Runtime/nDisplay/Source/DisplayCluster/Public/Render/Viewport/IDisplayClusterViewportManager.h
+++ b/Engine/Plugins/Runtime/nDisplay/Source/DisplayCluster/Public/Render/Viewport/IDisplayClusterViewportManager.h
@@ -5,8 +5,9 @@
 #include "CoreMinimal.h"
 
 #include "Render/Viewport/IDisplayClusterViewport.h"
-#include "Render/Viewport/RenderFrame/DisplayClusterRenderFrame.h"
 #include "Render/Viewport/Containers/DisplayClusterPreviewSettings.h"
+#include "Render/Viewport/RenderFrame/DisplayClusterRenderFrameEnums.h"
+
 #include "SceneView.h"
 
 class UWorld;
@@ -16,13 +17,23 @@ class ADisplayClusterRootActor;
 class UDisplayClusterConfigurationViewport;
 class IDisplayClusterViewportManagerProxy;
 class FReferenceCollector;
+class FDisplayClusterRenderFrame;
+struct FDisplayClusterRenderFrameTarget;
+struct FDisplayClusterRenderFrameTargetViewFamily;
 
+/**
+ * nDisplay ViewportManager (interface for GameThread)
+ */
 class DISPLAYCLUSTER_API IDisplayClusterViewportManager
 {
 public:
 	virtual ~IDisplayClusterViewportManager() = default;
 
 public:
+	/** Get TSharedPtr from self. */
+	virtual TSharedPtr<IDisplayClusterViewportManager, ESPMode::ThreadSafe> ToSharedPtr() = 0;
+	virtual TSharedPtr<const IDisplayClusterViewportManager, ESPMode::ThreadSafe> ToSharedPtr() const = 0;
+
 	virtual const IDisplayClusterViewportManagerProxy* GetProxy() const = 0;
 	virtual       IDisplayClusterViewportManagerProxy* GetProxy() = 0;
 
@@ -31,7 +42,7 @@ public:
 	virtual ADisplayClusterRootActor* GetRootActor() const = 0;
 
 	/**
-	* Return current scene status
+	* Returns true if the scene is open now (The current world is assigned and DCRA has already initialized for it).
 	* [Game thread func]
 	*/
 	virtual bool IsSceneOpened() const = 0;
@@ -101,7 +112,7 @@ public:
 	* @return - The ConstructionValues object that was created, meant to initialize a view family context.
 	*/
 	virtual FSceneViewFamily::ConstructionValues CreateViewFamilyConstructionValues(
-		const FDisplayClusterRenderFrame::FFrameRenderTarget& InFrameTarget,
+		const FDisplayClusterRenderFrameTarget& InFrameTarget,
 		FSceneInterface* InScene,
 		FEngineShowFlags InEngineShowFlags,
 		const bool bInAdditionalViewFamily
@@ -115,11 +126,15 @@ public:
 	* @param OutRenderOutViewFamily - output family
 	*
 	*/
-	virtual void ConfigureViewFamily(const FDisplayClusterRenderFrame::FFrameRenderTarget& InFrameTarget, const FDisplayClusterRenderFrame::FFrameViewFamily& InFrameViewFamily, FSceneViewFamilyContext& InOutViewFamily) = 0;
+	virtual void ConfigureViewFamily(const FDisplayClusterRenderFrameTarget& InFrameTarget, const FDisplayClusterRenderFrameTargetViewFamily& InFrameViewFamily, FSceneViewFamilyContext& InOutViewFamily) = 0;
 
-	// Send to render thread
+	/** Send to render thread. */
 	virtual void RenderFrame(FViewport* InViewport) = 0;
 
+	/** Return current render mode. */
+	virtual EDisplayClusterRenderFrameMode GetRenderMode() const = 0;
+
+	/** Add internal DCVM objects to the reference collector. */
 	virtual void AddReferencedObjects(FReferenceCollector& Collector) = 0;
 
 #if WITH_EDITOR
@@ -136,7 +151,7 @@ public:
 	*
 	* @return - true, if render success
 	*/
-	virtual bool RenderInEditor(class FDisplayClusterRenderFrame& InRenderFrame, FViewport* InViewport, const uint32 InFirstViewportNum, const int32 InViewportsAmount, int32& OutViewportsAmount, bool& bOutFrameRendered) = 0;
+	virtual bool RenderInEditor(FDisplayClusterRenderFrame& InRenderFrame, FViewport* InViewport, const uint32 InFirstViewportNum, const int32 InViewportsAmount, int32& OutViewportsAmount, bool& bOutFrameRendered) = 0;
 #endif
 
 	///////////////////////////////////////////////////////////////////////////////////////////////////////////////////////////////////////////////////////////////////////////
@@ -167,7 +182,7 @@ public:
 	*
 	* @return - arrays with viewport objects refs
 	*/
-	virtual const TArrayView<IDisplayClusterViewport*> GetViewports() const = 0;
+	virtual const TArrayView<TSharedPtr<IDisplayClusterViewport, ESPMode::ThreadSafe>> GetViewports() const = 0;
 
 	/**
 	* Mark the geometry of the referenced component(s) as dirty (ProceduralMesh, etc)
diff --git a/Engine/Plugins/Runtime/nDisplay/Source/DisplayCluster/Public/Render/Viewport/IDisplayClusterViewportManagerProxy.h b/Engine/Plugins/Runtime/nDisplay/Source/DisplayCluster/Public/Render/Viewport/IDisplayClusterViewportManagerProxy.h
--- a/Engine/Plugins/Runtime/nDisplay/Source/DisplayCluster/Public/Render/Viewport/IDisplayClusterViewportManagerProxy.h
+++ b/Engine/Plugins/Runtime/nDisplay/Source/DisplayCluster/Public/Render/Viewport/IDisplayClusterViewportManagerProxy.h
@@ -6,13 +6,22 @@
 
 #include "Render/Viewport/IDisplayClusterViewportProxy.h"
 
+/**
+ * nDisplay ViewportManagerProxy (interface for RenderThread)
+ */
 class DISPLAYCLUSTER_API IDisplayClusterViewportManagerProxy
 {
 public:
 	virtual ~IDisplayClusterViewportManagerProxy() = default;
 
 public:
-	///////////////////////////////////////////////////////////////////////////////////////////////////////////////////////////////////////////////////////////////////////////
+	/** Get TSharedPtr from self. */
+	virtual TSharedPtr<IDisplayClusterViewportManagerProxy, ESPMode::ThreadSafe> ToSharedPtr() = 0;
+	virtual TSharedPtr<const IDisplayClusterViewportManagerProxy, ESPMode::ThreadSafe> ToSharedPtr() const = 0;
+
+	/** Return current render mode. */
+	virtual EDisplayClusterRenderFrameMode GetRenderMode() const = 0;
+
 	/**
 	* Find viewport render thread proxy object by name
 	* [Rendering thread func]
diff --git a/Engine/Plugins/Runtime/nDisplay/Source/DisplayCluster/Public/Render/Viewport/IDisplayClusterViewportProxy.h b/Engine/Plugins/Runtime/nDisplay/Source/DisplayCluster/Public/Render/Viewport/IDisplayClusterViewportProxy.h
--- a/Engine/Plugins/Runtime/nDisplay/Source/DisplayCluster/Public/Render/Viewport/IDisplayClusterViewportProxy.h
+++ b/Engine/Plugins/Runtime/nDisplay/Source/DisplayCluster/Public/Render/Viewport/IDisplayClusterViewportProxy.h
@@ -15,12 +15,24 @@
 #include "Render/Viewport/Containers/DisplayClusterViewport_PostRenderSettings.h"
 #include "Render/Viewport/Containers/DisplayClusterViewportProxy_Context.h"
 
+#include "Render/Viewport/RenderFrame/DisplayClusterRenderFrameEnums.h"
+
+/**
+ * nDisplay: ViewportProxy (interface for RenderThread)
+ */
 class DISPLAYCLUSTER_API IDisplayClusterViewportProxy
 {
 public:
 	virtual ~IDisplayClusterViewportProxy() = default;
 
 public:
+	/** Get TSharedPtr from self. */
+	virtual TSharedPtr<IDisplayClusterViewportProxy, ESPMode::ThreadSafe> ToSharedPtr() = 0;
+	virtual TSharedPtr<const IDisplayClusterViewportProxy, ESPMode::ThreadSafe> ToSharedPtr() const = 0;
+
+	/** Return current render mode. */
+	virtual EDisplayClusterRenderFrameMode GetRenderMode() const = 0;
+
 	virtual FString GetId() const = 0;
 	virtual FString GetClusterNodeId() const = 0;
 
@@ -66,7 +78,11 @@ public:
 	/** Return output resource type (support preview, remap, etc). */
 	virtual EDisplayClusterViewportResourceType   GetOutputResourceType_RenderThread() const = 0;
 
-	virtual const class IDisplayClusterViewportManagerProxy& GetOwner_RenderThread() const = 0;
+	/** Returns ptr to ViewportManagerProxy (the owner of this viewport proxy) if it still exists. */
+	virtual const class IDisplayClusterViewportManagerProxy* GetViewportManagerProxy_RenderThread() const = 0;
+
+	/** Get render frame settings from viewport proxy owner, if it still exists. */
+	virtual const struct FDisplayClusterRenderFrameSettings* GetRenderFrameSettings_RenderThread() const = 0;
 
 	virtual void SetRenderSettings_RenderThread(const FDisplayClusterViewport_RenderSettings& InRenderSettings) const = 0;
 	virtual void SetContexts_RenderThread(const TArray<FDisplayClusterViewport_Context>& InContexts) const = 0;
diff --git a/Engine/Plugins/Runtime/nDisplay/Source/DisplayCluster/Public/Render/Viewport/RenderFrame/DisplayClusterRenderFrame.h b/Engine/Plugins/Runtime/nDisplay/Source/DisplayCluster/Public/Render/Viewport/RenderFrame/DisplayClusterRenderFrame.h
--- a/Engine/Plugins/Runtime/nDisplay/Source/DisplayCluster/Public/Render/Viewport/RenderFrame/DisplayClusterRenderFrame.h
+++ b/Engine/Plugins/Runtime/nDisplay/Source/DisplayCluster/Public/Render/Viewport/RenderFrame/DisplayClusterRenderFrame.h
@@ -6,96 +6,92 @@
 
 #include "DisplayClusterRenderFrameEnums.h"
 #include "Render/Viewport/Containers/DisplayClusterViewport_RenderSettings.h"
+
 #include "SceneViewExtension.h"
 
+class IDisplayClusterViewport;
 class IDisplayClusterViewportManager;
+class FRenderTarget;
 
-// Render frame container
-class FDisplayClusterRenderFrame
+/**
+ * nDisplay: DCViewport context for render in UE View
+ */
+struct FDisplayClusterRenderFrameTargetView
 {
-public:
-	~FDisplayClusterRenderFrame()
+	/** Is this viewport context can be rendered. */
+	inline bool IsViewportContextCanBeRendered() const
 	{
-		RenderTargets.Empty();
+		return !bDisableRender && !bFreezeRendering;
 	}
 
-public:
-	class FFrameView
-	{
-	public:
-		inline bool ShouldRenderSceneView() const
-		{
-			return !bDisableRender && !bFreezeRendering;
-		}
+	// Viewport context index for this view
+	uint32 ContextNum = 0;
 
-	public:
-		// Viewport game-thread data
-		class IDisplayClusterViewport* Viewport = nullptr;
-
-		// Viewport context index for this view
-		uint32 ContextNum = 0;
-
-		bool bDisableRender = false;
-		bool bFreezeRendering = false;
-	};
-
-	class FFrameViewFamily
-	{
-	public:
-		~FFrameViewFamily()
-		{
-			ViewExtensions.Empty();
-			Views.Empty();
-		}
+	// Some viewports may skip rendering (for example: overridden)
+	bool bDisableRender = false;
 
-	public:
-		// Customize ScreenPercentage feature for viewfamily
-		float CustomBufferRatio = 1;
+	// In special cases, viewports can reuse an existing RTT image
+	// For example: the preview is rendered over several frames, so the already rendered viewports remain "frozen" and can be used during subsequent frames.
+	bool bFreezeRendering = false;
 
-		// Extensions that can modify view parameters
-		TArray<FSceneViewExtensionRef> ViewExtensions;
+	// Viewport game-thread data
+	TSharedPtr<IDisplayClusterViewport, ESPMode::ThreadSafe> Viewport;
+};
 
-		// Vieports, rendered at once for tthis family
-		TArray<FFrameView> Views;
+/**
+ * nDisplay: A group of DCViewports that can be render within a single viewfamily
+ */
+struct FDisplayClusterRenderFrameTargetViewFamily
+{
+	// Customize ScreenPercentage feature for viewfamily
+	float CustomBufferRatio = 1;
 
-		int32 NumViewsForRender = 0;
-	};
+	// Extensions that can modify view parameters
+	TArray<FSceneViewExtensionRef> ViewExtensions;
 
-	class FFrameRenderTarget
-	{
-	public:
-		~FFrameRenderTarget()
-		{
-			ViewFamilies.Empty();
-		}
+	// Vieports, rendered at once for tthis family
+	TArray<FDisplayClusterRenderFrameTargetView> Views;
+};
 
-	public:
-		// Discard some RTT (when view render disabled)
-		// Also when RTT Atlasing used, this viewports excluded from atlas map (reduce size)
-		bool bShouldUseRenderTarget = true;
+/**
+ * nDisplay: Target texture for rendering.
+ * Contains structured information on how to render DCViewports via viewfamilies in RTT
+ */
+struct FDisplayClusterRenderFrameTarget
+{
+	// Discard some RTT (when view render disabled)
+	// Also when RTT Atlasing used, this viewports excluded from atlas map (reduce size)
+	bool bShouldUseRenderTarget = true;
 
-		// required Render target size (resource can be bigger)
-		FIntPoint RenderTargetSize;
+	// required Render target size (resource can be bigger)
+	FIntPoint RenderTargetSize;
 
-		EDisplayClusterViewportCaptureMode CaptureMode = EDisplayClusterViewportCaptureMode::Default;
+	EDisplayClusterViewportCaptureMode CaptureMode = EDisplayClusterViewportCaptureMode::Default;
 
-		// Render target resource ref
-		class FRenderTarget* RenderTargetPtr = nullptr;
+	// Render target resource ref
+	FRenderTarget* RenderTargetPtr = nullptr;
 
-		// Families, rendered on this target
-		TArray<FFrameViewFamily> ViewFamilies;
-	};
+	// Families, rendered on this target
+	TArray<FDisplayClusterRenderFrameTargetViewFamily> ViewFamilies;
+};
 
+/**
+ * nDisplay: Render frame container
+ * The final frame is composed of DCViewports, which are rendered in the correct order on a few RTTs.
+ * Contains all information about how to render all DCViewports for the current frame and the settings for it.
+ */
+class FDisplayClusterRenderFrame
+{
 public:
 	// Render frame to this targets
-	TArray<FFrameRenderTarget> RenderTargets;
+	TArray<FDisplayClusterRenderFrameTarget> RenderTargets;
 	
 	// Frame rect on final backbuffer
 	FIntRect FrameRect;
 
+	// Desired numbers of view
 	int32 DesiredNumberOfViews = 0;
-	int32 ViewportsAmount = 0;
 
-	IDisplayClusterViewportManager* ViewportManager = nullptr;
+	// Ref to the DC Viewport Manager
+	TSharedPtr<IDisplayClusterViewportManager, ESPMode::ThreadSafe> ViewportManager;
 };
-
diff --git a/Engine/Plugins/Runtime/nDisplay/Source/DisplayCluster/Public/Render/Viewport/RenderFrame/DisplayClusterRenderFrameEnums.h b/Engine/Plugins/Runtime/nDisplay/Source/DisplayCluster/Public/Render/Viewport/RenderFrame/DisplayClusterRenderFrameEnums.h
--- a/Engine/Plugins/Runtime/nDisplay/Source/DisplayCluster/Public/Render/Viewport/RenderFrame/DisplayClusterRenderFrameEnums.h
+++ b/Engine/Plugins/Runtime/nDisplay/Source/DisplayCluster/Public/Render/Viewport/RenderFrame/DisplayClusterRenderFrameEnums.h
@@ -6,6 +6,8 @@
 
 enum class EDisplayClusterRenderFrameMode: uint8
 {
+	Unknown = 0,
+
 	// Render to single backbuffer texture
 	Mono,
 
diff --git a/Engine/Plugins/Runtime/nDisplay/Source/DisplayClusterProjection/Private/Policy/Camera/DisplayClusterProjectionCameraPolicy.cpp b/Engine/Plugins/Runtime/nDisplay/Source/DisplayClusterProjection/Private/Policy/Camera/DisplayClusterProjectionCameraPolicy.cpp
--- a/Engine/Plugins/Runtime/nDisplay/Source/DisplayClusterProjection/Private/Policy/Camera/DisplayClusterProjectionCameraPolicy.cpp
+++ b/Engine/Plugins/Runtime/nDisplay/Source/DisplayClusterProjection/Private/Policy/Camera/DisplayClusterProjectionCameraPolicy.cpp
@@ -49,11 +49,9 @@ APlayerCameraManager* const GetCurPlayerCameraManager(IDisplayClusterViewport* I
 {
 	if (InViewport)
 	{
-		UWorld* World = InViewport->GetOwner().GetCurrentWorld();
-		if (World)
+		if (UWorld* World = InViewport->GetCurrentWorld())
 		{
-			APlayerController* const CurPlayerController = World->GetFirstPlayerController();
-			if (CurPlayerController)
+			if (APlayerController* const CurPlayerController = World->GetFirstPlayerController())
 			{
 				return CurPlayerController->PlayerCameraManager;
 			}
diff --git a/Engine/Plugins/Runtime/nDisplay/Source/DisplayClusterProjection/Private/Policy/DisplayClusterProjectionPolicyBase.cpp b/Engine/Plugins/Runtime/nDisplay/Source/DisplayClusterProjection/Private/Policy/DisplayClusterProjectionPolicyBase.cpp
--- a/Engine/Plugins/Runtime/nDisplay/Source/DisplayClusterProjection/Private/Policy/DisplayClusterProjectionPolicyBase.cpp
+++ b/Engine/Plugins/Runtime/nDisplay/Source/DisplayClusterProjection/Private/Policy/DisplayClusterProjectionPolicyBase.cpp
@@ -46,26 +46,7 @@ bool FDisplayClusterProjectionPolicyBase::IsEditorOperationMode(IDisplayClusterV
 	}
 
 	// Get state from viewport world (UE-114493)
-	if (InViewport)
-	{
-		if (UWorld* CurrentWorld = InViewport->GetOwner().GetCurrentWorld())
-		{
-			if (CurrentWorld)
-			{
-				switch (CurrentWorld->WorldType)
-				{
-				case EWorldType::Editor:
-				case EWorldType::EditorPreview:
-					return true;
-
-				default:
-					break;
-				}
-			}
-		}
-	}
-
-	return false;
+	return InViewport && InViewport->IsCurrentWorldHasAnyType(EWorldType::Editor, EWorldType::EditorPreview);
 }
 
 bool FDisplayClusterProjectionPolicyBase::IsConfigurationChanged(const FDisplayClusterConfigurationProjection* InConfigurationProjectionPolicy) const
@@ -95,8 +76,7 @@ void FDisplayClusterProjectionPolicyBase::InitializeOriginComponent(IDisplayClus
 	if (InViewport)
 	{
 		USceneComponent* PolicyOriginComp = nullptr;
-		ADisplayClusterRootActor* RootActor = InViewport->GetOwner().GetRootActor();
-		if (RootActor)
+		if (ADisplayClusterRootActor* RootActor = InViewport->GetRootActor())
 		{
 			// default use root actor as Origin
 			PolicyOriginComp = RootActor->GetRootComponent();
diff --git a/Engine/Plugins/Runtime/nDisplay/Source/DisplayClusterProjection/Private/Policy/Link/DisplayClusterProjectionLinkPolicy.cpp b/Engine/Plugins/Runtime/nDisplay/Source/DisplayClusterProjection/Private/Policy/Link/DisplayClusterProjectionLinkPolicy.cpp
--- a/Engine/Plugins/Runtime/nDisplay/Source/DisplayClusterProjection/Private/Policy/Link/DisplayClusterProjectionLinkPolicy.cpp
+++ b/Engine/Plugins/Runtime/nDisplay/Source/DisplayClusterProjection/Private/Policy/Link/DisplayClusterProjectionLinkPolicy.cpp
@@ -32,23 +32,25 @@ bool FDisplayClusterProjectionLinkPolicy::CalculateView(IDisplayClusterViewport*
 	check(IsInGameThread());
 	check(InViewport);
 
-	IDisplayClusterViewport* ParentViewport = InViewport->GetOwner().FindViewport(InViewport->GetRenderSettings().GetParentViewportId());
-	if (ParentViewport != nullptr)
+	if (IDisplayClusterViewportManager* ViewportManager = InViewport->GetViewportManager())
 	{
-		const TArray<FDisplayClusterViewport_Context>& ParentContexts = ParentViewport->GetContexts();
-		if (InContextNum < (uint32)ParentContexts.Num())
+		if (IDisplayClusterViewport* ParentViewport = ViewportManager->FindViewport(InViewport->GetRenderSettings().GetParentViewportId()))
 		{
-			if (ParentContexts[InContextNum].bDisableRender && ParentViewport->GetProjectionPolicy().IsValid())
+			const TArray<FDisplayClusterViewport_Context>& ParentContexts = ParentViewport->GetContexts();
+			if (InContextNum < (uint32)ParentContexts.Num())
 			{
-				return ParentViewport->GetProjectionPolicy()->CalculateView(ParentViewport, InContextNum, InOutViewLocation, InOutViewRotation, ViewOffset, WorldToMeters, NCP, FCP);
-			}
-			else
-			{
-				const FDisplayClusterViewport_Context& ParentContext = ParentContexts[InContextNum];
-				InOutViewLocation = ParentContext.ViewLocation;
-				InOutViewRotation = ParentContext.ViewRotation;
-
-				return true;
+				if (ParentContexts[InContextNum].bDisableRender && ParentViewport->GetProjectionPolicy().IsValid())
+				{
+					return ParentViewport->GetProjectionPolicy()->CalculateView(ParentViewport, InContextNum, InOutViewLocation, InOutViewRotation, ViewOffset, WorldToMeters, NCP, FCP);
+				}
+				else
+				{
+					const FDisplayClusterViewport_Context& ParentContext = ParentContexts[InContextNum];
+					InOutViewLocation = ParentContext.ViewLocation;
+					InOutViewRotation = ParentContext.ViewRotation;
+
+					return true;
+				}
 			}
 		}
 	}
@@ -61,22 +63,24 @@ bool FDisplayClusterProjectionLinkPolicy::GetProjectionMatrix(IDisplayClusterVie
 	check(IsInGameThread());
 	check(InViewport);
 
-	IDisplayClusterViewport* ParentViewport = InViewport->GetOwner().FindViewport(InViewport->GetRenderSettings().GetParentViewportId());
-	if (ParentViewport != nullptr)
+	if (IDisplayClusterViewportManager* ViewportManager = InViewport->GetViewportManager())
 	{
-		const TArray<FDisplayClusterViewport_Context>& ParentContexts = ParentViewport->GetContexts();
-		if (InContextNum < (uint32)ParentContexts.Num())
+		if (IDisplayClusterViewport* ParentViewport = ViewportManager->FindViewport(InViewport->GetRenderSettings().GetParentViewportId()))
 		{
-			if (ParentContexts[InContextNum].bDisableRender && ParentViewport->GetProjectionPolicy().IsValid())
+			const TArray<FDisplayClusterViewport_Context>& ParentContexts = ParentViewport->GetContexts();
+			if (InContextNum < (uint32)ParentContexts.Num())
 			{
-				return ParentViewport->GetProjectionPolicy()->GetProjectionMatrix(ParentViewport, InContextNum, OutPrjMatrix);
-			}
-			else
-			{
-				const FDisplayClusterViewport_Context& ParentContext = ParentContexts[InContextNum];
-				OutPrjMatrix = ParentContext.ProjectionMatrix;
-
-				return true;
+				if (ParentContexts[InContextNum].bDisableRender && ParentViewport->GetProjectionPolicy().IsValid())
+				{
+					return ParentViewport->GetProjectionPolicy()->GetProjectionMatrix(ParentViewport, InContextNum, OutPrjMatrix);
+				}
+				else
+				{
+					const FDisplayClusterViewport_Context& ParentContext = ParentContexts[InContextNum];
+					OutPrjMatrix = ParentContext.ProjectionMatrix;
+
+					return true;
+				}
 			}
 		}
 	}
diff --git a/Engine/Plugins/Runtime/nDisplay/Source/DisplayClusterProjection/Private/Policy/MPCDI/DisplayClusterProjectionMPCDIPolicy.cpp b/Engine/Plugins/Runtime/nDisplay/Source/DisplayClusterProjection/Private/Policy/MPCDI/DisplayClusterProjectionMPCDIPolicy.cpp
--- a/Engine/Plugins/Runtime/nDisplay/Source/DisplayClusterProjection/Private/Policy/MPCDI/DisplayClusterProjectionMPCDIPolicy.cpp
+++ b/Engine/Plugins/Runtime/nDisplay/Source/DisplayClusterProjection/Private/Policy/MPCDI/DisplayClusterProjectionMPCDIPolicy.cpp
@@ -235,6 +235,12 @@ void FDisplayClusterProjectionMPCDIPolicy::ApplyWarpBlend_RenderThread(FRHIComma
 		return;
 	}
 
+	const IDisplayClusterViewportManagerProxy* ViewportManagerProxyPtr = InViewportProxy->GetViewportManagerProxy_RenderThread();
+	if (!ViewportManagerProxyPtr)
+	{
+		return;
+	}
+
 	TArray<FRHITexture2D*> InputTextures, OutputTextures;
 	TArray<FIntRect> InputRects, OutputRects;
 
@@ -268,12 +274,11 @@ void FDisplayClusterProjectionMPCDIPolicy::ApplyWarpBlend_RenderThread(FRHIComma
 		for (int32 ContextNum = 0; ContextNum < InputTextures.Num(); ContextNum++)
 		{
 			// Update referenced resources and math in the ShaderICVFX.
-			ShaderICVFX.IterateViewportResourcesByPredicate([ContextNum, InViewportProxy, &ShaderICVFX](FDisplayClusterShaderParametersICVFX_ViewportResource& ViewportResourceIt)
+			ShaderICVFX.IterateViewportResourcesByPredicate([ContextNum, ViewportManagerProxyPtr, &ShaderICVFX](FDisplayClusterShaderParametersICVFX_ViewportResource& ViewportResourceIt)
 			{
 				// reset prev resource reference
 				ViewportResourceIt.Texture = nullptr;
-
-				if(const IDisplayClusterViewportProxy* SrcViewportProxy = InViewportProxy->GetOwner_RenderThread().FindViewport_RenderThread(ViewportResourceIt.ViewportId))
+				if(const IDisplayClusterViewportProxy* SrcViewportProxy = ViewportManagerProxyPtr->FindViewport_RenderThread(ViewportResourceIt.ViewportId))
 				{
 					if (!SrcViewportProxy->GetRenderSettings_RenderThread().bSkipRendering)
 					{
diff --git a/Engine/Plugins/Runtime/nDisplay/Source/DisplayClusterProjection/Private/Policy/Mesh/DisplayClusterProjectionMeshPolicy.cpp b/Engine/Plugins/Runtime/nDisplay/Source/DisplayClusterProjection/Private/Policy/Mesh/DisplayClusterProjectionMeshPolicy.cpp
--- a/Engine/Plugins/Runtime/nDisplay/Source/DisplayClusterProjection/Private/Policy/Mesh/DisplayClusterProjectionMeshPolicy.cpp
+++ b/Engine/Plugins/Runtime/nDisplay/Source/DisplayClusterProjection/Private/Policy/Mesh/DisplayClusterProjectionMeshPolicy.cpp
@@ -121,8 +121,8 @@ bool FDisplayClusterProjectionMeshPolicy::GetWarpMeshConfiguration(IDisplayClust
 	check(InViewport);
 
 	// Get our VR root
-	ADisplayClusterRootActor* const Root = InViewport->GetOwner().GetRootActor();
-	if (!Root)
+	ADisplayClusterRootActor* RootActorPtr = InViewport->GetRootActor();
+	if (!RootActorPtr)
 	{
 		UE_LOG(LogDisplayClusterProjectionMesh, Error, TEXT("Couldn't get a VR root object"));
 		return false;
@@ -179,7 +179,7 @@ bool FDisplayClusterProjectionMeshPolicy::GetWarpMeshConfiguration(IDisplayClust
 	OutWarpCfg.ChromakeyUVIndex = CfgChromakeyUVIndex;
 
 	// Get the StaticMeshComponent
-	OutWarpCfg.StaticMeshComponent = Root->GetComponentByName<UStaticMeshComponent>(ComponentId);
+	OutWarpCfg.StaticMeshComponent = RootActorPtr->GetComponentByName<UStaticMeshComponent>(ComponentId);
 	if (OutWarpCfg.StaticMeshComponent != nullptr)
 	{
 		int CfgLODIndex;
@@ -200,7 +200,7 @@ bool FDisplayClusterProjectionMeshPolicy::GetWarpMeshConfiguration(IDisplayClust
 	else
 	{
 		// Get the ProceduralMeshComponent
-		OutWarpCfg.ProceduralMeshComponent = Root->GetComponentByName<UProceduralMeshComponent>(ComponentId);
+		OutWarpCfg.ProceduralMeshComponent = RootActorPtr->GetComponentByName<UProceduralMeshComponent>(ComponentId);
 		if (OutWarpCfg.ProceduralMeshComponent == nullptr)
 		{
 			if (!IsEditorOperationMode(InViewport))
diff --git a/Engine/Plugins/Runtime/nDisplay/Source/DisplayClusterProjection/Private/Policy/Simple/DisplayClusterProjectionSimplePolicy.cpp b/Engine/Plugins/Runtime/nDisplay/Source/DisplayClusterProjection/Private/Policy/Simple/DisplayClusterProjectionSimplePolicy.cpp
--- a/Engine/Plugins/Runtime/nDisplay/Source/DisplayClusterProjection/Private/Policy/Simple/DisplayClusterProjectionSimplePolicy.cpp
+++ b/Engine/Plugins/Runtime/nDisplay/Source/DisplayClusterProjection/Private/Policy/Simple/DisplayClusterProjectionSimplePolicy.cpp
@@ -184,8 +184,8 @@ bool FDisplayClusterProjectionSimplePolicy::InitializeMeshData(IDisplayClusterVi
 	UE_LOG(LogDisplayClusterProjectionSimple, Verbose, TEXT("Initializing screen geometry for viewport policy  %s"), *GetId());
 
 	// Get our VR root
-	ADisplayClusterRootActor* const Root = InViewport->GetOwner().GetRootActor();
-	if (!Root)
+	ADisplayClusterRootActor* const RootActorPtr = InViewport->GetRootActor();
+	if (!RootActorPtr)
 	{
 		if (!IsEditorOperationMode(InViewport))
 		{
@@ -196,7 +196,7 @@ bool FDisplayClusterProjectionSimplePolicy::InitializeMeshData(IDisplayClusterVi
 	}
 
 	// Get screen component
-	UDisplayClusterScreenComponent* ScreenComp = Root->GetComponentByName<UDisplayClusterScreenComponent>(ScreenId);
+	UDisplayClusterScreenComponent* ScreenComp = RootActorPtr->GetComponentByName<UDisplayClusterScreenComponent>(ScreenId);
 	if (!ScreenComp)
 	{
 		if (!IsEditorOperationMode(InViewport))
diff --git a/Engine/Plugins/VirtualProduction/TextureShare/Source/TextureShareDisplayCluster/Private/PostProcess/TextureSharePostprocess_Frame.cpp b/Engine/Plugins/VirtualProduction/TextureShare/Source/TextureShareDisplayCluster/Private/PostProcess/TextureSharePostprocess_Frame.cpp
--- a/Engine/Plugins/VirtualProduction/TextureShare/Source/TextureShareDisplayCluster/Private/PostProcess/TextureSharePostprocess_Frame.cpp
+++ b/Engine/Plugins/VirtualProduction/TextureShare/Source/TextureShareDisplayCluster/Private/PostProcess/TextureSharePostprocess_Frame.cpp
@@ -25,9 +25,9 @@ void FTextureSharePostprocess::UpdateSupportedViews(IDisplayClusterViewportManag
 	if (InViewportManager)
 	{
 		// Get all existing viewports on DC node
-		for (const IDisplayClusterViewport* ViewportIt : InViewportManager->GetViewports())
+		for (const TSharedPtr<IDisplayClusterViewport, ESPMode::ThreadSafe>& ViewportIt : InViewportManager->GetViewports())
 		{
-			if (ViewportIt)
+			if (ViewportIt.IsValid())
 			{
 				const TArray<FDisplayClusterViewport_Context>& Contexts = ViewportIt->GetContexts();
 				for (const FDisplayClusterViewport_Context& ViewportContextIt : Contexts)
@@ -44,27 +44,30 @@ void FTextureSharePostprocess::UpdateManualProjectionPolicy(IDisplayClusterViewp
 {
 	if (InViewportManager)
 	{
-		for (IDisplayClusterViewport* ViewportIt : InViewportManager->GetViewports())
+		for (const TSharedPtr<IDisplayClusterViewport, ESPMode::ThreadSafe>& ViewportIt : InViewportManager->GetViewports())
 		{
-			TSharedPtr<IDisplayClusterProjectionPolicy, ESPMode::ThreadSafe> PrjPolicy = ViewportIt->GetProjectionPolicy();
-			if (PrjPolicy.IsValid() && PrjPolicy->GetType() == UE::TextureShare::DisplayClusterStrings::Projection::TextureShare)
+			if (ViewportIt.IsValid())
 			{
-				// Search the custom manual projection for viewport with the "TextureShare" projection policy type 
-				TArray<FTextureShareCoreManualProjection> ManualProjections;
-				for (const FTextureShareCoreObjectData& FrameDataIt : Object->GetReceivedCoreObjectData())
+				TSharedPtr<IDisplayClusterProjectionPolicy, ESPMode::ThreadSafe> PrjPolicy = ViewportIt->GetProjectionPolicy();
+				if (PrjPolicy.IsValid() && PrjPolicy->GetType() == UE::TextureShare::DisplayClusterStrings::Projection::TextureShare)
 				{
-					if (FrameDataIt.Data.ManualProjections.GetValuesByEqualsFunc(ViewportIt->GetId(), ManualProjections))
+					// Search the custom manual projection for viewport with the "TextureShare" projection policy type 
+					TArray<FTextureShareCoreManualProjection> ManualProjections;
+					for (const FTextureShareCoreObjectData& FrameDataIt : Object->GetReceivedCoreObjectData())
 					{
-						// Apply manual projection
-						static ITextureShareDisplayClusterAPI& TextureShareDisplayClusterAPI = ITextureShareDisplayCluster::Get().GetTextureShareDisplayClusterAPI();
-						TextureShareDisplayClusterAPI.TextureSharePolicySetProjectionData(PrjPolicy, ManualProjections);
-
-						// Save used manual projections sources
-						for (const FTextureShareCoreManualProjection& It : ManualProjections)
+						if (FrameDataIt.Data.ManualProjections.GetValuesByEqualsFunc(ViewportIt->GetId(), ManualProjections))
 						{
-							Object->GetCoreData().ManualProjectionsSources.Add(FTextureShareCoreManualProjectionSource(It.ViewDesc, FTextureShareCoreObjectFrameMarker(FrameDataIt.Desc, FrameDataIt.Data.FrameMarker)));
+							// Apply manual projection
+							static ITextureShareDisplayClusterAPI& TextureShareDisplayClusterAPI = ITextureShareDisplayCluster::Get().GetTextureShareDisplayClusterAPI();
+							TextureShareDisplayClusterAPI.TextureSharePolicySetProjectionData(PrjPolicy, ManualProjections);
+
+							// Save used manual projections sources
+							for (const FTextureShareCoreManualProjection& It : ManualProjections)
+							{
+								Object->GetCoreData().ManualProjectionsSources.Add(FTextureShareCoreManualProjectionSource(It.ViewDesc, FTextureShareCoreObjectFrameMarker(FrameDataIt.Desc, FrameDataIt.Data.FrameMarker)));
+							}
+							break;
 						}
-						break;
 					}
 				}
 			}
@@ -77,9 +80,9 @@ void FTextureSharePostprocess::UpdateViews(IDisplayClusterViewportManager* InVie
 	if (InViewportManager)
 	{
 		//Register all viewport for scene eye data capture by default
-		for (IDisplayClusterViewport* ViewportIt : InViewportManager->GetViewports())
+		for (const TSharedPtr<IDisplayClusterViewport, ESPMode::ThreadSafe>& ViewportIt : InViewportManager->GetViewports())
 		{
-			if (ViewportIt)
+			if (ViewportIt.IsValid())
 			{
 				const bool bMonoscopic = ViewportIt->GetContexts().Num() == 1;
 				for (const FDisplayClusterViewport_Context& ContextIt : ViewportIt->GetContexts())
diff --git a/Engine/Plugins/VirtualProduction/TextureShare/Source/TextureShareDisplayCluster/Private/Projection/TextureShareProjectionPolicyBase.cpp b/Engine/Plugins/VirtualProduction/TextureShare/Source/TextureShareDisplayCluster/Private/Projection/TextureShareProjectionPolicyBase.cpp
--- a/Engine/Plugins/VirtualProduction/TextureShare/Source/TextureShareDisplayCluster/Private/Projection/TextureShareProjectionPolicyBase.cpp
+++ b/Engine/Plugins/VirtualProduction/TextureShare/Source/TextureShareDisplayCluster/Private/Projection/TextureShareProjectionPolicyBase.cpp
@@ -47,26 +47,7 @@ bool FTextureShareProjectionPolicyBase::IsEditorOperationMode(class IDisplayClus
 	}
 
 	// Get state from viewport world (UE-114493)
-	if (InViewport)
-	{
-		if (UWorld* CurrentWorld = InViewport->GetOwner().GetCurrentWorld())
-		{
-			if (CurrentWorld)
-			{
-				switch (CurrentWorld->WorldType)
-				{
-				case EWorldType::Editor:
-				case EWorldType::EditorPreview:
-					return true;
-
-				default:
-					break;
-				}
-			}
-		}
-	}
-
-	return false;
+	return InViewport && InViewport->IsCurrentWorldHasAnyType(EWorldType::Editor, EWorldType::EditorPreview);
 }
 
 bool FTextureShareProjectionPolicyBase::IsConfigurationChanged(const struct FDisplayClusterConfigurationProjection* InConfigurationProjectionPolicy) const
