nDisplay: the UI logic related to DisplayClusterConfiguration structures has been redesigned

- the logic that worked only with DisplayClusterConfiguration structures has been moved to the corresponding module.
- numerous duplicates in the DisplayCluster module code were removed, and the code was redesigned for the new functions from the DisplayClusterConfiguration module.
- Fixed some bugs related to incorrect duplication of UI logic in the code (see comments in swarm review).

diff --git a/Engine/Plugins/Runtime/nDisplay/Source/DisplayCluster/Private/Blueprints/DisplayClusterBlueprintLib.cpp b/Engine/Plugins/Runtime/nDisplay/Source/DisplayCluster/Private/Blueprints/DisplayClusterBlueprintLib.cpp
--- a/Engine/Plugins/Runtime/nDisplay/Source/DisplayCluster/Private/Blueprints/DisplayClusterBlueprintLib.cpp
+++ b/Engine/Plugins/Runtime/nDisplay/Source/DisplayCluster/Private/Blueprints/DisplayClusterBlueprintLib.cpp
@@ -204,16 +204,20 @@ void UDisplayClusterBlueprintLib::FindChromakeyCardsForRootActor(const ADisplayC
 
 	for (const UDisplayClusterICVFXCameraComponent* Camera : CameraComponents)
 	{
-		const FDisplayClusterConfigurationICVFX_VisibilityList& RootActorChromakeyCards = Camera->GetCameraSettingsICVFX().Chromakey.ChromakeyRenderTexture.ShowOnlyList;
+		const FDisplayClusterConfigurationICVFX_ChromakeyRenderSettings* CameraChromakeyRenderSettings = Camera->GetCameraSettingsICVFX().Chromakey.GetChromakeyRenderSettings(RootActor->GetStageSettings());
+		const FDisplayClusterConfigurationICVFX_VisibilityList* CameraChromakeyShowOnlyList = CameraChromakeyRenderSettings ? &CameraChromakeyRenderSettings->ShowOnlyList : nullptr;
 
-		for (const TSoftObjectPtr<AActor>& ChromakeyCardActor : RootActorChromakeyCards.Actors)
+		if (CameraChromakeyShowOnlyList)
 		{
-			if (!ChromakeyCardActor.IsValid() || !ChromakeyCardActor->IsA<ADisplayClusterChromakeyCardActor>())
+			for (const TSoftObjectPtr<AActor>& ChromakeyCardActor : CameraChromakeyShowOnlyList->Actors)
 			{
-				continue;
-			}
+				if (!ChromakeyCardActor.IsValid() || !ChromakeyCardActor->IsA<ADisplayClusterChromakeyCardActor>())
+				{
+					continue;
+				}
 
-			OutChromakeyCards.Add(Cast<ADisplayClusterChromakeyCardActor>(ChromakeyCardActor.Get()));
+				OutChromakeyCards.Add(Cast<ADisplayClusterChromakeyCardActor>(ChromakeyCardActor.Get()));
+			}
 		}
 
 		if (const UWorld* World = RootActor->GetWorld())
@@ -229,9 +233,9 @@ void UDisplayClusterBlueprintLib::FindChromakeyCardsForRootActor(const ADisplayC
 				{
 					OutChromakeyCards.Add(*ActorIt);
 				}
-				else
+				else if(CameraChromakeyShowOnlyList)
 				{
-					for (const FActorLayer& ActorLayer : RootActorChromakeyCards.ActorLayers)
+					for (const FActorLayer& ActorLayer : CameraChromakeyShowOnlyList->ActorLayers)
 					{
 						if (ActorIt->Layers.Contains(ActorLayer.Name))
 						{
diff --git a/Engine/Plugins/Runtime/nDisplay/Source/DisplayCluster/Private/Game/EngineClasses/Scene/DisplayClusterChromakeyCardActor.cpp b/Engine/Plugins/Runtime/nDisplay/Source/DisplayCluster/Private/Game/EngineClasses/Scene/DisplayClusterChromakeyCardActor.cpp
--- a/Engine/Plugins/Runtime/nDisplay/Source/DisplayCluster/Private/Game/EngineClasses/Scene/DisplayClusterChromakeyCardActor.cpp
+++ b/Engine/Plugins/Runtime/nDisplay/Source/DisplayCluster/Private/Game/EngineClasses/Scene/DisplayClusterChromakeyCardActor.cpp
@@ -48,15 +48,21 @@ bool ADisplayClusterChromakeyCardActor::IsReferencedByICVFXCamera(const UDisplay
 		}
 	}
 	
-	const FDisplayClusterConfigurationICVFX_CameraSettings& CameraSettings = InCamera->GetCameraSettingsICVFX();
-	if (CameraSettings.Chromakey.ChromakeyRenderTexture.ShowOnlyList.Actors.Contains(this))
+	const ADisplayClusterRootActor* RootActor = GetRootActorOwner();
+	const FDisplayClusterConfigurationICVFX_ChromakeyRenderSettings* ChromakeyRenderSettings = RootActor ? InCamera->GetCameraSettingsICVFX().Chromakey.GetChromakeyRenderSettings(RootActor->GetStageSettings()) : nullptr;
+	if (!ChromakeyRenderSettings)
+	{
+		return false;
+	}
+
+	if (ChromakeyRenderSettings->ShowOnlyList.Actors.Contains(this))
 	{
 		return true;
 	}
 
 	for (const FName& ThisLayer : Layers)
 	{
-		if (const FActorLayer* ExistingLayer = CameraSettings.Chromakey.ChromakeyRenderTexture.ShowOnlyList.ActorLayers.FindByPredicate([ThisLayer](const FActorLayer& Layer)
+		if (const FActorLayer* ExistingLayer = ChromakeyRenderSettings->ShowOnlyList.ActorLayers.FindByPredicate([ThisLayer](const FActorLayer& Layer)
 		{
 			return Layer.Name == ThisLayer;
 		}))
@@ -81,18 +87,7 @@ void ADisplayClusterChromakeyCardActor::UpdateChromakeySettings()
 		{
 			if (IsReferencedByICVFXCamera(Component))
 			{
-				const FDisplayClusterConfigurationICVFX_ChromakeySettings& CameraChromakey = Component->GetCameraSettingsICVFX().Chromakey;
-
-				if (CameraChromakey.ChromakeySettingsSource == EDisplayClusterConfigurationICVFX_ChromakeySettingsSource::Viewport)
-				{
-					const FDisplayClusterConfigurationICVFX_GlobalChromakeySettings& GlobalChromakey = RootActor->GetStageSettings().GlobalChromakey;
-					ChromaColor += GlobalChromakey.ChromakeyColor;
-				}
-				else
-				{
-					ChromaColor += Component->GetCameraSettingsICVFX().Chromakey.ChromakeyColor;
-				}
-
+				ChromaColor += Component->GetCameraSettingsICVFX().Chromakey.GetChromakeyColor(RootActor->GetStageSettings());
 				Count++;
 			}
 		}
diff --git a/Engine/Plugins/Runtime/nDisplay/Source/DisplayCluster/Private/Game/EngineClasses/Scene/DisplayClusterICVFXCameraComponent.cpp b/Engine/Plugins/Runtime/nDisplay/Source/DisplayCluster/Private/Game/EngineClasses/Scene/DisplayClusterICVFXCameraComponent.cpp
--- a/Engine/Plugins/Runtime/nDisplay/Source/DisplayCluster/Private/Game/EngineClasses/Scene/DisplayClusterICVFXCameraComponent.cpp
+++ b/Engine/Plugins/Runtime/nDisplay/Source/DisplayCluster/Private/Game/EngineClasses/Scene/DisplayClusterICVFXCameraComponent.cpp
@@ -146,14 +146,13 @@ void UDisplayClusterICVFXCameraComponent::UpdateOverscanEstimatedFrameSize()
 		return;
 	}
 
-	float InnerFrustumResolutionWidth = RootActor->GetStageSettings().DefaultFrameSize.Width * CameraSettings.BufferRatio;
-	float InnerFrustumResolutionHeight = RootActor->GetStageSettings().DefaultFrameSize.Height * CameraSettings.BufferRatio;
+	const FDisplayClusterConfigurationICVFX_StageSettings& StageSettings = RootActor->GetStageSettings();
 
-	if (CameraSettings.RenderSettings.CustomFrameSize.bUseCustomSize)
-	{
-		InnerFrustumResolutionWidth = CameraSettings.RenderSettings.CustomFrameSize.CustomWidth * CameraSettings.BufferRatio;
-		InnerFrustumResolutionHeight = CameraSettings.RenderSettings.CustomFrameSize.CustomHeight * CameraSettings.BufferRatio;
-	}
+	const float CameraBufferRatio = CameraSettings.GetCameraBufferRatio(StageSettings);
+	const FIntPoint InnerFrustumSize = CameraSettings.GetCameraFrameSize(StageSettings);
+
+	float InnerFrustumResolutionWidth = InnerFrustumSize.X * CameraBufferRatio;
+	float InnerFrustumResolutionHeight = InnerFrustumSize.Y * CameraBufferRatio;
 
 	float EstimatedOverscanResolutionWidth = InnerFrustumResolutionWidth;
 	float EstimatedOverscanResolutionHeight = InnerFrustumResolutionHeight;
diff --git a/Engine/Plugins/Runtime/nDisplay/Source/DisplayCluster/Private/Game/EngineClasses/Scene/DisplayClusterRootActor.cpp b/Engine/Plugins/Runtime/nDisplay/Source/DisplayCluster/Private/Game/EngineClasses/Scene/DisplayClusterRootActor.cpp
--- a/Engine/Plugins/Runtime/nDisplay/Source/DisplayCluster/Private/Game/EngineClasses/Scene/DisplayClusterRootActor.cpp
+++ b/Engine/Plugins/Runtime/nDisplay/Source/DisplayCluster/Private/Game/EngineClasses/Scene/DisplayClusterRootActor.cpp
@@ -800,7 +800,14 @@ void ADisplayClusterRootActor::SetLightCardOwnership()
 		
 			for (UDisplayClusterICVFXCameraComponent* Camera : ICVFXComponents)
 			{
-				FDisplayClusterConfigurationICVFX_VisibilityList& ChromakeyCards = Camera->CameraSettings.Chromakey.ChromakeyRenderTexture.ShowOnlyList;
+				FDisplayClusterConfigurationICVFX_ChromakeyRenderSettings* CameraChromakeyRenderSettings = Camera->CameraSettings.Chromakey.GetWritableChromakeyRenderSettings(GetStageSettings());
+				if (!CameraChromakeyRenderSettings)
+				{
+					// Updating the CK visibility list only for cameras using chromakey actor rendering
+					continue;
+				}
+
+				FDisplayClusterConfigurationICVFX_VisibilityList& ChromakeyCards = CameraChromakeyRenderSettings->ShowOnlyList;
 				ChromakeyCards.AutoAddedActors.Reset();
 				
 				for (const TSoftObjectPtr<AActor>& Actor : ChromakeyCards.Actors)
@@ -824,7 +831,14 @@ void ADisplayClusterRootActor::SetLightCardOwnership()
 				{
 					for (UDisplayClusterICVFXCameraComponent* Camera : ICVFXComponents)
 					{
-						FDisplayClusterConfigurationICVFX_VisibilityList& ChromakeyCards = Camera->CameraSettings.Chromakey.ChromakeyRenderTexture.ShowOnlyList;
+						FDisplayClusterConfigurationICVFX_ChromakeyRenderSettings* CameraChromakeyRenderSettings = Camera->CameraSettings.Chromakey.GetWritableChromakeyRenderSettings(GetStageSettings());
+						if (!CameraChromakeyRenderSettings)
+						{
+							// Updating the CK visibility list only for cameras using chromakey actor rendering
+							continue;
+						}
+
+						FDisplayClusterConfigurationICVFX_VisibilityList& ChromakeyCards = CameraChromakeyRenderSettings->ShowOnlyList;
 
 						const UDisplayClusterChromakeyCardStageActorComponent* ChromakeyStageActor = Cast<UDisplayClusterChromakeyCardStageActorComponent>(ChromakeyCardActor->GetStageActorComponent());
 						if (ChromakeyStageActor && ChromakeyStageActor->IsReferencedByICVFXCamera(Camera))
diff --git a/Engine/Plugins/Runtime/nDisplay/Source/DisplayCluster/Private/Render/Viewport/Configuration/DisplayClusterViewportConfigurationHelpers.cpp b/Engine/Plugins/Runtime/nDisplay/Source/DisplayCluster/Private/Render/Viewport/Configuration/DisplayClusterViewportConfigurationHelpers.cpp
--- a/Engine/Plugins/Runtime/nDisplay/Source/DisplayCluster/Private/Render/Viewport/Configuration/DisplayClusterViewportConfigurationHelpers.cpp
+++ b/Engine/Plugins/Runtime/nDisplay/Source/DisplayCluster/Private/Render/Viewport/Configuration/DisplayClusterViewportConfigurationHelpers.cpp
@@ -231,71 +231,13 @@ void FDisplayClusterViewportConfigurationHelpers::UpdateBaseViewportSetting(FDis
 		}
 	}
 
+
 	// FDisplayClusterConfigurationViewport_ICVFX property:
 	{
-		EDisplayClusterViewportICVFXFlags& TargetFlags = DstViewport.RenderSettingsICVFX.Flags;
-
-		if (InConfigurationViewport.ICVFX.bAllowICVFX)
-		{
-			EnumAddFlags(TargetFlags, EDisplayClusterViewportICVFXFlags::Enable);
-
-			EDisplayClusterConfigurationICVFX_OverrideCameraRenderMode CameraRenderMode = InConfigurationViewport.ICVFX.CameraRenderMode;
-
-			const FDisplayClusterConfigurationICVFX_StageSettings& StageSettings = RootActor.GetStageSettings();
-			if (InConfigurationViewport.ICVFX.bAllowInnerFrustum == false || StageSettings.bEnableInnerFrustums == false)
-			{
-				CameraRenderMode = EDisplayClusterConfigurationICVFX_OverrideCameraRenderMode::Disabled;
-			}
-
-			switch (CameraRenderMode)
-			{
-				// Disable camera frame render for this viewport
-			case EDisplayClusterConfigurationICVFX_OverrideCameraRenderMode::Disabled:
-				EnumAddFlags(TargetFlags, EDisplayClusterViewportICVFXFlags::DisableCamera | EDisplayClusterViewportICVFXFlags::DisableChromakey | EDisplayClusterViewportICVFXFlags::DisableChromakeyMarkers);
-				break;
-
-				// Disable chromakey render for this viewport
-			case EDisplayClusterConfigurationICVFX_OverrideCameraRenderMode::DisableChromakey:
-				EnumAddFlags(TargetFlags, EDisplayClusterViewportICVFXFlags::DisableChromakey | EDisplayClusterViewportICVFXFlags::DisableChromakeyMarkers);
-				break;
-
-				// Disable chromakey markers render for this viewport
-			case EDisplayClusterConfigurationICVFX_OverrideCameraRenderMode::DisableChromakeyMarkers:
-
-				EnumAddFlags(TargetFlags, EDisplayClusterViewportICVFXFlags::DisableChromakeyMarkers);
-				break;
+		const FDisplayClusterConfigurationICVFX_StageSettings& StageSettings = RootActor.GetStageSettings();
 
-			default:
-				// Use default rendering rules
-				break;
-			}
-
-			switch (InConfigurationViewport.ICVFX.LightcardRenderMode)
-			{
-				// Render incamera frame over lightcard for this viewport
-			case EDisplayClusterConfigurationICVFX_OverrideLightcardRenderMode::Over:
-
-				EnumAddFlags(TargetFlags, EDisplayClusterViewportICVFXFlags::OverrideLightcardMode);
-				DstViewport.RenderSettingsICVFX.ICVFX.LightCardMode = EDisplayClusterShaderParametersICVFX_LightCardRenderMode::Over;
-				break;
-
-				// Over lightcard over incamera frame  for this viewport
-			case EDisplayClusterConfigurationICVFX_OverrideLightcardRenderMode::Under:
-
-				EnumAddFlags(TargetFlags, EDisplayClusterViewportICVFXFlags::OverrideLightcardMode);
-				DstViewport.RenderSettingsICVFX.ICVFX.LightCardMode = EDisplayClusterShaderParametersICVFX_LightCardRenderMode::Under;
-				break;
-
-			case EDisplayClusterConfigurationICVFX_OverrideLightcardRenderMode::Disabled:
-
-				EnumAddFlags(TargetFlags, EDisplayClusterViewportICVFXFlags::DisableLightcard);
-				break;
-
-			default:
-				// Use default lightcard mode
-				break;
-			}
-		}
+		DstViewport.RenderSettingsICVFX.Flags = InConfigurationViewport.ICVFX.GetViewportICVFXFlags(StageSettings);
+		DstViewport.RenderSettingsICVFX.ICVFX.LightCardMode = InConfigurationViewport.ICVFX.GetLightCardRenderMode(StageSettings);
 	}
 }
 
diff --git a/Engine/Plugins/Runtime/nDisplay/Source/DisplayCluster/Private/Render/Viewport/Configuration/DisplayClusterViewportConfigurationHelpers_ICVFX.cpp b/Engine/Plugins/Runtime/nDisplay/Source/DisplayCluster/Private/Render/Viewport/Configuration/DisplayClusterViewportConfigurationHelpers_ICVFX.cpp
--- a/Engine/Plugins/Runtime/nDisplay/Source/DisplayCluster/Private/Render/Viewport/Configuration/DisplayClusterViewportConfigurationHelpers_ICVFX.cpp
+++ b/Engine/Plugins/Runtime/nDisplay/Source/DisplayCluster/Private/Render/Viewport/Configuration/DisplayClusterViewportConfigurationHelpers_ICVFX.cpp
@@ -58,16 +58,11 @@ static FAutoConsoleVariableRef CVarDisplayClusterPreviewEnableReuseViewportInClu
 ////////////////////////////////////////////////////////////////////////////////
 namespace DisplayClusterViewportConfigurationHelpers_ICVFX_Impl
 {
-	static float ImplGetFieldOfViewMultiplier(const FDisplayClusterConfigurationICVFX_CameraSettings& CameraSettings)
-	{
-		return CameraSettings.CustomFrustum.bEnable ? CameraSettings.CustomFrustum.FieldOfViewMultiplier : 1.f;
-	}
-
 	// Initialize camera policy with camera component and settings
-	static bool ImplUpdateCameraProjectionSettings(TSharedPtr<IDisplayClusterProjectionPolicy, ESPMode::ThreadSafe>& InOutCameraProjection, const FDisplayClusterConfigurationICVFX_CameraSettings& CameraSettings, UCameraComponent* const CameraComponent)
+	static bool ImplUpdateCameraProjectionSettings(TSharedPtr<IDisplayClusterProjectionPolicy, ESPMode::ThreadSafe>& InOutCameraProjection, ADisplayClusterRootActor& RootActor, const FDisplayClusterConfigurationICVFX_CameraSettings& CameraSettings, UCameraComponent* const CameraComponent)
 	{
 		FDisplayClusterProjectionCameraPolicySettings PolicyCameraSettings;
-		PolicyCameraSettings.FOVMultiplier = ImplGetFieldOfViewMultiplier(CameraSettings);
+		PolicyCameraSettings.FOVMultiplier = CameraSettings.GetCameraFieldOfViewMultiplier(RootActor.GetStageSettings());
 
 		// Lens correction
 		PolicyCameraSettings.FrustumRotation = CameraSettings.FrustumRotation;
@@ -444,64 +439,25 @@ FDisplayClusterShaderParameters_ICVFX::FCameraSettings FDisplayClusterViewportCo
 {
 	const USceneComponent* OriginComp = RootActor.GetRootComponent();
 	const FDisplayClusterConfigurationICVFX_CameraSettings& CameraSettings = InCameraComponent.GetCameraSettingsICVFX();
+	const FDisplayClusterConfigurationICVFX_StageSettings& StageSettings = RootActor.GetStageSettings();
 
 	FDisplayClusterShaderParameters_ICVFX::FCameraSettings Result;
 
 	Result.Resource.ViewportId = InCameraViewport.GetId();
 	Result.Local2WorldTransform = OriginComp->GetComponentTransform();
 
-	const float RealThicknessScaleValue = 0.1f;
-	Result.InnerCameraBorderColor = CameraSettings.Border.Color;
-	Result.InnerCameraBorderThickness = CameraSettings.Border.Enable ? CameraSettings.Border.Thickness * RealThicknessScaleValue : 0.0f;
-	Result.InnerCameraFrameAspectRatio = (float)CameraSettings.RenderSettings.CustomFrameSize.CustomWidth / (float)CameraSettings.RenderSettings.CustomFrameSize.CustomHeight;
+	// Get camera border settings
+	CameraSettings.GetCameraBorder(StageSettings, Result.InnerCameraBorderColor, Result.InnerCameraBorderThickness);
+	Result.InnerCameraFrameAspectRatio = CameraSettings.GetCameraFrameAspectRatio(StageSettings);
 
+	// Soft edges
+	Result.SoftEdge = CameraSettings.GetCameraSoftEdge(StageSettings);
+
+	// Rendering order for camera overlap
 	const FString InnerFrustumID = InCameraComponent.GetCameraUniqueId();
 	const int32 CameraRenderOrder = RootActor.GetInnerFrustumPriority(InnerFrustumID);
 	Result.RenderOrder = (CameraRenderOrder<0) ? CameraSettings.RenderSettings.RenderOrder : CameraRenderOrder;
 
-	const float FieldOfViewMultiplier = ImplGetFieldOfViewMultiplier(CameraSettings);
-
-	// softedge adjustments	
-	const float Overscan = (FieldOfViewMultiplier > 0) ? FieldOfViewMultiplier : 1;
-	
-	// remap values from 0-1 GUI range into acceptable 0.0 - 0.25 shader range
-	Result.SoftEdge.X = FMath::GetMappedRangeValueClamped(FVector2D(0.0, 1.0f), FVector2D(0.0, 0.25), CameraSettings.SoftEdge.Horizontal) / Overscan; // Left
-	Result.SoftEdge.Y = FMath::GetMappedRangeValueClamped(FVector2D(0.0, 1.0f), FVector2D(0.0, 0.25), CameraSettings.SoftEdge.Vertical) / Overscan; // Top
-
-	// ZW now used in other way
-	// Z for new parameter Feather
-	Result.SoftEdge.Z = CameraSettings.SoftEdge.Feather;
-
-	// default - percents
-	const float ConvertToPercent = 0.01f;
-	float Left = FDisplayClusterViewport_OverscanSettings::ClampPercent(CameraSettings.CustomFrustum.Left * ConvertToPercent);
-	float Right = FDisplayClusterViewport_OverscanSettings::ClampPercent(CameraSettings.CustomFrustum.Right * ConvertToPercent);
-	float Top = FDisplayClusterViewport_OverscanSettings::ClampPercent(CameraSettings.CustomFrustum.Top * ConvertToPercent);
-	float Bottom = FDisplayClusterViewport_OverscanSettings::ClampPercent(CameraSettings.CustomFrustum.Bottom * ConvertToPercent);
-
-	if (CameraSettings.CustomFrustum.Mode == EDisplayClusterConfigurationViewportCustomFrustumMode::Pixels)
-	{
-		float FrameWidth = RootActor.GetStageSettings().DefaultFrameSize.Width * CameraSettings.BufferRatio;
-		float FrameHeight = RootActor.GetStageSettings().DefaultFrameSize.Height * CameraSettings.BufferRatio;
-
-		if (CameraSettings.RenderSettings.CustomFrameSize.bUseCustomSize)
-		{
-			FrameWidth = CameraSettings.RenderSettings.CustomFrameSize.CustomWidth * CameraSettings.BufferRatio;
-			FrameHeight = CameraSettings.RenderSettings.CustomFrameSize.CustomHeight * CameraSettings.BufferRatio;
-		}
-
-		Left = FDisplayClusterViewport_OverscanSettings::ClampPercent(CameraSettings.CustomFrustum.Left / FrameWidth);
-		Right = FDisplayClusterViewport_OverscanSettings::ClampPercent(CameraSettings.CustomFrustum.Right / FrameWidth);
-		Top = FDisplayClusterViewport_OverscanSettings::ClampPercent(CameraSettings.CustomFrustum.Top / FrameHeight);
-		Bottom = FDisplayClusterViewport_OverscanSettings::ClampPercent(CameraSettings.CustomFrustum.Bottom / FrameHeight);
-	}
-
-	// recalculate soft edge related offsets based on frustum
-	Result.SoftEdge.X /= (1 + Left + Right);
-	Result.SoftEdge.Y /= (1 + Top + Bottom);
-	Result.SoftEdge.Z /= (1 + Left + Right);
-	Result.SoftEdge.W /= (1 + Top + Bottom);
-
 	return Result;
 }
 
@@ -517,7 +473,7 @@ bool FDisplayClusterViewportConfigurationHelpers_ICVFX::GetCameraContext(ADispla
 	}
 
 	// Initialize camera policy with camera component and settings
-	if (!ImplUpdateCameraProjectionSettings(CameraProjectionPolicy, CameraSettings, InCameraComponent.GetCameraComponent()))
+	if (!ImplUpdateCameraProjectionSettings(CameraProjectionPolicy, RootActor, CameraSettings, InCameraComponent.GetCameraComponent()))
 	{
 		return false;
 	}
@@ -548,14 +504,13 @@ void FDisplayClusterViewportConfigurationHelpers_ICVFX::UpdateCameraViewportSett
 	DstViewport.RenderSettings.bVisible = false;
 
 	// use framecolor instead of viewport rendering
-	const FDisplayClusterConfigurationICVFX_ChromakeySettings& ChromakeySettings = CameraSettings.Chromakey;
-	if (ChromakeySettings.bEnable && ChromakeySettings.ChromakeyType == EDisplayClusterConfigurationICVFX_ChromakeyType::InnerFrustum)
+	if (CameraSettings.Chromakey.GetChromakeyType(StageSettings) == EDisplayClusterShaderParametersICVFX_ChromakeySource::FrameColor)
 	{
 		DstViewport.RenderSettings.bSkipRendering = true;
 	}
 
 	// Update camera viewport projection policy settings
-	ImplUpdateCameraProjectionSettings(DstViewport.ProjectionPolicy, CameraSettings, InCameraComponent.GetCameraComponent());
+	ImplUpdateCameraProjectionSettings(DstViewport.ProjectionPolicy, RootActor, CameraSettings, InCameraComponent.GetCameraComponent());
 
 	// Update OCIO for Camera Viewport
 	FDisplayClusterViewportConfigurationHelpers_OpenColorIO::UpdateCameraViewport(DstViewport, RootActor, InCameraComponent);
@@ -568,18 +523,7 @@ void FDisplayClusterViewportConfigurationHelpers_ICVFX::UpdateCameraViewportSett
 	DstViewport.RenderSettings.CameraId.Empty();
 
 	// UDisplayClusterConfigurationICVFX_CameraRenderSettings
-	FIntPoint DesiredSize(0);
-	// Camera viewport frame size:
-	if (CameraSettings.RenderSettings.CustomFrameSize.bUseCustomSize)
-	{
-		DesiredSize.X = CameraSettings.RenderSettings.CustomFrameSize.CustomWidth;
-		DesiredSize.Y = CameraSettings.RenderSettings.CustomFrameSize.CustomHeight;
-	}
-	else
-	{
-		DesiredSize.X = StageSettings.DefaultFrameSize.Width;
-		DesiredSize.Y = StageSettings.DefaultFrameSize.Height;
-	}
+	const FIntPoint DesiredSize = CameraSettings.GetCameraFrameSize(StageSettings);
 
 	DstViewport.RenderSettings.Rect = DstViewport.GetValidRect(FIntRect(FIntPoint(0, 0), DesiredSize), TEXT("Configuration Camera Frame Size"));
 
@@ -603,7 +547,12 @@ void FDisplayClusterViewportConfigurationHelpers_ICVFX::UpdateCameraViewportSett
 
 	// Support inner camera custom frustum
 	FDisplayClusterViewportConfigurationHelpers_ICVFX::UpdateCameraCustomFrustum(DstViewport, CameraSettings.CustomFrustum);
-	UpdateCameraViewportBufferRatio(DstViewport, CameraSettings);
+
+	// Set RenderTargetAdaptRatio
+	DstViewport.RenderSettings.RenderTargetAdaptRatio = CameraSettings.GetCameraAdaptResolutionRatio(StageSettings);
+
+	// Set viewport buffer ratio
+	DstViewport.SetViewportBufferRatio(CameraSettings.GetCameraBufferRatio(StageSettings));
 
 	// Set media related configuration (runtime only for now)
 	if (IDisplayCluster::Get().GetOperationMode() == EDisplayClusterOperationMode::Cluster)
@@ -634,8 +583,6 @@ void FDisplayClusterViewportConfigurationHelpers_ICVFX::UpdateChromakeyViewportS
 	const FDisplayClusterConfigurationICVFX_CameraSettings& CameraSettings = InCameraComponent.GetCameraSettingsICVFX();
 	const FDisplayClusterConfigurationICVFX_ChromakeySettings& ChromakeySettings = CameraSettings.Chromakey;
 
-	check(ChromakeySettings.bEnable && ChromakeySettings.ChromakeyType == EDisplayClusterConfigurationICVFX_ChromakeyType::CustomChromakey);
-
 	// Reset runtime flags from prev frame:
 	DstViewport.ResetRuntimeParameters();
 
@@ -646,22 +593,38 @@ void FDisplayClusterViewportConfigurationHelpers_ICVFX::UpdateChromakeyViewportS
 	DstViewport.RenderSettings.CaptureMode = EDisplayClusterViewportCaptureMode::Chromakey;
 
 	// UDisplayClusterConfigurationICVFX_ChromakeyRenderSettings
-	const FDisplayClusterConfigurationICVFX_ChromakeyRenderSettings& InRenderSettings = ChromakeySettings.ChromakeyRenderTexture;
+	if (const FDisplayClusterConfigurationICVFX_ChromakeyRenderSettings* ChromakeyRenderSettings = ChromakeySettings.GetChromakeyRenderSettings(StageSettings))
 	{
-		FDisplayClusterViewportConfigurationHelpers::UpdateViewportSetting_Override(DstViewport, InRenderSettings.Replace);
-		FDisplayClusterViewportConfigurationHelpers::UpdateViewportSetting_PostprocessBlur(DstViewport, InRenderSettings.PostprocessBlur);
-		FDisplayClusterViewportConfigurationHelpers::UpdateViewportSetting_GenerateMips(DstViewport, InRenderSettings.GenerateMips);
+		FDisplayClusterViewportConfigurationHelpers::UpdateViewportSetting_Override(DstViewport, ChromakeyRenderSettings->Replace);
+		FDisplayClusterViewportConfigurationHelpers::UpdateViewportSetting_PostprocessBlur(DstViewport, ChromakeyRenderSettings->PostprocessBlur);
+		FDisplayClusterViewportConfigurationHelpers::UpdateViewportSetting_GenerateMips(DstViewport, ChromakeyRenderSettings->GenerateMips);
 
 		// Update visibility settings only for rendered viewports
 		if (!DstViewport.PostRenderSettings.Replace.IsEnabled())
 		{
-			check(FDisplayClusterViewportConfigurationHelpers_Visibility::IsVisibilityListValid(InRenderSettings.ShowOnlyList));
+			check(ChromakeyRenderSettings->ShowOnlyList.IsVisibilityListValid());
 
-			FDisplayClusterViewportConfigurationHelpers_Visibility::UpdateShowOnlyList(DstViewport, RootActor, InRenderSettings.ShowOnlyList);
+			FDisplayClusterViewportConfigurationHelpers_Visibility::UpdateShowOnlyList(DstViewport, RootActor, ChromakeyRenderSettings->ShowOnlyList);
+		}
+
+		FDisplayClusterViewportConfigurationHelpers::UpdateViewportSetting_OverlayRenderSettings(DstViewport, ChromakeyRenderSettings->AdvancedRenderSettings);
+
+		// Support custom overlay size
+		if (ChromakeyRenderSettings->CustomSize.bUseCustomSize)
+		{
+			FIntPoint DesiredSize;
+			DesiredSize.X = ChromakeyRenderSettings->CustomSize.CustomWidth;
+			DesiredSize.Y = ChromakeyRenderSettings->CustomSize.CustomHeight;
+
+			DstViewport.RenderSettings.Rect = DstViewport.GetValidRect(FIntRect(FIntPoint(0, 0), DesiredSize), TEXT("Configuration custom chromakey Frame Size"));
 		}
-	}
 
-	FDisplayClusterViewportConfigurationHelpers::UpdateViewportSetting_OverlayRenderSettings(DstViewport, InRenderSettings.AdvancedRenderSettings);
+		// Debug: override the texture of the target viewport from this chromakeyRTT
+		if (ChromakeyRenderSettings->bReplaceCameraViewport)
+		{
+			InCameraViewport.RenderSettings.SetViewportOverride(DstViewport.GetId());
+		}
+	}
 
 	// Update OCIO for Chromakey Viewport
 	// Note: Chromakey OCIO is temporarily disabled
@@ -672,139 +635,88 @@ void FDisplayClusterViewportConfigurationHelpers_ICVFX::UpdateChromakeyViewportS
 
 	// Attach to parent viewport
 	DstViewport.RenderSettings.AssignParentViewport(InCameraViewport.GetId(), InCameraViewport.RenderSettings);
-
-	// Support custom overlay size
-	if (ChromakeySettings.ChromakeyRenderTexture.CustomSize.bUseCustomSize)
-	{
-		FIntPoint DesiredSize;
-		DesiredSize.X = ChromakeySettings.ChromakeyRenderTexture.CustomSize.CustomWidth;
-		DesiredSize.Y = ChromakeySettings.ChromakeyRenderTexture.CustomSize.CustomHeight;
-
-		DstViewport.RenderSettings.Rect = DstViewport.GetValidRect(FIntRect(FIntPoint(0, 0), DesiredSize), TEXT("Configuration custom chromakey Frame Size"));
-	}
-
-	// Debug: override the texture of the target viewport from this chromakeyRTT
-	if (ChromakeySettings.ChromakeyRenderTexture.bReplaceCameraViewport)
-	{
-		InCameraViewport.RenderSettings.SetViewportOverride(DstViewport.GetId());
-	}
 }
 
-void FDisplayClusterViewportConfigurationHelpers_ICVFX::UpdateCameraSettings_Chromakey(FDisplayClusterShaderParameters_ICVFX::FCameraSettings& InOutCameraSettings, const FDisplayClusterConfigurationICVFX_GlobalChromakeySettings& InGlobalChromakeySettings, const FDisplayClusterConfigurationICVFX_ChromakeySettings& InChromakeySettings, FDisplayClusterViewport* InChromakeyViewport)
+void FDisplayClusterViewportConfigurationHelpers_ICVFX::UpdateCameraSettings_Chromakey(FDisplayClusterShaderParameters_ICVFX::FCameraSettings& InOutCameraSettings, const FDisplayClusterConfigurationICVFX_StageSettings& InStageSettings, const FDisplayClusterConfigurationICVFX_CameraSettings& InCameraSettings, bool bEnableChromakeyMarkers, const FString& InChromakeyViewportId)
 {
-	if (InChromakeySettings.ChromakeySettingsSource == EDisplayClusterConfigurationICVFX_ChromakeySettingsSource::Viewport)
-	{
-		InOutCameraSettings.ChromakeyColor = InGlobalChromakeySettings.ChromakeyColor;
-	}
-	else
-	{
-		InOutCameraSettings.ChromakeyColor = InChromakeySettings.ChromakeyColor;
-	}
+	// Set chromakey color
+	InOutCameraSettings.ChromakeyColor = InCameraSettings.Chromakey.GetChromakeyColor(InStageSettings);
 
-	if (InChromakeySettings.bEnable)
+	// Set chromakey source
+	switch (InOutCameraSettings.ChromakeySource = InCameraSettings.Chromakey.GetChromakeyType(InStageSettings))
 	{
-		if (InChromakeySettings.ChromakeyType == EDisplayClusterConfigurationICVFX_ChromakeyType::CustomChromakey)
+	case EDisplayClusterShaderParametersICVFX_ChromakeySource::ChromakeyLayers:
+		if (InChromakeyViewportId.IsEmpty())
 		{
-			if (InChromakeySettings.ChromakeyRenderTexture.bReplaceCameraViewport)
+			// Disable chromakey: CK viewport required
+			InOutCameraSettings.ChromakeySource = EDisplayClusterShaderParametersICVFX_ChromakeySource::Disabled;
+			return;
+		}
+
+		if (const FDisplayClusterConfigurationICVFX_ChromakeyRenderSettings* ChromakeyRenderSettings = InCameraSettings.Chromakey.GetChromakeyRenderSettings(InStageSettings))
+		{
+			if (ChromakeyRenderSettings->bReplaceCameraViewport)
 			{
+				// Do not show Chromakey layers in this in-camera viewport, because they will be replaced by the Chromakey image (for debugging purposes).
 				InOutCameraSettings.ChromakeySource = EDisplayClusterShaderParametersICVFX_ChromakeySource::Disabled;
 				return;
 			}
-			else
-			{
-				check(InChromakeyViewport);
-
-				InOutCameraSettings.ChromakeySource = EDisplayClusterShaderParametersICVFX_ChromakeySource::ChromakeyLayers;
-				InOutCameraSettings.Chromakey.ViewportId = InChromakeyViewport->GetId();
-			}
-		}
-		else
-		{
-			InOutCameraSettings.ChromakeySource = EDisplayClusterShaderParametersICVFX_ChromakeySource::FrameColor;
 		}
-	}
-	else
-	{
-		InOutCameraSettings.ChromakeySource = EDisplayClusterShaderParametersICVFX_ChromakeySource::Disabled;
-	}
-}
 
-void FDisplayClusterViewportConfigurationHelpers_ICVFX::UpdateCameraSettings_ChromakeyMarkers(FDisplayClusterShaderParameters_ICVFX::FCameraSettings& InOutCameraSettings, const FDisplayClusterConfigurationICVFX_ChromakeyMarkers& InChromakeyMarkers)
-{
-	InOutCameraSettings.ChromakeMarkerTextureRHI.SafeRelease();
+		// Set the chromakey viewport name in the ICVFX shader params for current camera
+		InOutCameraSettings.Chromakey.ViewportId = InChromakeyViewportId;
 
-	if (InChromakeyMarkers.bEnable && InChromakeyMarkers.MarkerTileRGBA != nullptr)
-	{
-		InOutCameraSettings.ChromakeyMarkersColor = InChromakeyMarkers.MarkerColor;
-		InOutCameraSettings.ChromakeyMarkersScale = InChromakeyMarkers.MarkerSizeScale;
-		InOutCameraSettings.ChromakeyMarkersDistance = InChromakeyMarkers.MarkerTileDistance;
-		InOutCameraSettings.ChromakeyMarkersOffset = InChromakeyMarkers.MarkerTileOffset;
+		break;
 
-		// Assign texture RHI ref
-		FTextureResource* MarkersResource = InChromakeyMarkers.MarkerTileRGBA->GetResource();
-		if (MarkersResource)
-		{
-			InOutCameraSettings.ChromakeMarkerTextureRHI = MarkersResource->TextureRHI;
-		}
-	}
-}
+	case EDisplayClusterShaderParametersICVFX_ChromakeySource::Disabled:
+		return;
 
-void FDisplayClusterViewportConfigurationHelpers_ICVFX::UpdateCameraSettings_OverlapChromakeyMarkers(FDisplayClusterShaderParameters_ICVFX::FCameraSettings& InOutCameraSettings, const FDisplayClusterConfigurationICVFX_ChromakeyMarkers& InChromakeyMarkers)
-{
-	InOutCameraSettings.OverlapChromakeyMarkerTextureRHI.SafeRelease();
+	default:
+		break;
+	}
 
-	if (InChromakeyMarkers.bEnable && InChromakeyMarkers.MarkerTileRGBA != nullptr)
+	//Setup chromakey markers
+	if (bEnableChromakeyMarkers)
 	{
-		InOutCameraSettings.OverlapChromakeyMarkersColor = InChromakeyMarkers.MarkerColor;
-		InOutCameraSettings.OverlapChromakeyMarkersScale = InChromakeyMarkers.MarkerSizeScale;
-		InOutCameraSettings.OverlapChromakeyMarkersDistance = InChromakeyMarkers.MarkerTileDistance;
-		InOutCameraSettings.OverlapChromakeyMarkersOffset = InChromakeyMarkers.MarkerTileOffset;
-
-		// Assign texture RHI ref
-		FTextureResource* MarkersResource = InChromakeyMarkers.MarkerTileRGBA->GetResource();
-		if (MarkersResource)
+		// Setup chromakey markers
+		if (const FDisplayClusterConfigurationICVFX_ChromakeyMarkers* ChromakeyMarkers = InCameraSettings.Chromakey.GetChromakeyMarkers(InStageSettings))
 		{
-			InOutCameraSettings.OverlapChromakeyMarkerTextureRHI = MarkersResource->TextureRHI;
+			InOutCameraSettings.ChromakeyMarkersColor = ChromakeyMarkers->MarkerColor;
+			InOutCameraSettings.ChromakeyMarkersScale = ChromakeyMarkers->MarkerSizeScale;
+			InOutCameraSettings.ChromakeyMarkersDistance = ChromakeyMarkers->MarkerTileDistance;
+			InOutCameraSettings.ChromakeyMarkersOffset = ChromakeyMarkers->MarkerTileOffset;
+
+			// Assign texture RHI ref
+			if (FTextureResource* MarkersResource = ChromakeyMarkers->MarkerTileRGBA->GetResource())
+			{
+				InOutCameraSettings.ChromakeMarkerTextureRHI = MarkersResource->TextureRHI;
+			}
 		}
 	}
 }
 
-bool FDisplayClusterViewportConfigurationHelpers_ICVFX::IsShouldUseLightcard(const FDisplayClusterConfigurationICVFX_LightcardSettings& InLightcardSettings)
+void FDisplayClusterViewportConfigurationHelpers_ICVFX::UpdateCameraSettings_OverlapChromakey(FDisplayClusterShaderParameters_ICVFX::FCameraSettings& InOutCameraSettings, const FDisplayClusterConfigurationICVFX_StageSettings& InStageSettings, const FDisplayClusterConfigurationICVFX_CameraSettings& InCameraSettings, bool bEnableChromakeyMarkers)
 {
-	if (InLightcardSettings.bEnable == false)
-	{
-		// dont use lightcard if disabled
-		return false;
-	}
+	// Setup overlap chromakey color
+	InOutCameraSettings.OverlapChromakeyColor = InCameraSettings.Chromakey.GetOverlapChromakeyColor(InStageSettings);
 
-	if (InLightcardSettings.RenderSettings.Replace.bAllowReplace)
+	// Setup overlap chromakey markers
+	if (bEnableChromakeyMarkers)
 	{
-		if (InLightcardSettings.RenderSettings.Replace.SourceTexture == nullptr)
+		if (const FDisplayClusterConfigurationICVFX_ChromakeyMarkers* OverlapChromakeyMarkers = InCameraSettings.Chromakey.GetOverlapChromakeyMarkers(InStageSettings))
 		{
-			// LightcardSettings.Override require source texture.
-			return false;
-		}
+			InOutCameraSettings.OverlapChromakeyMarkersColor = OverlapChromakeyMarkers->MarkerColor;
+			InOutCameraSettings.OverlapChromakeyMarkersScale = OverlapChromakeyMarkers->MarkerSizeScale;
+			InOutCameraSettings.OverlapChromakeyMarkersDistance = OverlapChromakeyMarkers->MarkerTileDistance;
+			InOutCameraSettings.OverlapChromakeyMarkersOffset = OverlapChromakeyMarkers->MarkerTileOffset;
 
-		return true;
-	}
-
-	// Lightcard require layers for render
-	return FDisplayClusterViewportConfigurationHelpers_Visibility::IsVisibilityListValid(InLightcardSettings.ShowOnlyList);
-}
-
-bool FDisplayClusterViewportConfigurationHelpers_ICVFX::IsShouldUseUVLightcard(FDisplayClusterViewportManager& InViewportManager, const FDisplayClusterConfigurationICVFX_LightcardSettings& InLightcardSettings)
-{
-	if (InLightcardSettings.bEnable)
-	{
-		const TSharedPtr<FDisplayClusterViewportLightCardManager, ESPMode::ThreadSafe> LightCardManager = InViewportManager.GetLightCardManager();
-		if (LightCardManager.IsValid() && LightCardManager->IsUVLightCardEnabled())
-		{
-			return FDisplayClusterViewportConfigurationHelpers_Visibility::IsVisibilityListValid(InLightcardSettings.ShowOnlyList);
+			// Assign texture RHI ref
+			if (FTextureResource* OverlapMarkersResource = OverlapChromakeyMarkers->MarkerTileRGBA->GetResource())
+			{
+				InOutCameraSettings.OverlapChromakeyMarkerTextureRHI = OverlapMarkersResource->TextureRHI;
+			}
 		}
 	}
-
-	// dont use lightcard if disabled
-	return false;
 }
 
 void FDisplayClusterViewportConfigurationHelpers_ICVFX::UpdateLightcardViewportSetting(FDisplayClusterViewport& DstViewport, FDisplayClusterViewport& BaseViewport, ADisplayClusterRootActor& RootActor)
@@ -841,7 +753,7 @@ void FDisplayClusterViewportConfigurationHelpers_ICVFX::UpdateLightcardViewportS
 		// Update visibility settings only for rendered viewports
 		if (!DstViewport.PostRenderSettings.Replace.IsEnabled())
 		{
-			check(FDisplayClusterViewportConfigurationHelpers_Visibility::IsVisibilityListValid(LightcardSettings.ShowOnlyList));
+			check(LightcardSettings.ShowOnlyList.IsVisibilityListValid());
 
 			FDisplayClusterViewportConfigurationHelpers_Visibility::UpdateShowOnlyList(DstViewport, RootActor, LightcardSettings.ShowOnlyList);
 		}
@@ -852,22 +764,6 @@ void FDisplayClusterViewportConfigurationHelpers_ICVFX::UpdateLightcardViewportS
 	// Attach to parent viewport
 	DstViewport.RenderSettings.AssignParentViewport(BaseViewport.GetId(), BaseViewport.RenderSettings);
 
-	// Global lighcard rendering mode
-	if (!EnumHasAllFlags(BaseViewport.RenderSettingsICVFX.Flags, EDisplayClusterViewportICVFXFlags::OverrideLightcardMode))
-	{
-		// Use global lightcard blending mode
-		switch (LightcardSettings.Blendingmode)
-		{
-		default:
-		case EDisplayClusterConfigurationICVFX_LightcardRenderMode::Over:
-			BaseViewport.RenderSettingsICVFX.ICVFX.LightCardMode = EDisplayClusterShaderParametersICVFX_LightCardRenderMode::Over;
-			break;
-		case EDisplayClusterConfigurationICVFX_LightcardRenderMode::Under:
-			BaseViewport.RenderSettingsICVFX.ICVFX.LightCardMode = EDisplayClusterShaderParametersICVFX_LightCardRenderMode::Under;
-			break;
-		};
-	}
-
 	if (EnumHasAnyFlags(DstViewport.RenderSettingsICVFX.RuntimeFlags, EDisplayClusterViewportRuntimeICVFXFlags::Lightcard))
 	{
 		// Debug: override the texture of the target viewport from this lightcard RTT
@@ -922,23 +818,3 @@ void FDisplayClusterViewportConfigurationHelpers_ICVFX::UpdateCameraCustomFrustu
 		DstViewport.CustomFrustumRendering.Set(CustomFrustumSettings);
 	}
 };
-
-void FDisplayClusterViewportConfigurationHelpers_ICVFX::UpdateCameraViewportBufferRatio(FDisplayClusterViewport& DstViewport, const FDisplayClusterConfigurationICVFX_CameraSettings& CameraSettings)
-{
-	float BufferRatio = CameraSettings.BufferRatio;
-
-	if (CameraSettings.CustomFrustum.bAdaptResolution)
-	{
-		const float FieldOfViewMultiplier = ImplGetFieldOfViewMultiplier(CameraSettings);
-
-		// adapt resolution should work as a shortcut to improve rendering quality
-		DstViewport.RenderSettings.RenderTargetAdaptRatio = FieldOfViewMultiplier;
-	}
-	else
-	{
-		// Don't use an adaptive resolution multiplier
-		DstViewport.RenderSettings.RenderTargetAdaptRatio = 1.f;
-	}
-
-	DstViewport.SetViewportBufferRatio(BufferRatio);
-}
diff --git a/Engine/Plugins/Runtime/nDisplay/Source/DisplayCluster/Private/Render/Viewport/Configuration/DisplayClusterViewportConfigurationHelpers_ICVFX.h b/Engine/Plugins/Runtime/nDisplay/Source/DisplayCluster/Private/Render/Viewport/Configuration/DisplayClusterViewportConfigurationHelpers_ICVFX.h
--- a/Engine/Plugins/Runtime/nDisplay/Source/DisplayCluster/Private/Render/Viewport/Configuration/DisplayClusterViewportConfigurationHelpers_ICVFX.h
+++ b/Engine/Plugins/Runtime/nDisplay/Source/DisplayCluster/Private/Render/Viewport/Configuration/DisplayClusterViewportConfigurationHelpers_ICVFX.h
@@ -13,8 +13,8 @@ class FDisplayClusterViewportManager;
 class ADisplayClusterRootActor;
 class UDisplayClusterICVFXCameraComponent;
 
+struct FDisplayClusterConfigurationICVFX_StageSettings;
 struct FDisplayClusterConfigurationICVFX_ChromakeySettings;
-struct FDisplayClusterConfigurationICVFX_GlobalChromakeySettings;
 struct FDisplayClusterConfigurationICVFX_ChromakeyMarkers;
 struct FDisplayClusterConfigurationICVFX_LightcardSettings;
 struct FDisplayClusterConfigurationICVFX_CameraCustomFrustum;
@@ -55,16 +55,6 @@ public:
 	 */
 	static void ReuseUVLightCardViewportWithinClusterNode(FDisplayClusterViewport& InUVLightCardViewport);
 
-	/* Returns true if the use of the LightCard viewport is allowed in InLightcardSettings.
-	 * These viewports also require visible LightCard actors to render.
-	 */
-	static bool IsShouldUseLightcard(const FDisplayClusterConfigurationICVFX_LightcardSettings& InLightcardSettings);
-
-	/* Returns true if the use of the UVLightCard viewport is allowed in InLightcardSettings.
-	 * These viewports also require visible UVLightCard actors to render.
-	 */
-	static bool IsShouldUseUVLightcard(FDisplayClusterViewportManager& InViewportManager, const FDisplayClusterConfigurationICVFX_LightcardSettings& InLightcardSettings);
-
 #if WITH_EDITOR
 	static TArray<TSharedPtr<FDisplayClusterViewport, ESPMode::ThreadSafe>> PreviewGetRenderedInCameraViewports(ADisplayClusterRootActor& RootActor, UDisplayClusterICVFXCameraComponent& InCameraComponent, bool bGetChromakey = false);
 	static void PreviewReuseInnerFrustumViewportWithinClusterNodes_Editor(FDisplayClusterViewport& InCameraViewport, ADisplayClusterRootActor& RootActor, UDisplayClusterICVFXCameraComponent& InCameraComponent);
@@ -81,16 +71,10 @@ public:
 	static FDisplayClusterShaderParameters_ICVFX::FCameraSettings GetShaderParametersCameraSettings(const FDisplayClusterViewport& InCameraViewport, ADisplayClusterRootActor& RootActor, UDisplayClusterICVFXCameraComponent& InCameraComponent);
 	
 	/** Configures the camera's chromakey render settings to match the specified configuration settings  */
-	static void UpdateCameraSettings_Chromakey(FDisplayClusterShaderParameters_ICVFX::FCameraSettings& InOutCameraSettings, const FDisplayClusterConfigurationICVFX_GlobalChromakeySettings& InGlobalChromakeySettings, const FDisplayClusterConfigurationICVFX_ChromakeySettings& InChromakeySettings, FDisplayClusterViewport* InChromakeyViewport);
-
-	/** Configures the camera's chromakey marker render settings to match the specified configuration settings */
-	static void UpdateCameraSettings_ChromakeyMarkers(FDisplayClusterShaderParameters_ICVFX::FCameraSettings& InOutCameraSettings, const FDisplayClusterConfigurationICVFX_ChromakeyMarkers& InChromakeyMarkers);
-
-	/** Configures the camera's overlap chromakey marker render settings to match the specified configuration settings */
-	static void UpdateCameraSettings_OverlapChromakeyMarkers(FDisplayClusterShaderParameters_ICVFX::FCameraSettings& InOutCameraSettings, const FDisplayClusterConfigurationICVFX_ChromakeyMarkers& InChromakeyMarkers);
+	static void UpdateCameraSettings_Chromakey(FDisplayClusterShaderParameters_ICVFX::FCameraSettings& InOutCameraSettings, const FDisplayClusterConfigurationICVFX_StageSettings& InStageSettings, const FDisplayClusterConfigurationICVFX_CameraSettings& InCameraSettings, bool bEnableChromakeyMarkers, const FString& InChromakeyViewportId);
+	static void UpdateCameraSettings_OverlapChromakey(FDisplayClusterShaderParameters_ICVFX::FCameraSettings& InOutCameraSettings, const FDisplayClusterConfigurationICVFX_StageSettings& InStageSettings, const FDisplayClusterConfigurationICVFX_CameraSettings& InCameraSettings, bool bEnableChromakeyMarkers);
 
 	static void UpdateCameraCustomFrustum(FDisplayClusterViewport& DstViewport, const FDisplayClusterConfigurationICVFX_CameraCustomFrustum& InCameraCustomFrustum);
-	static void UpdateCameraViewportBufferRatio(FDisplayClusterViewport& DstViewport, const FDisplayClusterConfigurationICVFX_CameraSettings& CameraSettings);
 
 	static FDisplayClusterViewport* ImplFindViewport(ADisplayClusterRootActor& RootActor, const FString& InViewportId, const FString& InResourceId);
 private:
diff --git a/Engine/Plugins/Runtime/nDisplay/Source/DisplayCluster/Private/Render/Viewport/Configuration/DisplayClusterViewportConfigurationHelpers_Visibility.cpp b/Engine/Plugins/Runtime/nDisplay/Source/DisplayCluster/Private/Render/Viewport/Configuration/DisplayClusterViewportConfigurationHelpers_Visibility.cpp
--- a/Engine/Plugins/Runtime/nDisplay/Source/DisplayCluster/Private/Render/Viewport/Configuration/DisplayClusterViewportConfigurationHelpers_Visibility.cpp
+++ b/Engine/Plugins/Runtime/nDisplay/Source/DisplayCluster/Private/Render/Viewport/Configuration/DisplayClusterViewportConfigurationHelpers_Visibility.cpp
@@ -52,43 +52,6 @@ static void ImplCollectVisibility(ADisplayClusterRootActor& RootActor, const FDi
 	}
 }
 
-bool FDisplayClusterViewportConfigurationHelpers_Visibility::IsVisibilityListValid(const FDisplayClusterConfigurationICVFX_VisibilityList& InVisibilityList)
-{
-	for (const FString& ComponentNameIt : InVisibilityList.RootActorComponentNames)
-	{
-		if (!ComponentNameIt.IsEmpty())
-		{
-			return true;
-		}
-	}
-
-	for (const TSoftObjectPtr<AActor>& ActorSOPtrIt : InVisibilityList.Actors)
-	{
-		if (ActorSOPtrIt.IsValid())
-		{
-			return true;
-		}
-	}
-
-	for (const FActorLayer& ActorLayerIt : InVisibilityList.ActorLayers)
-	{
-		if (!ActorLayerIt.Name.IsNone())
-		{
-			return true;
-		}
-	}
-
-	for (const TSoftObjectPtr<AActor>& AutoAddedActor : InVisibilityList.AutoAddedActors)
-	{
-		if (AutoAddedActor.IsValid())
-		{
-			return true;
-		}
-	}
-	
-	return false;
-}
-
 void FDisplayClusterViewportConfigurationHelpers_Visibility::UpdateShowOnlyList(FDisplayClusterViewport& DstViewport, ADisplayClusterRootActor& RootActor, const FDisplayClusterConfigurationICVFX_VisibilityList& InVisibilityList)
 {
 	TSet<FPrimitiveComponentId> AdditionalComponentsList;
@@ -129,8 +92,10 @@ void FDisplayClusterViewportConfigurationHelpers_Visibility::UpdateHideList_ICVF
 		{
 			if (ICVFXCameraIt)
 			{
-				const FDisplayClusterConfigurationICVFX_CameraSettings& CameraSettings = ICVFXCameraIt->GetCameraSettingsICVFX();
-				ImplCollectVisibility(InRootActor, CameraSettings.Chromakey.ChromakeyRenderTexture.ShowOnlyList, ActorLayerNames, AdditionalComponentsList);
+				if (const FDisplayClusterConfigurationICVFX_ChromakeyRenderSettings* ChromakeyRenderSettings = ICVFXCameraIt->GetCameraSettingsICVFX().Chromakey.GetChromakeyRenderSettings(StageSettings))
+				{
+					ImplCollectVisibility(InRootActor, ChromakeyRenderSettings->ShowOnlyList, ActorLayerNames, AdditionalComponentsList);
+				}
 			}
 		}
 
diff --git a/Engine/Plugins/Runtime/nDisplay/Source/DisplayCluster/Private/Render/Viewport/Configuration/DisplayClusterViewportConfigurationHelpers_Visibility.h b/Engine/Plugins/Runtime/nDisplay/Source/DisplayCluster/Private/Render/Viewport/Configuration/DisplayClusterViewportConfigurationHelpers_Visibility.h
--- a/Engine/Plugins/Runtime/nDisplay/Source/DisplayCluster/Private/Render/Viewport/Configuration/DisplayClusterViewportConfigurationHelpers_Visibility.h
+++ b/Engine/Plugins/Runtime/nDisplay/Source/DisplayCluster/Private/Render/Viewport/Configuration/DisplayClusterViewportConfigurationHelpers_Visibility.h
@@ -12,9 +12,6 @@ struct FDisplayClusterConfigurationICVFX_VisibilityList;
 class FDisplayClusterViewportConfigurationHelpers_Visibility
 {
 public:
-	// Return true, if visibility list not empty
-	static bool IsVisibilityListValid(const FDisplayClusterConfigurationICVFX_VisibilityList& InVisibilityList);
-
 	// Update ShowOnly list for DstViewport
 	static void UpdateShowOnlyList(FDisplayClusterViewport& DstViewport, ADisplayClusterRootActor& RootActor, const FDisplayClusterConfigurationICVFX_VisibilityList& InVisibilityList);
 
diff --git a/Engine/Plugins/Runtime/nDisplay/Source/DisplayCluster/Private/Render/Viewport/Configuration/DisplayClusterViewportConfigurationICVFX.cpp b/Engine/Plugins/Runtime/nDisplay/Source/DisplayCluster/Private/Render/Viewport/Configuration/DisplayClusterViewportConfigurationICVFX.cpp
--- a/Engine/Plugins/Runtime/nDisplay/Source/DisplayCluster/Private/Render/Viewport/Configuration/DisplayClusterViewportConfigurationICVFX.cpp
+++ b/Engine/Plugins/Runtime/nDisplay/Source/DisplayCluster/Private/Render/Viewport/Configuration/DisplayClusterViewportConfigurationICVFX.cpp
@@ -19,6 +19,8 @@
 #include "Render/Viewport/DisplayClusterViewportStrings.h"
 #include "Render/Viewport/RenderFrame/DisplayClusterRenderFrameSettings.h"
 
+#include "Render/Viewport/LightCard/DisplayClusterViewportLightCardManager.h"
+
 #include "Components/DisplayClusterICVFXCameraComponent.h"
 
 #include "Misc/DisplayClusterLog.h"
@@ -98,52 +100,6 @@ public:
 		return false;
 	}
 
-	void AssignOverlapChromakeyToTargetViewport(FDisplayClusterViewport& DstViewport, ADisplayClusterRootActor& InRootActor)
-	{
-		if (EnumHasAnyFlags(DstViewport.RenderSettingsICVFX.Flags, EDisplayClusterViewportICVFXFlags::DisableChromakey))
-		{
-			// chromakey disabled for this viewport
-			return;
-		}
-
-		const FDisplayClusterConfigurationICVFX_StageSettings& StageSettings = InRootActor.GetStageSettings();
-		FDisplayClusterShaderParameters_ICVFX::FCameraSettings& DstCameraData = DstViewport.RenderSettingsICVFX.ICVFX.Cameras.Last();
-
-		DstCameraData.OverlapChromakeyColor = StageSettings.GlobalChromakey.ChromakeyColor;
-
-		if (!EnumHasAnyFlags(DstViewport.RenderSettingsICVFX.Flags, EDisplayClusterViewportICVFXFlags::DisableChromakeyMarkers))
-		{
-			FDisplayClusterViewportConfigurationHelpers_ICVFX::UpdateCameraSettings_OverlapChromakeyMarkers(DstCameraData, StageSettings.GlobalChromakey.ChromakeyMarkers);
-		}
-	}
-
-	void AssignChromakeyToTargetViewport(FDisplayClusterViewport& DstViewport, ADisplayClusterRootActor& InRootActor, UDisplayClusterICVFXCameraComponent& InCameraComponent)
-	{
-		if (EnumHasAnyFlags(DstViewport.RenderSettingsICVFX.Flags, EDisplayClusterViewportICVFXFlags::DisableChromakey))
-		{
-			// chromakey disabled for this viewport
-			return;
-		}
-
-		const FDisplayClusterConfigurationICVFX_StageSettings& StageSettings = InRootActor.GetStageSettings();
-		const FDisplayClusterConfigurationICVFX_CameraSettings& CameraSettings = InCameraComponent.GetCameraSettingsICVFX();
-		FDisplayClusterShaderParameters_ICVFX::FCameraSettings& DstCameraData = DstViewport.RenderSettingsICVFX.ICVFX.Cameras.Last();
-
-		FDisplayClusterViewportConfigurationHelpers_ICVFX::UpdateCameraSettings_Chromakey(DstCameraData, StageSettings.GlobalChromakey, CameraSettings.Chromakey, ChromakeyViewport);
-
-		if (!EnumHasAnyFlags(DstViewport.RenderSettingsICVFX.Flags, EDisplayClusterViewportICVFXFlags::DisableChromakeyMarkers))
-		{
-			if (CameraSettings.Chromakey.ChromakeySettingsSource == EDisplayClusterConfigurationICVFX_ChromakeySettingsSource::Viewport)
-			{
-				FDisplayClusterViewportConfigurationHelpers_ICVFX::UpdateCameraSettings_ChromakeyMarkers(DstCameraData, StageSettings.GlobalChromakey.ChromakeyMarkers);
-			}
-			else
-			{
-				FDisplayClusterViewportConfigurationHelpers_ICVFX::UpdateCameraSettings_ChromakeyMarkers(DstCameraData, CameraSettings.Chromakey.ChromakeyMarkers);
-			}
-		}
-	}
-
 	bool IsCameraProjectionVisibleOnViewport(FDisplayClusterViewport* TargetViewport)
 	{
 		if (TargetViewport && TargetViewport->ProjectionPolicy.IsValid())
@@ -198,8 +154,6 @@ public:
 		const FDisplayClusterConfigurationICVFX_CameraSettings& CameraSettings = CameraComponent.GetCameraSettingsICVFX();
 		if (CameraConfiguration.CreateCameraViewport(RootActor, CameraComponent))
 		{
-			bool bAllowChromakey = false;
-
 			FDisplayClusterShaderParameters_ICVFX::FCameraSettings ShaderParametersCameraSettings;
 			if (CameraConfiguration.GetShaderParametersCameraSettings(RootActor, CameraComponent, ShaderParametersCameraSettings))
 			{
@@ -209,21 +163,12 @@ public:
 					if (ViewportIt.IsValid())
 					{
 						ViewportIt->RenderSettingsICVFX.ICVFX.Cameras.Add(ShaderParametersCameraSettings);
-
-						// At lest one target must accept chromakey
-						if (!EnumHasAnyFlags(ViewportIt->RenderSettingsICVFX.Flags, EDisplayClusterViewportICVFXFlags::DisableChromakey))
-						{
-							bAllowChromakey = true;
-						}
 					}
 				}
 			}
 
 			// Create and assign chromakey for all targets for this camera
-			if (bAllowChromakey)
-			{
-				CreateChromakey();
-			}
+			CreateAndSetupInCameraChromakey();
 		}
 	}
 
@@ -248,48 +193,42 @@ public:
 	}
 
 protected:
-	bool CreateChromakey()
+	bool CreateAndSetupInCameraChromakey()
 	{
+		const FDisplayClusterConfigurationICVFX_StageSettings& StageSettings = RootActor.GetStageSettings();
 		const FDisplayClusterConfigurationICVFX_CameraSettings& CameraSettings = CameraComponent.GetCameraSettingsICVFX();
 		const FDisplayClusterConfigurationICVFX_ChromakeySettings& ChromakeySettings = CameraSettings.Chromakey;
 
-		// Assign the overlap chromakey to all supported targets, always done regardless of custom chromakey settings
-		for (TSharedPtr<FDisplayClusterViewport, ESPMode::ThreadSafe>& TargetIt : VisibleTargets)
+		// Try create chromakey render on demand
+		if(const FDisplayClusterConfigurationICVFX_ChromakeyRenderSettings* ChromakeyRenderSettings = ChromakeySettings.GetChromakeyRenderSettings(StageSettings))
 		{
-			if (TargetIt.IsValid())
+			if (ChromakeyRenderSettings->ShouldUseChromakeyViewport(StageSettings))
 			{
-				CameraConfiguration.AssignOverlapChromakeyToTargetViewport(*TargetIt, RootActor);
+				CameraConfiguration.CreateChromakeyViewport(RootActor, CameraComponent);
 			}
 		}
 
-		// Try create chromakey render on demand
-		if (ChromakeySettings.bEnable && ChromakeySettings.ChromakeyType == EDisplayClusterConfigurationICVFX_ChromakeyType::CustomChromakey)
+		// Chromakey viewport name with alpha channel
+		const FString ChromakeyViewportId(CameraConfiguration.ChromakeyViewport ? CameraConfiguration.ChromakeyViewport->GetId() : TEXT(""));
+
+		// Assign this chromakey to all supported targets
+		for (const TSharedPtr<FDisplayClusterViewport, ESPMode::ThreadSafe>& OuterViewportIt : VisibleTargets)
 		{
-			if (ChromakeySettings.ChromakeyRenderTexture.Replace.bAllowReplace && ChromakeySettings.ChromakeyRenderTexture.Replace.SourceTexture == nullptr)
+			if (OuterViewportIt.IsValid())
 			{
-				// ChromakeyRenderTexture.Override require source texture.
-				return false;
-			}
+				FDisplayClusterShaderParameters_ICVFX::FCameraSettings& DstCameraData = OuterViewportIt->RenderSettingsICVFX.ICVFX.Cameras.Last();
 
-			const bool bIsChromakeyHasAnyRenderComponent = FDisplayClusterViewportConfigurationHelpers_Visibility::IsVisibilityListValid(ChromakeySettings.ChromakeyRenderTexture.ShowOnlyList);
-			if (ChromakeySettings.ChromakeyRenderTexture.Replace.bAllowReplace == false && !bIsChromakeyHasAnyRenderComponent)
-			{
-				// ChromakeyRenderTexture requires actors for render.
-				return false;
-			}
+				const bool bEnableChromakey = !EnumHasAnyFlags(OuterViewportIt->RenderSettingsICVFX.Flags, EDisplayClusterViewportICVFXFlags::DisableChromakey);
+				const bool bEnableChromakeyMarkers = !EnumHasAnyFlags(OuterViewportIt->RenderSettingsICVFX.Flags, EDisplayClusterViewportICVFXFlags::DisableChromakeyMarkers);
 
-			if (!CameraConfiguration.CreateChromakeyViewport(RootActor, CameraComponent))
-			{
-				return false;
-			}
-		}
+				if (bEnableChromakey)
+				{
+					// Setup chromakey with markers
+					FDisplayClusterViewportConfigurationHelpers_ICVFX::UpdateCameraSettings_Chromakey(DstCameraData, StageSettings, CameraSettings, bEnableChromakeyMarkers, ChromakeyViewportId);
+				}
 
-		// Assign this chromakey to all supported targets
-		for (TSharedPtr<FDisplayClusterViewport, ESPMode::ThreadSafe>& TargetIt : VisibleTargets)
-		{
-			if (TargetIt.IsValid())
-			{
-				CameraConfiguration.AssignChromakeyToTargetViewport(*TargetIt, RootActor, CameraComponent);
+				// Setup overlap chromakey with markers
+				FDisplayClusterViewportConfigurationHelpers_ICVFX::UpdateCameraSettings_OverlapChromakey(DstCameraData, StageSettings, CameraSettings, bEnableChromakeyMarkers);
 			}
 		}
 
@@ -426,26 +365,26 @@ void FDisplayClusterViewportConfigurationICVFX::Update()
 			// If not all viewports disable lightcard
 			if (!EnumHasAnyFlags(TargetViewportsFlags, EDisplayClusterViewportICVFXFlags::DisableLightcard))
 			{
+				// UVLightCard must be enabled in LC manager
+				const bool bUVLightCardEnabled = ViewportManager->GetLightCardManager().IsValid() && ViewportManager->GetLightCardManager()->IsUVLightCardEnabled();
+					
 				// Allocate and assign lightcard resources
-				{
-
-					const bool bUseLightCard = FDisplayClusterViewportConfigurationHelpers_ICVFX::IsShouldUseLightcard(StageSettings.Lightcard);
-					const bool bUseUVLightCard = FDisplayClusterViewportConfigurationHelpers_ICVFX::IsShouldUseUVLightcard(*ViewportManager, StageSettings.Lightcard);
+				const bool bUseLightCard = StageSettings.Lightcard.ShouldUseLightCard(StageSettings);
+				const bool bUseUVLightCard = bUVLightCardEnabled && StageSettings.Lightcard.ShouldUseUVLightCard(StageSettings);
 
-					for (const TSharedPtr<FDisplayClusterViewport, ESPMode::ThreadSafe>& TargetIt : TargetViewports)
+				for (const TSharedPtr<FDisplayClusterViewport, ESPMode::ThreadSafe>& TargetIt : TargetViewports)
+				{
+					// only for support targets
+					if (TargetIt.IsValid() && !EnumHasAnyFlags(TargetIt->RenderSettingsICVFX.Flags, EDisplayClusterViewportICVFXFlags::DisableLightcard))
 					{
-						// only for support targets
-						if (TargetIt.IsValid() && !EnumHasAnyFlags(TargetIt->RenderSettingsICVFX.Flags, EDisplayClusterViewportICVFXFlags::DisableLightcard))
+						if (bUseLightCard)
 						{
-							if (bUseLightCard)
-							{
-								CreateLightcardViewport(*TargetIt);
-							}
+							CreateLightcardViewport(*TargetIt);
+						}
 
-							if (bUseUVLightCard)
-							{
-								CreateUVLightcardViewport(*TargetIt);
-							}
+						if (bUseUVLightCard)
+						{
+							CreateUVLightcardViewport(*TargetIt);
 						}
 					}
 				}
@@ -459,8 +398,7 @@ void FDisplayClusterViewportConfigurationICVFX::Update()
 					TargetIt->RenderSettingsICVFX.ICVFX.SortCamerasRenderOrder();
 
 					// Setup incamera overlap mode
-					TargetIt->RenderSettingsICVFX.ICVFX.CameraOverlappingRenderMode = StageSettings.bEnableInnerFrustumChromakeyOverlap ?
-						EDisplayClusterShaderParametersICVFX_CameraOverlappingRenderMode::FinalPass : EDisplayClusterShaderParametersICVFX_CameraOverlappingRenderMode::None;
+					TargetIt->RenderSettingsICVFX.ICVFX.CameraOverlappingRenderMode = StageSettings.GetCameraOverlappingRenderMode();
 				}
 			}
 
diff --git a/Engine/Plugins/Runtime/nDisplay/Source/DisplayCluster/Private/Render/Viewport/DisplayClusterViewportProxy.cpp b/Engine/Plugins/Runtime/nDisplay/Source/DisplayCluster/Private/Render/Viewport/DisplayClusterViewportProxy.cpp
--- a/Engine/Plugins/Runtime/nDisplay/Source/DisplayCluster/Private/Render/Viewport/DisplayClusterViewportProxy.cpp
+++ b/Engine/Plugins/Runtime/nDisplay/Source/DisplayCluster/Private/Render/Viewport/DisplayClusterViewportProxy.cpp
@@ -248,7 +248,7 @@ void FDisplayClusterViewportProxy::HandleResourceDelete_RenderThread(FDisplayClu
 }
 
 
-bool FDisplayClusterViewportProxy::IsShouldOverrideViewportResource(const EDisplayClusterViewportResourceType InResourceType) const
+bool FDisplayClusterViewportProxy::ShouldOverrideViewportResource(const EDisplayClusterViewportResourceType InResourceType) const
 {
 	// Override resources from other viewport
 	if (RenderSettings.IsViewportOverrided())
@@ -340,7 +340,7 @@ bool FDisplayClusterViewportProxy::ImplGetResources_RenderThread(const EDisplayC
 	check(IsInRenderingThread());
 
 	// Override resources from other viewport
-	if (IsShouldOverrideViewportResource(InResourceType))
+	if (ShouldOverrideViewportResource(InResourceType))
 	{
 		if (InRecursionDepth < DisplayClusterViewportProxyResourcesOverrideRecursionDepthMax)
 		{
@@ -617,7 +617,7 @@ bool FDisplayClusterViewportProxy::ImplGetResourcesWithRects_RenderThread(const
 	check(IsInRenderingThread());
 
 	// Override resources from other viewport
-	if(IsShouldOverrideViewportResource(InResourceType))
+	if(ShouldOverrideViewportResource(InResourceType))
 	{
 		if (InRecursionDepth < DisplayClusterViewportProxyResourcesOverrideRecursionDepthMax)
 		{
diff --git a/Engine/Plugins/Runtime/nDisplay/Source/DisplayCluster/Private/Render/Viewport/DisplayClusterViewportProxy.h b/Engine/Plugins/Runtime/nDisplay/Source/DisplayCluster/Private/Render/Viewport/DisplayClusterViewportProxy.h
--- a/Engine/Plugins/Runtime/nDisplay/Source/DisplayCluster/Private/Render/Viewport/DisplayClusterViewportProxy.h
+++ b/Engine/Plugins/Runtime/nDisplay/Source/DisplayCluster/Private/Render/Viewport/DisplayClusterViewportProxy.h
@@ -272,7 +272,9 @@ private:
 	void ImplPreviewReadPixels_RenderThread(FRHICommandListImmediate& RHICmdList) const;
 
 	bool ImplResolveResources_RenderThread(FRHICommandListImmediate& RHICmdList, FDisplayClusterViewportProxy const* SourceProxy, const EDisplayClusterViewportResourceType InputResourceType, const EDisplayClusterViewportResourceType OutputResourceType, const int32 InContextNum) const;
-	bool IsShouldOverrideViewportResource(const EDisplayClusterViewportResourceType InResourceType) const;
+
+	/** When a resource by type can be overridden from another viewport, true is returned. */
+	bool ShouldOverrideViewportResource(const EDisplayClusterViewportResourceType InResourceType) const;
 
 	/** Copy pixel channels between resources.
 	 *
diff --git a/Engine/Plugins/Runtime/nDisplay/Source/DisplayCluster/Public/Render/Viewport/Containers/DisplayClusterViewport_EnumsICVFX.h b/Engine/Plugins/Runtime/nDisplay/Source/DisplayCluster/Public/Render/Viewport/Containers/DisplayClusterViewport_EnumsICVFX.h
new file mode 100644
--- /dev/null
+++ b/Engine/Plugins/Runtime/nDisplay/Source/DisplayCluster/Public/Render/Viewport/Containers/DisplayClusterViewport_EnumsICVFX.h
@@ -0,0 +1,53 @@
+// Copyright Epic Games, Inc. All Rights Reserved.
+
+#pragma once
+
+#include "CoreMinimal.h"
+
+/**
+ * Runtime configuration from DCRA.
+ */
+enum class EDisplayClusterViewportICVFXFlags : uint8
+{
+	None = 0,
+
+	// Allow to use ICVFX for this viewport (Must be supported by projection policy)
+	Enable = 1 << 0,
+
+	// Disable incamera render to this viewport
+	DisableCamera = 1 << 1,
+
+	// Disable chromakey render to this viewport
+	DisableChromakey = 1 << 2,
+
+	// Disable chromakey markers render to this viewport
+	DisableChromakeyMarkers = 1 << 3,
+
+	// Disable lightcard render to this viewport
+	DisableLightcard = 1 << 4,
+};
+ENUM_CLASS_FLAGS(EDisplayClusterViewportICVFXFlags);
+
+/**
+ * This flag raised only from icvfx manager.
+*/
+enum class EDisplayClusterViewportRuntimeICVFXFlags: uint8
+{
+	None = 0,
+
+	// Enable use icvfx only from projection policy for this viewport.
+	Target = 1 << 0,
+
+	// viewport ICVFX usage
+	InCamera    = 1 << 1,
+	Chromakey   = 1 << 2,
+	Lightcard   = 1 << 3,
+	UVLightcard = 1 << 4,
+
+	// This viewport used as internal icvfx composing resource (created and deleted inside icvfx logic)
+	InternalResource = 1 << 5,
+
+	// Mark unused icvfx dynamic viewports
+	Unused = 1 << 6,
+};
+ENUM_CLASS_FLAGS(EDisplayClusterViewportRuntimeICVFXFlags);
diff --git a/Engine/Plugins/Runtime/nDisplay/Source/DisplayCluster/Public/Render/Viewport/Containers/DisplayClusterViewport_RenderSettingsICVFX.h b/Engine/Plugins/Runtime/nDisplay/Source/DisplayCluster/Public/Render/Viewport/Containers/DisplayClusterViewport_RenderSettingsICVFX.h
--- a/Engine/Plugins/Runtime/nDisplay/Source/DisplayCluster/Public/Render/Viewport/Containers/DisplayClusterViewport_RenderSettingsICVFX.h
+++ b/Engine/Plugins/Runtime/nDisplay/Source/DisplayCluster/Public/Render/Viewport/Containers/DisplayClusterViewport_RenderSettingsICVFX.h
@@ -4,59 +4,12 @@
 
 #include "CoreMinimal.h"
 
+#include "Render/Viewport/Containers/DisplayClusterViewport_EnumsICVFX.h"
 #include "ShaderParameters/DisplayClusterShaderParameters_ICVFX.h"
 
-//***************************************************
-//* Runtime configuration from DCRA
-//***************************************************
-enum class EDisplayClusterViewportICVFXFlags : uint8
-{
-	None = 0,
-
-	// Allow to use ICVFX for this viewport (Must be supported by projection policy)
-	Enable = 1 << 0,
-
-	// Disable incamera render to this viewport
-	DisableCamera = 1 << 1,
-
-	// Disable chromakey render to this viewport
-	DisableChromakey = 1 << 2,
-
-	// Disable chromakey markers render to this viewport
-	DisableChromakeyMarkers = 1 << 3,
-
-	// Disable lightcard render to this viewport
-	DisableLightcard = 1 << 4,
-
-	// Use unique lightcard mode for this viewport from param 'OverrideLightcardMode'
-	OverrideLightcardMode = 1 << 5,
-};
-ENUM_CLASS_FLAGS(EDisplayClusterViewportICVFXFlags);
-
-//***************************************************
-//* This flag raised only from icvfx manager
-//***************************************************
-enum class EDisplayClusterViewportRuntimeICVFXFlags: uint8
-{
-	None = 0,
-
-	// Enable use icvfx only from projection policy for this viewport.
-	Target = 1 << 0,
-
-	// viewport ICVFX usage
-	InCamera    = 1 << 1,
-	Chromakey   = 1 << 2,
-	Lightcard   = 1 << 3,
-	UVLightcard = 1 << 4,
-
-	// This viewport used as internal icvfx composing resource (created and deleted inside icvfx logic)
-	InternalResource = 1 << 5,
-
-	// Mark unused icvfx dynamic viewports
-	Unused = 1 << 6,
-};
-ENUM_CLASS_FLAGS(EDisplayClusterViewportRuntimeICVFXFlags);
-
+/**
+ * ICVFX settings for viewport
+ */
 class FDisplayClusterViewport_RenderSettingsICVFX
 {
 public:
diff --git a/Engine/Plugins/Runtime/nDisplay/Source/DisplayClusterConfiguration/DisplayClusterConfiguration.Build.cs b/Engine/Plugins/Runtime/nDisplay/Source/DisplayClusterConfiguration/DisplayClusterConfiguration.Build.cs
--- a/Engine/Plugins/Runtime/nDisplay/Source/DisplayClusterConfiguration/DisplayClusterConfiguration.Build.cs
+++ b/Engine/Plugins/Runtime/nDisplay/Source/DisplayClusterConfiguration/DisplayClusterConfiguration.Build.cs
@@ -11,6 +11,7 @@ public class DisplayClusterConfiguration : ModuleRules
 			new string[] {
 				"DisplayCluster",
 				"DisplayClusterProjection",
+				"DisplayClusterShaders",
 			});
 
 		PublicDependencyModuleNames.AddRange(
diff --git a/Engine/Plugins/Runtime/nDisplay/Source/DisplayClusterConfiguration/Private/DisplayClusterConfigurationTypes_ICVFX.cpp b/Engine/Plugins/Runtime/nDisplay/Source/DisplayClusterConfiguration/Private/DisplayClusterConfigurationTypes_ICVFX.cpp
--- a/Engine/Plugins/Runtime/nDisplay/Source/DisplayClusterConfiguration/Private/DisplayClusterConfigurationTypes_ICVFX.cpp
+++ b/Engine/Plugins/Runtime/nDisplay/Source/DisplayClusterConfiguration/Private/DisplayClusterConfigurationTypes_ICVFX.cpp
@@ -2,6 +2,17 @@
 
 #include "DisplayClusterConfigurationTypes_ICVFX.h"
 
+namespace UE::DisplayClusterConfiguration::ICVFX
+{
+	static float ClampPercent(float InValue)
+	{
+		static const float MaxCustomFrustumValue = 5.f;
+
+		return FMath::Clamp(InValue, -MaxCustomFrustumValue, MaxCustomFrustumValue);
+	}
+};
+using namespace UE::DisplayClusterConfiguration::ICVFX;
+
 ///////////////////////////////////////////////////////////////////////////////////////
 // FDisplayClusterConfigurationICVFX_ChromakeyMarkers
 ///////////////////////////////////////////////////////////////////////////////////////
@@ -191,9 +202,120 @@ bool FDisplayClusterConfigurationICVFX_CameraSettings::IsChromakeyViewportSettin
 {
 	return CameraOCIO.IsChromakeyViewportSettingsEqual_Editor(InClusterNodeId1, InClusterNodeId2);
 }
-
 #endif
 
+float FDisplayClusterConfigurationICVFX_CameraSettings::GetCameraFieldOfViewMultiplier(const FDisplayClusterConfigurationICVFX_StageSettings& InStageSettings) const
+{
+	if (CustomFrustum.bEnable)
+	{
+		return CustomFrustum.FieldOfViewMultiplier;
+	}
+
+	return  1.f;
+}
+
+float FDisplayClusterConfigurationICVFX_CameraSettings::GetCameraAdaptResolutionRatio(const FDisplayClusterConfigurationICVFX_StageSettings& InStageSettings) const
+{
+	if (CustomFrustum.bAdaptResolution)
+	{
+		return GetCameraFieldOfViewMultiplier(InStageSettings);
+	}
+
+	// Don't use an adaptive resolution multiplier
+	return 1.f;
+}
+
+float FDisplayClusterConfigurationICVFX_CameraSettings::GetCameraBufferRatio(const FDisplayClusterConfigurationICVFX_StageSettings& InStageSettings) const
+{
+	return BufferRatio;
+}
+
+bool FDisplayClusterConfigurationICVFX_CameraSettings::GetCameraBorder(const FDisplayClusterConfigurationICVFX_StageSettings& InStageSettings, FLinearColor& OutBorderColor, float& OutBorderThickness) const
+{
+	if (!Border.Enable)
+	{
+		OutBorderColor = FLinearColor::Black;
+		OutBorderThickness = 0.0f;
+
+		return false;
+	}
+
+	const float RealThicknessScaleValue = 0.1f;
+
+	OutBorderColor = Border.Color;
+	OutBorderThickness = Border.Thickness * RealThicknessScaleValue;
+
+	return true;
+}
+
+FIntPoint FDisplayClusterConfigurationICVFX_CameraSettings::GetCameraFrameSize(const FDisplayClusterConfigurationICVFX_StageSettings& InStageSettings) const
+{
+	if (RenderSettings.CustomFrameSize.bUseCustomSize)
+	{
+		// Custom camera size
+		return FIntPoint(RenderSettings.CustomFrameSize.CustomWidth, RenderSettings.CustomFrameSize.CustomHeight);
+	}
+
+	// global camera size
+	return FIntPoint(InStageSettings.DefaultFrameSize.Width, InStageSettings.DefaultFrameSize.Height);
+}
+
+float FDisplayClusterConfigurationICVFX_CameraSettings::GetCameraFrameAspectRatio(const FDisplayClusterConfigurationICVFX_StageSettings& InStageSettings) const
+{
+	FIntPoint FrameSize = GetCameraFrameSize(InStageSettings);
+
+	return (FrameSize.Y > 0 && FrameSize.X > 0) ? (float)FrameSize.X / float(FrameSize.Y) : 0;
+}
+
+FVector4 FDisplayClusterConfigurationICVFX_CameraSettings::GetCameraSoftEdge(const FDisplayClusterConfigurationICVFX_StageSettings& InStageSettings) const
+{
+	FVector4 ResultSoftEdge(ForceInitToZero);
+
+	const float FieldOfViewMultiplier = GetCameraFieldOfViewMultiplier(InStageSettings);
+
+	// softedge adjustments	
+	const float Overscan = (FieldOfViewMultiplier > 0) ? FieldOfViewMultiplier : 1;
+
+	// remap values from 0-1 GUI range into acceptable 0.0 - 0.25 shader range
+	ResultSoftEdge.X = FMath::GetMappedRangeValueClamped(FVector2D(0.0, 1.0f), FVector2D(0.0, 0.25), SoftEdge.Horizontal) / Overscan; // Left
+	ResultSoftEdge.Y = FMath::GetMappedRangeValueClamped(FVector2D(0.0, 1.0f), FVector2D(0.0, 0.25), SoftEdge.Vertical) / Overscan; // Top
+
+	// ZW now used in other way
+	// Z for new parameter Feather
+	ResultSoftEdge.Z = SoftEdge.Feather;
+
+	const float CameraBufferRatio = GetCameraBufferRatio(InStageSettings);
+
+	// default - percents
+	const float ConvertToPercent = 0.01f;
+
+	float Left = ClampPercent(CustomFrustum.Left * ConvertToPercent);
+	float Right = ClampPercent(CustomFrustum.Right * ConvertToPercent);
+	float Top = ClampPercent(CustomFrustum.Top * ConvertToPercent);
+	float Bottom = ClampPercent(CustomFrustum.Bottom * ConvertToPercent);
+
+	if (CustomFrustum.Mode == EDisplayClusterConfigurationViewportCustomFrustumMode::Pixels)
+	{
+		const FIntPoint FrameSize = GetCameraFrameSize(InStageSettings);
+
+		const float  FrameWidth = FrameSize.X * CameraBufferRatio;
+		const float FrameHeight = FrameSize.Y * CameraBufferRatio;
+
+		Left = ClampPercent(CustomFrustum.Left / FrameWidth);
+		Right = ClampPercent(CustomFrustum.Right / FrameWidth);
+		Top = ClampPercent(CustomFrustum.Top / FrameHeight);
+		Bottom = ClampPercent(CustomFrustum.Bottom / FrameHeight);
+	}
+
+	// recalculate soft edge related offsets based on frustum
+	ResultSoftEdge.X /= (1 + Left + Right);
+	ResultSoftEdge.Y /= (1 + Top + Bottom);
+	ResultSoftEdge.Z /= (1 + Left + Right);
+	ResultSoftEdge.W /= (1 + Top + Bottom);
+
+	return ResultSoftEdge;
+}
+
 ///////////////////////////////////////////////////////////////////////////////////////
 // FDisplayClusterConfigurationICVFX_StageSettings
 ///////////////////////////////////////////////////////////////////////////////////////
@@ -206,3 +328,206 @@ const FOpenColorIOColorConversionSettings* FDisplayClusterConfigurationICVFX_Sta
 {
 	return Lightcard.LightcardOCIO.FindOCIOConfiguration(InViewportId, ViewportOCIO);
 }
+
+EDisplayClusterShaderParametersICVFX_CameraOverlappingRenderMode FDisplayClusterConfigurationICVFX_StageSettings::GetCameraOverlappingRenderMode() const
+{
+	if (bEnableInnerFrustumChromakeyOverlap)
+	{
+		return EDisplayClusterShaderParametersICVFX_CameraOverlappingRenderMode::FinalPass;
+	}
+
+	return EDisplayClusterShaderParametersICVFX_CameraOverlappingRenderMode::None;
+}
+
+///////////////////////////////////////////////////////////////////////////////////////
+// FDisplayClusterConfigurationICVFX_ChromakeySettings
+///////////////////////////////////////////////////////////////////////////////////////
+EDisplayClusterShaderParametersICVFX_ChromakeySource FDisplayClusterConfigurationICVFX_ChromakeySettings::GetChromakeyType(const FDisplayClusterConfigurationICVFX_StageSettings& InStageSettings) const
+{
+	if (!bEnable)
+	{
+		return EDisplayClusterShaderParametersICVFX_ChromakeySource::Disabled;
+	}
+
+	switch (ChromakeyType)
+	{
+	case EDisplayClusterConfigurationICVFX_ChromakeyType::CustomChromakey:
+		return EDisplayClusterShaderParametersICVFX_ChromakeySource::ChromakeyLayers;
+
+	default:
+		break;
+	}
+
+	return EDisplayClusterShaderParametersICVFX_ChromakeySource::FrameColor;
+}
+
+FDisplayClusterConfigurationICVFX_ChromakeyRenderSettings* FDisplayClusterConfigurationICVFX_ChromakeySettings::GetWritableChromakeyRenderSettings(const FDisplayClusterConfigurationICVFX_StageSettings& InStageSettings)
+{
+	// Note: Here we can add an override of the CK rendering settings from StageSetings
+	if (GetChromakeyType(InStageSettings) == EDisplayClusterShaderParametersICVFX_ChromakeySource::ChromakeyLayers)
+	{
+		return &ChromakeyRenderTexture;
+	}
+
+	return nullptr;
+}
+
+const FDisplayClusterConfigurationICVFX_ChromakeyRenderSettings* FDisplayClusterConfigurationICVFX_ChromakeySettings::GetChromakeyRenderSettings(const FDisplayClusterConfigurationICVFX_StageSettings& InStageSettings) const
+{
+	// Note: Here we can add an override of the CK rendering settings from StageSetings
+	if (GetChromakeyType(InStageSettings) == EDisplayClusterShaderParametersICVFX_ChromakeySource::ChromakeyLayers)
+	{
+		return &ChromakeyRenderTexture;
+	}
+
+	return nullptr;
+}
+
+const FLinearColor& FDisplayClusterConfigurationICVFX_ChromakeySettings::GetChromakeyColor(const FDisplayClusterConfigurationICVFX_StageSettings& InStageSettings) const
+{
+	switch(ChromakeySettingsSource)
+	{
+	case EDisplayClusterConfigurationICVFX_ChromakeySettingsSource::Viewport:
+		// Override Chromakey color from stage settings
+		return InStageSettings.GlobalChromakey.ChromakeyColor;
+
+	default:
+		break;
+	}
+
+	// Use Chromakey color from camera
+	return ChromakeyColor;
+}
+
+const FLinearColor& FDisplayClusterConfigurationICVFX_ChromakeySettings::GetOverlapChromakeyColor(const FDisplayClusterConfigurationICVFX_StageSettings& InStageSettings) const
+{
+	// Note: Here we can add an override of the CK overlap area color from camera
+
+	// Use overlay color from stage settings
+	return InStageSettings.GlobalChromakey.ChromakeyColor;
+}
+
+const FDisplayClusterConfigurationICVFX_ChromakeyMarkers* FDisplayClusterConfigurationICVFX_ChromakeySettings::ImplGetChromakeyMarkers(const FDisplayClusterConfigurationICVFX_ChromakeyMarkers* InValue) const
+{
+	// Chromakey markers require texture
+	if (!InValue || !InValue->bEnable || InValue->MarkerTileRGBA == nullptr)
+	{
+		return nullptr;
+	}
+
+	// This CK markers can be used
+	return InValue;
+}
+
+const FDisplayClusterConfigurationICVFX_ChromakeyMarkers* FDisplayClusterConfigurationICVFX_ChromakeySettings::GetChromakeyMarkers(const FDisplayClusterConfigurationICVFX_StageSettings& InStageSettings) const
+{
+	switch (ChromakeySettingsSource)
+	{
+		case EDisplayClusterConfigurationICVFX_ChromakeySettingsSource::Viewport:
+			// Use global CK markers
+			return ImplGetChromakeyMarkers(&InStageSettings.GlobalChromakey.ChromakeyMarkers);
+
+		default:
+			break;
+	}
+
+	// Use CK markers from camera
+	return ImplGetChromakeyMarkers(&ChromakeyMarkers);
+}
+
+const FDisplayClusterConfigurationICVFX_ChromakeyMarkers* FDisplayClusterConfigurationICVFX_ChromakeySettings::GetOverlapChromakeyMarkers(const FDisplayClusterConfigurationICVFX_StageSettings& InStageSettings) const
+{
+	// Note: Here we can add an override of the CK overlap markers from camera
+	FDisplayClusterConfigurationICVFX_ChromakeyMarkers const* OutChromakeyMarkers = &InStageSettings.GlobalChromakey.ChromakeyMarkers;
+
+	// Use CK overlap markers from stage settings:
+	return ImplGetChromakeyMarkers(OutChromakeyMarkers);
+}
+
+///////////////////////////////////////////////////////////////////////////////////////
+// FDisplayClusterConfigurationICVFX_VisibilityList
+///////////////////////////////////////////////////////////////////////////////////////
+bool FDisplayClusterConfigurationICVFX_VisibilityList::IsVisibilityListValid() const
+{
+	for (const FString& ComponentNameIt : RootActorComponentNames)
+	{
+		if (!ComponentNameIt.IsEmpty())
+		{
+			return true;
+		}
+	}
+
+	for (const TSoftObjectPtr<AActor>& ActorSOPtrIt : Actors)
+	{
+		if (ActorSOPtrIt.IsValid())
+		{
+			return true;
+		}
+	}
+
+	for (const FActorLayer& ActorLayerIt : ActorLayers)
+	{
+		if (!ActorLayerIt.Name.IsNone())
+		{
+			return true;
+		}
+	}
+
+	for (const TSoftObjectPtr<AActor>& AutoAddedActor : AutoAddedActors)
+	{
+		if (AutoAddedActor.IsValid())
+		{
+			return true;
+		}
+	}
+
+	return false;
+}
+
+///////////////////////////////////////////////////////////////////////////////////////
+// FDisplayClusterConfigurationICVFX_LightcardSettings
+///////////////////////////////////////////////////////////////////////////////////////
+bool FDisplayClusterConfigurationICVFX_LightcardSettings::ShouldUseLightCard(const FDisplayClusterConfigurationICVFX_StageSettings& InStageSettings) const
+{
+	if (!bEnable)
+	{
+		// dont use lightcard if disabled
+		return false;
+	}
+
+	if (RenderSettings.Replace.bAllowReplace)
+	{
+		if (RenderSettings.Replace.SourceTexture == nullptr)
+		{
+			// LightcardSettings.Override require source texture.
+			return false;
+		}
+	}
+
+	// Lightcard require layers for render
+	return ShowOnlyList.IsVisibilityListValid();
+}
+
+bool FDisplayClusterConfigurationICVFX_LightcardSettings::ShouldUseUVLightCard(const FDisplayClusterConfigurationICVFX_StageSettings& InStageSettings) const
+{
+	//Note: Here we can add custom rules for UV lightcards
+	return ShouldUseLightCard(InStageSettings);
+}
+
+///////////////////////////////////////////////////////////////////////////////////////
+// FDisplayClusterConfigurationICVFX_ChromakeyRenderSettings
+///////////////////////////////////////////////////////////////////////////////////////
+bool FDisplayClusterConfigurationICVFX_ChromakeyRenderSettings::ShouldUseChromakeyViewport(const FDisplayClusterConfigurationICVFX_StageSettings& InStageSettings) const
+{
+	if (Replace.bAllowReplace)
+	{
+		if (Replace.SourceTexture == nullptr)
+		{
+			// ChromakeyRender Override require source texture.
+			return false;
+		}
+	}
+
+	// ChromakeyRender requires actors for render.
+	return ShowOnlyList.IsVisibilityListValid();
+}
diff --git a/Engine/Plugins/Runtime/nDisplay/Source/DisplayClusterConfiguration/Private/DisplayClusterConfigurationTypes_Viewport.cpp b/Engine/Plugins/Runtime/nDisplay/Source/DisplayClusterConfiguration/Private/DisplayClusterConfigurationTypes_Viewport.cpp
new file mode 100644
--- /dev/null
+++ b/Engine/Plugins/Runtime/nDisplay/Source/DisplayClusterConfiguration/Private/DisplayClusterConfigurationTypes_Viewport.cpp
@@ -0,0 +1,94 @@
+// Copyright Epic Games, Inc. All Rights Reserved.
+
+#include "DisplayClusterConfigurationTypes_Viewport.h"
+#include "DisplayClusterConfigurationTypes_ICVFX.h"
+
+///////////////////////////////////////////////////////////////////////////////////////
+// FDisplayClusterConfigurationViewport_ICVFX
+///////////////////////////////////////////////////////////////////////////////////////
+EDisplayClusterShaderParametersICVFX_LightCardRenderMode FDisplayClusterConfigurationViewport_ICVFX::GetLightCardRenderMode(const FDisplayClusterConfigurationICVFX_StageSettings& InStageSettings) const
+{
+	// Note: please check if this rule is valid
+	if (!bAllowICVFX || !InStageSettings.Lightcard.bEnable)
+	{
+		// When ICVFX is disabled we dont render lightcards
+		return EDisplayClusterShaderParametersICVFX_LightCardRenderMode::None;
+	}
+
+	if (LightcardRenderMode != EDisplayClusterConfigurationICVFX_OverrideLightcardRenderMode::Default)
+	{
+		// Use overrided values from viewport:
+		switch (LightcardRenderMode)
+		{
+		case EDisplayClusterConfigurationICVFX_OverrideLightcardRenderMode::Over:
+			return EDisplayClusterShaderParametersICVFX_LightCardRenderMode::Over;
+
+		case EDisplayClusterConfigurationICVFX_OverrideLightcardRenderMode::Under:
+			return EDisplayClusterShaderParametersICVFX_LightCardRenderMode::Under;
+
+		default:
+			break;
+		}
+
+		return EDisplayClusterShaderParametersICVFX_LightCardRenderMode::None;
+	}
+
+	// Use global lightcard settings:
+	switch (InStageSettings.Lightcard.Blendingmode)
+	{
+	case EDisplayClusterConfigurationICVFX_LightcardRenderMode::Under:
+		return EDisplayClusterShaderParametersICVFX_LightCardRenderMode::Under;
+
+	default:
+		break;
+	};
+
+	// By default lightcard render in 'Over' mode
+	return EDisplayClusterShaderParametersICVFX_LightCardRenderMode::Over;
+}
+
+EDisplayClusterViewportICVFXFlags FDisplayClusterConfigurationViewport_ICVFX::GetViewportICVFXFlags(const FDisplayClusterConfigurationICVFX_StageSettings& InStageSettings) const
+{
+	EDisplayClusterViewportICVFXFlags OutFlags = EDisplayClusterViewportICVFXFlags::None;
+	if (bAllowICVFX)
+	{
+		EnumAddFlags(OutFlags, EDisplayClusterViewportICVFXFlags::Enable);
+	}
+
+	// Override camera render mode
+	EDisplayClusterConfigurationICVFX_OverrideCameraRenderMode UsedCameraRenderMode = CameraRenderMode;
+	if (!bAllowInnerFrustum || !InStageSettings.bEnableInnerFrustums)
+	{
+		UsedCameraRenderMode = EDisplayClusterConfigurationICVFX_OverrideCameraRenderMode::Disabled;
+	}
+
+	switch (UsedCameraRenderMode)
+	{
+	// Disable camera frame render for this viewport
+	case EDisplayClusterConfigurationICVFX_OverrideCameraRenderMode::Disabled:
+		EnumAddFlags(OutFlags, EDisplayClusterViewportICVFXFlags::DisableCamera | EDisplayClusterViewportICVFXFlags::DisableChromakey | EDisplayClusterViewportICVFXFlags::DisableChromakeyMarkers);
+		break;
+
+	// Disable chromakey render for this viewport
+	case EDisplayClusterConfigurationICVFX_OverrideCameraRenderMode::DisableChromakey:
+		EnumAddFlags(OutFlags, EDisplayClusterViewportICVFXFlags::DisableChromakey | EDisplayClusterViewportICVFXFlags::DisableChromakeyMarkers);
+		break;
+
+	// Disable chromakey markers render for this viewport
+	case EDisplayClusterConfigurationICVFX_OverrideCameraRenderMode::DisableChromakeyMarkers:
+		EnumAddFlags(OutFlags, EDisplayClusterViewportICVFXFlags::DisableChromakeyMarkers);
+		break;
+
+	default:
+		break;
+	}
+
+	// Disable lightcards rendering
+	const EDisplayClusterShaderParametersICVFX_LightCardRenderMode LightCardRenderMode = GetLightCardRenderMode(InStageSettings);
+	if (LightCardRenderMode == EDisplayClusterShaderParametersICVFX_LightCardRenderMode::None)
+	{
+		EnumAddFlags(OutFlags, EDisplayClusterViewportICVFXFlags::DisableLightcard);
+	}
+
+	return OutFlags;
+}
diff --git a/Engine/Plugins/Runtime/nDisplay/Source/DisplayClusterConfiguration/Public/DisplayClusterConfigurationTypes_ICVFX.h b/Engine/Plugins/Runtime/nDisplay/Source/DisplayClusterConfiguration/Public/DisplayClusterConfigurationTypes_ICVFX.h
--- a/Engine/Plugins/Runtime/nDisplay/Source/DisplayClusterConfiguration/Public/DisplayClusterConfigurationTypes_ICVFX.h
+++ b/Engine/Plugins/Runtime/nDisplay/Source/DisplayClusterConfiguration/Public/DisplayClusterConfigurationTypes_ICVFX.h
@@ -17,6 +17,8 @@
 #include "DisplayClusterConfigurationTypes_Postprocess.h"
 #include "DisplayClusterConfigurationTypes_OCIO.h"
 
+#include "Containers/DisplayClusterShader_Enums.h"
+
 #include "DisplayClusterConfigurationTypes_ICVFX.generated.h"
 
 USTRUCT(Blueprintable)
@@ -108,6 +110,10 @@ struct DISPLAYCLUSTERCONFIGURATION_API FDisplayClusterConfigurationICVFX_Visibil
 {
 	GENERATED_BODY()
 
+public:
+	/** returns true if there is at least one valid value in the visibility list. */
+	bool IsVisibilityListValid() const;
+
 public:
 	/** Actor Layers */
 	UPROPERTY(EditAnywhere, BlueprintReadWrite, Category = NDisplay, meta = (DisplayName = "Layers", ToolTip = "{LayersTooltip}"))
@@ -207,6 +213,10 @@ public:
 	FDisplayClusterConfigurationICVFX_ChromakeyRenderSettings& operator=(FDisplayClusterConfigurationICVFX_ChromakeyRenderSettings&&) = default;
 	PRAGMA_ENABLE_DEPRECATION_WARNINGS
 
+public:
+	/** true if chromakey viewport is used in  ICVFX. */
+	bool ShouldUseChromakeyViewport(const FDisplayClusterConfigurationICVFX_StageSettings& InStageSettings) const;
+
 public:
 	/** Set to True to use custom chromakey content. */
 	UE_DEPRECATED(5.3, "Use the ChromakeyType enum in FDisplayClusterConfigurationICVFX_ChromakeySettings instead")
@@ -286,6 +296,40 @@ public:
 		: ChromakeyColor(0, 0.5f, 0)
 	{ }
 
+public:
+	/** Get Type of Chromakey: disabled, filling all internal frustum or rendering Chromakey actors. */
+	EDisplayClusterShaderParametersICVFX_ChromakeySource GetChromakeyType(const FDisplayClusterConfigurationICVFX_StageSettings& InStageSettings) const;
+
+	/** Get the Chromakey render settings
+	 * Return nullptr if Chromakey: is disabled or its type does not use rendering
+	 */
+	const FDisplayClusterConfigurationICVFX_ChromakeyRenderSettings* GetChromakeyRenderSettings(const FDisplayClusterConfigurationICVFX_StageSettings& InStageSettings) const;
+
+	/** Get the writable Chromakey render settings
+	 * Return nullptr if Chromakey: is disabled or its type does not use rendering
+	 */
+	FDisplayClusterConfigurationICVFX_ChromakeyRenderSettings* GetWritableChromakeyRenderSettings(const FDisplayClusterConfigurationICVFX_StageSettings& InStageSettings);
+
+	/** Get the Chromakey color value */
+	const FLinearColor& GetChromakeyColor(const FDisplayClusterConfigurationICVFX_StageSettings& InStageSettings) const;
+
+	/** Get the Chromakey overlap color value */
+	const FLinearColor& GetOverlapChromakeyColor(const FDisplayClusterConfigurationICVFX_StageSettings& InStageSettings) const;
+
+	/** Get the Chromakey markers settings
+	 * Return nullptr if Chromakey markers is :disabled
+	 */
+	const FDisplayClusterConfigurationICVFX_ChromakeyMarkers* GetChromakeyMarkers(const FDisplayClusterConfigurationICVFX_StageSettings& InStageSettings) const;
+	
+	/** Get the Chromakey overlap markers settings
+	 * Return nullptr if Chromakey overlap markers is disabled
+	 */
+	const FDisplayClusterConfigurationICVFX_ChromakeyMarkers* GetOverlapChromakeyMarkers(const FDisplayClusterConfigurationICVFX_StageSettings& InStageSettings) const;
+
+private:
+	/** Checks the input parameter and returns it or nullptr if it cannot be used. */
+	const FDisplayClusterConfigurationICVFX_ChromakeyMarkers* ImplGetChromakeyMarkers(const FDisplayClusterConfigurationICVFX_ChromakeyMarkers* InValue) const;
+
 public:
 	/** Set to True to fill the inner frustum with the specified Chromakey Color. */
 	UPROPERTY(EditAnywhere, BlueprintReadWrite, Category = NDisplay, meta = (DisplayName = "Enable Inner Frustum Chromakey"))
@@ -364,6 +408,13 @@ struct DISPLAYCLUSTERCONFIGURATION_API FDisplayClusterConfigurationICVFX_Lightca
 {
 	GENERATED_BODY()
 
+public:
+	/* Returns true if the use of the LightCard is allowed */
+	bool ShouldUseLightCard(const FDisplayClusterConfigurationICVFX_StageSettings& InStageSettings) const;
+
+	/* Returns true if the use of the UVLightCard is allowed */
+	bool ShouldUseUVLightCard(const FDisplayClusterConfigurationICVFX_StageSettings& InStageSettings) const;
+
 public:
 	/** Enable Light Cards */
 	UPROPERTY(EditAnywhere, BlueprintReadWrite, Category = "Light Cards", meta = (DisplayName = "Enable Light Cards"))
@@ -610,6 +661,27 @@ public:
 	bool IsChromakeyViewportSettingsEqual_Editor(const FString& InClusterNodeId1, const FString& InClusterNodeId2) const;
 #endif
 
+	/** Return calculated soft edges values. */
+	FVector4 GetCameraSoftEdge(const FDisplayClusterConfigurationICVFX_StageSettings& InStageSettings) const;
+
+	/** Get camera FOV multiplier. */
+	float GetCameraFieldOfViewMultiplier(const FDisplayClusterConfigurationICVFX_StageSettings& InStageSettings) const;
+
+	/** Get camera adapt resolution ratio. */
+	float GetCameraAdaptResolutionRatio(const FDisplayClusterConfigurationICVFX_StageSettings& InStageSettings) const;
+
+	/** Get camera buffer ratio. */
+	float GetCameraBufferRatio(const FDisplayClusterConfigurationICVFX_StageSettings& InStageSettings) const;
+
+	/** Get camera frame size. */
+	FIntPoint GetCameraFrameSize(const FDisplayClusterConfigurationICVFX_StageSettings& InStageSettings) const;
+
+	/** Get camera frame aspect ratio. */
+	float GetCameraFrameAspectRatio(const FDisplayClusterConfigurationICVFX_StageSettings& InStageSettings) const;
+
+	/** Get camera border settings. */
+	bool GetCameraBorder(const FDisplayClusterConfigurationICVFX_StageSettings& InStageSettings, FLinearColor& OutBorderColor, float& OutBorderThickness) const;
+
 public:
 	/** Render the inner frustum for this ICVFX camera. */
 	UPROPERTY(EditAnywhere, BlueprintReadWrite, Category = "In Camera VFX", meta = (DisplayName = "Enable Inner Frustum"))
@@ -690,6 +762,9 @@ public:
 	/** Return LightCard OCIO configuration for the specified viewport. Return nullptr if no OCIO. */
 	const FOpenColorIOColorConversionSettings* FindLightcardOCIOConfiguration(const FString& InViewportId) const;
 
+	/** Get camera overlapping rendering mode.*/
+	EDisplayClusterShaderParametersICVFX_CameraOverlappingRenderMode GetCameraOverlappingRenderMode() const;
+
 public:
 	/** Enable/disable the inner frustum on all ICVFX cameras. */
 	UPROPERTY(EditAnywhere, BlueprintReadWrite, Category = "In Camera VFX", meta = (DisplayName = "Enable Inner Frustum"))
diff --git a/Engine/Plugins/Runtime/nDisplay/Source/DisplayClusterConfiguration/Public/DisplayClusterConfigurationTypes_Viewport.h b/Engine/Plugins/Runtime/nDisplay/Source/DisplayClusterConfiguration/Public/DisplayClusterConfigurationTypes_Viewport.h
--- a/Engine/Plugins/Runtime/nDisplay/Source/DisplayClusterConfiguration/Public/DisplayClusterConfigurationTypes_Viewport.h
+++ b/Engine/Plugins/Runtime/nDisplay/Source/DisplayClusterConfiguration/Public/DisplayClusterConfigurationTypes_Viewport.h
@@ -14,11 +14,16 @@
 #include "DisplayClusterConfigurationTypes_Postprocess.h"
 #include "DisplayClusterConfigurationTypes_ViewportRemap.h"
 
+#include "Containers/DisplayClusterShader_Enums.h"
+#include "Render/Viewport/Containers/DisplayClusterViewport_EnumsICVFX.h"
+
 #include "OpenColorIOColorSpace.h"
 #include "Engine/Scene.h"
 
 #include "DisplayClusterConfigurationTypes_Viewport.generated.h"
 
+struct FDisplayClusterConfigurationICVFX_StageSettings;
+
 USTRUCT(Blueprintable)
 struct DISPLAYCLUSTERCONFIGURATION_API FDisplayClusterConfigurationViewport_Overscan
 {
@@ -59,6 +64,13 @@ struct DISPLAYCLUSTERCONFIGURATION_API FDisplayClusterConfigurationViewport_ICVF
 {
 	GENERATED_BODY()
 
+public:
+	/** Get ligthcard render mode for this viewport. */
+	EDisplayClusterShaderParametersICVFX_LightCardRenderMode GetLightCardRenderMode(const FDisplayClusterConfigurationICVFX_StageSettings& InStageSettings) const;
+
+	/** Get ICVFX settings flags for viewport*/
+	EDisplayClusterViewportICVFXFlags GetViewportICVFXFlags(const FDisplayClusterConfigurationICVFX_StageSettings& InStageSettings) const;
+
 public:
 	/** Enable in-camera VFX for this Viewport (works only with supported Projection Policies) */
 	UPROPERTY(EditAnywhere, BlueprintReadWrite, Category = "In Camera VFX")
diff --git a/Engine/Plugins/Runtime/nDisplay/Source/DisplayClusterShaders/Public/Containers/DisplayClusterShader_Enums.h b/Engine/Plugins/Runtime/nDisplay/Source/DisplayClusterShaders/Public/Containers/DisplayClusterShader_Enums.h
new file mode 100644
--- /dev/null
+++ b/Engine/Plugins/Runtime/nDisplay/Source/DisplayClusterShaders/Public/Containers/DisplayClusterShader_Enums.h
@@ -0,0 +1,48 @@
+// Copyright Epic Games, Inc. All Rights Reserved.
+
+#pragma once
+
+#include "CoreMinimal.h"
+
+/**
+ * Rendering mode for chromakey
+ */
+enum class EDisplayClusterShaderParametersICVFX_ChromakeySource : uint8
+{
+	// Dont use chromakey
+	Disabled = 0,
+
+	// Render color over camera frame
+	FrameColor,
+
+	// Render specified layer from scene
+	ChromakeyLayers,
+
+	Unknown,
+};
+
+/**
+ * The rendering order of the lightcard layer
+ */
+enum class EDisplayClusterShaderParametersICVFX_LightCardRenderMode : uint8
+{
+	None = 0,
+
+	// Render incamera frame over lightcard
+	Under,
+
+	// Over lightcard over incamera frame
+	Over,
+};
+
+/**
+ * Rendering mode for overlapping areas of camera projections
+ */
+enum class EDisplayClusterShaderParametersICVFX_CameraOverlappingRenderMode : uint8
+{
+	// Disabled
+	None = 0,
+
+	// The overlapping area is rendered at the very end, on top of all the layers
+	FinalPass
+};
diff --git a/Engine/Plugins/Runtime/nDisplay/Source/DisplayClusterShaders/Public/ShaderParameters/DisplayClusterShaderParameters_ICVFX.h b/Engine/Plugins/Runtime/nDisplay/Source/DisplayClusterShaders/Public/ShaderParameters/DisplayClusterShaderParameters_ICVFX.h
--- a/Engine/Plugins/Runtime/nDisplay/Source/DisplayClusterShaders/Public/ShaderParameters/DisplayClusterShaderParameters_ICVFX.h
+++ b/Engine/Plugins/Runtime/nDisplay/Source/DisplayClusterShaders/Public/ShaderParameters/DisplayClusterShaderParameters_ICVFX.h
@@ -3,49 +3,9 @@
 #pragma once
 
 #include "CoreMinimal.h"
+#include "Containers/DisplayClusterShader_Enums.h"
 #include "RHI.h"
 
-/**
- * Rendering mode for chromakey
- */
-enum class EDisplayClusterShaderParametersICVFX_ChromakeySource : uint8
-{
-	// Dont use chromakey
-	Disabled = 0,
-
-	// Render color over camera frame
-	FrameColor,
-
-	// Render specified layer from scene
-	ChromakeyLayers,
-
-	Unknown,
-};
-
-/**
- * The rendering order of the lightcard layer
- */
-enum class EDisplayClusterShaderParametersICVFX_LightCardRenderMode : uint8
-{
-	// Render incamera frame over lightcard
-	Under = 0,
-
-	// Over lightcard over incamera frame
-	Over,
-};
-
-/**
- * Rendering mode for overlapping areas of camera projections
- */
-enum class EDisplayClusterShaderParametersICVFX_CameraOverlappingRenderMode : uint8
-{
-	// Disabled
-	None = 0,
-
-	// The overlapping area is rendered at the very end, on top of all the layers
-	FinalPass
-};
-
 /**
  * Camera view projection data
  */
@@ -192,7 +152,7 @@ public:
 		/** Gets whether a valid texture is available to render chromakey markers for any overlapping ICVFX frustums  */
 		inline bool IsOverlapChromakeyMarkerUsed() const
 		{
-			return OverlapChromakeyMarkerTextureRHI.IsValid();	
+			return OverlapChromakeyMarkerTextureRHI.IsValid();
 		}
 
 		inline void SetViewProjection(const FDisplayClusterShaderParametersICVFX_CameraViewProjection& InLocalSpaceViewProjection)
