nDisplay: Preview Share Component can now referenece another nDisplay actor and pull its textures directly.

This should faciliate the techviz workflow of having the source and destination nDisplay actors in the same Unreal instance.

In addition, on option for this PullActor mode to sync Icvfx cameras was added, which will sync lens and transform settings between the two nDisplay actors (in either direction).

diff --git a/Engine/Plugins/Runtime/nDisplay/Source/DisplayCluster/Private/Game/EngineClasses/Scene/DisplayClusterPreviewShareComponent.cpp b/Engine/Plugins/Runtime/nDisplay/Source/DisplayCluster/Private/Game/EngineClasses/Scene/DisplayClusterPreviewShareComponent.cpp
--- a/Engine/Plugins/Runtime/nDisplay/Source/DisplayCluster/Private/Game/EngineClasses/Scene/DisplayClusterPreviewShareComponent.cpp
+++ b/Engine/Plugins/Runtime/nDisplay/Source/DisplayCluster/Private/Game/EngineClasses/Scene/DisplayClusterPreviewShareComponent.cpp
@@ -2,6 +2,8 @@
 
 #include "Components/DisplayClusterPreviewShareComponent.h"
 
+#include "CineCameraComponent.h"
+#include "Components/DisplayClusterICVFXCameraComponent.h"
 #include "Components/DisplayClusterPreviewComponent.h"
 #include "DisplayClusterRootActor.h"
 #include "Engine/TextureRenderTarget2D.h"
@@ -12,6 +14,10 @@
 #include "SharedMemoryMediaOutput.h"
 #include "SharedMemoryMediaSource.h"
 
+#if WITH_EDITOR
+	#include "LevelEditor.h"
+#endif // WITH_EDITOR
+
 #if WITH_EDITOR
 
 namespace UE::DisplayClusterPreviewShare
@@ -99,7 +105,7 @@ bool UDisplayClusterPreviewShareComponent::AllowedToShare() const
 		return false;
 	}
 
-	// We don't allow sharing is the parent is the active root actor
+	// We don't allow sharing if the parent is the active root actor
 	{
 		IDisplayCluster& Display = IDisplayCluster::Get();
 		IDisplayClusterGameManager* GameMgr = Display.GetGameMgr();
@@ -130,6 +136,11 @@ void UDisplayClusterPreviewShareComponent::PostEditChangeProperty(FPropertyChang
 		// All the names are now invalid, so we need to close all media.
 		CloseAllMedia();
 	}
+	else if (PropertyChangedEvent.GetPropertyName() == GET_MEMBER_NAME_CHECKED(UDisplayClusterPreviewShareComponent, SourceNDisplayActor))
+	{
+		// Since we're changing the source actor, we should restore the original settings before taking the new ones.
+		RestoreRootActorOriginalSettings();
+	}
 }
 #endif // WITH_EDITOR
 
@@ -157,6 +168,18 @@ void UDisplayClusterPreviewShareComponent::ModeChanged()
 
 		break;
 	}
+	case EDisplayClusterPreviewShareMode::PullActor:
+	{
+		if (RootActor)
+		{
+			// Tell root actor that it doesn't need to keep up the preview enabled for us.
+			RootActor->RemovePreviewEnableOverride(reinterpret_cast<uint8*>(this));
+		}
+
+		SetTickEnable(true);
+
+		break;
+	}
 	case EDisplayClusterPreviewShareMode::Send:
 	{
 		// If we're sending, we need to tell the root actor to generate the preview textures for us.
@@ -209,6 +232,11 @@ void UDisplayClusterPreviewShareComponent::RestoreRootActorOriginalSettings()
 		return;
 	}
 
+	if (OriginalSourceTextures.IsEmpty() && OriginalTextureReplaces.IsEmpty())
+	{
+		return;
+	}
+
 	for (UActorComponent* ActorComponent : RootActor->K2_GetComponentsByClass(UDisplayClusterPreviewComponent::StaticClass()))
 	{
 		UDisplayClusterPreviewComponent* PreviewComponent = Cast<UDisplayClusterPreviewComponent>(ActorComponent);
@@ -435,6 +463,260 @@ void UDisplayClusterPreviewShareComponent::TickSend()
 
 }
 
+void UDisplayClusterPreviewShareComponent::TickPullActor()
+{
+	ADisplayClusterRootActor* SourceRootActor = SourceNDisplayActor.Get();
+
+	if (!IsValid(SourceRootActor))
+	{
+		return;
+	}
+
+	ADisplayClusterRootActor* RootActor = Cast<ADisplayClusterRootActor>(GetOwner());
+
+	if (!RootActor)
+	{
+		return;
+	}
+
+	// Pull preview textures
+	PullPreviewFromSourceActor(SourceRootActor);
+
+	// Sync Icvfx cameras
+	switch (IcvfxCamerasSyncType)
+	{
+	case EDisplayClusterPreviewShareIcvfxSync::PullActor:
+		SyncIcvxCamerasFromSourceActor(SourceRootActor, RootActor);
+		break;
+
+	case EDisplayClusterPreviewShareIcvfxSync::PushActor:
+		SyncIcvxCamerasFromSourceActor(RootActor, SourceRootActor);
+		break;
+
+	default:
+		break;
+	}
+}
+
+void UDisplayClusterPreviewShareComponent::PullPreviewFromSourceActor(const ADisplayClusterRootActor* SourceRootActor)
+{
+	using namespace UE::DisplayClusterPreviewShare;
+
+	if (!SourceRootActor)
+	{
+		return;
+	}
+
+	// Grab destination root actor.
+
+	ADisplayClusterRootActor* RootActor = Cast<ADisplayClusterRootActor>(GetOwner());
+
+	if (!RootActor)
+	{
+		return;
+	}
+
+	// Disallow self referencing
+
+	if (RootActor == SourceRootActor)
+	{
+		return;
+	}
+
+	// Gather source preview components
+
+	TMap<FString, UDisplayClusterPreviewComponent*> SrcPreviewComponentWithUniqueName;
+
+	for (UActorComponent* SrcComponent : SourceRootActor->K2_GetComponentsByClass(UDisplayClusterPreviewComponent::StaticClass()))
+	{
+		UDisplayClusterPreviewComponent* SrcPreviewComponent = Cast<UDisplayClusterPreviewComponent>(SrcComponent);
+
+		if (!SrcPreviewComponent)
+		{
+			continue;
+		}
+
+		const FString SrcViewportKey = GenerateViewportKey(
+			SrcPreviewComponent->GetClusterNodeId(),
+			SrcPreviewComponent->GetViewportId()
+		);
+
+		SrcPreviewComponentWithUniqueName.Add(SrcViewportKey, SrcPreviewComponent);
+	}
+
+	// Iterate over destination preview components, find the matching source preview component, and assign the preview texture accordingly.
+
+	for (UActorComponent* DstComponent : RootActor->K2_GetComponentsByClass(UDisplayClusterPreviewComponent::StaticClass()))
+	{
+		UDisplayClusterPreviewComponent* DstPreviewComponent = Cast<UDisplayClusterPreviewComponent>(DstComponent);
+
+		if (!DstPreviewComponent)
+		{
+			continue;
+		}
+
+		// We will ultimately update the viewport so let's make sure it exists.
+
+		UDisplayClusterConfigurationViewport* DstViewport = GetViewportFromDCRA(RootActor, DstPreviewComponent->GetClusterNodeId(), DstPreviewComponent->GetViewportId());
+
+		if (!DstViewport)
+		{
+			continue;
+		}
+
+		// Each viewport gets associated with a uniquely name
+
+		const FString DstViewportKey = GenerateViewportKey(
+			DstPreviewComponent->GetClusterNodeId(),
+			DstPreviewComponent->GetViewportId()
+		);
+
+		// Find corresponding source preview component
+
+		UDisplayClusterPreviewComponent** SrcPreviewComponentPtr = SrcPreviewComponentWithUniqueName.Find(DstViewportKey);
+
+		if (!SrcPreviewComponentPtr || !*SrcPreviewComponentPtr)
+		{
+			continue;
+		}
+
+		// Assign destination preview texture from source preview texture.
+
+		const UDisplayClusterPreviewComponent* SrcPreviewComponent = *SrcPreviewComponentPtr;
+		check(SrcPreviewComponent);
+
+		UTextureRenderTarget2D* SrcPreviewTexture = SrcPreviewComponent->GetRenderTargetTexture();
+
+		if (!SrcPreviewTexture)
+		{
+			continue;
+		}
+
+		// Remember the original replacement textures, if they haven't been saved yet.
+
+		if (!OriginalTextureReplaces.Find(DstViewportKey))
+		{
+			OriginalTextureReplaces.Add(DstViewportKey, DstViewport->RenderSettings.Replace.bAllowReplace);
+		}
+
+		if (!OriginalSourceTextures.Find(DstViewportKey))
+		{
+			OriginalSourceTextures.Add(DstViewportKey, DstViewport->RenderSettings.Replace.SourceTexture);
+		}
+
+		// Override the replace texture values
+		DstViewport->RenderSettings.Replace.bAllowReplace = true;
+		DstViewport->RenderSettings.Replace.SourceTexture = SrcPreviewTexture;
+	}
+}
+
+void UDisplayClusterPreviewShareComponent::SyncIcvxCamerasFromSourceActor(const ADisplayClusterRootActor* SourceRootActor, const ADisplayClusterRootActor* DstRootActor)
+{
+	using namespace UE::DisplayClusterPreviewShare;
+
+	if (!SourceRootActor || !DstRootActor)
+	{
+		return;
+	}
+
+	// Disallow self referencing
+
+	if (DstRootActor == SourceRootActor)
+	{
+		return;
+	}
+
+	// Gather source icvfx cameras
+
+	TMap<FString, UDisplayClusterICVFXCameraComponent*> SrcCameraComponentWithId;
+
+	for (UActorComponent* SrcComponent : SourceRootActor->K2_GetComponentsByClass(UDisplayClusterICVFXCameraComponent::StaticClass()))
+	{
+		UDisplayClusterICVFXCameraComponent* SrcCameraComponent = Cast<UDisplayClusterICVFXCameraComponent>(SrcComponent);
+
+		if (!SrcCameraComponent)
+		{
+			continue;
+		}
+
+		SrcCameraComponentWithId.Add(SrcCameraComponent->GetCameraUniqueId(), SrcCameraComponent);
+	}
+
+	// Iterate over destination cameras and sync the desired settings
+
+	for (UActorComponent* DstComponent : DstRootActor->K2_GetComponentsByClass(UDisplayClusterICVFXCameraComponent::StaticClass()))
+	{
+		UDisplayClusterICVFXCameraComponent* DstIcvfxCameraComponent = Cast<UDisplayClusterICVFXCameraComponent>(DstComponent);
+
+		if (!DstIcvfxCameraComponent)
+		{
+			continue;
+		}
+
+		UDisplayClusterICVFXCameraComponent** SrcIcvfxCameraComponentPtr = SrcCameraComponentWithId.Find(DstIcvfxCameraComponent->GetCameraUniqueId());
+
+		if (!SrcIcvfxCameraComponentPtr || !*SrcIcvfxCameraComponentPtr)
+		{
+			continue;
+		}
+
+		UDisplayClusterICVFXCameraComponent* SrcIcvfxCameraComponent = *SrcIcvfxCameraComponentPtr;
+		check(SrcIcvfxCameraComponent);
+
+		// This one sould be the actual source camera, e.g. the camera component of the referenced cine camera.
+		const UCameraComponent* SrcCameraComponent = SrcIcvfxCameraComponent->GetCameraComponent();
+		check(SrcCameraComponent);
+		
+		// Get the source camera transform with respect to the Source root actor
+		const FTransform SrcCameraToSrcRootActorTransform = SrcCameraComponent->GetComponentTransform().GetRelativeTransform(SourceRootActor->GetTransform());
+
+		// Apply this transform to the destination camera
+
+		UCameraComponent* DstCameraComponent = DstIcvfxCameraComponent->GetCameraComponent();
+
+		AActor* DstCameraOwnerActor = DstCameraComponent->GetOwner();
+		
+		if (!DstCameraOwnerActor)
+		{
+			continue;
+		}
+
+		const FTransform DstCameraComponentWorldTransform = SrcCameraToSrcRootActorTransform * DstRootActor->GetTransform();
+
+		// If the dst owning actor is a cine camera, then we move the cine camera itself instead of the component
+		if (DstCameraOwnerActor->IsA(ACineCameraActor::StaticClass()))
+		{
+			// DCr * Aw = DCw
+			// Aw = DCr_i * DCw
+
+			const FTransform DstCameraToOwingActorTransform = DstCameraComponent->GetComponentTransform().GetRelativeTransform(DstCameraOwnerActor->GetTransform());
+
+			DstCameraOwnerActor->SetActorTransform(DstCameraToOwingActorTransform.Inverse() * DstCameraComponentWorldTransform);
+		}
+		else
+		{
+			// Otherwise we set the world transform of the component directly
+			DstCameraComponent->SetWorldTransform(DstCameraComponentWorldTransform);
+		}
+
+		// Copy relevant lens settings. 
+		{
+			// Currently we only support this when source and destination are cine camera components.
+
+			const UCineCameraComponent* SrcCineCameraComponent = Cast<UCineCameraComponent>(SrcCameraComponent);
+			UCineCameraComponent* DstCineCameraComponent = Cast<UCineCameraComponent>(DstCameraComponent);
+
+			if (!SrcCineCameraComponent || !DstCineCameraComponent)
+			{
+				continue;
+			}
+
+			DstCineCameraComponent->Filmback = SrcCineCameraComponent->Filmback;
+			DstCineCameraComponent->SetCurrentFocalLength(SrcCineCameraComponent->CurrentFocalLength);
+		}
+	}
+}
+
 void UDisplayClusterPreviewShareComponent::TickReceive()
 {
 	using namespace UE::DisplayClusterPreviewShare;
@@ -530,7 +812,7 @@ void UDisplayClusterPreviewShareComponent::TickReceive()
 			MediaTexture->SetMediaPlayer(MediaPlayer);
 			MediaTexture->UpdateResource();
 
-			// Assign replacement texture
+			// Remember original replacement textures
 
 			OriginalTextureReplaces.Add(ViewportKey, Viewport->RenderSettings.Replace.bAllowReplace);
 			OriginalSourceTextures.Add(ViewportKey, Viewport->RenderSettings.Replace.SourceTexture);
@@ -606,6 +888,11 @@ void UDisplayClusterPreviewShareComponent::TickComponent(float DeltaTime, ELevel
 	{
 		break;
 	}
+	case EDisplayClusterPreviewShareMode::PullActor:
+	{
+		TickPullActor();
+		break;
+	}
 	case EDisplayClusterPreviewShareMode::Send:
 	{
 		TickSend();
@@ -638,14 +925,72 @@ void UDisplayClusterPreviewShareComponent::DestroyComponent(bool bPromoteChildre
 	// We try to leave the root actor as it originally was.
 	RestoreRootActorOriginalSettings();
 
+	if (GIsEditor)
+	{
+		// Unregister Map Change Events
+		if (FLevelEditorModule* LevelEditor = FModuleManager::GetModulePtr<FLevelEditorModule>("LevelEditor"))
+		{
+			LevelEditor->OnMapChanged().RemoveAll(this);
+		}
+	}
+
 #endif // WITH_EDITOR
 
 	Super::DestroyComponent(bPromoteChildren);
 }
 
+void UDisplayClusterPreviewShareComponent::HandleMapChanged(UWorld* InWorld, EMapChangeType InMapChangeType)
+{
+	// We remove the reference to the source actor preview texture to avoid this being interpreted as a 
+	// memory leak when unloadin the source actor map while the current map is kept open.
+	if (Mode == EDisplayClusterPreviewShareMode::PullActor)
+	{
+		RestoreRootActorOriginalSettings();
+	}
+}
+
 #endif // WITH_EDITOR in bulk
 
 
+void UDisplayClusterPreviewShareComponent::Serialize(FArchive& Ar)
+{
+#if WITH_EDITOR
+	if (Ar.IsSaving())
+	{
+		// Because for speed we may be directly referencing the source nDisplay actor preview textures,
+		// but these may belong to a different world, we must clear these references when saving.
+
+		if (Mode == EDisplayClusterPreviewShareMode::PullActor)
+		{
+			RestoreRootActorOriginalSettings();
+		}
+	}
+#endif // WITH_EDITOR
+
+	Super::Serialize(Ar);
+}
+
+void UDisplayClusterPreviewShareComponent::PostLoad()
+{
+	Super::PostLoad();
+
+#if WITH_EDITOR
+	// Make sure we are in the right state after loading
+	ModeChanged();
+
+	if (GIsEditor)
+	{
+		// Register for Map Change Events. We do this to remove texture external texture reference that
+		// prevents proper world GC. In particular, our owing actor reference to the source nDisplay actor preview texture.
+		if (FLevelEditorModule* LevelEditor = FModuleManager::GetModulePtr<FLevelEditorModule>("LevelEditor"))
+		{
+			LevelEditor->OnMapChanged().AddUObject(this, &UDisplayClusterPreviewShareComponent::HandleMapChanged);
+		}
+	}
+
+#endif // WITH_EDITOR
+}
+
 void UDisplayClusterPreviewShareComponent::SetMode(EDisplayClusterPreviewShareMode NewMode)
 {
 #if WITH_EDITORONLY_DATA
@@ -687,4 +1032,3 @@ void UDisplayClusterPreviewShareComponent::SetUniqueName(const FString& NewUniqu
 
 #endif // WITH_EDITOR
 }
-
diff --git a/Engine/Plugins/Runtime/nDisplay/Source/DisplayCluster/Public/Components/DisplayClusterPreviewShareComponent.h b/Engine/Plugins/Runtime/nDisplay/Source/DisplayCluster/Public/Components/DisplayClusterPreviewShareComponent.h
--- a/Engine/Plugins/Runtime/nDisplay/Source/DisplayCluster/Public/Components/DisplayClusterPreviewShareComponent.h
+++ b/Engine/Plugins/Runtime/nDisplay/Source/DisplayCluster/Public/Components/DisplayClusterPreviewShareComponent.h
@@ -5,9 +5,14 @@
 #include "Components/ActorComponent.h"
 #include "Containers/Map.h"
 
+#if WITH_EDITOR
+	#include "UnrealEdMisc.h"
+#endif //WITH_EDITOR
+
 #include "DisplayClusterPreviewShareComponent.generated.h"
 
 
+class ADisplayClusterRootActor;
 class UMediaCapture;
 class UMediaOutput;
 class UMediaPlayer;
@@ -20,9 +25,12 @@ class UTexture;
 UENUM()
 enum class EDisplayClusterPreviewShareMode
 {
-	/** Sharing disabled */
+	/** Preview share disabled. */
 	None,
 
+	/** Pull from Source nDisplay Actor (if set) */
+	PullActor,
+
 	/** Sends the viewport textures for sharing */
 	Send,
 
@@ -30,6 +38,20 @@ enum class EDisplayClusterPreviewShareMode
 	Receive,
 };
 
+/** Available Icvfx camera sync types */
+UENUM()
+enum class EDisplayClusterPreviewShareIcvfxSync
+{
+	/** Icvfx cameras will not be synced */
+	None,
+
+	/** Pull from Source nDisplay Actor */
+	PullActor,
+
+	/** Pushes to source display actor */
+	PushActor,
+};
+
 
 /**
  * nDisplay Viewport preview share component
@@ -62,9 +84,14 @@ public:
 #endif // WITH_EDITOR
 
 	//~ UObject interface begin
+
+	virtual void Serialize(FArchive& Ar) override;
+	virtual void PostLoad() override;
+
 #if WITH_EDITOR
 	virtual void PostEditChangeProperty(FPropertyChangedEvent& PropertyChangedEvent) override;
 #endif // WITH_EDITOR
+
 	//~ UObject interface end
 
 	/** Sets the sharing mode */
@@ -80,13 +107,21 @@ public:
 #if WITH_EDITORONLY_DATA
 
 	/** Current sharing mode of this component */
-	UPROPERTY(Transient, EditAnywhere, Setter=SetMode, BlueprintSetter=SetMode, Category=Sharing)
+	UPROPERTY(EditAnywhere, Setter=SetMode, BlueprintSetter=SetMode, Category=Sharing)
 	EDisplayClusterPreviewShareMode Mode = EDisplayClusterPreviewShareMode::None;
 
 	/** Current unique name of this component, which should match between sender and receiver of viewport textures */
-	UPROPERTY(EditAnywhere, Setter = SetUniqueName, BlueprintSetter = SetUniqueName, Category = Sharing)
+	UPROPERTY(EditAnywhere, Setter = SetUniqueName, BlueprintSetter = SetUniqueName, Category = Sharing, meta = (EditCondition = "Mode == EDisplayClusterPreviewShareMode::Send || Mode == EDisplayClusterPreviewShareMode::Receive"))
 	FString UniqueName;
 
+	/** The source nDisplay actor to pull the preview from */
+	UPROPERTY(EditAnywhere, meta = (DisplayName = "Source nDisplay Actor"), Category = Sharing, meta = (EditCondition = "Mode == EDisplayClusterPreviewShareMode::PullActor"))
+	TSoftObjectPtr<ADisplayClusterRootActor> SourceNDisplayActor;
+
+	/** Type of Icvfx camera sync to be performed between the Source nDisplay actor and the owner of this component */
+	UPROPERTY(EditAnywhere, Category = Sharing, meta = (EditCondition = "Mode == EDisplayClusterPreviewShareMode::PullActor"))
+	EDisplayClusterPreviewShareIcvfxSync IcvfxCamerasSyncType = EDisplayClusterPreviewShareIcvfxSync::PullActor;
+
 #endif // WITH_EDITORONLY_DATA
 
 private:
@@ -114,12 +149,24 @@ private:
 	/** Logic that should run every tick when in Receive mode */
 	void TickReceive();
 
+	/** Logic that should run every tick when in Pull Actor mode */
+	void TickPullActor();
+
+	/** Retrieves preview textures from given source nDisplay actor and uses them to drive the preview textures of the owning actor */
+	void PullPreviewFromSourceActor(const ADisplayClusterRootActor* SourceRootActor);
+
+	/** Syncs the Icvfx cameras from the given source nDisplay actor to the given destination nDisplay actor */
+	void SyncIcvxCamerasFromSourceActor(const ADisplayClusterRootActor* SrcRootActor, const ADisplayClusterRootActor* DstRootActor);
+
 	/** Restores settings in the nDisplay actor that were altered by this component to achieve its intended purpose */
 	void RestoreRootActorOriginalSettings();
 
 	/** Enables/Disables component ticking */
 	void SetTickEnable(const bool bEnable);
 
+	/** Called when the editor map is changed. Used to remove unwanted references to external maps. */
+	void HandleMapChanged(UWorld* InWorld, EMapChangeType InMapChangeType);
+
 #endif // WITH_EDITOR
 
 private:
