Make Preload DDC command non-blocking

diff --git a/Engine/Plugins/Runtime/nDisplay/Source/DisplayClusterPreloadDerivedDataCache/Private/DisplayClusterPreloadDerivedDataCacheModule.cpp b/Engine/Plugins/Runtime/nDisplay/Source/DisplayClusterPreloadDerivedDataCache/Private/DisplayClusterPreloadDerivedDataCacheModule.cpp
--- a/Engine/Plugins/Runtime/nDisplay/Source/DisplayClusterPreloadDerivedDataCache/Private/DisplayClusterPreloadDerivedDataCacheModule.cpp
+++ b/Engine/Plugins/Runtime/nDisplay/Source/DisplayClusterPreloadDerivedDataCache/Private/DisplayClusterPreloadDerivedDataCacheModule.cpp
@@ -3,15 +3,11 @@
 #include "DisplayClusterPreloadDerivedDataCacheModule.h"
 
 #include "DisplayClusterPreloadDerivedDataCacheLog.h"
+#include "DisplayClusterPreloadDerivedDataCacheWorker.h"
 
 #include "Commandlets/DerivedDataCacheCommandlet.h"
-#include "Framework/Docking/TabManager.h"
 #include "Framework/Notifications/NotificationManager.h"
 #include "Interfaces/IProjectManager.h"
-#include "Internationalization/Regex.h"
-#include "Misc/App.h"
-#include "Misc/DataDrivenPlatformInfoRegistry.h"
-#include "Misc/ScopedSlowTask.h"
 #include "Modules/ModuleManager.h"
 #include "ToolMenu.h"
 #include "ToolMenus.h"
@@ -40,248 +36,28 @@ void FDisplayClusterPreloadDerivedDataCacheModule::OnFEngineLoopInitComplete()
 		LOCTEXT("PreloadDerivedDataCache", "Preload Derived Data Cache"),
 		LOCTEXT("PreloadDerivedDataCacheTooltip", "Precompile all shaders and preprocess Nanite and other heavy data ahead of time for all assets in this project for all checked platforms in Project Settings > Supported Platforms. \nThis is useful to avoid shader compilation and pop-in when opening a level or running PIE/Standalone. \nNote that this operation can take a very long time to complete and can take up a lot of disk space, depending on the size of your project."),
 		FSlateIcon(),
-		FUIAction(FExecuteAction::CreateRaw(this, &FDisplayClusterPreloadDerivedDataCacheModule::LaunchAndCommunicateWithProcess)));
+		FUIAction(FExecuteAction::CreateRaw(this, &FDisplayClusterPreloadDerivedDataCacheModule::CreateAsyncTaskWorker)));
 }
 
 void FDisplayClusterPreloadDerivedDataCacheModule::ShutdownModule()
 {
 	UToolMenus::UnregisterOwner(this);
 	FCoreDelegates::OnFEngineLoopInitComplete.RemoveAll(this);
-}
-
-void FDisplayClusterPreloadDerivedDataCacheModule::RegexParseForLoadingProgress(
-	const FText& LoadingPackagesFormat, const FText& UnknownAmount, FScopedSlowTask& SlowTask, int32& LoadingTotal, FRegexMatcher& LoadingProgressRegex)
-{
-	while (LoadingProgressRegex.FindNext())
-	{
-		SlowTask.CompletedWork = FCString::Atoi(*LoadingProgressRegex.GetCaptureGroup(1));
-		SlowTask.TotalAmountOfWork = LoadingTotal;
-
-		if (LoadingTotal > INDEX_NONE)
-		{
-			SlowTask.DefaultMessage = FText::Format(LoadingPackagesFormat, FText::AsNumber(SlowTask.CompletedWork), FText::AsNumber(SlowTask.TotalAmountOfWork));
-		}
-		else
-		{
-			SlowTask.DefaultMessage = FText::Format(LoadingPackagesFormat, FText::AsNumber(SlowTask.CompletedWork), UnknownAmount);
-		}
-	}
-}
-
-void FDisplayClusterPreloadDerivedDataCacheModule::RegexParseForCompilationProgress(
-	const FText& CompilingAssetsFormat, const FText& UnknownAmount, FScopedSlowTask& SlowTask,
-	int32& CompileTotal, int32& AssetsWaitingToCompile, FRegexMatcher& CompileProgressRegex)
-{
-	while (CompileProgressRegex.FindNext())
-	{
-		AssetsWaitingToCompile = FCString::Atoi(*CompileProgressRegex.GetCaptureGroup(1));
-		if (CompileTotal == INDEX_NONE)
-		{
-			CompileTotal = AssetsWaitingToCompile;
-			SlowTask.TotalAmountOfWork = CompileTotal;
-		}
-
-		if (const int32 Progress = CompileTotal - AssetsWaitingToCompile; Progress > -1)
-		{
-			SlowTask.CompletedWork = Progress;
-			SlowTask.DefaultMessage = FText::Format(CompilingAssetsFormat, FText::AsNumber(SlowTask.CompletedWork), FText::AsNumber(SlowTask.TotalAmountOfWork));
-		}
-		else
-		{
-			SlowTask.DefaultMessage = FText::Format(CompilingAssetsFormat, UnknownAmount, UnknownAmount);
-		}
-	}
-}
-
-void FDisplayClusterPreloadDerivedDataCacheModule::LaunchAndCommunicateWithProcess()
-{
-	int32 OutResult = 0;
-	bool bOutCancelled = false;
 
-	FProcHandle ProcessHandle;
-
-	const FText EnumeratingPackages = LOCTEXT("EnumeratingPackages", "Enumerating Packages...");
-	const FText LoadingPackagesFormat = LOCTEXT("LoadingPackagesFormat", "Loading Packages ({0}/{1})...");
-	const FText CompilingAssetsFormat = LOCTEXT("CompilingAssetsFormat", "Compiling Assets ({0}/{1})...");
-	const FText UnknownAmount = LOCTEXT("Unknown", "Unknown");
-	
-	FScopedSlowTask SlowTask(1.0f, EnumeratingPackages);
-	SlowTask.MakeDialog(true);
-
-	// Show a partially completed progress bar so the end user can see something is happening
-	SlowTask.CompletedWork = 1;
-	SlowTask.TotalAmountOfWork = 2;
-
-	void* ReadPipe = nullptr;
-	void* WritePipe = nullptr;
-	verify(FPlatformProcess::CreatePipe(ReadPipe, WritePipe));
-
-	const FString CurrentExecutableName = FPlatformProcess::ExecutablePath();
-	const FString ProjectPath = FPaths::SetExtension(
-		FPaths::Combine(FPaths::ProjectDir(), FApp::GetProjectName()),".uproject");
-	const FString Arguments = ProjectPath + " " + GetDdcCommandletParams();
-	
-	UE_LOG(LogDisplayClusterPreloadDerivedDataCache, Display, TEXT("Running commandlet: %s %s"), *CurrentExecutableName, *Arguments);
-
-	uint32 ProcessID;
-	const bool bLaunchDetached = true;
-	const bool bLaunchHidden = true;
-	const bool bLaunchReallyHidden = true;
-	ProcessHandle = FPlatformProcess::CreateProc(
-		*CurrentExecutableName, *Arguments, bLaunchDetached, bLaunchHidden, bLaunchReallyHidden, &ProcessID,
-		0, nullptr, WritePipe, ReadPipe);
-
-	int32 LoadingTotal = INDEX_NONE;
-	int32 CompileTotal = INDEX_NONE;
-	int32 AssetsWaitingToCompile = INDEX_NONE;
-	
-	const FRegexPattern LoadingTotalPattern(TEXT("Display:\\s([0-9]+)\\spackages to load from command line arguments"));
-	const FRegexPattern LoadingProgressPattern(TEXT("Display:\\sLoading\\s\\(([0-9]+)"));
-	const FRegexPattern CompileProgressPattern(TEXT("Display:\\sWaiting for\\s([0-9]+)"));
-	
-	while (FPlatformProcess::IsProcRunning(ProcessHandle))
+	if (AsyncTaskWorker)
 	{
-		if (SlowTask.ShouldCancel())
-		{
-			bOutCancelled = true;
-			FPlatformProcess::TerminateProc(ProcessHandle);
-			break;
-		}
-
-		const FString LogString = FPlatformProcess::ReadPipe(ReadPipe);
-		if (LogString.IsEmpty())
-		{
-			continue;
-		}
-		UE_LOG(LogDisplayClusterPreloadDerivedDataCache, Display, TEXT("Commandlet Output: %s"), *LogString);
-
-		// Total number of assets to load
-		if (LoadingTotal == INDEX_NONE)
-		{
-			FRegexMatcher LoadingTotalRegex(LoadingTotalPattern, *LogString);
-
-			while (LoadingTotalRegex.FindNext())
-			{
-				LoadingTotal = FCString::Atoi(*LoadingTotalRegex.GetCaptureGroup(1));
-			}
-		}
-
-		// Progress on asset loading
-		FRegexMatcher LoadingProgressRegex(LoadingProgressPattern, *LogString);
-		RegexParseForLoadingProgress(LoadingPackagesFormat, UnknownAmount, SlowTask, LoadingTotal, LoadingProgressRegex);
-
-		// Progress on asset compilation
-		FRegexMatcher CompileProgressRegex(CompileProgressPattern, *LogString);
-		RegexParseForCompilationProgress(CompilingAssetsFormat, UnknownAmount, SlowTask, CompileTotal, AssetsWaitingToCompile,
-		                                 CompileProgressRegex);
-
-		SlowTask.EnterProgressFrame(0);
-		FPlatformProcess::Sleep(0.1);
+		AsyncTaskWorker->CancelTask();
+		AsyncTaskWorker = nullptr;
 	}
-	
-	FPlatformProcess::GetProcReturnCode(ProcessHandle, &OutResult);
-
-	FPlatformProcess::ClosePipe(ReadPipe, WritePipe);
-
-	CompleteCommandletAndShowNotification(OutResult, bOutCancelled, CurrentExecutableName, Arguments);
 }
 
-void FDisplayClusterPreloadDerivedDataCacheModule::CompleteCommandletAndShowNotification(
-	const int32 ResultCode, const bool bWasCancelled,
-	const FString& CurrentExecutableName, const FString& Arguments)
+void FDisplayClusterPreloadDerivedDataCacheModule::CreateAsyncTaskWorker()
 {
-	const bool bFinishedWithFailures = ResultCode != 0;
-
-	FText NotificationText = LOCTEXT("ToastCommandletSuccess", "Commandlet Finished Without Error!");
-	SNotificationItem::ECompletionState CompletionState = SNotificationItem::CS_Success;
-	if (bWasCancelled)
-	{
-		UE_LOG(LogDisplayClusterPreloadDerivedDataCache, Display, TEXT("#### Commandlet Process Cancelled ####"));
-		NotificationText = LOCTEXT("ToastCommandletCancelled", "Commandlet Process Cancelled");
-		CompletionState = SNotificationItem::CS_None;
-	}
-	else if (bFinishedWithFailures)
-	{
-		UE_LOG(LogDisplayClusterPreloadDerivedDataCache, Error, TEXT("#### Commandlet Finished With Errors ####"));
-		UE_LOG(LogDisplayClusterPreloadDerivedDataCache, Error, TEXT("%s %s"), *CurrentExecutableName, *Arguments);
-		UE_LOG(LogDisplayClusterPreloadDerivedDataCache, Error, TEXT("Return Code: %i"), ResultCode);
-
-		NotificationText = FText::Format(LOCTEXT("ToastCommandletFailure", "Commandlet Finished With Error(s), Return Code: {0}"), FText::AsNumber(ResultCode));
-		CompletionState = SNotificationItem::CS_Fail;
-	}
-	else
-	{
-		UE_LOG(LogDisplayClusterPreloadDerivedDataCache, Display, TEXT("#### Commandlet Finished Without Error ####"));
-	}
-
-	FNotificationInfo NotificationInfo(NotificationText);
-	NotificationInfo.bFireAndForget = bWasCancelled;
-
-	// "close" button
-	FNotificationButtonInfo ButtonOK(LOCTEXT("ButtonOK", "OK"), FText(),
-		FSimpleDelegate::CreateLambda([this]()
-		{
-			if (Notification.IsValid())
-			{
-				Notification->ExpireAndFadeout();
-			}
-		}));
-
-	// Should be visible in all cases
-	ButtonOK.VisibilityOnFail = ButtonOK.VisibilityOnNone = ButtonOK.VisibilityOnPending = ButtonOK.VisibilityOnSuccess =EVisibility::Visible;
-	NotificationInfo.ButtonDetails.Add(ButtonOK);
-
-	if (bFinishedWithFailures)
-	{
-		NotificationInfo.Hyperlink = FSimpleDelegate::CreateLambda([]() { FGlobalTabmanager::Get()->TryInvokeTab(FName("OutputLog")); });
-		NotificationInfo.HyperlinkText = LOCTEXT("ShowOutputLogHyperlink", "Show Output Log");
-	}
-
-	// Close previous notification if it's still persistent
-	if (Notification.IsValid())
-	{
-		Notification->ExpireAndFadeout();
-	}
-
-	Notification = FSlateNotificationManager::Get().AddNotification( NotificationInfo );
-	if ( Notification.IsValid() )
-	{
-		Notification->SetCompletionState(CompletionState);
-	}
-}
-
-FString FDisplayClusterPreloadDerivedDataCacheModule::GetDdcCommandletParams() const
-{
-	return FString::Printf(TEXT("-run=DerivedDataCache %s -fill -DDC=CreateInstalledEnginePak"), *GetTargetPlatformParams());
-}
-
-FString FDisplayClusterPreloadDerivedDataCacheModule::GetTargetPlatformParams() const
-{
-	TArray<FString> SupportedPlatforms;
-	FProjectStatus ProjectStatus;
-	
-	if(IProjectManager::Get().QueryStatusForCurrentProject(ProjectStatus))
-	{
-		for (const FDataDrivenPlatformInfo* DDPI : FDataDrivenPlatformInfoRegistry::GetSortedPlatformInfos(EPlatformInfoType::TruePlatformsOnly))
-		{
-			if (ProjectStatus.IsTargetPlatformSupported(DDPI->IniPlatformName))
-			{
-				const FString PlatformName = DDPI->IniPlatformName.ToString();
-				SupportedPlatforms.Add(PlatformName);
-
-				// Not all platforms have an editor target but they may in the future.
-				// It doesn't hurt to add targets that don't exist as they are ignored.
-				SupportedPlatforms.Add(PlatformName + "Editor");
-			}
-		}
-	}
-
-	if (SupportedPlatforms.Num() > 0)
+	if (AsyncTaskWorker)
 	{
-		return FString::Printf(TEXT("-TargetPlatform=%s"), *FString::Join(SupportedPlatforms, TEXT("+")));
+		AsyncTaskWorker->CancelTask();
 	}
-	
-	return FString();
+	AsyncTaskWorker = MakeUnique<FDisplayClusterPreloadDerivedDataCacheWorker>();
 }
 
 #undef LOCTEXT_NAMESPACE
diff --git a/Engine/Plugins/Runtime/nDisplay/Source/DisplayClusterPreloadDerivedDataCache/Private/DisplayClusterPreloadDerivedDataCacheWorker.cpp b/Engine/Plugins/Runtime/nDisplay/Source/DisplayClusterPreloadDerivedDataCache/Private/DisplayClusterPreloadDerivedDataCacheWorker.cpp
new file mode 100644
--- /dev/null
+++ b/Engine/Plugins/Runtime/nDisplay/Source/DisplayClusterPreloadDerivedDataCache/Private/DisplayClusterPreloadDerivedDataCacheWorker.cpp
@@ -0,0 +1,308 @@
+// Copyright Epic Games, Inc. All Rights Reserved.
+
+#include "DisplayClusterPreloadDerivedDataCacheWorker.h"
+
+#include "DisplayClusterPreloadDerivedDataCacheLog.h"
+
+#include "Commandlets/DerivedDataCacheCommandlet.h"
+#include "Framework/Docking/TabManager.h"
+#include "Framework/Notifications/NotificationManager.h"
+#include "HAL/PlatformProcess.h"
+#include "Interfaces/IProjectManager.h"
+#include "Internationalization/Regex.h"
+#include "Misc/App.h"
+#include "Misc/AsyncTaskNotification.h"
+#include "Misc/DataDrivenPlatformInfoRegistry.h"
+#include "Misc/ScopedSlowTask.h"
+#include "Widgets/Notifications/SNotificationList.h"
+
+#define LOCTEXT_NAMESPACE "FDisplayClusterPreloadDerivedDataCacheWorker"
+
+const FText EnumeratingAssetsTitle = LOCTEXT("EnumeratingAssetsTitle", "Enumerating Assets (Step 1/3)...");
+const FText EnumeratingAssetsProgressFormat = LOCTEXT("EnumeratingAssetsProgressFormat", "{0} assets discovered...");
+const FText LoadingAssetsTitle = LOCTEXT("LoadingAssetsTitle", "Loading Assets (Step 2/3)...");
+const FText CompilingAssetsTitle = LOCTEXT("CompilingAssetsTitle", "Compiling Assets (Step 3/3)...");
+const FText GenericProgressFormat = LOCTEXT("GenericProgressFormat", "{0}/{1} ({2}%)...");
+
+const FRegexPattern EnumeratingAssetsPattern(TEXT("Display:\\s+([0-9]+)\\)\\s.+")); // ex. "146) /Engine/EditorBlueprintResources/ActorMacros"
+const FRegexPattern LoadingTotalPattern(TEXT("Display:\\s([0-9]+)\\spackages to load from command line arguments")); // ex. "2833 packages to load from command line arguments"
+const FRegexPattern LoadingProgressPattern(TEXT("Display:\\sLoading\\s\\(([0-9]+)\\)")); // ex. "Loading (2833)"
+const FRegexPattern CompileProgressPattern(TEXT("Display:\\sWaiting for\\s([0-9]+)\\s.+ to finish")); // ex. "Waiting for 14163 Shaders to finish"
+
+FDisplayClusterPreloadDerivedDataCacheWorker::FDisplayClusterPreloadDerivedDataCacheWorker()
+{	
+	FAsyncTaskNotificationConfig NotificationConfig;
+	NotificationConfig.bCanCancel = true;
+	NotificationConfig.TitleText = EnumeratingAssetsTitle;
+	NotificationConfig.bKeepOpenOnFailure = true;
+	NotificationConfig.bKeepOpenOnSuccess = true;
+	NotificationConfig.ExpireDuration = 1.0f;
+	NotificationConfig.ProgressText = FText::Format(EnumeratingAssetsProgressFormat, FText::AsNumber(0));
+	ProgressNotification = MakeUnique<FAsyncTaskNotification>(NotificationConfig);
+	
+	verify(FPlatformProcess::CreatePipe(ReadPipe, WritePipe));
+
+	CurrentExecutableName = FPlatformProcess::ExecutablePath();
+	const FString ProjectPath = FPaths::SetExtension(
+		FPaths::Combine(FPaths::ProjectDir(), FApp::GetProjectName()),".uproject");
+	Arguments = ProjectPath + " " + GetDdcCommandletParams();
+	
+	UE_LOG(LogDisplayClusterPreloadDerivedDataCache, Display, TEXT("Running commandlet: %s %s"), *CurrentExecutableName, *Arguments);
+
+	uint32 ProcessID;
+	const bool bLaunchDetached = true;
+	const bool bLaunchHidden = true;
+	const bool bLaunchReallyHidden = true;
+	ProcessHandle = FPlatformProcess::CreateProc(
+		*CurrentExecutableName, *Arguments, bLaunchDetached, bLaunchHidden, bLaunchReallyHidden, &ProcessID,
+		0, nullptr, WritePipe, ReadPipe);
+	
+	Thread = TUniquePtr<FRunnableThread>(FRunnableThread::Create(this, TEXT("FDisplayClusterPreloadDerivedDataCacheWorker")));
+}
+
+FDisplayClusterPreloadDerivedDataCacheWorker::~FDisplayClusterPreloadDerivedDataCacheWorker()
+{
+	if (Thread)
+	{
+		constexpr bool bShouldWait = true;
+		Thread->Kill( bShouldWait );
+		Thread = nullptr;
+	}
+}
+
+void FDisplayClusterPreloadDerivedDataCacheWorker::CancelTask()
+{
+	bWasCancelled = true;
+	FPlatformProcess::TerminateProc(ProcessHandle);
+	CompleteCommandletAndShowNotification();
+}
+
+void FDisplayClusterPreloadDerivedDataCacheWorker::CompleteCommandletAndShowNotification()
+{
+	if (bWasCancelled)
+	{
+		UE_LOG(LogDisplayClusterPreloadDerivedDataCache, Display, TEXT("#### Commandlet Process Cancelled ####"));
+		ProgressNotification->SetComplete(
+			LOCTEXT("ToastCommandletCancelled", "Commandlet Process Cancelled"), FText(), false);
+
+		// Close the notification automatically
+		ProgressNotification->SetKeepOpenOnFailure(false);
+		return;
+	}
+	
+	FPlatformProcess::GetProcReturnCode(ProcessHandle, &ResultCode);
+	const bool bFinishedWithFailures = ResultCode != 0;
+
+	FAsyncNotificationStateData NotificationStateData;
+	NotificationStateData.ProgressText = FText();
+	NotificationStateData.TitleText = LOCTEXT("ToastCommandletSuccess", "Commandlet Finished Without Error!");
+	NotificationStateData.State = EAsyncTaskNotificationState::Success;
+
+	if (bFinishedWithFailures)
+	{
+		UE_LOG(LogDisplayClusterPreloadDerivedDataCache, Error, TEXT("#### Commandlet Finished With Errors ####"));
+		UE_LOG(LogDisplayClusterPreloadDerivedDataCache, Error, TEXT("%s %s"), *CurrentExecutableName, *Arguments);
+		UE_LOG(LogDisplayClusterPreloadDerivedDataCache, Error, TEXT("Return Code: %i"), ResultCode);
+
+		NotificationStateData.TitleText = LOCTEXT("ToastCommandletFailure", "Commandlet Finished With Error(s)");
+		NotificationStateData.ProgressText = FText::Format(LOCTEXT("ReturnCodeFormat", "Return Code: {0}"), FText::AsNumber(ResultCode));
+		NotificationStateData.State = EAsyncTaskNotificationState::Failure;
+		NotificationStateData.Hyperlink = FSimpleDelegate::CreateLambda([]() { FGlobalTabmanager::Get()->TryInvokeTab(FName("OutputLog")); });
+		NotificationStateData.HyperlinkText = LOCTEXT("ShowOutputLogHyperlink", "Show Output Log");
+	}
+	else
+	{
+		UE_LOG(LogDisplayClusterPreloadDerivedDataCache, Display, TEXT("#### Commandlet Finished Without Error ####"));
+	}
+
+	ProgressNotification->SetNotificationState(NotificationStateData);
+}
+
+void FDisplayClusterPreloadDerivedDataCacheWorker::RegexParseForEnumerationCount(const FString& LogString)
+{
+	FRegexMatcher EnumerationCountRegex(EnumeratingAssetsPattern, *LogString);
+	while (EnumerationCountRegex.FindNext())
+	{
+		ProgressNotification->SetTitleText(EnumeratingAssetsTitle);
+		
+		if (const int32 EnumerationCount = FCString::Atoi(*EnumerationCountRegex.GetCaptureGroup(1)); EnumerationCount > 0)
+		{
+			ProgressNotification->SetProgressText(
+			   FText::Format(EnumeratingAssetsProgressFormat, FText::AsNumber(EnumerationCount))
+		   );
+		}
+	}
+}
+
+void FDisplayClusterPreloadDerivedDataCacheWorker::RegexParseForLoadingProgress(const FString& LogString)
+{
+	FRegexMatcher LoadingProgressRegex(LoadingProgressPattern, *LogString);
+	while (LoadingProgressRegex.FindNext())
+	{
+		ProgressNotification->SetTitleText(LoadingAssetsTitle);
+		
+		if (LoadingTotal > INDEX_NONE)
+		{
+			if (const int32 LoadingProgress = FCString::Atoi(*LoadingProgressRegex.GetCaptureGroup(1)); LoadingProgress > 0)
+			{
+				ProgressNotification->SetProgressText(
+				   FText::Format(GenericProgressFormat,
+					   FText::AsNumber(LoadingProgress),
+					   FText::AsNumber(LoadingTotal),
+					   FText::AsNumber(FMath::FloorToInt((float)LoadingProgress / (float)LoadingTotal * 100))
+				   )
+			   );
+			}
+		}
+	}
+}
+
+void FDisplayClusterPreloadDerivedDataCacheWorker::RegexParseForCompilationProgress(const FString& LogString)
+{
+	FRegexMatcher CompileProgressRegex(CompileProgressPattern, *LogString);
+	while (CompileProgressRegex.FindNext())
+	{
+		AssetsWaitingToCompile = FCString::Atoi(*CompileProgressRegex.GetCaptureGroup(1));
+		if (CompileTotal == INDEX_NONE && AssetsWaitingToCompile > 0)
+		{
+			CompileTotal = AssetsWaitingToCompile;
+		}
+
+		if (CompileTotal > 0)
+		{
+			ProgressNotification->SetTitleText(CompilingAssetsTitle);
+
+			if (const int32 Progress = CompileTotal - AssetsWaitingToCompile; Progress > -1)
+			{
+				ProgressNotification->SetProgressText(
+					FText::Format(
+						GenericProgressFormat,
+						FText::AsNumber(Progress),
+						FText::AsNumber(CompileTotal),
+						FText::AsNumber(FMath::FloorToInt((float)Progress / (float)CompileTotal * 100))
+					)
+				);
+			}
+		}
+	}
+}
+
+FString FDisplayClusterPreloadDerivedDataCacheWorker::GetDdcCommandletParams() const
+{
+	return FString::Printf(TEXT("-run=DerivedDataCache %s -fill -DDC=CreateInstalledEnginePak"), *GetTargetPlatformParams());
+}
+
+FString FDisplayClusterPreloadDerivedDataCacheWorker::GetTargetPlatformParams() const
+{
+	TArray<FString> SupportedPlatforms;
+	FProjectStatus ProjectStatus;
+	
+	if(IProjectManager::Get().QueryStatusForCurrentProject(ProjectStatus))
+	{
+		for (const FDataDrivenPlatformInfo* DDPI : FDataDrivenPlatformInfoRegistry::GetSortedPlatformInfos(EPlatformInfoType::TruePlatformsOnly))
+		{
+			if (ProjectStatus.IsTargetPlatformSupported(DDPI->IniPlatformName))
+			{
+				const FString PlatformName = DDPI->IniPlatformName.ToString();
+				SupportedPlatforms.Add(PlatformName);
+
+				// Not all platforms have an editor target but they may in the future.
+				// It doesn't hurt to add targets that don't exist as they are ignored.
+				SupportedPlatforms.Add(PlatformName + "Editor");
+			}
+		}
+	}
+
+	if (SupportedPlatforms.Num() > 0)
+	{
+		return FString::Printf(TEXT("-TargetPlatform=%s"), *FString::Join(SupportedPlatforms, TEXT("+")));
+	}
+	
+	return FString();
+}
+
+FString FDisplayClusterPreloadDerivedDataCacheWorker::GetLastWholeLogBlock(const FString& LogString)
+{
+	FString CumulativeString = LastPartialLogLine + LogString;
+
+	// Look for a line break. If it exists, we can consider everything to the right of it to be a partial line while the left of it is a whole line
+	const int32 IndexOfLastNewLine = CumulativeString.Find(TEXT("\n"), ESearchCase::IgnoreCase, ESearchDir::FromEnd);
+
+	if (IndexOfLastNewLine > INDEX_NONE)
+	{
+		LastPartialLogLine = CumulativeString.RightChop(IndexOfLastNewLine).TrimStartAndEnd();
+		
+		return CumulativeString.Left(IndexOfLastNewLine).TrimStartAndEnd();
+	}
+
+	// If we don't find a break in the line, we should consider the entire line partial
+	LastPartialLogLine = CumulativeString;
+
+	// Then return nothing so we know we have nothing to do this iteration
+	return FString();
+}
+
+void FDisplayClusterPreloadDerivedDataCacheWorker::ReadCommandletOutputAndUpdateEditorNotification()
+{
+	while (FPlatformProcess::IsProcRunning(ProcessHandle))
+	{
+		if (ProgressNotification->GetPromptAction() == EAsyncTaskNotificationPromptAction::Cancel)
+		{
+			UE_LOG(LogDisplayClusterPreloadDerivedDataCache, Display, TEXT("----------CANCEL BUTTON PRESSED----------"));
+			CancelTask();
+			break;
+		}
+
+		const FString LogString = FPlatformProcess::ReadPipe(ReadPipe);
+		if (LogString.TrimStartAndEnd().IsEmpty())
+		{
+			continue;
+		}
+		
+		const FString LastWholeLogBlock = GetLastWholeLogBlock(LogString);
+		if (LastWholeLogBlock.TrimStartAndEnd().IsEmpty())
+		{
+			continue;
+		}
+
+		TArray<FString> EachLineInLogBlock;
+		LastWholeLogBlock.ParseIntoArray(EachLineInLogBlock, TEXT("\n"), true);
+
+		for (const FString& Line : EachLineInLogBlock)
+		{
+			if (Line.TrimStartAndEnd().IsEmpty())
+			{
+				continue;
+			}
+			
+			UE_LOG(LogDisplayClusterPreloadDerivedDataCache, Display, TEXT("Commandlet Output: %s"), *Line)
+
+			if (LoadingTotal == INDEX_NONE)
+			{			
+				// Try to find total number of assets to load
+				FRegexMatcher LoadingTotalRegex(LoadingTotalPattern, *Line);
+
+				while (LoadingTotalRegex.FindNext())
+				{
+					LoadingTotal = FCString::Atoi(*LoadingTotalRegex.GetCaptureGroup(1));
+				}
+			}
+		
+			// Progress on asset enumeration
+			RegexParseForEnumerationCount(Line);
+		
+			// Progress on asset loading
+			RegexParseForLoadingProgress(Line);
+		
+			// Progress on asset compilation
+			RegexParseForCompilationProgress(Line);
+		}
+
+		FPlatformProcess::Sleep(0.5f);
+	}
+	
+	UE_LOG(LogDisplayClusterPreloadDerivedDataCache, Display, TEXT("----------PROC NOT RUNNING----------"));
+	CompleteCommandletAndShowNotification();
+}
+
+#undef LOCTEXT_NAMESPACE
diff --git a/Engine/Plugins/Runtime/nDisplay/Source/DisplayClusterPreloadDerivedDataCache/Public/DisplayClusterPreloadDerivedDataCacheModule.h b/Engine/Plugins/Runtime/nDisplay/Source/DisplayClusterPreloadDerivedDataCache/Public/DisplayClusterPreloadDerivedDataCacheModule.h
--- a/Engine/Plugins/Runtime/nDisplay/Source/DisplayClusterPreloadDerivedDataCache/Public/DisplayClusterPreloadDerivedDataCacheModule.h
+++ b/Engine/Plugins/Runtime/nDisplay/Source/DisplayClusterPreloadDerivedDataCache/Public/DisplayClusterPreloadDerivedDataCacheModule.h
@@ -2,9 +2,8 @@
 
 #pragma once
 
-#include "Containers/UnrealString.h"
 #include "Modules/ModuleInterface.h"
-#include "Templates/SharedPointer.h"
+#include "Templates/UniquePtr.h"
 
 class FDisplayClusterPreloadDerivedDataCacheModule : public IModuleInterface
 {
@@ -17,22 +16,10 @@ public:
 	virtual void ShutdownModule() override;
 	//~ End IModuleInterface Interface
 
-	void LaunchAndCommunicateWithProcess();
-
 protected:
 	
-	void RegexParseForLoadingProgress(const FText& LoadingPackagesFormat, const FText& UnknownAmount, struct FScopedSlowTask& SlowTask,
-									  int32& LoadingTotal, class FRegexMatcher& LoadingProgressRegex);
-	void RegexParseForCompilationProgress(const FText& CompilingAssetsFormat, const FText& UnknownAmount, FScopedSlowTask& SlowTask,
-	                                      int32& CompileTotal, int32& AssetsWaitingToCompile,
-	                                      FRegexMatcher& CompileProgressRegex);
-	void CompleteCommandletAndShowNotification(const int32 ResultCode, const bool bWasCancelled,
-	const FString& CurrentExecutableName, const FString& Arguments);
-
-	FString GetDdcCommandletParams() const;
-	FString GetTargetPlatformParams() const;
-
-	TSharedPtr<class SNotificationItem> Notification;
+	void CreateAsyncTaskWorker();
+	TUniquePtr<class FDisplayClusterPreloadDerivedDataCacheWorker> AsyncTaskWorker;
 
 private:
 	
diff --git a/Engine/Plugins/Runtime/nDisplay/Source/DisplayClusterPreloadDerivedDataCache/Public/DisplayClusterPreloadDerivedDataCacheWorker.h b/Engine/Plugins/Runtime/nDisplay/Source/DisplayClusterPreloadDerivedDataCache/Public/DisplayClusterPreloadDerivedDataCacheWorker.h
new file mode 100644
--- /dev/null
+++ b/Engine/Plugins/Runtime/nDisplay/Source/DisplayClusterPreloadDerivedDataCache/Public/DisplayClusterPreloadDerivedDataCacheWorker.h
@@ -0,0 +1,85 @@
+// Copyright Epic Games, Inc. All Rights Reserved.
+
+#pragma once
+
+#include "Containers/UnrealString.h"
+#include "HAL/PlatformProcess.h"
+#include "HAL/Runnable.h"
+#include "HAL/RunnableThread.h"
+#include "Misc/SingleThreadRunnable.h"
+#include "Templates/UniquePtr.h"
+
+class FAsyncTaskNotification;
+
+struct FProcHandle;
+
+class FDisplayClusterPreloadDerivedDataCacheWorker : public FRunnable, public FSingleThreadRunnable
+{
+public:
+
+	FDisplayClusterPreloadDerivedDataCacheWorker();
+	
+	virtual ~FDisplayClusterPreloadDerivedDataCacheWorker() override;
+
+	//~ Begin FRunnable implementation
+	virtual uint32 Run() override
+	{
+		ReadCommandletOutputAndUpdateEditorNotification();
+		return 0;
+	}
+	virtual void Stop() override
+	{
+		CancelTask();
+	}
+	//~ End FRunnable implementation
+	
+	//~ Begin FSingleThreadRunnable implementation
+	virtual void Tick() override
+	{
+		ReadCommandletOutputAndUpdateEditorNotification();
+	}
+	virtual FSingleThreadRunnable* GetSingleThreadInterface() override
+	{
+		return this;
+	}
+	//~ End FSingleThreadRunnable implementation
+
+	void CancelTask();
+	void CompleteCommandletAndShowNotification();
+
+protected:
+	
+	void RegexParseForEnumerationCount(const FString& LogString);
+	void RegexParseForLoadingProgress(const FString& LogString);
+	void RegexParseForCompilationProgress(const FString& LogString);
+
+	FString GetDdcCommandletParams() const;
+	FString GetTargetPlatformParams() const;
+
+	FString GetLastWholeLogBlock(const FString& LogString);
+
+	void ReadCommandletOutputAndUpdateEditorNotification();
+	
+	/** Progress notification */
+	TUniquePtr<FAsyncTaskNotification> ProgressNotification;
+
+private:
+
+	TUniquePtr<FRunnableThread> Thread;
+	
+	int32 ResultCode = 0;
+	bool bWasCancelled = false;
+	FProcHandle ProcessHandle;
+
+	void* ReadPipe = nullptr;
+	void* WritePipe = nullptr;
+
+	FString LastPartialLogLine;
+
+	int32 LoadingTotal = INDEX_NONE;
+	int32 CompileTotal = INDEX_NONE;
+	int32 AssetsWaitingToCompile = INDEX_NONE;
+
+	FString CurrentExecutableName;
+	FString Arguments;
+};
