nDisplay: Default to non spatially loaded.

Unhid WorldPartition category in details panel so that the user can change this setting.

diff --git a/Engine/Plugins/Runtime/nDisplay/Source/DisplayCluster/Private/Game/EngineClasses/Scene/DisplayClusterRootActor.cpp b/Engine/Plugins/Runtime/nDisplay/Source/DisplayCluster/Private/Game/EngineClasses/Scene/DisplayClusterRootActor.cpp
--- a/Engine/Plugins/Runtime/nDisplay/Source/DisplayCluster/Private/Game/EngineClasses/Scene/DisplayClusterRootActor.cpp
+++ b/Engine/Plugins/Runtime/nDisplay/Source/DisplayCluster/Private/Game/EngineClasses/Scene/DisplayClusterRootActor.cpp
@@ -101,6 +101,12 @@ ADisplayClusterRootActor::ADisplayClusterRootActor(const FObjectInitializer& Obj
 	// Enabled by default to avoid being active on Stage
 	SetActorHiddenInGame(true);
 
+	// By default, we don't want the NDC to be spatially loaded.
+#if WITH_EDITORONLY_DATA
+	bIsSpatiallyLoaded = false;
+#endif
+
+
 #if WITH_EDITOR
 	Constructor_Editor();
 #endif
diff --git a/Engine/Plugins/Runtime/nDisplay/Source/DisplayCluster/Public/DisplayClusterRootActor.h b/Engine/Plugins/Runtime/nDisplay/Source/DisplayCluster/Public/DisplayClusterRootActor.h
--- a/Engine/Plugins/Runtime/nDisplay/Source/DisplayCluster/Public/DisplayClusterRootActor.h
+++ b/Engine/Plugins/Runtime/nDisplay/Source/DisplayCluster/Public/DisplayClusterRootActor.h
@@ -45,7 +45,7 @@ class UProceduralMeshComponent;
 /**
  * VR root. This contains nDisplay VR hierarchy in the game.
  */
-UCLASS(HideCategories=(Replication, Collision, Input, Actor, HLOD, Cooking, Physics, Activation, AssetUserData, ActorTick, Advanced, WorldPartition, DataLayers, Events), meta=(DisplayName = "nDisplay Root Actor"))
+UCLASS(HideCategories=(Replication, Collision, Input, Actor, HLOD, Cooking, Physics, Activation, AssetUserData, ActorTick, Advanced, DataLayers, Events), meta=(DisplayName = "nDisplay Root Actor"))
 class DISPLAYCLUSTER_API ADisplayClusterRootActor
 	: public AActor
 {
diff --git a/Engine/Plugins/Runtime/nDisplay/Source/DisplayClusterConfigurator/Private/Views/Details/DisplayClusterRootActorDetailsCustomization.cpp b/Engine/Plugins/Runtime/nDisplay/Source/DisplayClusterConfigurator/Private/Views/Details/DisplayClusterRootActorDetailsCustomization.cpp
--- a/Engine/Plugins/Runtime/nDisplay/Source/DisplayClusterConfigurator/Private/Views/Details/DisplayClusterRootActorDetailsCustomization.cpp
+++ b/Engine/Plugins/Runtime/nDisplay/Source/DisplayClusterConfigurator/Private/Views/Details/DisplayClusterRootActorDetailsCustomization.cpp
@@ -36,6 +36,7 @@ namespace DisplayClusterRootActorDetailsCustomizationUtils
 			DisplayClusterConfigurationStrings::categories::LightcardCategory,
 			DisplayClusterConfigurationStrings::categories::OverrideCategory,
 			TEXT("Rendering"),
+			TEXT("WorldPartition"),
 			DisplayClusterConfigurationStrings::categories::PreviewCategory,
 			DisplayClusterConfigurationStrings::categories::ConfigurationCategory,
 			DisplayClusterConfigurationStrings::categories::AdvancedCategory,
