nDisplay: Added singleton class that controls the preview rendering of multiple DCRA objects.
fixed bug when  DC Viewport Manager object is not destroyed because the TSharedPtr has references

diff --git a/Engine/Plugins/Runtime/nDisplay/Source/DisplayCluster/Private/Game/EngineClasses/Basics/DisplayClusterViewportClient.cpp b/Engine/Plugins/Runtime/nDisplay/Source/DisplayCluster/Private/Game/EngineClasses/Basics/DisplayClusterViewportClient.cpp
--- a/Engine/Plugins/Runtime/nDisplay/Source/DisplayCluster/Private/Game/EngineClasses/Basics/DisplayClusterViewportClient.cpp
+++ b/Engine/Plugins/Runtime/nDisplay/Source/DisplayCluster/Private/Game/EngineClasses/Basics/DisplayClusterViewportClient.cpp
@@ -473,6 +473,13 @@ void UDisplayClusterViewportClient::Draw(FViewport* InViewport, FCanvas* SceneCa
 		return;
 	}
 
+	IDisplayClusterViewportManager* RenderFrameViewportManager = RenderFrame.GetViewportManager();
+	if (!RenderFrameViewportManager)
+	{
+		// skip rendering: Can't find render manager
+		return;
+	}
+
 	// Handle special viewports game-thread logic at frame begin
 	DCRenderDevice->InitializeNewFrame();
 
@@ -481,7 +488,7 @@ void UDisplayClusterViewportClient::Draw(FViewport* InViewport, FCanvas* SceneCa
 		for (FDisplayClusterRenderFrameTargetViewFamily& DCViewFamily : DCRenderTarget.ViewFamilies)
 		{
 			// Create the view family for rendering the world scene to the viewport's render target
-			ViewFamilies.Add(new FSceneViewFamilyContext(RenderFrame.ViewportManager->CreateViewFamilyConstructionValues(
+			ViewFamilies.Add(new FSceneViewFamilyContext(RenderFrameViewportManager->CreateViewFamilyConstructionValues(
 				DCRenderTarget,
 				MyWorld->Scene,
 				EngineShowFlags,
@@ -491,7 +498,7 @@ void UDisplayClusterViewportClient::Draw(FViewport* InViewport, FCanvas* SceneCa
 			bool bIsFamilyVisible = false;
 
 			// Configure family
-			RenderFrame.ViewportManager->ConfigureViewFamily(DCRenderTarget, DCViewFamily, ViewFamily);
+			RenderFrameViewportManager->ConfigureViewFamily(DCRenderTarget, DCViewFamily, ViewFamily);
 
 #if WITH_EDITOR
 			if (GIsEditor)
@@ -552,7 +559,7 @@ void UDisplayClusterViewportClient::Draw(FViewport* InViewport, FCanvas* SceneCa
 				// Calculate the player's view information.
 				FVector		ViewLocation;
 				FRotator	ViewRotation;
-				FSceneView* View = RenderFrame.ViewportManager->CalcSceneView(LocalPlayer, &ViewFamily, ViewLocation, ViewRotation, InViewport, nullptr, ViewportContext.StereoViewIndex);
+				FSceneView* View = RenderFrameViewportManager->CalcSceneView(LocalPlayer, &ViewFamily, ViewLocation, ViewRotation, InViewport, nullptr, ViewportContext.StereoViewIndex);
 
 				if (View && !DCView.IsViewportContextCanBeRendered())
 				{
@@ -931,7 +938,7 @@ void UDisplayClusterViewportClient::Draw(FViewport* InViewport, FCanvas* SceneCa
 		SceneCanvas->Flush_GameThread();
 
 		// After all render target rendered call nDisplay frame rendering
-		RenderFrame.ViewportManager->RenderFrame(InViewport);
+		RenderFrameViewportManager->RenderFrame(InViewport);
 
 		OnDrawn().Broadcast();
 
diff --git a/Engine/Plugins/Runtime/nDisplay/Source/DisplayCluster/Private/Game/EngineClasses/Scene/DisplayClusterRootActor.cpp b/Engine/Plugins/Runtime/nDisplay/Source/DisplayCluster/Private/Game/EngineClasses/Scene/DisplayClusterRootActor.cpp
--- a/Engine/Plugins/Runtime/nDisplay/Source/DisplayCluster/Private/Game/EngineClasses/Scene/DisplayClusterRootActor.cpp
+++ b/Engine/Plugins/Runtime/nDisplay/Source/DisplayCluster/Private/Game/EngineClasses/Scene/DisplayClusterRootActor.cpp
@@ -1,6 +1,7 @@
 // Copyright Epic Games, Inc. All Rights Reserved.
 
 #include "DisplayClusterRootActor.h"
+#include "DisplayClusterRootActorPreviewRenderingManager.h"
 
 #include "Async/ParallelFor.h"
 #include "Components/SceneComponent.h"
@@ -130,9 +131,21 @@ void ADisplayClusterRootActor::CreateViewportManagerImpl()
 	{
 		ViewportManager = MakeShared<FDisplayClusterViewportManager, ESPMode::ThreadSafe>();
 		ViewportManager->Initialize();
+
+		// Preview rendering depends on the DC VM
+		FDisplayClusterRootActorPreviewRenderingManager::HandleEvent(EDisplayClusterRootActorPreviewEvent::Create, this);
 	}
 }
 
+void ADisplayClusterRootActor::RemoveViewportManagerImpl()
+{
+	// Preview rendering depends on the  DC VM
+	FDisplayClusterRootActorPreviewRenderingManager::HandleEvent(EDisplayClusterRootActorPreviewEvent::Remove, this);
+
+	// Immediately release the viewport manager with resources
+	ViewportManager.Reset();
+}
+
 bool ADisplayClusterRootActor::IsRunningGameOrPIE() const
 {
 	if (!IsRunningGame())
@@ -1018,7 +1031,14 @@ void ADisplayClusterRootActor::Tick(float DeltaSeconds)
 
 #if WITH_EDITOR
 	// Tick editor preview
-	Tick_Editor(DeltaSeconds);
+	if (!IsPreviewEnabled())
+	{
+		ResetPreviewInternals_Editor();
+	}
+	else
+	{
+		FDisplayClusterRootActorPreviewRenderingManager::HandleEvent(EDisplayClusterRootActorPreviewEvent::Render, this);
+	}
 #endif
 
 	SetLightCardOwnership();
@@ -1064,7 +1084,7 @@ void ADisplayClusterRootActor::Destroyed()
 #endif
 
 	// Release viewport manager with resources immediatelly
-	ViewportManager.Reset();
+	RemoveViewportManagerImpl();
 
 	Super::Destroyed();
 }
@@ -1075,7 +1095,7 @@ void ADisplayClusterRootActor::BeginDestroy()
 	BeginDestroy_Editor();
 #endif
 
-	ViewportManager.Reset();
+	RemoveViewportManagerImpl();
 
 	Super::BeginDestroy();
 }
diff --git a/Engine/Plugins/Runtime/nDisplay/Source/DisplayCluster/Private/Game/EngineClasses/Scene/DisplayClusterRootActorEditor.cpp b/Engine/Plugins/Runtime/nDisplay/Source/DisplayCluster/Private/Game/EngineClasses/Scene/DisplayClusterRootActorEditor.cpp
--- a/Engine/Plugins/Runtime/nDisplay/Source/DisplayCluster/Private/Game/EngineClasses/Scene/DisplayClusterRootActorEditor.cpp
+++ b/Engine/Plugins/Runtime/nDisplay/Source/DisplayCluster/Private/Game/EngineClasses/Scene/DisplayClusterRootActorEditor.cpp
@@ -124,7 +124,7 @@ void ADisplayClusterRootActor::Destructor_Editor()
 	}
 }
 
-void ADisplayClusterRootActor::Tick_Editor(float DeltaSeconds)
+void ADisplayClusterRootActor::RenderPreview_Editor()
 {
 	if (IsPreviewEnabled())
 	{
@@ -162,7 +162,7 @@ void ADisplayClusterRootActor::Tick_Editor(float DeltaSeconds)
 			if(ViewportManager->GetRenderMode() == EDisplayClusterRenderFrameMode::PreviewInScene)
 			{
 				// Release viewport manager with resources immediatelly for preview in scene
-				ViewportManager.Reset();
+				RemoveViewportManagerImpl();
 			}
 		}
 	}
diff --git a/Engine/Plugins/Runtime/nDisplay/Source/DisplayCluster/Private/Game/EngineClasses/Scene/DisplayClusterRootActorPreviewRenderingManager.cpp b/Engine/Plugins/Runtime/nDisplay/Source/DisplayCluster/Private/Game/EngineClasses/Scene/DisplayClusterRootActorPreviewRenderingManager.cpp
new file mode 100644
--- /dev/null
+++ b/Engine/Plugins/Runtime/nDisplay/Source/DisplayCluster/Private/Game/EngineClasses/Scene/DisplayClusterRootActorPreviewRenderingManager.cpp
@@ -0,0 +1,252 @@
+// Copyright Epic Games, Inc. All Rights Reserved.
+
+#include "DisplayClusterRootActorPreviewRenderingManager.h"
+#include "DisplayClusterRootActor.h"
+
+#include "Engine/World.h"
+#include "Engine/GameViewportClient.h"
+#include "HAL/IConsoleManager.h"
+
+///////////////////////////////////////////////////////////////////
+int32 GDisplayClusterRootActorPreviewRenderingEnable = 1;
+static FAutoConsoleVariableRef CVarDisplayClusterRootActorPreviewRenderingEnable(
+	TEXT("nDisplay.preview.render.enable"),
+	GDisplayClusterRootActorPreviewRenderingEnable,
+	TEXT("Globally enable DCRA preview rendering\n")
+	TEXT("0 : Disable preview render \n"),
+	ECVF_Default
+);
+
+int32 GDisplayClusterRootActorPreviewRenderingOptimization = 1;
+static FAutoConsoleVariableRef CVarDisplayClusterRootActorPreviewRenderingOptimization(
+	TEXT("nDisplay.preview.render.optimize"),
+	GDisplayClusterRootActorPreviewRenderingOptimization,
+	TEXT("Optimize DCRA preview rendering only 1 DCRA per frame\n")
+	TEXT("0 : Disable optimization\n"),
+	ECVF_Default
+);
+
+///////////////////////////////////////////////////////////////////
+namespace UE::DisplayCluster::DCRAPreview
+{
+	/**
+	 * We can skip the preview rendering for invisible DCRA objects:
+	 * minimized or inactive window, invisible tab, etc.
+	 */
+	static bool IsRootActorVisible(ADisplayClusterRootActor* InRootActor)
+	{
+		check(InRootActor);
+
+		// Todo: Add logic to get DCRA visibility.
+
+		return true;
+	}
+
+	static FDisplayClusterRootActorPreviewRenderingManager& GetPreviewRenderingManager()
+	{
+		static FDisplayClusterRootActorPreviewRenderingManager DCRAPreviewManagerSingletonInstance;
+
+		return DCRAPreviewManagerSingletonInstance;
+	}
+
+};
+using namespace UE::DisplayCluster::DCRAPreview;
+
+//------------------------------------------------------------------
+// FDisplayClusterRootActorPreviewObject
+//------------------------------------------------------------------
+void FDisplayClusterRootActorPreviewObject::HandleRenderRequest()
+{
+	if (!GDisplayClusterRootActorPreviewRenderingEnable)
+	{
+		return;
+	}
+
+	if (GDisplayClusterRootActorPreviewRenderingOptimization)
+	{
+		// Request for preview render
+		bRenderPreviewRequested = true;
+	}
+	else
+	{
+		RenderPreview();
+	}
+}
+
+void FDisplayClusterRootActorPreviewObject::RenderPreview()
+{
+#if WITH_EDITOR
+	if (ADisplayClusterRootActor* RootActor = GetRootActor())
+	{
+		RootActor->RenderPreview_Editor();
+	}
+#endif
+}
+
+//------------------------------------------------------------------
+// FDisplayClusterRootActorPreviewRenderingManager
+//------------------------------------------------------------------
+void FDisplayClusterRootActorPreviewRenderingManager::HandleEvent(const EDisplayClusterRootActorPreviewEvent InDCRAEvent, ADisplayClusterRootActor* InRootActor)
+{
+	GetPreviewRenderingManager().HandleEventImpl(InDCRAEvent, InRootActor);
+}
+
+void FDisplayClusterRootActorPreviewRenderingManager::HandleEventImpl(const EDisplayClusterRootActorPreviewEvent InDCRAEvent, ADisplayClusterRootActor* InRootActor)
+{
+	const int32 PreviewObjectIndex = PreviewObjects.IndexOfByPredicate([InRootActor](const FDisplayClusterRootActorPreviewObject& PreviewObjectIt)
+		{
+			return PreviewObjectIt.IsEqual(InRootActor);
+		});
+
+	switch (InDCRAEvent)
+	{
+	case EDisplayClusterRootActorPreviewEvent::Create:
+		if (PreviewObjectIndex == INDEX_NONE)
+		{
+			PreviewObjects.Add(FDisplayClusterRootActorPreviewObject(InRootActor));
+		}
+		break;
+
+	case EDisplayClusterRootActorPreviewEvent::Remove:
+		if (PreviewObjectIndex != INDEX_NONE)
+		{
+			PreviewObjects.RemoveAt(PreviewObjectIndex);
+		}
+		break;
+
+	case EDisplayClusterRootActorPreviewEvent::Render:
+		if (PreviewObjectIndex != INDEX_NONE)
+		{
+			PreviewObjects[PreviewObjectIndex].HandleRenderRequest();
+		}
+		else
+		{
+			// Create new ref by render request
+			const int32 NewIndex = PreviewObjects.Add(FDisplayClusterRootActorPreviewObject(InRootActor));
+			PreviewObjects[NewIndex].HandleRenderRequest();
+		}
+		break;
+
+	default:
+		break;
+	}
+
+	UpdateTickableGameObject();
+}
+
+void FDisplayClusterRootActorPreviewRenderingManager::Tick(float DeltaTime)
+{
+	if (!GDisplayClusterRootActorPreviewRenderingOptimization || !GDisplayClusterRootActorPreviewRenderingEnable)
+	{
+		return;
+	}
+
+	// Render one DCRA per frame:
+	int32 NextPreviewObjectIndex = GetNextPreviewObject();
+	if (NextPreviewObjectIndex == INDEX_NONE)
+	{
+		// Reset render cycle
+		BeginNewRenderCycle();
+
+		// start new cycle
+		NextPreviewObjectIndex = GetNextPreviewObject();
+	}
+
+	if (NextPreviewObjectIndex != INDEX_NONE)
+	{
+		// Render this DCRA
+		PreviewObjects[NextPreviewObjectIndex].RenderPreview();
+	}
+}
+
+int32 FDisplayClusterRootActorPreviewRenderingManager::GetNextPreviewObject()
+{
+	int32 OutIndex = INDEX_NONE;
+
+	// rebuild list - remove invalid DCRA's
+	TArray<FDisplayClusterRootActorPreviewObject> NewPreviewObjects;
+	NewPreviewObjects.Reserve(PreviewObjects.Num());
+
+	for (FDisplayClusterRootActorPreviewObject& PreviewObjectIt : PreviewObjects)
+	{
+		if (PreviewObjectIt.IsValid())
+		{
+			if (OutIndex == INDEX_NONE && PreviewObjectIt.ShouldRenderPreview())
+			{
+				// Mark this object as rendered
+				PreviewObjectIt.MarkAsRendered();
+
+				// Render only visible DCRA objects
+				if (IsRootActorVisible(PreviewObjectIt.GetRootActor()))
+				{
+					OutIndex = NewPreviewObjects.Num();
+				}
+			}
+
+			NewPreviewObjects.Add(PreviewObjectIt);
+		}
+	}
+
+	PreviewObjects = NewPreviewObjects;
+
+	return OutIndex;
+}
+
+void FDisplayClusterRootActorPreviewRenderingManager::BeginNewRenderCycle()
+{
+	for (FDisplayClusterRootActorPreviewObject& PreviewObject : PreviewObjects)
+	{
+		if (PreviewObject.ShouldRenderPreview())
+		{
+			// This object was not render in the current loop
+			return;
+		}
+	}
+
+	for (FDisplayClusterRootActorPreviewObject& PreviewObject : PreviewObjects)
+	{
+		// Allow to render again
+		PreviewObject.MarkAsNotRendered();
+	}
+}
+
+//---------------------------------------------------------------------------------
+// FDisplayClusterPreviewTickableGameObject
+//---------------------------------------------------------------------------------
+class FDisplayClusterPreviewTickableGameObject : public FTickableGameObject
+{
+public:
+	virtual bool IsTickableInEditor() const override { return true; }
+
+	virtual void Tick(float DeltaTime) override
+	{
+		GetPreviewRenderingManager().Tick(DeltaTime);
+	}
+
+	virtual ETickableTickType GetTickableTickType() const override
+	{ 
+		return ETickableTickType::Always;
+	}
+
+	virtual TStatId GetStatId() const override
+	{
+		RETURN_QUICK_DECLARE_CYCLE_STAT(FDisplayClusterPreviewTickableGameObject, STATGROUP_Tickables);
+	}
+};
+
+void FDisplayClusterRootActorPreviewRenderingManager::UpdateTickableGameObject()
+{
+	if (PreviewObjects.IsEmpty())
+	{
+		// delete when no DCRAs
+		TickableGameObject.Reset();
+	}
+	else
+	{
+		// When we have at least one DCRA, a tickable object is created
+		if (!TickableGameObject.IsValid())
+		{
+			TickableGameObject = MakeUnique<FDisplayClusterPreviewTickableGameObject>();
+		}
+	}
+}
diff --git a/Engine/Plugins/Runtime/nDisplay/Source/DisplayCluster/Private/Game/EngineClasses/Scene/DisplayClusterRootActorPreviewRenderingManager.h b/Engine/Plugins/Runtime/nDisplay/Source/DisplayCluster/Private/Game/EngineClasses/Scene/DisplayClusterRootActorPreviewRenderingManager.h
new file mode 100644
--- /dev/null
+++ b/Engine/Plugins/Runtime/nDisplay/Source/DisplayCluster/Private/Game/EngineClasses/Scene/DisplayClusterRootActorPreviewRenderingManager.h
@@ -0,0 +1,126 @@
+// Copyright Epic Games, Inc. All Rights Reserved.
+
+#pragma once
+#include "CoreMinimal.h"
+
+#include "Engine/EngineTypes.h"
+#include "Tickable.h"
+
+class ADisplayClusterRootActor;
+
+/**
+ * Preview manager receive events from DCRAs
+ */
+enum class EDisplayClusterRootActorPreviewEvent : uint8
+{
+	// DCRA created
+	Create = 0,
+
+	// DCRA removed
+	Remove,
+
+	// DCRA wants to render preview
+	Render
+};
+
+/**
+ * Container for DCRA objects that are used for preview rendering.
+ */
+class FDisplayClusterRootActorPreviewObject
+{
+public:
+	FDisplayClusterRootActorPreviewObject() = default;
+	FDisplayClusterRootActorPreviewObject(ADisplayClusterRootActor* InRootActor)
+		: RootActorWeakPtr(InRootActor)
+	{ }
+
+public:
+	/** Returns true if the DCRA object still exists. */
+	inline bool IsValid() const
+	{
+		return RootActorWeakPtr.IsValid() && !RootActorWeakPtr.IsStale();
+	}
+
+	/** Returns true if the DCRA preview can be performed right now. */
+	inline bool ShouldRenderPreview() const
+	{
+		return bRenderPreviewRequested && !bPreviewRenderingCompleted;
+	}
+
+	/** Changes the internal flags of this DCRA container to a state indicating that the preview rendering is complete. */
+	inline void MarkAsRendered()
+	{
+		bRenderPreviewRequested = false;
+		bPreviewRenderingCompleted = true;
+	}
+
+	/** Changes the internal flags of this DCRA container to a state that allows render previews. */
+	inline void MarkAsNotRendered()
+	{
+		bPreviewRenderingCompleted = false;
+	}
+
+	/** Returns true if the InRootActor is the same as the one used inside this container. */
+	inline bool IsEqual(const ADisplayClusterRootActor* InRootActor) const
+	{
+		return GetRootActor() == InRootActor;
+	}
+
+	/** Return the pointer to the DCRA actor if it still exists. */
+	inline ADisplayClusterRootActor* GetRootActor() const
+	{
+		return RootActorWeakPtr.IsValid() ? RootActorWeakPtr.Get() : nullptr;
+	}
+
+	/** When DCRA tries to call the preview rendering, this function will be called. */
+	void HandleRenderRequest();
+
+	/** Performs a preview of the rendering for DCRA in this container. */
+	void RenderPreview();
+
+private:
+	// Weak pointer to the DCRA
+	TWeakObjectPtr<ADisplayClusterRootActor> RootActorWeakPtr;
+
+	// Preview is already rendered for this actor
+	bool bPreviewRenderingCompleted = false;
+
+	// Preview render request from DCRA
+	bool bRenderPreviewRequested = false;
+};
+
+/**
+ * This is a singleton class that controls the preview rendering of multiple DCRA objects.
+ */
+class FDisplayClusterRootActorPreviewRenderingManager
+{
+public:
+	/** The entry point for DCRA, which is used to handle all workflow events related to preview rendering. */
+	static void HandleEvent(const EDisplayClusterRootActorPreviewEvent InDCRAEvent, ADisplayClusterRootActor* InRootActor);
+
+protected:
+	friend class FDisplayClusterPreviewTickableGameObject;
+
+	/** This function is used to perform delayed DCRA preview rendering requests. */
+	void Tick(float DeltaTime);
+
+protected:
+	/** Implementation of the static HandleEvent() function. */
+	void HandleEventImpl(const EDisplayClusterRootActorPreviewEvent InDCRAEvent, ADisplayClusterRootActor* InRootActor);
+
+	/** Creating or deleting a Tickable object depends on whether the DCRAs list is not empty. */
+	void UpdateTickableGameObject();
+
+	/** Gets the index of the next DCRA container, which can be used for rendering the preview. */
+	int32  GetNextPreviewObject();
+
+	/** When all registered DCRA objects are drawn, you must call this function to reset the internal flags. */
+	void BeginNewRenderCycle();
+
+private:
+	// List of existing DCRAs
+	TArray<FDisplayClusterRootActorPreviewObject> PreviewObjects;
+
+	// Tickable object to call preview rendering for single DCRA per frame
+	TUniquePtr<FTickableGameObject> TickableGameObject;
+};
diff --git a/Engine/Plugins/Runtime/nDisplay/Source/DisplayCluster/Private/Render/Device/DisplayClusterDeviceBase.cpp b/Engine/Plugins/Runtime/nDisplay/Source/DisplayCluster/Private/Render/Device/DisplayClusterDeviceBase.cpp
--- a/Engine/Plugins/Runtime/nDisplay/Source/DisplayCluster/Private/Render/Device/DisplayClusterDeviceBase.cpp
+++ b/Engine/Plugins/Runtime/nDisplay/Source/DisplayCluster/Private/Render/Device/DisplayClusterDeviceBase.cpp
@@ -64,6 +64,16 @@ FDisplayClusterDeviceBase::~FDisplayClusterDeviceBase()
 	//@todo: delete singleton object IDisplayClusterViewportManager
 }
 
+IDisplayClusterViewportManager* FDisplayClusterDeviceBase::GetViewportManager() const
+{
+	return ViewportManagerWeakPtr.IsValid() ? ViewportManagerWeakPtr.Pin().Get() : nullptr;
+}
+
+FDisplayClusterViewportManagerProxy* FDisplayClusterDeviceBase::GetViewportManagerProxy_RenderThread() const
+{
+	return ViewportManagerProxyWeakPtr.IsValid() ? ViewportManagerProxyWeakPtr.Pin().Get() : nullptr;
+}
+
 //////////////////////////////////////////////////////////////////////////////////////////////
 // IDisplayClusterStereoDevice
 //////////////////////////////////////////////////////////////////////////////////////////////
@@ -170,7 +180,7 @@ EStereoscopicPass FDisplayClusterDeviceBase::GetViewPassForIndex(bool bStereoReq
 	{
 		if (IsInRenderingThread())
 		{
-			if (ViewportManagerProxy.IsValid())
+			if (IDisplayClusterViewportManagerProxy* ViewportManagerProxy = GetViewportManagerProxy_RenderThread())
 			{
 				uint32 ViewportContextNum = 0;
 				IDisplayClusterViewportProxy* ViewportProxy = ViewportManagerProxy->FindViewport_RenderThread(ViewIndex, &ViewportContextNum);
@@ -183,7 +193,7 @@ EStereoscopicPass FDisplayClusterDeviceBase::GetViewPassForIndex(bool bStereoReq
 		}
 		else
 		{
-			if (ViewportManager)
+			if (IDisplayClusterViewportManager* ViewportManager = GetViewportManager())
 			{
 				uint32 ViewportContextNum = 0;
 				IDisplayClusterViewport* ViewportPtr = ViewportManager->FindViewport(ViewIndex, &ViewportContextNum);
@@ -203,6 +213,7 @@ void FDisplayClusterDeviceBase::AdjustViewRect(int32 ViewIndex, int32& X, int32&
 {
 	check(IsInGameThread());
 
+	IDisplayClusterViewportManager* ViewportManager = GetViewportManager();
 	if (ViewportManager == nullptr || ViewportManager->IsSceneOpened() == false)
 	{
 		return;
@@ -239,6 +250,8 @@ void FDisplayClusterDeviceBase::CalculateStereoViewOffset(const int32 ViewIndex,
 	check(IsInGameThread());
 	check(WorldToMeters > 0.f);
 
+	IDisplayClusterViewportManager* ViewportManager = GetViewportManager();
+
 	uint32 ViewportContextNum = 0;
 	IDisplayClusterViewport* ViewportPtr = ViewportManager ? ViewportManager->FindViewport(ViewIndex, &ViewportContextNum) : nullptr;
 	if (ViewportPtr == nullptr)
@@ -279,6 +292,8 @@ FMatrix FDisplayClusterDeviceBase::GetStereoProjectionMatrix(const int32 ViewInd
 {
 	check(IsInGameThread());
 
+	IDisplayClusterViewportManager* ViewportManager = GetViewportManager();
+
 	FMatrix PrjMatrix = FMatrix::Identity;
 
 	// ViewIndex == eSSE_MONOSCOPIC(-1) is a special case called for ISR culling math.
@@ -336,10 +351,10 @@ void FDisplayClusterDeviceBase::InitializeNewFrame()
 	IDisplayCluster& DisplayCluster = IDisplayCluster::Get();
 	if (ADisplayClusterRootActor* RootActor = DisplayCluster.GetGameMgr()->GetRootActor())
 	{
-		if (IDisplayClusterViewportManager* ViewportManagerPtr = RootActor->GetViewportManager())
+		if (IDisplayClusterViewportManager* ViewportManager = RootActor->GetViewportManager())
 		{
 			// Begin use viewport manager for current frame
-			ViewportManager = ViewportManagerPtr->ToSharedPtr();
+			ViewportManagerWeakPtr = ViewportManager->ToSharedPtr();
 
 			// Initialize frame for render
 			ViewportManager->InitializeNewFrame();
@@ -350,7 +365,7 @@ void FDisplayClusterDeviceBase::InitializeNewFrame()
 			ENQUEUE_RENDER_COMMAND(DisplayClusterDevice_SetViewportManagerProxy)(
 				[DCRenderDevice = this, NewViewportManagerProxy = ViewportManagerProxyPtr->AsShared()](FRHICommandListImmediate& RHICmdList)
 				{
-					DCRenderDevice->ViewportManagerProxy = NewViewportManagerProxy;
+					DCRenderDevice->ViewportManagerProxyWeakPtr = NewViewportManagerProxy;
 				});
 		}
 	}
@@ -358,19 +373,20 @@ void FDisplayClusterDeviceBase::InitializeNewFrame()
 
 void FDisplayClusterDeviceBase::FinalizeNewFrame()
 {
+	IDisplayClusterViewportManager* ViewportManager = GetViewportManager();
 	if (ViewportManager)
 	{
 		ViewportManager->FinalizeNewFrame();
 	}
 
 	// reset viewport manager ptr on game thread
-	ViewportManager = nullptr;
+	ViewportManagerWeakPtr.Reset();
 
 	// Reset viewport manager proxy on render thread
 	ENQUEUE_RENDER_COMMAND(DisplayClusterDevice_ResetViewportManagerProxy)(
 		[DCRenderDevice = this](FRHICommandListImmediate& RHICmdList)
 		{
-			DCRenderDevice->ViewportManagerProxy.Reset();
+			DCRenderDevice->ViewportManagerProxyWeakPtr.Reset();
 		});
 }
 
@@ -399,10 +415,13 @@ void FDisplayClusterDeviceBase::RenderTexture_RenderThread(FRHICommandListImmedi
 			TransitionAndCopyTexture(RHICmdList, SrcTexture, BackBuffer, CopyInfo);
 		}
 
-		if (RenderFrameMode == EDisplayClusterRenderFrameMode::Stereo && ViewportManagerProxy.IsValid())
+		if (RenderFrameMode == EDisplayClusterRenderFrameMode::Stereo)
 		{
-			// QuadBufStereo: Copy RIGHT_EYE to backbuffer
-			ViewportManagerProxy->ResolveFrameTargetToBackBuffer_RenderThread(RHICmdList, 1, 1, BackBuffer, WindowSize);
+			if (FDisplayClusterViewportManagerProxy* ViewportManagerProxy = GetViewportManagerProxy_RenderThread())
+			{
+				// QuadBufStereo: Copy RIGHT_EYE to backbuffer
+				ViewportManagerProxy->ResolveFrameTargetToBackBuffer_RenderThread(RHICmdList, 1, 1, BackBuffer, WindowSize);
+			}
 		}
 
 		const bool bClearTextureEnabled = CVarClearTextureEnabled.GetValueOnRenderThread() != 0;
@@ -468,6 +487,7 @@ void FDisplayClusterDeviceBase::StartFinalPostprocessSettings(struct FPostProces
 	check(IsInGameThread());
 
 	// eSSP_FULL pass reserved for UE internal render
+	IDisplayClusterViewportManager* ViewportManager = GetViewportManager();
 	if (StereoPassType != EStereoscopicPass::eSSP_FULL && ViewportManager)
 	{
 		IDisplayClusterViewport* ViewportPtr = ViewportManager->FindViewport(StereoViewIndex);
@@ -483,6 +503,7 @@ bool FDisplayClusterDeviceBase::OverrideFinalPostprocessSettings(struct FPostPro
 	check(IsInGameThread());
 
 	// eSSP_FULL pass reserved for UE internal render
+	IDisplayClusterViewportManager* ViewportManager = GetViewportManager();
 	if (StereoPassType != EStereoscopicPass::eSSP_FULL && ViewportManager)
 	{
 		IDisplayClusterViewport* ViewportPtr = ViewportManager->FindViewport(StereoViewIndex);
@@ -500,6 +521,7 @@ void FDisplayClusterDeviceBase::EndFinalPostprocessSettings(struct FPostProcessS
 	check(IsInGameThread());
 
 	// eSSP_FULL pass reserved for UE internal render
+	IDisplayClusterViewportManager* ViewportManager = GetViewportManager();
 	if (StereoPassType != EStereoscopicPass::eSSP_FULL && ViewportManager && FinalPostProcessingSettings != nullptr)
 	{
 		IDisplayClusterViewport* ViewportPtr = ViewportManager->FindViewport(StereoViewIndex);
diff --git a/Engine/Plugins/Runtime/nDisplay/Source/DisplayCluster/Private/Render/Device/DisplayClusterDeviceBase.h b/Engine/Plugins/Runtime/nDisplay/Source/DisplayCluster/Private/Render/Device/DisplayClusterDeviceBase.h
--- a/Engine/Plugins/Runtime/nDisplay/Source/DisplayCluster/Private/Render/Device/DisplayClusterDeviceBase.h
+++ b/Engine/Plugins/Runtime/nDisplay/Source/DisplayCluster/Private/Render/Device/DisplayClusterDeviceBase.h
@@ -91,10 +91,18 @@ protected:
 	virtual bool OverrideFinalPostprocessSettings(struct FPostProcessSettings* OverridePostProcessingSettings, const enum EStereoscopicPass StereoPassType, const int32 StereoViewIndex, float& BlendWeight) override;
 	virtual void EndFinalPostprocessSettings(struct FPostProcessSettings* FinalPostProcessingSettings, const enum EStereoscopicPass StereoPassType, const int32 StereoViewIndex) override;
 
+	/** Get a pointer to the DC ViewportManager if it still exists. */
+	IDisplayClusterViewportManager* GetViewportManager() const;
+
+	/** Get a pointer to the DC ViewportManagerProxy if it still exists. */
+	FDisplayClusterViewportManagerProxy* GetViewportManagerProxy_RenderThread() const;
+
 private:
-	// Runtime viewport manager api for game and render threads. Internal usage only
-	TSharedPtr<IDisplayClusterViewportManager, ESPMode::ThreadSafe> ViewportManager;
-	TSharedPtr<FDisplayClusterViewportManagerProxy, ESPMode::ThreadSafe> ViewportManagerProxy;
+	// Pointer to the DC ViewportManager from the active DCRA that is currently being used for rendering.
+	TWeakPtr<IDisplayClusterViewportManager, ESPMode::ThreadSafe> ViewportManagerWeakPtr;
+
+	// Pointer to the DC ViewportManagerProxy from the active DCRA that is currently being used for rendering.
+	TWeakPtr<FDisplayClusterViewportManagerProxy, ESPMode::ThreadSafe> ViewportManagerProxyWeakPtr;
 
 	EDisplayClusterRenderFrameMode RenderFrameMode = EDisplayClusterRenderFrameMode::Mono;
 	int32 DesiredNumberOfViews = 0;
diff --git a/Engine/Plugins/Runtime/nDisplay/Source/DisplayCluster/Private/Render/Viewport/Configuration/DisplayClusterViewportConfiguration.cpp b/Engine/Plugins/Runtime/nDisplay/Source/DisplayCluster/Private/Render/Viewport/Configuration/DisplayClusterViewportConfiguration.cpp
--- a/Engine/Plugins/Runtime/nDisplay/Source/DisplayCluster/Private/Render/Viewport/Configuration/DisplayClusterViewportConfiguration.cpp
+++ b/Engine/Plugins/Runtime/nDisplay/Source/DisplayCluster/Private/Render/Viewport/Configuration/DisplayClusterViewportConfiguration.cpp
@@ -120,7 +120,7 @@ static FAutoConsoleVariableRef CVarDisplayClusterAlphaChannelCaptureMode(
 // FDisplayClusterViewportConfiguration
 ///////////////////////////////////////////////////////////////////
 FDisplayClusterViewportConfiguration::FDisplayClusterViewportConfiguration(FDisplayClusterViewportManager& InViewportManager)
-	: ViewportManager(InViewportManager.AsShared())
+	: ViewportManagerWeakPtr(InViewportManager.AsShared())
 { }
 
 FDisplayClusterViewportConfiguration::~FDisplayClusterViewportConfiguration()
@@ -200,7 +200,8 @@ bool FDisplayClusterViewportConfiguration::ImplUpdateConfiguration(EDisplayClust
 	if (RootActor)
 	{
 		const UDisplayClusterConfigurationData* ConfigurationData = RootActor->GetConfigData();
-		if (ConfigurationData)
+		FDisplayClusterViewportManager* ViewportManager = GetViewportManager();
+		if (ConfigurationData && ViewportManager)
 		{
 			FDisplayClusterViewportConfigurationBase ConfigurationBase(*ViewportManager, *RootActor, *ConfigurationData);
 			FDisplayClusterViewportConfigurationICVFX ConfigurationICVFX(*RootActor);
@@ -333,8 +334,10 @@ bool FDisplayClusterViewportConfiguration::UpdatePreviewConfiguration(EDisplayCl
 void FDisplayClusterViewportConfiguration::ImplUpdateConfigurationVisibility(ADisplayClusterRootActor& RootActor, const UDisplayClusterConfigurationData& ConfigurationData) const
 {
 	// Hide root actor components for all viewports
+	FDisplayClusterViewportManager* ViewportManager = GetViewportManager();
+
 	TSet<FPrimitiveComponentId> RootActorHidePrimitivesList;
-	if (RootActor.GetHiddenInGamePrimitives(RootActorHidePrimitivesList))
+	if (ViewportManager && RootActor.GetHiddenInGamePrimitives(RootActorHidePrimitivesList))
 	{
 		for (const TSharedPtr<FDisplayClusterViewport, ESPMode::ThreadSafe>& ViewportIt : ViewportManager->ImplGetViewports())
 		{
@@ -348,9 +351,12 @@ void FDisplayClusterViewportConfiguration::ImplUpdateConfigurationVisibility(ADi
 
 void FDisplayClusterViewportConfiguration::ImplPostUpdateRenderFrameConfiguration()
 {
-	// Some frame postprocess require additional render targetable resources
-	RenderFrameSettings.bShouldUseAdditionalFrameTargetableResource = ViewportManager->ShouldUseAdditionalFrameTargetableResource();
-	RenderFrameSettings.bShouldUseFullSizeFrameTargetableResource = ViewportManager->ShouldUseFullSizeFrameTargetableResource();
+	if (FDisplayClusterViewportManager* ViewportManager = GetViewportManager())
+	{
+		// Some frame postprocess require additional render targetable resources
+		RenderFrameSettings.bShouldUseAdditionalFrameTargetableResource = ViewportManager->ShouldUseAdditionalFrameTargetableResource();
+		RenderFrameSettings.bShouldUseFullSizeFrameTargetableResource = ViewportManager->ShouldUseFullSizeFrameTargetableResource();
+	}
 }
 
 void FDisplayClusterViewportConfiguration::ImplUpdateRenderFrameConfiguration(const FDisplayClusterConfigurationRenderFrame& InRenderFrameConfiguration)
diff --git a/Engine/Plugins/Runtime/nDisplay/Source/DisplayCluster/Private/Render/Viewport/Configuration/DisplayClusterViewportConfiguration.h b/Engine/Plugins/Runtime/nDisplay/Source/DisplayCluster/Private/Render/Viewport/Configuration/DisplayClusterViewportConfiguration.h
--- a/Engine/Plugins/Runtime/nDisplay/Source/DisplayCluster/Private/Render/Viewport/Configuration/DisplayClusterViewportConfiguration.h
+++ b/Engine/Plugins/Runtime/nDisplay/Source/DisplayCluster/Private/Render/Viewport/Configuration/DisplayClusterViewportConfiguration.h
@@ -51,8 +51,15 @@ private:
 	/** Get alpha channel capture mode (for LightCard, ChromaKey).*/
 	EDisplayClusterRenderFrameAlphaChannelCaptureMode GetAlphaChannelCaptureMode() const;
 
+	/** Get a pointer to the DC ViewportManager if it still exists. */
+	FDisplayClusterViewportManager* GetViewportManager() const
+	{
+		return ViewportManagerWeakPtr.IsValid() ? ViewportManagerWeakPtr.Pin().Get() : nullptr;
+	}
+
 private:
-	const TSharedRef<FDisplayClusterViewportManager> ViewportManager;
+	/** A reference to the owning viewport manager */
+	const TWeakPtr<FDisplayClusterViewportManager> ViewportManagerWeakPtr;
 
 	FDisplayClusterActorRef            RootActorRef;
 	FDisplayClusterRenderFrameSettings RenderFrameSettings;
diff --git a/Engine/Plugins/Runtime/nDisplay/Source/DisplayCluster/Private/Render/Viewport/DisplayClusterViewportManager.cpp b/Engine/Plugins/Runtime/nDisplay/Source/DisplayCluster/Private/Render/Viewport/DisplayClusterViewportManager.cpp
--- a/Engine/Plugins/Runtime/nDisplay/Source/DisplayCluster/Private/Render/Viewport/DisplayClusterViewportManager.cpp
+++ b/Engine/Plugins/Runtime/nDisplay/Source/DisplayCluster/Private/Render/Viewport/DisplayClusterViewportManager.cpp
@@ -475,7 +475,7 @@ bool FDisplayClusterViewportManager::BeginNewFrame(FViewport* InViewport, UWorld
 {
 	check(IsInGameThread());
 
-	OutRenderFrame.ViewportManager = this->AsShared();
+	OutRenderFrame.ViewportManagerWeakPtr = this->AsShared();
 
 	// Handle world runtime update
 	UWorld* CurrentWorld = GetCurrentWorld();
diff --git a/Engine/Plugins/Runtime/nDisplay/Source/DisplayCluster/Private/Render/Viewport/DisplayClusterViewportManagerViewExtension.cpp b/Engine/Plugins/Runtime/nDisplay/Source/DisplayCluster/Private/Render/Viewport/DisplayClusterViewportManagerViewExtension.cpp
--- a/Engine/Plugins/Runtime/nDisplay/Source/DisplayCluster/Private/Render/Viewport/DisplayClusterViewportManagerViewExtension.cpp
+++ b/Engine/Plugins/Runtime/nDisplay/Source/DisplayCluster/Private/Render/Viewport/DisplayClusterViewportManagerViewExtension.cpp
@@ -214,7 +214,7 @@ bool FDisplayClusterViewportManagerViewExtension::IsActiveThisFrame_Internal(con
 	if (Context.IsA(MoveTempIfPossible(DCViewExtensionContext)))
 	{
 		const FDisplayClusterSceneViewExtensionContext& DisplayContext = static_cast<const FDisplayClusterSceneViewExtensionContext&>(Context);
-		if (DisplayContext.ViewportManager.Get() == GetViewportManager())
+		if (DisplayContext.GetViewportManager() == GetViewportManager())
 		{
 			// Apply only for DC viewports
 			return true;
diff --git a/Engine/Plugins/Runtime/nDisplay/Source/DisplayCluster/Private/Render/Viewport/LightCard/DisplayClusterViewportLightCardManager.cpp b/Engine/Plugins/Runtime/nDisplay/Source/DisplayCluster/Private/Render/Viewport/LightCard/DisplayClusterViewportLightCardManager.cpp
--- a/Engine/Plugins/Runtime/nDisplay/Source/DisplayCluster/Private/Render/Viewport/LightCard/DisplayClusterViewportLightCardManager.cpp
+++ b/Engine/Plugins/Runtime/nDisplay/Source/DisplayCluster/Private/Render/Viewport/LightCard/DisplayClusterViewportLightCardManager.cpp
@@ -26,7 +26,7 @@ static TAutoConsoleVariable<int32> CVarUVLightCardTextureSize(
 // FDisplayClusterViewportLightCardManager
 ///////////////////////////////////////////////////////////////////////////////////////////////
 FDisplayClusterViewportLightCardManager::FDisplayClusterViewportLightCardManager(FDisplayClusterViewportManager& InViewportManager)
-	: ViewportManager(InViewportManager.AsShared())
+	: ViewportManagerWeakPtr(InViewportManager.AsShared())
 {
 	LightCardManagerProxy = MakeShared<FDisplayClusterViewportLightCardManagerProxy, ESPMode::ThreadSafe>();
 }
@@ -89,6 +89,12 @@ void FDisplayClusterViewportLightCardManager::UpdateUVLightCardData()
 {
 	ReleaseUVLightCardData();
 
+	FDisplayClusterViewportManager* ViewportManager = GetViewportManager();
+	if (!ViewportManager)
+	{
+		return;
+	}
+
 	/** The list of UV light card actors that are referenced by the root actor */
 	TArray<ADisplayClusterLightCardActor*> UVLightCardActors;
 
@@ -165,7 +171,8 @@ void FDisplayClusterViewportLightCardManager::UpdateUVLightCardResource()
 void FDisplayClusterViewportLightCardManager::RenderUVLightCard()
 {
 	// Render UV LightCard:
-	UWorld* World = ViewportManager->GetCurrentWorld();
+	FDisplayClusterViewportManager* ViewportManager = GetViewportManager();
+	UWorld* World = ViewportManager ? ViewportManager->GetCurrentWorld() : nullptr;
 	if (IsUVLightCardEnabled() && World && LightCardManagerProxy.IsValid())
 	{
 		UpdateUVLightCardResource();
diff --git a/Engine/Plugins/Runtime/nDisplay/Source/DisplayCluster/Private/Render/Viewport/LightCard/DisplayClusterViewportLightCardManager.h b/Engine/Plugins/Runtime/nDisplay/Source/DisplayCluster/Private/Render/Viewport/LightCard/DisplayClusterViewportLightCardManager.h
--- a/Engine/Plugins/Runtime/nDisplay/Source/DisplayCluster/Private/Render/Viewport/LightCard/DisplayClusterViewportLightCardManager.h
+++ b/Engine/Plugins/Runtime/nDisplay/Source/DisplayCluster/Private/Render/Viewport/LightCard/DisplayClusterViewportLightCardManager.h
@@ -73,9 +73,15 @@ private:
 	/** Release UVLightCard data game thread*/
 	void ReleaseUVLightCardData();
 
+	/** Get a pointer to the DC ViewportManager if it still exists. */
+	FDisplayClusterViewportManager* GetViewportManager() const
+	{
+		return ViewportManagerWeakPtr.IsValid() ? ViewportManagerWeakPtr.Pin().Get() : nullptr;
+	}
+
 private:
 	/** A reference to the owning viewport manager */
-	const TSharedRef<FDisplayClusterViewportManager, ESPMode::ThreadSafe> ViewportManager;
+	const TWeakPtr<FDisplayClusterViewportManager, ESPMode::ThreadSafe> ViewportManagerWeakPtr;
 
 	/** A list of primitive components that have been added to the preview scene for rendering in the current frame */
 	TArray<UPrimitiveComponent*> UVLightCardPrimitiveComponents;
diff --git a/Engine/Plugins/Runtime/nDisplay/Source/DisplayCluster/Private/Render/Viewport/Postprocess/DisplayClusterViewportPostProcessManager.cpp b/Engine/Plugins/Runtime/nDisplay/Source/DisplayCluster/Private/Render/Viewport/Postprocess/DisplayClusterViewportPostProcessManager.cpp
--- a/Engine/Plugins/Runtime/nDisplay/Source/DisplayCluster/Private/Render/Viewport/Postprocess/DisplayClusterViewportPostProcessManager.cpp
+++ b/Engine/Plugins/Runtime/nDisplay/Source/DisplayCluster/Private/Render/Viewport/Postprocess/DisplayClusterViewportPostProcessManager.cpp
@@ -62,7 +62,7 @@ static TAutoConsoleVariable<int32> CVarPostprocessFrameAfterWarpBlend(
 // FDisplayClusterViewportPostProcessManager
 //----------------------------------------------------------------------------------------------------------
 FDisplayClusterViewportPostProcessManager::FDisplayClusterViewportPostProcessManager(FDisplayClusterViewportManager& InViewportManager)
-	: ViewportManager(InViewportManager.AsShared())
+	: ViewportManagerWeakPtr(InViewportManager.AsShared())
 {
 	OutputRemap = MakeShared<FDisplayClusterViewportPostProcessOutputRemap, ESPMode::ThreadSafe>();
 }
@@ -74,7 +74,7 @@ FDisplayClusterViewportPostProcessManager::~FDisplayClusterViewportPostProcessMa
 
 FDisplayClusterViewportManager* FDisplayClusterViewportPostProcessManager::GetViewportManager() const
 {
-	return &(ViewportManager.Get());
+	return ViewportManagerWeakPtr.IsValid() ? ViewportManagerWeakPtr.Pin().Get() : nullptr;
 }
 
 void FDisplayClusterViewportPostProcessManager::Release()
@@ -96,13 +96,14 @@ bool FDisplayClusterViewportPostProcessManager::HandleStartScene()
 {
 	check(IsInGameThread());;
 
-	if (ViewportManager->IsSceneOpened())
+	FDisplayClusterViewportManager* ViewportManager = GetViewportManager();
+	if (ViewportManager && ViewportManager->IsSceneOpened())
 	{
 		for (const TSharedPtr<IDisplayClusterPostProcess, ESPMode::ThreadSafe>& It : Postprocess)
 		{
 			if (It.IsValid())
 			{
-				It->HandleStartScene(GetViewportManager());
+				It->HandleStartScene(ViewportManager);
 			}
 		}
 	}
@@ -114,11 +115,14 @@ void FDisplayClusterViewportPostProcessManager::HandleEndScene()
 {
 	check(IsInGameThread());
 
-	for (const TSharedPtr<IDisplayClusterPostProcess, ESPMode::ThreadSafe>& It : Postprocess)
+	if (FDisplayClusterViewportManager* ViewportManager = GetViewportManager())
 	{
-		if (It.IsValid())
+		for (const TSharedPtr<IDisplayClusterPostProcess, ESPMode::ThreadSafe>& It : Postprocess)
 		{
-			It->HandleEndScene(GetViewportManager());
+			if (It.IsValid())
+			{
+				It->HandleEndScene(ViewportManager);
+			}
 		}
 	}
 }
@@ -162,9 +166,10 @@ bool FDisplayClusterViewportPostProcessManager::CreatePostprocess(const FString&
 		{
 			UE_LOG(LogDisplayClusterRender, Log, TEXT("PostProcess '%s', type '%s' : Created"), *InConfigurationPostprocess->Type, *InPostprocessId);
 
-			if (ViewportManager->IsSceneOpened())
+			FDisplayClusterViewportManager* ViewportManager = GetViewportManager();
+			if (ViewportManager && ViewportManager->IsSceneOpened())
 			{
-				PostProcessInstance->HandleStartScene(GetViewportManager());
+				PostProcessInstance->HandleStartScene(ViewportManager);
 			}
 
 			for (int32 Order = 0; Order < Postprocess.Num(); Order++)
@@ -339,22 +344,24 @@ void FDisplayClusterViewportPostProcessManager::Tick()
 
 void FDisplayClusterViewportPostProcessManager::HandleSetupNewFrame()
 {
-	if (CVarCustomPPEnabled.GetValueOnGameThread() != 0)
+	FDisplayClusterViewportManager* ViewportManager = GetViewportManager();
+	if (ViewportManager && CVarCustomPPEnabled.GetValueOnGameThread() != 0)
 	{
 		for (const TSharedPtr<IDisplayClusterPostProcess, ESPMode::ThreadSafe >& It : Postprocess)
 		{
-			It->HandleSetupNewFrame(GetViewportManager());
+			It->HandleSetupNewFrame(ViewportManager);
 		}
 	}
 }
 
 void FDisplayClusterViewportPostProcessManager::HandleBeginNewFrame(FDisplayClusterRenderFrame& InOutRenderFrame)
 {
-	if (CVarCustomPPEnabled.GetValueOnGameThread() != 0)
+	FDisplayClusterViewportManager* ViewportManager = GetViewportManager();
+	if (ViewportManager && CVarCustomPPEnabled.GetValueOnGameThread() != 0)
 	{
 		for (const TSharedPtr<IDisplayClusterPostProcess, ESPMode::ThreadSafe >& It : Postprocess)
 		{
-			It->HandleBeginNewFrame(GetViewportManager(), InOutRenderFrame);
+			It->HandleBeginNewFrame(ViewportManager, InOutRenderFrame);
 		}
 	}
 }
diff --git a/Engine/Plugins/Runtime/nDisplay/Source/DisplayCluster/Private/Render/Viewport/Postprocess/DisplayClusterViewportPostProcessManager.h b/Engine/Plugins/Runtime/nDisplay/Source/DisplayCluster/Private/Render/Viewport/Postprocess/DisplayClusterViewportPostProcessManager.h
--- a/Engine/Plugins/Runtime/nDisplay/Source/DisplayCluster/Private/Render/Viewport/Postprocess/DisplayClusterViewportPostProcessManager.h
+++ b/Engine/Plugins/Runtime/nDisplay/Source/DisplayCluster/Private/Render/Viewport/Postprocess/DisplayClusterViewportPostProcessManager.h
@@ -71,6 +71,8 @@ public:
 
 private:
 	TSharedPtr<IDisplayClusterPostProcess, ESPMode::ThreadSafe> ImplFindPostProcess(const FString& InPostprocessId) const;
+
+	/** Get a pointer to the DC ViewportManager if it still exists. */
 	FDisplayClusterViewportManager* GetViewportManager() const;
 
 protected:
@@ -84,7 +86,8 @@ protected:
 	TArray<TSharedPtr<IDisplayClusterPostProcess, ESPMode::ThreadSafe>> Postprocess;
 
 private:
-	const TSharedRef<FDisplayClusterViewportManager, ESPMode::ThreadSafe> ViewportManager;
+	/** A reference to the owning viewport manager */
+	const TWeakPtr<FDisplayClusterViewportManager, ESPMode::ThreadSafe> ViewportManagerWeakPtr;
 
 	// Render thread instances
 	TArray<TSharedPtr<IDisplayClusterPostProcess, ESPMode::ThreadSafe>> PostprocessProxy;
diff --git a/Engine/Plugins/Runtime/nDisplay/Source/DisplayCluster/Public/DisplayClusterRootActor.h b/Engine/Plugins/Runtime/nDisplay/Source/DisplayCluster/Public/DisplayClusterRootActor.h
--- a/Engine/Plugins/Runtime/nDisplay/Source/DisplayCluster/Public/DisplayClusterRootActor.h
+++ b/Engine/Plugins/Runtime/nDisplay/Source/DisplayCluster/Public/DisplayClusterRootActor.h
@@ -205,11 +205,14 @@ public:
 	}
 
 private:
-	/** Create new viewport manager instance if not exists. */
+	/** Create a new instance of the viewport manager if it does not exist. */
 	void CreateViewportManagerImpl();
 
+	/** Release the viewport manager instance, if it exists. */
+	void RemoveViewportManagerImpl();
+
 private:
-	// Viewport manager for this configuration
+	// DC ViewportManager instance for this DCRA
 	TSharedPtr<class FDisplayClusterViewportManager, ESPMode::ThreadSafe> ViewportManager;
 
 //////////////////////////////////////////////////////////////////////////////////////////////
@@ -458,7 +461,9 @@ public:
 	void Constructor_Editor();
 	void Destructor_Editor();
 
-	void Tick_Editor(float DeltaSeconds);
+	/** Perform a rendering of the DCRA preview. */
+	void RenderPreview_Editor();
+
 	void PostLoad_Editor();
 	void PostActorCreated_Editor();
 	void EndPlay_Editor(const EEndPlayReason::Type EndPlayReason);
diff --git a/Engine/Plugins/Runtime/nDisplay/Source/DisplayCluster/Public/DisplayClusterSceneViewExtensions.h b/Engine/Plugins/Runtime/nDisplay/Source/DisplayCluster/Public/DisplayClusterSceneViewExtensions.h
--- a/Engine/Plugins/Runtime/nDisplay/Source/DisplayCluster/Public/DisplayClusterSceneViewExtensions.h
+++ b/Engine/Plugins/Runtime/nDisplay/Source/DisplayCluster/Public/DisplayClusterSceneViewExtensions.h
@@ -7,7 +7,9 @@
 
 class IDisplayClusterViewportManager;
 
-/** Contains information about the context in which this scene view extension will be used. */
+/**
+ * Contains information about the context in which this scene view extension will be used.
+ */
 struct FDisplayClusterSceneViewExtensionContext : public FSceneViewExtensionContext
 {
 private:
@@ -21,12 +23,6 @@ private:
 	}
 
 public:
-	// The id of the nDisplay viewport being rendered.
-	const FString ViewportId;
-
-	// Reference to viewport manager
-	TSharedPtr<IDisplayClusterViewportManager, ESPMode::ThreadSafe> ViewportManager;
-
 	FDisplayClusterSceneViewExtensionContext()
 		: FSceneViewExtensionContext()
 	{ }
@@ -34,12 +30,31 @@ public:
 	FDisplayClusterSceneViewExtensionContext(FViewport* InViewport, const TSharedPtr<IDisplayClusterViewportManager, ESPMode::ThreadSafe>& InViewportManager, const FString& InViewportId)
 		: FSceneViewExtensionContext(InViewport)
 		, ViewportId(InViewportId)
-		, ViewportManager(InViewportManager)
+		, ViewportManagerWeakPtr(InViewportManager)
 	{ }
 
 	FDisplayClusterSceneViewExtensionContext(FSceneInterface* InScene, const TSharedPtr<IDisplayClusterViewportManager, ESPMode::ThreadSafe>& InViewportManager, const FString& InViewportId)
 		: FSceneViewExtensionContext(InScene)
 		, ViewportId(InViewportId)
-		, ViewportManager(InViewportManager)
+		, ViewportManagerWeakPtr(InViewportManager)
 	{ }
+
+	/** Returns a pointer to the DC ViewportManager used for this VE context. */
+	IDisplayClusterViewportManager* GetViewportManager() const
+	{
+		return ViewportManagerWeakPtr.IsValid() ? ViewportManagerWeakPtr.Pin().Get() : nullptr;
+	}
+
+	/** Returns the name of the DC viewport used for this VE context. */
+	const FString& GetViewportId() const
+	{
+		return ViewportId;
+	}
+
+private:
+	// The id of the nDisplay viewport being rendered.
+	const FString ViewportId;
+
+	// Reference to viewport manager
+	const TWeakPtr<IDisplayClusterViewportManager, ESPMode::ThreadSafe> ViewportManagerWeakPtr;
 };
diff --git a/Engine/Plugins/Runtime/nDisplay/Source/DisplayCluster/Public/Render/Viewport/RenderFrame/DisplayClusterRenderFrame.h b/Engine/Plugins/Runtime/nDisplay/Source/DisplayCluster/Public/Render/Viewport/RenderFrame/DisplayClusterRenderFrame.h
--- a/Engine/Plugins/Runtime/nDisplay/Source/DisplayCluster/Public/Render/Viewport/RenderFrame/DisplayClusterRenderFrame.h
+++ b/Engine/Plugins/Runtime/nDisplay/Source/DisplayCluster/Public/Render/Viewport/RenderFrame/DisplayClusterRenderFrame.h
@@ -82,6 +82,13 @@ struct FDisplayClusterRenderFrameTarget
  */
 class FDisplayClusterRenderFrame
 {
+public:
+	/** Returns a pointer to the DCViewportManager used for this frame. */
+	IDisplayClusterViewportManager* GetViewportManager() const
+	{
+		return ViewportManagerWeakPtr.IsValid() ? ViewportManagerWeakPtr.Pin().Get() : nullptr;
+	}
+
 public:
 	// Render frame to this targets
 	TArray<FDisplayClusterRenderFrameTarget> RenderTargets;
@@ -93,5 +100,5 @@ public:
 	int32 DesiredNumberOfViews = 0;
 
 	// Ref to the DC Viewport Manager
-	TSharedPtr<IDisplayClusterViewportManager, ESPMode::ThreadSafe> ViewportManager;
+	TWeakPtr<IDisplayClusterViewportManager, ESPMode::ThreadSafe> ViewportManagerWeakPtr;
 };
