nDisplay: Fix crash when applying a material that is assigned to a PPV and an nDisplay actor is in the scene with editor preview enabled. When GC runs the MID pool holds on to invalid material objects which we now clear.

diff --git a/Engine/Plugins/Runtime/nDisplay/Source/DisplayCluster/Private/Render/Viewport/DisplayClusterViewport.h b/Engine/Plugins/Runtime/nDisplay/Source/DisplayCluster/Private/Render/Viewport/DisplayClusterViewport.h
--- a/Engine/Plugins/Runtime/nDisplay/Source/DisplayCluster/Private/Render/Viewport/DisplayClusterViewport.h
+++ b/Engine/Plugins/Runtime/nDisplay/Source/DisplayCluster/Private/Render/Viewport/DisplayClusterViewport.h
@@ -245,6 +245,8 @@ private:
 	// Support view states for preview
 private:
 	FSceneViewStateInterface* GetViewState(uint32 ViewIndex);
+
+public:
 	void CleanupViewState();
 
 private:
diff --git a/Engine/Plugins/Runtime/nDisplay/Source/DisplayCluster/Private/Render/Viewport/DisplayClusterViewportManager.cpp b/Engine/Plugins/Runtime/nDisplay/Source/DisplayCluster/Private/Render/Viewport/DisplayClusterViewportManager.cpp
--- a/Engine/Plugins/Runtime/nDisplay/Source/DisplayCluster/Private/Render/Viewport/DisplayClusterViewportManager.cpp
+++ b/Engine/Plugins/Runtime/nDisplay/Source/DisplayCluster/Private/Render/Viewport/DisplayClusterViewportManager.cpp
@@ -59,6 +59,8 @@ FDisplayClusterViewportManager::FDisplayClusterViewportManager()
 
 	// Always reset RTT when root actor re-created
 	ResetSceneRenderTargetSize();
+
+	RegisterCallbacks();
 }
 
 void FDisplayClusterViewportManager::Initialize()
@@ -77,6 +79,8 @@ void FDisplayClusterViewportManager::Initialize()
 
 FDisplayClusterViewportManager::~FDisplayClusterViewportManager()
 {
+	UnregisterCallbacks();
+	
 	// Remove viewports
 	Viewports.Reset();
 	ClusterNodeViewports.Reset();
@@ -433,6 +437,20 @@ void FDisplayClusterViewportManager::ImplUpdateClusterNodeViewports(const EDispl
 	}
 }
 
+void FDisplayClusterViewportManager::RegisterCallbacks()
+{
+#if WITH_EDITOR
+	PreGarbageCollectHandle = FCoreUObjectDelegates::GetPreGarbageCollectDelegate().AddRaw(this, &FDisplayClusterViewportManager::OnPreGarbageCollect);
+#endif
+}
+
+void FDisplayClusterViewportManager::UnregisterCallbacks()
+{
+#if WITH_EDITOR
+	FCoreUObjectDelegates::GetPreGarbageCollectDelegate().Remove(PreGarbageCollectHandle);
+#endif
+}
+
 int32 FDisplayClusterViewportManager::GetViewPerViewportAmount() const
 {
 	switch (GetRenderMode())
diff --git a/Engine/Plugins/Runtime/nDisplay/Source/DisplayCluster/Private/Render/Viewport/DisplayClusterViewportManager.h b/Engine/Plugins/Runtime/nDisplay/Source/DisplayCluster/Private/Render/Viewport/DisplayClusterViewportManager.h
--- a/Engine/Plugins/Runtime/nDisplay/Source/DisplayCluster/Private/Render/Viewport/DisplayClusterViewportManager.h
+++ b/Engine/Plugins/Runtime/nDisplay/Source/DisplayCluster/Private/Render/Viewport/DisplayClusterViewportManager.h
@@ -86,7 +86,13 @@ public:
 	virtual bool RenderInEditor(FDisplayClusterRenderFrame& InRenderFrame, FViewport* InViewport, const uint32 InFirstViewportNum, const int32 InViewportsAmount, int32& OutViewportsAmount, bool& bOutFrameRendered) override;
 	
 	void ImplUpdatePreviewRTTResources();
+
+private:
+	/** Called before garbage collection is run */
+	void OnPreGarbageCollect();
+
 #endif
+public:
 
 	virtual IDisplayClusterViewport* FindViewport(const FString& InViewportId) const override;
 	virtual IDisplayClusterViewport* FindViewport(const int32 ViewIndex, uint32* OutContextNum = nullptr) const override;
@@ -167,6 +173,12 @@ private:
 	void HandleViewportRTTChanges(const TArray<FDisplayClusterViewport_Context>& PrevContexts, const TArray<FDisplayClusterViewport_Context>& Contexts);
 	void ImplUpdateClusterNodeViewports(const EDisplayClusterRenderFrameMode InRenderMode, const FString& InClusterNodeId);
 
+	/** Register any callbacks */
+	void RegisterCallbacks();
+	
+	/** Unregister any callbacks used */
+	void UnregisterCallbacks();
+	
 protected:
 	friend FDisplayClusterViewportManagerProxy;
 	friend FDisplayClusterViewportConfiguration;
@@ -207,4 +219,9 @@ private:
 
 	// This VE is used for LocalPlayer::GetViewPoint(). It calls the ISceneViewExtension::SetupViewPoint() function of this VE.
 	TSharedPtr<class FDisplayClusterViewportManagerViewPointExtension, ESPMode::ThreadSafe> ViewportManagerViewPointExtension;
+
+#if WITH_EDITOR
+	/** The handle for FCoreUObjectDelegates::GetPreGarbageCollectDelegate */
+	FDelegateHandle PreGarbageCollectHandle;
+#endif
 };
diff --git a/Engine/Plugins/Runtime/nDisplay/Source/DisplayCluster/Private/Render/Viewport/DisplayClusterViewportManager_EditorPreview.cpp b/Engine/Plugins/Runtime/nDisplay/Source/DisplayCluster/Private/Render/Viewport/DisplayClusterViewportManager_EditorPreview.cpp
--- a/Engine/Plugins/Runtime/nDisplay/Source/DisplayCluster/Private/Render/Viewport/DisplayClusterViewportManager_EditorPreview.cpp
+++ b/Engine/Plugins/Runtime/nDisplay/Source/DisplayCluster/Private/Render/Viewport/DisplayClusterViewportManager_EditorPreview.cpp
@@ -95,6 +95,18 @@ void FDisplayClusterViewportManager::ImplUpdatePreviewRTTResources()
 	}
 }
 
+void FDisplayClusterViewportManager::OnPreGarbageCollect()
+{
+	// The view state can reference materials from the world being cleaned up. (Example post process materials)
+	for (TSharedPtr<FDisplayClusterViewport, ESPMode::ThreadSafe>& ViewportIt : Viewports)
+	{
+		if (ViewportIt.IsValid())
+		{
+			ViewportIt->CleanupViewState();
+		}
+	}
+}
+
 bool FDisplayClusterViewportManager::RenderInEditor(FDisplayClusterRenderFrame& InRenderFrame, FViewport* InViewport, const uint32 InFirstViewportNum, const int32 InViewportsAmount, int32& OutViewportsAmount, bool& bOutFrameRendered)
 {
 	bOutFrameRendered = false;
