nDisplay: Fix crash when opening ICVFX editor viewport.

diff --git a/Engine/Plugins/Runtime/nDisplay/Source/DisplayClusterLightCardEditor/Private/Viewport/DisplayClusterLightCardEditorViewportClient.cpp b/Engine/Plugins/Runtime/nDisplay/Source/DisplayClusterLightCardEditor/Private/Viewport/DisplayClusterLightCardEditorViewportClient.cpp
--- a/Engine/Plugins/Runtime/nDisplay/Source/DisplayClusterLightCardEditor/Private/Viewport/DisplayClusterLightCardEditorViewportClient.cpp
+++ b/Engine/Plugins/Runtime/nDisplay/Source/DisplayClusterLightCardEditor/Private/Viewport/DisplayClusterLightCardEditorViewportClient.cpp
@@ -52,6 +52,7 @@
 #include "PreviewScene.h"
 #include "ProceduralMeshComponent.h"
 #include "RayTracingDebugVisualizationMenuCommands.h"
+#include "RenderResource.h"
 #include "ScopedTransaction.h"
 #include "Settings/LevelEditorViewportSettings.h"
 #include "Slate/SceneViewport.h"
@@ -1418,7 +1419,7 @@ EMouseCursor::Type FDisplayClusterLightCardEditorViewportClient::GetCursor(FView
 		HHitProxy* HitProxy = InViewport->GetHitProxy(X,Y);
 		if (HitProxy)
 		{
-			bShouldCheckHitProxy = true;
+			TrySetShouldCheckHitProxy();
 
 			if (HitProxy->IsA(HActor::StaticGetType()))
 			{
@@ -1606,7 +1607,7 @@ void FDisplayClusterLightCardEditorViewportClient::UpdatePreviewActor(ADisplayCl
 	auto Finalize = [this, ProxyType, StageActor]()
 	{
 		Viewport->InvalidateHitProxy();
-		bShouldCheckHitProxy = true;
+		TrySetShouldCheckHitProxy();
 
 		if (!StageActor)
 		{
@@ -1966,7 +1967,7 @@ void FDisplayClusterLightCardEditorViewportClient::SetProjectionMode(EDisplayClu
 		Viewport->InvalidateHitProxy();
 	}
 
-	bShouldCheckHitProxy = true;
+	TrySetShouldCheckHitProxy();
 }
 
 float FDisplayClusterLightCardEditorViewportClient::GetProjectionModeFOV(EDisplayClusterMeshProjectionType InProjectionMode) const
@@ -2006,7 +2007,7 @@ void FDisplayClusterLightCardEditorViewportClient::SetProjectionModeFOV(EDisplay
 	}
 
 	Viewport->InvalidateHitProxy();
-	bShouldCheckHitProxy = true;
+	TrySetShouldCheckHitProxy();
 }
 
 void FDisplayClusterLightCardEditorViewportClient::ResetCamera(bool bLocationOnly)
@@ -3072,6 +3073,17 @@ void FDisplayClusterLightCardEditorViewportClient::ResetProjectionViewConfigurat
 	RestoreProjectionCameraTransform();
 }
 
+void FDisplayClusterLightCardEditorViewportClient::TrySetShouldCheckHitProxy()
+{
+	if (GetEditorViewportWidget().IsValid() && GetEditorViewportWidget()->GetSceneViewport().IsValid() &&
+		((FRenderResource*)GetEditorViewportWidget()->GetSceneViewport().Get())->IsInitialized())
+	{
+		// If this is set before the SceneViewport is initialized it's possible the client will tick first,
+		// attempt to check hit proxies, and crash because RHI is assumed to be initialized at this point.
+		bShouldCheckHitProxy = true;
+	}
+}
+
 void FDisplayClusterLightCardEditorViewportClient::EnterDrawingLightCardMode()
 {
 	if (InputMode == EInputMode::Idle)
diff --git a/Engine/Plugins/Runtime/nDisplay/Source/DisplayClusterLightCardEditor/Private/Viewport/DisplayClusterLightCardEditorViewportClient.h b/Engine/Plugins/Runtime/nDisplay/Source/DisplayClusterLightCardEditor/Private/Viewport/DisplayClusterLightCardEditorViewportClient.h
--- a/Engine/Plugins/Runtime/nDisplay/Source/DisplayClusterLightCardEditor/Private/Viewport/DisplayClusterLightCardEditorViewportClient.h
+++ b/Engine/Plugins/Runtime/nDisplay/Source/DisplayClusterLightCardEditor/Private/Viewport/DisplayClusterLightCardEditorViewportClient.h
@@ -247,6 +247,9 @@ private:
 	/** Resets the stored view configurations to default values for each projection mode */
 	void ResetProjectionViewConfigurations();
 
+	/** Attempt to set bShouldCheckHitProxy */
+	void TrySetShouldCheckHitProxy();
+	
 	/** Creates a new light card using a polygon alpha mask as defined by the given mouse positions on the viewport */
 	void CreateDrawnLightCard(const TArray<FIntPoint>& MousePositions);
 
