Fill Derived Data Cache command:

-Edit terminology in all instances of "Preload DDC" to "Fill DDC"

diff --git a/Engine/Plugins/Runtime/nDisplay/Source/DisplayClusterFillDerivedDataCache/DisplayClusterFillDerivedDataCache.build.cs b/Engine/Plugins/Runtime/nDisplay/Source/DisplayClusterFillDerivedDataCache/DisplayClusterFillDerivedDataCache.build.cs
new file mode 100644
--- /dev/null
+++ b/Engine/Plugins/Runtime/nDisplay/Source/DisplayClusterFillDerivedDataCache/DisplayClusterFillDerivedDataCache.build.cs
@@ -0,0 +1,46 @@
+// Copyright Epic Games, Inc. All Rights Reserved.
+
+using System.IO;
+using UnrealBuildTool;
+
+public class DisplayClusterFillDerivedDataCache : ModuleRules
+{
+	public DisplayClusterFillDerivedDataCache(ReadOnlyTargetRules Target) : base(Target)
+	{
+		PCHUsage = ModuleRules.PCHUsageMode.UseExplicitOrSharedPCHs;
+
+		PublicDependencyModuleNames.AddRange(
+			new string[]
+			{
+				"Core",
+				"OutputLog", 
+				"UnrealEd"
+			}
+		);
+			
+		PrivateDependencyModuleNames.AddRange(
+			new string[]
+			{
+				"AssetRegistry",
+				"AssetTools",
+				"CoreUObject",
+				"ContentBrowser",
+				"Engine",
+				"EditorStyle",
+				"EditorWidgets",
+				"InputCore",
+				"Kismet",
+				"Messaging",
+				"Networking",
+				"Projects",
+				"PropertyEditor",
+				"Slate",
+				"SlateCore",
+				"ToolMenus", 
+				"ToolWidgets",
+				"UnrealEd",
+				"WorkspaceMenuStructure"
+			}
+		);
+	}
+}
diff --git a/Engine/Plugins/Runtime/nDisplay/Source/DisplayClusterFillDerivedDataCache/Private/DisplayClusterFillDerivedDataCacheModule.cpp b/Engine/Plugins/Runtime/nDisplay/Source/DisplayClusterFillDerivedDataCache/Private/DisplayClusterFillDerivedDataCacheModule.cpp
new file mode 100644
--- /dev/null
+++ b/Engine/Plugins/Runtime/nDisplay/Source/DisplayClusterFillDerivedDataCache/Private/DisplayClusterFillDerivedDataCacheModule.cpp
@@ -0,0 +1,65 @@
+// Copyright Epic Games, Inc. All Rights Reserved.
+
+#include "DisplayClusterFillDerivedDataCacheModule.h"
+
+#include "DisplayClusterFillDerivedDataCacheLog.h"
+#include "DisplayClusterFillDerivedDataCacheWorker.h"
+
+#include "Commandlets/DerivedDataCacheCommandlet.h"
+#include "Framework/Notifications/NotificationManager.h"
+#include "Interfaces/IProjectManager.h"
+#include "Modules/ModuleManager.h"
+#include "ToolMenu.h"
+#include "ToolMenus.h"
+#include "Widgets/Notifications/SNotificationList.h"
+
+#define LOCTEXT_NAMESPACE "FDisplayClusterFillDerivedDataCacheModule"
+
+FDisplayClusterFillDerivedDataCacheModule& FDisplayClusterFillDerivedDataCacheModule::Get()
+{
+	return FModuleManager::GetModuleChecked<FDisplayClusterFillDerivedDataCacheModule>("DisplayClusterFillDerivedDataCacheModule");
+}
+
+void FDisplayClusterFillDerivedDataCacheModule::StartupModule()
+{
+	FCoreDelegates::OnFEngineLoopInitComplete.AddRaw(this, &FDisplayClusterFillDerivedDataCacheModule::OnFEngineLoopInitComplete);
+}
+
+void FDisplayClusterFillDerivedDataCacheModule::OnFEngineLoopInitComplete()
+{
+	UToolMenu* Menu = UToolMenus::Get()->ExtendMenu("LevelEditor.MainMenu.Build");
+
+	FToolMenuSection& Section = Menu->FindOrAddSection("LevelEditorAutomation");
+		
+	Section.AddMenuEntry(
+		"FillDerivedDataCache",
+		LOCTEXT("FillDerivedDataCache", "Fill Derived Data Cache (Prepare all Shaders)"),
+		LOCTEXT("FillDerivedDataCacheTooltip", "Prepare all shaders and preprocess Nanite and other heavy data ahead of time for all assets in this project for all checked platforms in Project Settings > Supported Platforms. \nThis is useful to avoid shader compilation and pop-in when opening a level or running PIE/Standalone. \nNote that this operation can take a very long time to complete and can take up a lot of disk space, depending on the size of your project."),
+		FSlateIcon(),
+		FUIAction(FExecuteAction::CreateRaw(this, &FDisplayClusterFillDerivedDataCacheModule::CreateAsyncTaskWorker)));
+}
+
+void FDisplayClusterFillDerivedDataCacheModule::ShutdownModule()
+{
+	UToolMenus::UnregisterOwner(this);
+	FCoreDelegates::OnFEngineLoopInitComplete.RemoveAll(this);
+
+	if (AsyncTaskWorker)
+	{
+		AsyncTaskWorker->CancelTask();
+		AsyncTaskWorker = nullptr;
+	}
+}
+
+void FDisplayClusterFillDerivedDataCacheModule::CreateAsyncTaskWorker()
+{
+	if (AsyncTaskWorker)
+	{
+		AsyncTaskWorker->CancelTask();
+	}
+	AsyncTaskWorker = MakeUnique<FDisplayClusterFillDerivedDataCacheWorker>();
+}
+
+#undef LOCTEXT_NAMESPACE
+
+IMPLEMENT_MODULE(FDisplayClusterFillDerivedDataCacheModule, DisplayClusterFillDerivedDataCache)
diff --git a/Engine/Plugins/Runtime/nDisplay/Source/DisplayClusterFillDerivedDataCache/Private/DisplayClusterFillDerivedDataCacheWorker.cpp b/Engine/Plugins/Runtime/nDisplay/Source/DisplayClusterFillDerivedDataCache/Private/DisplayClusterFillDerivedDataCacheWorker.cpp
new file mode 100644
--- /dev/null
+++ b/Engine/Plugins/Runtime/nDisplay/Source/DisplayClusterFillDerivedDataCache/Private/DisplayClusterFillDerivedDataCacheWorker.cpp
@@ -0,0 +1,308 @@
+// Copyright Epic Games, Inc. All Rights Reserved.
+
+#include "DisplayClusterFillDerivedDataCacheWorker.h"
+
+#include "DisplayClusterFillDerivedDataCacheLog.h"
+
+#include "Commandlets/DerivedDataCacheCommandlet.h"
+#include "Framework/Docking/TabManager.h"
+#include "Framework/Notifications/NotificationManager.h"
+#include "HAL/PlatformProcess.h"
+#include "Interfaces/IProjectManager.h"
+#include "Internationalization/Regex.h"
+#include "Misc/App.h"
+#include "Misc/AsyncTaskNotification.h"
+#include "Misc/DataDrivenPlatformInfoRegistry.h"
+#include "Misc/ScopedSlowTask.h"
+#include "Widgets/Notifications/SNotificationList.h"
+
+#define LOCTEXT_NAMESPACE "FDisplayClusterFillDerivedDataCacheWorker"
+
+const FText EnumeratingAssetsTitle = LOCTEXT("EnumeratingAssetsTitle", "Enumerating Assets (Step 1/3)...");
+const FText EnumeratingAssetsProgressFormat = LOCTEXT("EnumeratingAssetsProgressFormat", "{0} assets discovered...");
+const FText LoadingAssetsTitle = LOCTEXT("LoadingAssetsTitle", "Loading Assets (Step 2/3)...");
+const FText CompilingAssetsTitle = LOCTEXT("CompilingAssetsTitle", "Compiling Assets (Step 3/3)...");
+const FText GenericProgressFormat = LOCTEXT("GenericProgressFormat", "{0}/{1} ({2}%)...");
+
+const FRegexPattern EnumeratingAssetsPattern(TEXT("Display:\\s+([0-9]+)\\)\\s.+")); // ex. "146) /Engine/EditorBlueprintResources/ActorMacros"
+const FRegexPattern LoadingTotalPattern(TEXT("Display:\\s([0-9]+)\\spackages to load from command line arguments")); // ex. "2833 packages to load from command line arguments"
+const FRegexPattern LoadingProgressPattern(TEXT("Display:\\sLoading\\s\\(([0-9]+)\\)")); // ex. "Loading (2833)"
+const FRegexPattern CompileProgressPattern(TEXT("Display:\\sWaiting for\\s([0-9]+)\\s.+ to finish")); // ex. "Waiting for 14163 Shaders to finish"
+
+FDisplayClusterFillDerivedDataCacheWorker::FDisplayClusterFillDerivedDataCacheWorker()
+{	
+	FAsyncTaskNotificationConfig NotificationConfig;
+	NotificationConfig.bCanCancel = true;
+	NotificationConfig.TitleText = EnumeratingAssetsTitle;
+	NotificationConfig.bKeepOpenOnFailure = true;
+	NotificationConfig.bKeepOpenOnSuccess = true;
+	NotificationConfig.ExpireDuration = 1.0f;
+	NotificationConfig.ProgressText = FText::Format(EnumeratingAssetsProgressFormat, FText::AsNumber(0));
+	ProgressNotification = MakeUnique<FAsyncTaskNotification>(NotificationConfig);
+	
+	verify(FPlatformProcess::CreatePipe(ReadPipe, WritePipe));
+
+	CurrentExecutableName = FPlatformProcess::ExecutablePath();
+	const FString ProjectPath = FPaths::SetExtension(
+		FPaths::Combine(FPaths::ProjectDir(), FApp::GetProjectName()),".uproject");
+	Arguments = ProjectPath + " " + GetDdcCommandletParams();
+	
+	UE_LOG(LogDisplayClusterFillDerivedDataCache, Display, TEXT("Running commandlet: %s %s"), *CurrentExecutableName, *Arguments);
+
+	uint32 ProcessID;
+	const bool bLaunchDetached = true;
+	const bool bLaunchHidden = true;
+	const bool bLaunchReallyHidden = true;
+	ProcessHandle = FPlatformProcess::CreateProc(
+		*CurrentExecutableName, *Arguments, bLaunchDetached, bLaunchHidden, bLaunchReallyHidden, &ProcessID,
+		0, nullptr, WritePipe, ReadPipe);
+	
+	Thread = TUniquePtr<FRunnableThread>(FRunnableThread::Create(this, TEXT("FDisplayClusterFillDerivedDataCacheWorker")));
+}
+
+FDisplayClusterFillDerivedDataCacheWorker::~FDisplayClusterFillDerivedDataCacheWorker()
+{
+	if (Thread)
+	{
+		constexpr bool bShouldWait = true;
+		Thread->Kill( bShouldWait );
+		Thread = nullptr;
+	}
+}
+
+void FDisplayClusterFillDerivedDataCacheWorker::CancelTask()
+{
+	bWasCancelled = true;
+	FPlatformProcess::TerminateProc(ProcessHandle);
+	CompleteCommandletAndShowNotification();
+}
+
+void FDisplayClusterFillDerivedDataCacheWorker::CompleteCommandletAndShowNotification()
+{
+	if (bWasCancelled)
+	{
+		UE_LOG(LogDisplayClusterFillDerivedDataCache, Display, TEXT("#### Commandlet Process Cancelled ####"));
+		ProgressNotification->SetComplete(
+			LOCTEXT("ToastCommandletCancelled", "Commandlet Process Cancelled"), FText(), false);
+
+		// Close the notification automatically
+		ProgressNotification->SetKeepOpenOnFailure(false);
+		return;
+	}
+	
+	FPlatformProcess::GetProcReturnCode(ProcessHandle, &ResultCode);
+	const bool bFinishedWithFailures = ResultCode != 0;
+
+	FAsyncNotificationStateData NotificationStateData;
+	NotificationStateData.ProgressText = FText();
+	NotificationStateData.TitleText = LOCTEXT("ToastCommandletSuccess", "Commandlet Finished Without Error!");
+	NotificationStateData.State = EAsyncTaskNotificationState::Success;
+
+	if (bFinishedWithFailures)
+	{
+		UE_LOG(LogDisplayClusterFillDerivedDataCache, Error, TEXT("#### Commandlet Finished With Errors ####"));
+		UE_LOG(LogDisplayClusterFillDerivedDataCache, Error, TEXT("%s %s"), *CurrentExecutableName, *Arguments);
+		UE_LOG(LogDisplayClusterFillDerivedDataCache, Error, TEXT("Return Code: %i"), ResultCode);
+
+		NotificationStateData.TitleText = LOCTEXT("ToastCommandletFailure", "Commandlet Finished With Error(s)");
+		NotificationStateData.ProgressText = FText::Format(LOCTEXT("ReturnCodeFormat", "Return Code: {0}"), FText::AsNumber(ResultCode));
+		NotificationStateData.State = EAsyncTaskNotificationState::Failure;
+		NotificationStateData.Hyperlink = FSimpleDelegate::CreateLambda([]() { FGlobalTabmanager::Get()->TryInvokeTab(FName("OutputLog")); });
+		NotificationStateData.HyperlinkText = LOCTEXT("ShowOutputLogHyperlink", "Show Output Log");
+	}
+	else
+	{
+		UE_LOG(LogDisplayClusterFillDerivedDataCache, Display, TEXT("#### Commandlet Finished Without Error ####"));
+	}
+
+	ProgressNotification->SetNotificationState(NotificationStateData);
+}
+
+void FDisplayClusterFillDerivedDataCacheWorker::RegexParseForEnumerationCount(const FString& LogString)
+{
+	FRegexMatcher EnumerationCountRegex(EnumeratingAssetsPattern, *LogString);
+	while (EnumerationCountRegex.FindNext())
+	{
+		ProgressNotification->SetTitleText(EnumeratingAssetsTitle);
+		
+		if (const int32 EnumerationCount = FCString::Atoi(*EnumerationCountRegex.GetCaptureGroup(1)); EnumerationCount > 0)
+		{
+			ProgressNotification->SetProgressText(
+			   FText::Format(EnumeratingAssetsProgressFormat, FText::AsNumber(EnumerationCount))
+		   );
+		}
+	}
+}
+
+void FDisplayClusterFillDerivedDataCacheWorker::RegexParseForLoadingProgress(const FString& LogString)
+{
+	FRegexMatcher LoadingProgressRegex(LoadingProgressPattern, *LogString);
+	while (LoadingProgressRegex.FindNext())
+	{
+		ProgressNotification->SetTitleText(LoadingAssetsTitle);
+		
+		if (LoadingTotal > INDEX_NONE)
+		{
+			if (const int32 LoadingProgress = FCString::Atoi(*LoadingProgressRegex.GetCaptureGroup(1)); LoadingProgress > 0)
+			{
+				ProgressNotification->SetProgressText(
+				   FText::Format(GenericProgressFormat,
+					   FText::AsNumber(LoadingProgress),
+					   FText::AsNumber(LoadingTotal),
+					   FText::AsNumber(FMath::FloorToInt((float)LoadingProgress / (float)LoadingTotal * 100))
+				   )
+			   );
+			}
+		}
+	}
+}
+
+void FDisplayClusterFillDerivedDataCacheWorker::RegexParseForCompilationProgress(const FString& LogString)
+{
+	FRegexMatcher CompileProgressRegex(CompileProgressPattern, *LogString);
+	while (CompileProgressRegex.FindNext())
+	{
+		AssetsWaitingToCompile = FCString::Atoi(*CompileProgressRegex.GetCaptureGroup(1));
+		if (CompileTotal == INDEX_NONE && AssetsWaitingToCompile > 0)
+		{
+			CompileTotal = AssetsWaitingToCompile;
+		}
+
+		if (CompileTotal > 0)
+		{
+			ProgressNotification->SetTitleText(CompilingAssetsTitle);
+
+			if (const int32 Progress = CompileTotal - AssetsWaitingToCompile; Progress > -1)
+			{
+				ProgressNotification->SetProgressText(
+					FText::Format(
+						GenericProgressFormat,
+						FText::AsNumber(Progress),
+						FText::AsNumber(CompileTotal),
+						FText::AsNumber(FMath::FloorToInt((float)Progress / (float)CompileTotal * 100))
+					)
+				);
+			}
+		}
+	}
+}
+
+FString FDisplayClusterFillDerivedDataCacheWorker::GetDdcCommandletParams() const
+{
+	return FString::Printf(TEXT("-run=DerivedDataCache %s -fill -DDC=CreateInstalledEnginePak"), *GetTargetPlatformParams());
+}
+
+FString FDisplayClusterFillDerivedDataCacheWorker::GetTargetPlatformParams() const
+{
+	TArray<FString> SupportedPlatforms;
+	FProjectStatus ProjectStatus;
+	
+	if(IProjectManager::Get().QueryStatusForCurrentProject(ProjectStatus))
+	{
+		for (const FDataDrivenPlatformInfo* DDPI : FDataDrivenPlatformInfoRegistry::GetSortedPlatformInfos(EPlatformInfoType::TruePlatformsOnly))
+		{
+			if (ProjectStatus.IsTargetPlatformSupported(DDPI->IniPlatformName))
+			{
+				const FString PlatformName = DDPI->IniPlatformName.ToString();
+				SupportedPlatforms.Add(PlatformName);
+
+				// Not all platforms have an editor target but they may in the future.
+				// It doesn't hurt to add targets that don't exist as they are ignored.
+				SupportedPlatforms.Add(PlatformName + "Editor");
+			}
+		}
+	}
+
+	if (SupportedPlatforms.Num() > 0)
+	{
+		return FString::Printf(TEXT("-TargetPlatform=%s"), *FString::Join(SupportedPlatforms, TEXT("+")));
+	}
+	
+	return FString();
+}
+
+FString FDisplayClusterFillDerivedDataCacheWorker::GetLastWholeLogBlock(const FString& LogString)
+{
+	FString CumulativeString = LastPartialLogLine + LogString;
+
+	// Look for a line break. If it exists, we can consider everything to the right of it to be a partial line while the left of it is a whole line
+	const int32 IndexOfLastNewLine = CumulativeString.Find(TEXT("\n"), ESearchCase::IgnoreCase, ESearchDir::FromEnd);
+
+	if (IndexOfLastNewLine > INDEX_NONE)
+	{
+		LastPartialLogLine = CumulativeString.RightChop(IndexOfLastNewLine).TrimStartAndEnd();
+		
+		return CumulativeString.Left(IndexOfLastNewLine).TrimStartAndEnd();
+	}
+
+	// If we don't find a break in the line, we should consider the entire line partial
+	LastPartialLogLine = CumulativeString;
+
+	// Then return nothing so we know we have nothing to do this iteration
+	return FString();
+}
+
+void FDisplayClusterFillDerivedDataCacheWorker::ReadCommandletOutputAndUpdateEditorNotification()
+{
+	while (FPlatformProcess::IsProcRunning(ProcessHandle))
+	{
+		if (ProgressNotification->GetPromptAction() == EAsyncTaskNotificationPromptAction::Cancel)
+		{
+			UE_LOG(LogDisplayClusterFillDerivedDataCache, Display, TEXT("----------CANCEL BUTTON PRESSED----------"));
+			CancelTask();
+			break;
+		}
+
+		const FString LogString = FPlatformProcess::ReadPipe(ReadPipe);
+		if (LogString.TrimStartAndEnd().IsEmpty())
+		{
+			continue;
+		}
+		
+		const FString LastWholeLogBlock = GetLastWholeLogBlock(LogString);
+		if (LastWholeLogBlock.TrimStartAndEnd().IsEmpty())
+		{
+			continue;
+		}
+
+		TArray<FString> EachLineInLogBlock;
+		LastWholeLogBlock.ParseIntoArray(EachLineInLogBlock, TEXT("\n"), true);
+
+		for (const FString& Line : EachLineInLogBlock)
+		{
+			if (Line.TrimStartAndEnd().IsEmpty())
+			{
+				continue;
+			}
+			
+			UE_LOG(LogDisplayClusterFillDerivedDataCache, Display, TEXT("Commandlet Output: %s"), *Line)
+
+			if (LoadingTotal == INDEX_NONE)
+			{			
+				// Try to find total number of assets to load
+				FRegexMatcher LoadingTotalRegex(LoadingTotalPattern, *Line);
+
+				while (LoadingTotalRegex.FindNext())
+				{
+					LoadingTotal = FCString::Atoi(*LoadingTotalRegex.GetCaptureGroup(1));
+				}
+			}
+		
+			// Progress on asset enumeration
+			RegexParseForEnumerationCount(Line);
+		
+			// Progress on asset loading
+			RegexParseForLoadingProgress(Line);
+		
+			// Progress on asset compilation
+			RegexParseForCompilationProgress(Line);
+		}
+
+		FPlatformProcess::Sleep(0.5f);
+	}
+	
+	UE_LOG(LogDisplayClusterFillDerivedDataCache, Display, TEXT("----------PROC NOT RUNNING----------"));
+	CompleteCommandletAndShowNotification();
+}
+
+#undef LOCTEXT_NAMESPACE
diff --git a/Engine/Plugins/Runtime/nDisplay/Source/DisplayClusterFillDerivedDataCache/Public/DisplayClusterFillDerivedDataCacheLog.h b/Engine/Plugins/Runtime/nDisplay/Source/DisplayClusterFillDerivedDataCache/Public/DisplayClusterFillDerivedDataCacheLog.h
new file mode 100644
--- /dev/null
+++ b/Engine/Plugins/Runtime/nDisplay/Source/DisplayClusterFillDerivedDataCache/Public/DisplayClusterFillDerivedDataCacheLog.h
@@ -0,0 +1,7 @@
+// Copyright Epic Games, Inc. All Rights Reserved.
+
+#pragma once
+
+#include "Logging/LogMacros.h"
+
+DECLARE_LOG_CATEGORY_CLASS(LogDisplayClusterFillDerivedDataCache, Log, All);
diff --git a/Engine/Plugins/Runtime/nDisplay/Source/DisplayClusterFillDerivedDataCache/Public/DisplayClusterFillDerivedDataCacheModule.h b/Engine/Plugins/Runtime/nDisplay/Source/DisplayClusterFillDerivedDataCache/Public/DisplayClusterFillDerivedDataCacheModule.h
new file mode 100644
--- /dev/null
+++ b/Engine/Plugins/Runtime/nDisplay/Source/DisplayClusterFillDerivedDataCache/Public/DisplayClusterFillDerivedDataCacheModule.h
@@ -0,0 +1,27 @@
+// Copyright Epic Games, Inc. All Rights Reserved.
+
+#pragma once
+
+#include "Modules/ModuleInterface.h"
+#include "Templates/UniquePtr.h"
+
+class FDisplayClusterFillDerivedDataCacheModule : public IModuleInterface
+{
+public:
+	
+	static FDisplayClusterFillDerivedDataCacheModule& Get();
+
+	//~ Begin IModuleInterface Interface
+	virtual void StartupModule() override;
+	virtual void ShutdownModule() override;
+	//~ End IModuleInterface Interface
+
+protected:
+	
+	void CreateAsyncTaskWorker();
+	TUniquePtr<class FDisplayClusterFillDerivedDataCacheWorker> AsyncTaskWorker;
+
+private:
+	
+	void OnFEngineLoopInitComplete();
+};
diff --git a/Engine/Plugins/Runtime/nDisplay/Source/DisplayClusterFillDerivedDataCache/Public/DisplayClusterFillDerivedDataCacheWorker.h b/Engine/Plugins/Runtime/nDisplay/Source/DisplayClusterFillDerivedDataCache/Public/DisplayClusterFillDerivedDataCacheWorker.h
new file mode 100644
--- /dev/null
+++ b/Engine/Plugins/Runtime/nDisplay/Source/DisplayClusterFillDerivedDataCache/Public/DisplayClusterFillDerivedDataCacheWorker.h
@@ -0,0 +1,85 @@
+// Copyright Epic Games, Inc. All Rights Reserved.
+
+#pragma once
+
+#include "Containers/UnrealString.h"
+#include "HAL/PlatformProcess.h"
+#include "HAL/Runnable.h"
+#include "HAL/RunnableThread.h"
+#include "Misc/SingleThreadRunnable.h"
+#include "Templates/UniquePtr.h"
+
+class FAsyncTaskNotification;
+
+struct FProcHandle;
+
+class FDisplayClusterFillDerivedDataCacheWorker : public FRunnable, public FSingleThreadRunnable
+{
+public:
+
+	FDisplayClusterFillDerivedDataCacheWorker();
+	
+	virtual ~FDisplayClusterFillDerivedDataCacheWorker() override;
+
+	//~ Begin FRunnable implementation
+	virtual uint32 Run() override
+	{
+		ReadCommandletOutputAndUpdateEditorNotification();
+		return 0;
+	}
+	virtual void Stop() override
+	{
+		CancelTask();
+	}
+	//~ End FRunnable implementation
+	
+	//~ Begin FSingleThreadRunnable implementation
+	virtual void Tick() override
+	{
+		ReadCommandletOutputAndUpdateEditorNotification();
+	}
+	virtual FSingleThreadRunnable* GetSingleThreadInterface() override
+	{
+		return this;
+	}
+	//~ End FSingleThreadRunnable implementation
+
+	void CancelTask();
+	void CompleteCommandletAndShowNotification();
+
+protected:
+	
+	void RegexParseForEnumerationCount(const FString& LogString);
+	void RegexParseForLoadingProgress(const FString& LogString);
+	void RegexParseForCompilationProgress(const FString& LogString);
+
+	FString GetDdcCommandletParams() const;
+	FString GetTargetPlatformParams() const;
+
+	FString GetLastWholeLogBlock(const FString& LogString);
+
+	void ReadCommandletOutputAndUpdateEditorNotification();
+	
+	/** Progress notification */
+	TUniquePtr<FAsyncTaskNotification> ProgressNotification;
+
+private:
+
+	TUniquePtr<FRunnableThread> Thread;
+	
+	int32 ResultCode = 0;
+	bool bWasCancelled = false;
+	FProcHandle ProcessHandle;
+
+	void* ReadPipe = nullptr;
+	void* WritePipe = nullptr;
+
+	FString LastPartialLogLine;
+
+	int32 LoadingTotal = INDEX_NONE;
+	int32 CompileTotal = INDEX_NONE;
+	int32 AssetsWaitingToCompile = INDEX_NONE;
+
+	FString CurrentExecutableName;
+	FString Arguments;
+};
diff --git a/Engine/Plugins/Runtime/nDisplay/Source/DisplayClusterPreloadDerivedDataCache/DisplayClusterPreloadDerivedDataCache.build.cs b/Engine/Plugins/Runtime/nDisplay/Source/DisplayClusterPreloadDerivedDataCache/DisplayClusterPreloadDerivedDataCache.build.cs
deleted file mode 100644
diff --git a/Engine/Plugins/Runtime/nDisplay/Source/DisplayClusterPreloadDerivedDataCache/Private/DisplayClusterPreloadDerivedDataCacheModule.cpp b/Engine/Plugins/Runtime/nDisplay/Source/DisplayClusterPreloadDerivedDataCache/Private/DisplayClusterPreloadDerivedDataCacheModule.cpp
deleted file mode 100644
diff --git a/Engine/Plugins/Runtime/nDisplay/Source/DisplayClusterPreloadDerivedDataCache/Private/DisplayClusterPreloadDerivedDataCacheWorker.cpp b/Engine/Plugins/Runtime/nDisplay/Source/DisplayClusterPreloadDerivedDataCache/Private/DisplayClusterPreloadDerivedDataCacheWorker.cpp
deleted file mode 100644
diff --git a/Engine/Plugins/Runtime/nDisplay/Source/DisplayClusterPreloadDerivedDataCache/Public/DisplayClusterPreloadDerivedDataCacheLog.h b/Engine/Plugins/Runtime/nDisplay/Source/DisplayClusterPreloadDerivedDataCache/Public/DisplayClusterPreloadDerivedDataCacheLog.h
deleted file mode 100644
diff --git a/Engine/Plugins/Runtime/nDisplay/Source/DisplayClusterPreloadDerivedDataCache/Public/DisplayClusterPreloadDerivedDataCacheModule.h b/Engine/Plugins/Runtime/nDisplay/Source/DisplayClusterPreloadDerivedDataCache/Public/DisplayClusterPreloadDerivedDataCacheModule.h
deleted file mode 100644
diff --git a/Engine/Plugins/Runtime/nDisplay/Source/DisplayClusterPreloadDerivedDataCache/Public/DisplayClusterPreloadDerivedDataCacheWorker.h b/Engine/Plugins/Runtime/nDisplay/Source/DisplayClusterPreloadDerivedDataCache/Public/DisplayClusterPreloadDerivedDataCacheWorker.h
deleted file mode 100644
diff --git a/Engine/Plugins/Runtime/nDisplay/nDisplay.uplugin b/Engine/Plugins/Runtime/nDisplay/nDisplay.uplugin
--- a/Engine/Plugins/Runtime/nDisplay/nDisplay.uplugin
+++ b/Engine/Plugins/Runtime/nDisplay/nDisplay.uplugin
@@ -29,7 +29,7 @@
 			]
 		},
 		{
-			"Name": "DisplayClusterPreloadDerivedDataCache",
+			"Name": "DisplayClusterFillDerivedDataCache",
 			"Type": "Editor",
 			"LoadingPhase": "Default",
 			"PlatformAllowList": [
