nDisplay:
- Added postprocess support for the projection policy interface
- A public function has been added to the viewport interface to get the viewpoint component.
- Renamed viewports containers for DC Manager and ProxyManager
- Corrected a few typos
- Added missing DISPLAYCLUSTER_API macro to public classes

(this subtask is part of the In-Frustum Fit for MPCDI/2D/3D/mesh)

diff --git a/Engine/Plugins/Runtime/nDisplay/Source/DisplayCluster/Private/Render/Viewport/Configuration/DisplayClusterViewportConfiguration.cpp b/Engine/Plugins/Runtime/nDisplay/Source/DisplayCluster/Private/Render/Viewport/Configuration/DisplayClusterViewportConfiguration.cpp
--- a/Engine/Plugins/Runtime/nDisplay/Source/DisplayCluster/Private/Render/Viewport/Configuration/DisplayClusterViewportConfiguration.cpp
+++ b/Engine/Plugins/Runtime/nDisplay/Source/DisplayCluster/Private/Render/Viewport/Configuration/DisplayClusterViewportConfiguration.cpp
@@ -339,7 +339,7 @@ void FDisplayClusterViewportConfiguration::ImplUpdateConfigurationVisibility(ADi
 	TSet<FPrimitiveComponentId> RootActorHidePrimitivesList;
 	if (ViewportManager && RootActor.GetHiddenInGamePrimitives(RootActorHidePrimitivesList))
 	{
-		for (const TSharedPtr<FDisplayClusterViewport, ESPMode::ThreadSafe>& ViewportIt : ViewportManager->ImplGetViewports())
+		for (const TSharedPtr<FDisplayClusterViewport, ESPMode::ThreadSafe>& ViewportIt : ViewportManager->ImplGetCurrentRenderFrameViewports())
 		{
 			if (ViewportIt.IsValid())
 			{
diff --git a/Engine/Plugins/Runtime/nDisplay/Source/DisplayCluster/Private/Render/Viewport/Configuration/DisplayClusterViewportConfigurationBase.cpp b/Engine/Plugins/Runtime/nDisplay/Source/DisplayCluster/Private/Render/Viewport/Configuration/DisplayClusterViewportConfigurationBase.cpp
--- a/Engine/Plugins/Runtime/nDisplay/Source/DisplayCluster/Private/Render/Viewport/Configuration/DisplayClusterViewportConfigurationBase.cpp
+++ b/Engine/Plugins/Runtime/nDisplay/Source/DisplayCluster/Private/Render/Viewport/Configuration/DisplayClusterViewportConfigurationBase.cpp
@@ -60,7 +60,7 @@ void FDisplayClusterViewportConfigurationBase::Update(const FString& ClusterNode
 		ConfigurationData.Cluster->GetNodeIds(ExistClusterNodesIDs);
 
 		TArray<TSharedPtr<FDisplayClusterViewport, ESPMode::ThreadSafe>> UnusedViewports;
-		for (const TSharedPtr<FDisplayClusterViewport, ESPMode::ThreadSafe>& ViewportIt : ViewportManager.ImplGetViewports())
+		for (const TSharedPtr<FDisplayClusterViewport, ESPMode::ThreadSafe>& ViewportIt : ViewportManager.ImplGetCurrentRenderFrameViewports())
 		{
 			if (ViewportIt.IsValid())
 			{
@@ -120,7 +120,7 @@ void FDisplayClusterViewportConfigurationBase::Update(const TArray<FString>& InV
 	// Collect unused viewports and delete
 	{
 		TArray<TSharedPtr<FDisplayClusterViewport, ESPMode::ThreadSafe>> UnusedViewports;
-		for (const TSharedPtr<FDisplayClusterViewport, ESPMode::ThreadSafe>& ViewportIt : ViewportManager.ImplGetViewports())
+		for (const TSharedPtr<FDisplayClusterViewport, ESPMode::ThreadSafe>& ViewportIt : ViewportManager.ImplGetCurrentRenderFrameViewports())
 		{
 			if (ViewportIt.IsValid())
 			{
diff --git a/Engine/Plugins/Runtime/nDisplay/Source/DisplayCluster/Private/Render/Viewport/Configuration/DisplayClusterViewportConfigurationHelpers.cpp b/Engine/Plugins/Runtime/nDisplay/Source/DisplayCluster/Private/Render/Viewport/Configuration/DisplayClusterViewportConfigurationHelpers.cpp
--- a/Engine/Plugins/Runtime/nDisplay/Source/DisplayCluster/Private/Render/Viewport/Configuration/DisplayClusterViewportConfigurationHelpers.cpp
+++ b/Engine/Plugins/Runtime/nDisplay/Source/DisplayCluster/Private/Render/Viewport/Configuration/DisplayClusterViewportConfigurationHelpers.cpp
@@ -284,4 +284,10 @@ void FDisplayClusterViewportConfigurationHelpers::UpdateProjectionPolicy(FDispla
 			}
 		}
 	}
+
+	// Override PostProcess from projection policy
+	if (DstViewport.ProjectionPolicy.IsValid())
+	{
+		DstViewport.ProjectionPolicy->OverridePostProcessSettings(&DstViewport);
+	}
 }
diff --git a/Engine/Plugins/Runtime/nDisplay/Source/DisplayCluster/Private/Render/Viewport/Configuration/DisplayClusterViewportConfigurationHelpers_ICVFX.cpp b/Engine/Plugins/Runtime/nDisplay/Source/DisplayCluster/Private/Render/Viewport/Configuration/DisplayClusterViewportConfigurationHelpers_ICVFX.cpp
--- a/Engine/Plugins/Runtime/nDisplay/Source/DisplayCluster/Private/Render/Viewport/Configuration/DisplayClusterViewportConfigurationHelpers_ICVFX.cpp
+++ b/Engine/Plugins/Runtime/nDisplay/Source/DisplayCluster/Private/Render/Viewport/Configuration/DisplayClusterViewportConfigurationHelpers_ICVFX.cpp
@@ -219,7 +219,7 @@ void FDisplayClusterViewportConfigurationHelpers_ICVFX::ReuseUVLightCardViewport
 {
 	if (FDisplayClusterViewportManager* ViewportManager = InUVLightCardViewport.GetViewportManagerImpl())
 	{
-		for (const TSharedPtr<FDisplayClusterViewport, ESPMode::ThreadSafe>& ViewportIt : ViewportManager->ImplGetWholeClusterViewports())
+		for (const TSharedPtr<FDisplayClusterViewport, ESPMode::ThreadSafe>& ViewportIt : ViewportManager->ImplGetEntireClusterViewports())
 		{
 			if (ViewportIt.IsValid() && ViewportIt != InUVLightCardViewport.AsShared() && !ViewportIt->RenderSettings.IsViewportOverrided()
 				&& EnumHasAnyFlags(ViewportIt->RenderSettingsICVFX.RuntimeFlags, EDisplayClusterViewportRuntimeICVFXFlags::UVLightcard))
@@ -247,7 +247,7 @@ TArray<TSharedPtr<FDisplayClusterViewport, ESPMode::ThreadSafe>> FDisplayCluster
 		const EDisplayClusterViewportRuntimeICVFXFlags InRuntimeFlagsMask = bGetChromakey ? EDisplayClusterViewportRuntimeICVFXFlags::Chromakey : EDisplayClusterViewportRuntimeICVFXFlags::InCamera;
 		const FString ViewportTypeId = bGetChromakey ? DisplayClusterViewportStrings::icvfx::chromakey : DisplayClusterViewportStrings::icvfx::camera;
 
-		for (const TSharedPtr<FDisplayClusterViewport, ESPMode::ThreadSafe>& ViewportIt : ViewportManager->ImplGetWholeClusterViewports())
+		for (const TSharedPtr<FDisplayClusterViewport, ESPMode::ThreadSafe>& ViewportIt : ViewportManager->ImplGetEntireClusterViewports())
 		{
 			if (ViewportIt.IsValid()
 				&& EnumHasAnyFlags(ViewportIt->RenderSettingsICVFX.RuntimeFlags, InRuntimeFlagsMask)
diff --git a/Engine/Plugins/Runtime/nDisplay/Source/DisplayCluster/Private/Render/Viewport/Configuration/DisplayClusterViewportConfigurationICVFX.cpp b/Engine/Plugins/Runtime/nDisplay/Source/DisplayCluster/Private/Render/Viewport/Configuration/DisplayClusterViewportConfigurationICVFX.cpp
--- a/Engine/Plugins/Runtime/nDisplay/Source/DisplayCluster/Private/Render/Viewport/Configuration/DisplayClusterViewportConfigurationICVFX.cpp
+++ b/Engine/Plugins/Runtime/nDisplay/Source/DisplayCluster/Private/Render/Viewport/Configuration/DisplayClusterViewportConfigurationICVFX.cpp
@@ -416,7 +416,7 @@ void FDisplayClusterViewportConfigurationICVFX::Update()
 				// Freeze render for lightcards when outer viewports freezed
 				if (StageSettings.Lightcard.bIgnoreOuterViewportsFreezingForLightcards == false)
 				{
-					for (const TSharedPtr<FDisplayClusterViewport, ESPMode::ThreadSafe>& ViewportIt : ViewportManager->ImplGetViewports())
+					for (const TSharedPtr<FDisplayClusterViewport, ESPMode::ThreadSafe>& ViewportIt : ViewportManager->ImplGetCurrentRenderFrameViewports())
 					{
 						if (ViewportIt.IsValid() && EnumHasAnyFlags(ViewportIt->RenderSettingsICVFX.RuntimeFlags, EDisplayClusterViewportRuntimeICVFXFlags::Lightcard | EDisplayClusterViewportRuntimeICVFXFlags::UVLightcard))
 						{
@@ -464,7 +464,7 @@ void FDisplayClusterViewportConfigurationICVFX::PostUpdate()
 		// Support alpha channel rendering for WarpBlend
 		if (GDisplayClusterEnableAlphaChannelRendering != 0)
 		{
-			for (const TSharedPtr<FDisplayClusterViewport, ESPMode::ThreadSafe>& ViewportIt : ViewportManager->ImplGetViewports())
+			for (const TSharedPtr<FDisplayClusterViewport, ESPMode::ThreadSafe>& ViewportIt : ViewportManager->ImplGetCurrentRenderFrameViewports())
 			{
 				if (ViewportIt.IsValid())
 				{
@@ -493,7 +493,7 @@ void FDisplayClusterViewportConfigurationICVFX::UpdateHideList(FDisplayClusterVi
 	TArray<TSharedPtr<FDisplayClusterViewport, ESPMode::ThreadSafe>> ICVFXViewports;
 
 	// Collect viewports, that use ICVFX hide list
-	for (const TSharedPtr<FDisplayClusterViewport, ESPMode::ThreadSafe>& ViewportIt : ViewportManager.ImplGetViewports())
+	for (const TSharedPtr<FDisplayClusterViewport, ESPMode::ThreadSafe>& ViewportIt : ViewportManager.ImplGetCurrentRenderFrameViewports())
 	{
 		if (ViewportIt.IsValid())
 		{
@@ -514,7 +514,7 @@ void FDisplayClusterViewportConfigurationICVFX::UpdateHideList(FDisplayClusterVi
 
 void FDisplayClusterViewportConfigurationICVFX::ImplBeginReallocateViewports(FDisplayClusterViewportManager& ViewportManager)
 {
-	for (const TSharedPtr<FDisplayClusterViewport, ESPMode::ThreadSafe>& ViewportIt : ViewportManager.ImplGetViewports())
+	for (const TSharedPtr<FDisplayClusterViewport, ESPMode::ThreadSafe>& ViewportIt : ViewportManager.ImplGetCurrentRenderFrameViewports())
 	{
 		// Runtime icvfx viewport support reallocate feature:
 		if (ViewportIt.IsValid() && EnumHasAllFlags(ViewportIt->RenderSettingsICVFX.RuntimeFlags, EDisplayClusterViewportRuntimeICVFXFlags::InternalResource))
@@ -530,7 +530,7 @@ void FDisplayClusterViewportConfigurationICVFX::ImplFinishReallocateViewports(FD
 	TArray<TSharedPtr<FDisplayClusterViewport, ESPMode::ThreadSafe>> UnusedViewports;
 
 	// Collect all unused viewports for remove
-	for (const TSharedPtr<FDisplayClusterViewport, ESPMode::ThreadSafe>& ViewportIt : ViewportManager.ImplGetViewports())
+	for (const TSharedPtr<FDisplayClusterViewport, ESPMode::ThreadSafe>& ViewportIt : ViewportManager.ImplGetCurrentRenderFrameViewports())
 	{
 		if (ViewportIt.IsValid() && EnumHasAllFlags(ViewportIt->RenderSettingsICVFX.RuntimeFlags, EDisplayClusterViewportRuntimeICVFXFlags::Unused))
 		{
@@ -579,7 +579,7 @@ EDisplayClusterViewportICVFXFlags FDisplayClusterViewportConfigurationICVFX::Imp
 	EDisplayClusterViewportICVFXFlags InvFlags = EDisplayClusterViewportICVFXFlags::None;
 
 	// Collect invertet disable flags from all target viewports
-	for (const TSharedPtr<FDisplayClusterViewport, ESPMode::ThreadSafe>& ViewportIt : ViewportManager.ImplGetViewports())
+	for (const TSharedPtr<FDisplayClusterViewport, ESPMode::ThreadSafe>& ViewportIt : ViewportManager.ImplGetCurrentRenderFrameViewports())
 	{
 		// Process only external viewports:
 		if (ViewportIt.IsValid() && !EnumHasAnyFlags(ViewportIt->RenderSettingsICVFX.RuntimeFlags, EDisplayClusterViewportRuntimeICVFXFlags::InternalResource))
diff --git a/Engine/Plugins/Runtime/nDisplay/Source/DisplayCluster/Private/Render/Viewport/Configuration/DisplayClusterViewportConfigurationProjectionPolicy.cpp b/Engine/Plugins/Runtime/nDisplay/Source/DisplayCluster/Private/Render/Viewport/Configuration/DisplayClusterViewportConfigurationProjectionPolicy.cpp
--- a/Engine/Plugins/Runtime/nDisplay/Source/DisplayCluster/Private/Render/Viewport/Configuration/DisplayClusterViewportConfigurationProjectionPolicy.cpp
+++ b/Engine/Plugins/Runtime/nDisplay/Source/DisplayCluster/Private/Render/Viewport/Configuration/DisplayClusterViewportConfigurationProjectionPolicy.cpp
@@ -28,7 +28,7 @@
 
 void FDisplayClusterViewportConfigurationProjectionPolicy::Update()
 {
-	for (const TSharedPtr<FDisplayClusterViewport, ESPMode::ThreadSafe>& ViewportIt : ViewportManager.ImplGetViewports())
+	for (const TSharedPtr<FDisplayClusterViewport, ESPMode::ThreadSafe>& ViewportIt : ViewportManager.ImplGetCurrentRenderFrameViewports())
 	{
 		// ignore ICVFX internal resources
 		if (ViewportIt.IsValid() && !EnumHasAnyFlags(ViewportIt->GetRenderSettingsICVFX().RuntimeFlags, EDisplayClusterViewportRuntimeICVFXFlags::InternalResource))
diff --git a/Engine/Plugins/Runtime/nDisplay/Source/DisplayCluster/Private/Render/Viewport/DisplayClusterViewport.h b/Engine/Plugins/Runtime/nDisplay/Source/DisplayCluster/Private/Render/Viewport/DisplayClusterViewport.h
--- a/Engine/Plugins/Runtime/nDisplay/Source/DisplayCluster/Private/Render/Viewport/DisplayClusterViewport.h
+++ b/Engine/Plugins/Runtime/nDisplay/Source/DisplayCluster/Private/Render/Viewport/DisplayClusterViewport.h
@@ -106,6 +106,7 @@ public:
 
 	virtual bool SetupViewPoint(FMinimalViewInfo& InOutViewInfo) override;
 	virtual float GetStereoEyeOffsetDistance(const uint32 InContextNum) override;
+	virtual class UDisplayClusterCameraComponent* GetViewPointCameraComponent() const override;
 
 	virtual const FDisplayClusterViewport_RenderSettingsICVFX& GetRenderSettingsICVFX() const override
 	{
@@ -137,6 +138,12 @@ public:
 		return CustomPostProcessSettings;
 	}
 
+	virtual IDisplayClusterViewport_CustomPostProcessSettings& GetViewport_CustomPostProcessSettings() override
+	{
+		check(IsInGameThread());
+		return CustomPostProcessSettings;
+	}
+
 	// Setup scene view for rendering specified Context
 	virtual void SetupSceneView(uint32 ContextNum, class UWorld* World, class FSceneViewFamily& InViewFamily, FSceneView& InView) const override;
 
@@ -199,8 +206,6 @@ public:
 	bool HandleStartScene();
 	void HandleEndScene();
 
-	class UDisplayClusterCameraComponent* GetViewPointCamera() const;
-
 	void AddReferencedObjects(FReferenceCollector& Collector);
 
 	void ResetRuntimeParameters()
diff --git a/Engine/Plugins/Runtime/nDisplay/Source/DisplayCluster/Private/Render/Viewport/DisplayClusterViewportManager.cpp b/Engine/Plugins/Runtime/nDisplay/Source/DisplayCluster/Private/Render/Viewport/DisplayClusterViewportManager.cpp
--- a/Engine/Plugins/Runtime/nDisplay/Source/DisplayCluster/Private/Render/Viewport/DisplayClusterViewportManager.cpp
+++ b/Engine/Plugins/Runtime/nDisplay/Source/DisplayCluster/Private/Render/Viewport/DisplayClusterViewportManager.cpp
@@ -82,8 +82,8 @@ FDisplayClusterViewportManager::~FDisplayClusterViewportManager()
 	UnregisterCallbacks();
 	
 	// Remove viewports
-	Viewports.Reset();
-	ClusterNodeViewports.Reset();
+	EntireClusterViewports.Reset();
+	CurrentRenderFrameViewports.Reset();
 
 	RenderTargetManager.Reset();
 	PostProcessManager.Reset();
@@ -152,7 +152,7 @@ void FDisplayClusterViewportManager::StartScene(UWorld* InWorld)
 
 	CurrentWorldRef = TWeakObjectPtr<UWorld>(InWorld);
 
-	for (TSharedPtr<FDisplayClusterViewport, ESPMode::ThreadSafe>& ViewportIt : Viewports)
+	for (const TSharedPtr<FDisplayClusterViewport, ESPMode::ThreadSafe>& ViewportIt : ImplGetEntireClusterViewports())
 	{
 		if (ViewportIt.IsValid())
 		{
@@ -175,7 +175,7 @@ void FDisplayClusterViewportManager::EndScene()
 {
 	check(IsInGameThread());
 
-	for (TSharedPtr<FDisplayClusterViewport, ESPMode::ThreadSafe>& ViewportIt : Viewports)
+	for (const TSharedPtr<FDisplayClusterViewport, ESPMode::ThreadSafe>& ViewportIt : ImplGetEntireClusterViewports())
 	{
 		if (ViewportIt.IsValid())
 		{
@@ -206,7 +206,7 @@ void FDisplayClusterViewportManager::ResetScene()
 		LightCardManager->HandleStartScene();
 	}
 
-	for (TSharedPtr<FDisplayClusterViewport, ESPMode::ThreadSafe>& ViewportIt : Viewports)
+	for (const TSharedPtr<FDisplayClusterViewport, ESPMode::ThreadSafe>& ViewportIt : ImplGetEntireClusterViewports())
 	{
 		if (ViewportIt.IsValid())
 		{
@@ -290,7 +290,7 @@ bool FDisplayClusterViewportManager::ShouldUseFullSizeFrameTargetableResource()
 		return true;
 	}
 
-	for (const TSharedPtr<FDisplayClusterViewport, ESPMode::ThreadSafe>& Viewport : ClusterNodeViewports)
+	for (const TSharedPtr<FDisplayClusterViewport, ESPMode::ThreadSafe>& Viewport : ImplGetCurrentRenderFrameViewports())
 	{
 		if (Viewport.IsValid() && Viewport->ShouldUseFullSizeFrameTargetableResource())
 		{
@@ -310,7 +310,7 @@ bool FDisplayClusterViewportManager::ShouldUseAdditionalFrameTargetableResource(
 		return true;
 	}
 
-	for (const TSharedPtr<FDisplayClusterViewport, ESPMode::ThreadSafe>& Viewport : ClusterNodeViewports)
+	for (const TSharedPtr<FDisplayClusterViewport, ESPMode::ThreadSafe>& Viewport : ImplGetCurrentRenderFrameViewports())
 	{
 		if (Viewport.IsValid() && Viewport->ShouldUseAdditionalFrameTargetableResource())
 		{
@@ -397,43 +397,58 @@ EDisplayClusterRenderFrameMode FDisplayClusterViewportManager::GetRenderMode() c
 
 void FDisplayClusterViewportManager::ImplUpdateClusterNodeViewports(const EDisplayClusterRenderFrameMode InRenderMode, const FString& InClusterNodeId)
 {
-	ClusterNodeViewports.Empty();
+	// Get viewports used to rendering of the current frame
+	if (InClusterNodeId.IsEmpty())
+	{
+		// When a cluster node name is empty, we render without the cluster nodes
+		CurrentRenderFrameViewports = ImplGetEntireClusterViewports();
+	}
+	else
+	{
+		// Get viewport for current cluster nodes
+		CurrentRenderFrameViewports.Reset();
+		for (const TSharedPtr<FDisplayClusterViewport, ESPMode::ThreadSafe>& Viewport : ImplGetEntireClusterViewports())
+		{
+			if (Viewport.IsValid() && (Viewport->GetClusterNodeId() == InClusterNodeId))
+			{
+				CurrentRenderFrameViewports.Add(Viewport);
+			}
+		}
+	}
 
-	switch (InRenderMode)
+	// Per-viewport preview render required special render order for viewports
+	// Sort viewports: move linked and overridden viewports to last
+	if (InRenderMode == EDisplayClusterRenderFrameMode::PreviewInScene)
 	{
-	case EDisplayClusterRenderFrameMode::PreviewInScene:
+		TArray<TSharedPtr<FDisplayClusterViewport, ESPMode::ThreadSafe>> LinkedViewports;
+		TArray<TSharedPtr<FDisplayClusterViewport, ESPMode::ThreadSafe>> OverriddenViewports;
+
+		TArray<TSharedPtr<FDisplayClusterViewport, ESPMode::ThreadSafe>> UsedViewports(CurrentRenderFrameViewports);
+		CurrentRenderFrameViewports.Reset();
+
+		for (const TSharedPtr<FDisplayClusterViewport, ESPMode::ThreadSafe>& Viewport : UsedViewports)
 		{
-			// Per-viewport preview render required special render order for viewports
-			// Sort viewports: move linked and overridden viewports to last
-			TArray<TSharedPtr<FDisplayClusterViewport, ESPMode::ThreadSafe>> LinkedViewports;
-			TArray<TSharedPtr<FDisplayClusterViewport, ESPMode::ThreadSafe>> OverriddenViewports;
-			for (TSharedPtr<FDisplayClusterViewport, ESPMode::ThreadSafe>& Viewport : Viewports)
+			if (!Viewport.IsValid())
 			{
-				if (Viewport.IsValid() && (InClusterNodeId.IsEmpty() || Viewport->GetClusterNodeId() == InClusterNodeId))
-				{
-					if (Viewport->RenderSettings.IsViewportOverrided())
-					{
-						OverriddenViewports.Add(Viewport);
-					}
-					else if (Viewport->RenderSettings.IsViewportHasParent())
-					{
-						LinkedViewports.Add(Viewport);
-					}
-					else
-					{
-						ClusterNodeViewports.Add(Viewport);
-					}
-				}
+				continue;
 			}
 
-			ClusterNodeViewports.Append(LinkedViewports);
-			ClusterNodeViewports.Append(OverriddenViewports);
+			if (Viewport->RenderSettings.IsViewportOverrided())
+			{
+				OverriddenViewports.Add(Viewport);
+			}
+			else if (Viewport->RenderSettings.IsViewportHasParent())
+			{
+				LinkedViewports.Add(Viewport);
+			}
+			else
+			{
+				CurrentRenderFrameViewports.Add(Viewport);
+			}
 		}
-		break;
-	default:
-		// Use all exist viewports for render
-		ClusterNodeViewports.Append(Viewports);
-		break;
+
+		CurrentRenderFrameViewports.Append(LinkedViewports);
+		CurrentRenderFrameViewports.Append(OverriddenViewports);
 	}
 }
 
@@ -476,7 +491,7 @@ int32 FDisplayClusterViewportManager::FindFirstViewportStereoViewIndex(const TSh
 	const int32 ViewPerViewportAmount = GetViewPerViewportAmount();
 	if (InViewport.IsValid() && ViewPerViewportAmount > 0)
 	{
-		const int32 ViewportIndex = Viewports.Find(InViewport);
+		const int32 ViewportIndex = ImplGetEntireClusterViewports().Find(InViewport);
 		if (ViewportIndex != INDEX_NONE)
 		{
 			// Begin from 1, because INDEX_NONE use ViewState[0] in LocalPlayer.cpp:786
@@ -538,7 +553,7 @@ bool FDisplayClusterViewportManager::BeginNewFrame(FViewport* InViewport, UWorld
 	const FDisplayClusterRenderFrameSettings& RenderFrameSettings = GetRenderFrameSettings();
 
 	// Initialize viewports from new render settings, and create new contexts, reset prev frame resources
-	for (TSharedPtr<FDisplayClusterViewport, ESPMode::ThreadSafe>& ViewportIt : ClusterNodeViewports)
+	for (const TSharedPtr<FDisplayClusterViewport, ESPMode::ThreadSafe>& ViewportIt : ImplGetCurrentRenderFrameViewports())
 	{
 		if (!ViewportIt.IsValid())
 		{
@@ -570,7 +585,7 @@ bool FDisplayClusterViewportManager::BeginNewFrame(FViewport* InViewport, UWorld
 	// in PostRenderViewFamily like it was previously so nothing should be broken.
 	const float DefaultDisplayGamma = GEngine ? GEngine->DisplayGamma : 2.2f;
 	const FEngineShowFlags* const DefaultEngineShowFlags = (InViewport && InViewport->GetClient()) ? InViewport->GetClient()->GetEngineShowFlags() : nullptr;
-	for (TSharedPtr<FDisplayClusterViewport, ESPMode::ThreadSafe>& ViewportIt : ClusterNodeViewports)
+	for (const TSharedPtr<FDisplayClusterViewport, ESPMode::ThreadSafe>& ViewportIt : ImplGetCurrentRenderFrameViewports())
 	{
 		if (ViewportIt.IsValid())
 		{
@@ -596,20 +611,20 @@ bool FDisplayClusterViewportManager::BeginNewFrame(FViewport* InViewport, UWorld
 	UpdateSceneRenderTargetSize();
 
 	// Build new frame structure
-	if (!RenderFrameManager->BuildRenderFrame(InViewport, RenderFrameSettings, ClusterNodeViewports, OutRenderFrame))
+	if (!RenderFrameManager->BuildRenderFrame(InViewport, RenderFrameSettings, ImplGetCurrentRenderFrameViewports(), OutRenderFrame))
 	{
 		return false;
 	}
 
 	// Allocate resources for frame
-	if (!RenderTargetManager->AllocateRenderFrameResources(InViewport, RenderFrameSettings, ClusterNodeViewports, OutRenderFrame))
+	if (!RenderTargetManager->AllocateRenderFrameResources(InViewport, RenderFrameSettings, ImplGetCurrentRenderFrameViewports(), OutRenderFrame))
 	{
 		return false;
 	}
 
 	const FIntPoint RenderFrameSize = OutRenderFrame.FrameRect.Size();
 
-	for (TSharedPtr<FDisplayClusterViewport, ESPMode::ThreadSafe>& Viewport : ClusterNodeViewports)
+	for (const TSharedPtr<FDisplayClusterViewport, ESPMode::ThreadSafe>& Viewport : ImplGetCurrentRenderFrameViewports())
 	{
 		if (Viewport.IsValid())
 		{
@@ -619,7 +634,7 @@ bool FDisplayClusterViewportManager::BeginNewFrame(FViewport* InViewport, UWorld
 	}
 
 	// Update desired views number
-	OutRenderFrame.DesiredNumberOfViews = (Viewports.Num() * GetViewPerViewportAmount()) + 1;
+	OutRenderFrame.DesiredNumberOfViews = (ImplGetEntireClusterViewports().Num() * GetViewPerViewportAmount()) + 1;
 
 #if WITH_EDITOR
 	// Get preview resources from root actor
@@ -639,7 +654,7 @@ void FDisplayClusterViewportManager::InitializeNewFrame()
 	check(IsInGameThread());
 
 	// Before render
-	for (TSharedPtr<FDisplayClusterViewport, ESPMode::ThreadSafe>& Viewport : ClusterNodeViewports)
+	for (const TSharedPtr<FDisplayClusterViewport, ESPMode::ThreadSafe>& Viewport : ImplGetCurrentRenderFrameViewports())
 	{
 		if (Viewport.IsValid())
 		{
@@ -661,7 +676,7 @@ void FDisplayClusterViewportManager::InitializeNewFrame()
 	 * So we are sending viewports from the game stream data to the proxy.
 	 * Lastly, after the viewport math data has been updated, we need to send it again.
 	 */
-	ViewportManagerProxy->ImplUpdateViewports(ClusterNodeViewports);
+	ViewportManagerProxy->ImplUpdateViewportProxies_GameThread(ImplGetCurrentRenderFrameViewports());
 
 	// Send postprocess data to render thread
 	if (PostProcessManager.IsValid())
@@ -679,9 +694,9 @@ void FDisplayClusterViewportManager::FinalizeNewFrame()
 	check(IsInGameThread());
 
 	// [2] The viewport math has been updated, we need to send it to the proxy again.
-	ViewportManagerProxy->ImplUpdateViewports(ClusterNodeViewports);
+	ViewportManagerProxy->ImplUpdateViewportProxies_GameThread(ImplGetCurrentRenderFrameViewports());
 
-	for (TSharedPtr<FDisplayClusterViewport, ESPMode::ThreadSafe>& Viewport : ClusterNodeViewports)
+	for (const TSharedPtr<FDisplayClusterViewport, ESPMode::ThreadSafe>& Viewport : ImplGetCurrentRenderFrameViewports())
 	{
 		if (Viewport.IsValid())
 		{
@@ -698,7 +713,7 @@ void FDisplayClusterViewportManager::FinalizeNewFrame()
 	}
 
 	// Finish update settings
-	for (TSharedPtr<FDisplayClusterViewport, ESPMode::ThreadSafe>& Viewport : ClusterNodeViewports)
+	for (const TSharedPtr<FDisplayClusterViewport, ESPMode::ThreadSafe>& Viewport : ImplGetCurrentRenderFrameViewports())
 	{
 		if (Viewport.IsValid())
 		{
@@ -857,7 +872,7 @@ void FDisplayClusterViewportManager::RenderFrame(FViewport* InViewport)
 		LightCardManager->RenderFrame();
 	}
 
-	ViewportManagerProxy->ImplRenderFrame(InViewport);
+	ViewportManagerProxy->ImplRenderFrame_GameThread(InViewport);
 }
 
 bool FDisplayClusterViewportManager::CreateViewport(const FString& InViewportId, const class UDisplayClusterConfigurationViewport& ConfigurationViewport)
@@ -917,7 +932,7 @@ TSharedPtr<FDisplayClusterViewport, ESPMode::ThreadSafe> FDisplayClusterViewport
 	check(IsInGameThread());
 
 	// Ok, we have a request for a particular viewport. Let's find it.
-	TSharedPtr<FDisplayClusterViewport, ESPMode::ThreadSafe> const* DesiredViewport = Viewports.FindByPredicate([ViewportId](const TSharedPtr<FDisplayClusterViewport, ESPMode::ThreadSafe>& ItemViewport)
+	TSharedPtr<FDisplayClusterViewport, ESPMode::ThreadSafe> const* DesiredViewport = ImplGetEntireClusterViewports().FindByPredicate([ViewportId](const TSharedPtr<FDisplayClusterViewport, ESPMode::ThreadSafe>& ItemViewport)
 	{
 		return ItemViewport.IsValid() && ViewportId.Equals(ItemViewport->ViewportId, ESearchCase::IgnoreCase);
 	});
@@ -971,8 +986,8 @@ TSharedPtr<FDisplayClusterViewport, ESPMode::ThreadSafe> FDisplayClusterViewport
 	if (NewViewport.IsValid())
 	{
 		// Add viewport on gamethread
-		Viewports.Add(NewViewport);
-		ClusterNodeViewports.Add(NewViewport);
+		EntireClusterViewports.Add(NewViewport);
+		CurrentRenderFrameViewports.Add(NewViewport);
 
 		// Handle start scene for viewport
 		NewViewport->HandleStartScene();
@@ -990,22 +1005,22 @@ void FDisplayClusterViewportManager::ImplDeleteViewport(const TSharedPtr<FDispla
 	ExistViewport->UninitializedProjectionPolicy.Reset();
 
 	{
-		// Remove viewport from the whole viewports list
-		int32 ViewportIndex = Viewports.Find(ExistViewport);
+		// Remove viewport from the entire viewports list
+		int32 ViewportIndex = EntireClusterViewports.Find(ExistViewport);
 		if (ViewportIndex != INDEX_NONE)
 		{
-			Viewports[ViewportIndex] = nullptr;
-			Viewports.RemoveAt(ViewportIndex);
+			EntireClusterViewports[ViewportIndex] = nullptr;
+			EntireClusterViewports.RemoveAt(ViewportIndex);
 		}
 	}
 
 	{
 		// Remove viewport from the cluster viewports list
-		int32 ViewportIndex = ClusterNodeViewports.Find(ExistViewport);
+		int32 ViewportIndex = CurrentRenderFrameViewports.Find(ExistViewport);
 		if (ViewportIndex != INDEX_NONE)
 		{
-			ClusterNodeViewports[ViewportIndex] = nullptr;
-			ClusterNodeViewports.RemoveAt(ViewportIndex);
+			CurrentRenderFrameViewports[ViewportIndex] = nullptr;
+			CurrentRenderFrameViewports.RemoveAt(ViewportIndex);
 		}
 	}
 
@@ -1025,7 +1040,7 @@ IDisplayClusterViewport* FDisplayClusterViewportManager::FindViewport(const int3
 		return nullptr;
 	}
 
-	for (const TSharedPtr<FDisplayClusterViewport, ESPMode::ThreadSafe>& ViewportIt : Viewports)
+	for (const TSharedPtr<FDisplayClusterViewport, ESPMode::ThreadSafe>& ViewportIt : ImplGetEntireClusterViewports())
 	{
 		if (ViewportIt.IsValid() && ViewportIt->FindContext(ViewIndex, OutContextNum))
 		{
@@ -1074,7 +1089,8 @@ void FDisplayClusterViewportManager::MarkComponentGeometryDirty(const FName InCo
 	check(IsInGameThread());
 
 	// 1. Update all ProceduralMeshComponent references for projection policies
-	for (const TSharedPtr<FDisplayClusterViewport, ESPMode::ThreadSafe>& ViewportIt : ImplGetViewports())
+	// (We need to update for all cluster viewports)
+	for (const TSharedPtr<FDisplayClusterViewport, ESPMode::ThreadSafe>& ViewportIt : ImplGetEntireClusterViewports())
 	{
 		if (ViewportIt.IsValid())
 		{
@@ -1104,7 +1120,7 @@ void FDisplayClusterViewportManager::MarkComponentGeometryDirty(const FName InCo
 
 void FDisplayClusterViewportManager::AddReferencedObjects(FReferenceCollector& Collector)
 {
-	for (TSharedPtr<FDisplayClusterViewport, ESPMode::ThreadSafe>& ViewportIt : Viewports)
+	for (const TSharedPtr<FDisplayClusterViewport, ESPMode::ThreadSafe>& ViewportIt : ImplGetEntireClusterViewports())
 	{
 		if (ViewportIt.IsValid())
 		{
diff --git a/Engine/Plugins/Runtime/nDisplay/Source/DisplayCluster/Private/Render/Viewport/DisplayClusterViewportManager.h b/Engine/Plugins/Runtime/nDisplay/Source/DisplayCluster/Private/Render/Viewport/DisplayClusterViewportManager.h
--- a/Engine/Plugins/Runtime/nDisplay/Source/DisplayCluster/Private/Render/Viewport/DisplayClusterViewportManager.h
+++ b/Engine/Plugins/Runtime/nDisplay/Source/DisplayCluster/Private/Render/Viewport/DisplayClusterViewportManager.h
@@ -99,7 +99,7 @@ public:
 
 	virtual const TArrayView<TSharedPtr<IDisplayClusterViewport, ESPMode::ThreadSafe>> GetViewports() const override
 	{
-		return TArrayView<TSharedPtr<IDisplayClusterViewport, ESPMode::ThreadSafe>>((TSharedPtr<IDisplayClusterViewport, ESPMode::ThreadSafe>*)(ClusterNodeViewports.GetData()), ClusterNodeViewports.Num());
+		return TArrayView<TSharedPtr<IDisplayClusterViewport, ESPMode::ThreadSafe>>((TSharedPtr<IDisplayClusterViewport, ESPMode::ThreadSafe>*)(CurrentRenderFrameViewports.GetData()), CurrentRenderFrameViewports.Num());
 	}
 
 	virtual void MarkComponentGeometryDirty(const FName InComponentName = NAME_None) override;
@@ -109,17 +109,15 @@ public:
 
 	//~~ End IDisplayClusterViewportManager
 
-	const TArray<TSharedPtr<FDisplayClusterViewport, ESPMode::ThreadSafe>>& ImplGetViewports() const
+	/** Getting the viewports of the current rendering frame (viewports from the current node or from a special named list). */
+	inline const TArray<TSharedPtr<FDisplayClusterViewport, ESPMode::ThreadSafe>>& ImplGetCurrentRenderFrameViewports() const
 	{
-		check(IsInGameThread());
-		return ClusterNodeViewports;
+		return CurrentRenderFrameViewports;
 	}
 
-	const TArray<TSharedPtr<FDisplayClusterViewport, ESPMode::ThreadSafe>>& ImplGetWholeClusterViewports() const
+	inline const TArray<TSharedPtr<FDisplayClusterViewport, ESPMode::ThreadSafe>>& ImplGetEntireClusterViewports() const
 	{
-		check(IsInGameThread());
-
-		return Viewports;
+		return EntireClusterViewports;
 	}
 
 	// internal use only
@@ -188,15 +186,20 @@ protected:
 	TSharedPtr<FDisplayClusterViewportLightCardManager, ESPMode::ThreadSafe>   LightCardManager;
 
 public:
+	// Configuration of the current cluster node
 	TUniquePtr<FDisplayClusterViewportConfiguration> Configuration;
 
 private:
+	// Manager for creating the render frame stucture
 	TUniquePtr<FDisplayClusterRenderFrameManager>  RenderFrameManager;
 
-	TArray<TSharedPtr<FDisplayClusterViewport, ESPMode::ThreadSafe>> Viewports;
-	TArray<TSharedPtr<FDisplayClusterViewport, ESPMode::ThreadSafe>> ClusterNodeViewports;
+	// Viewports of the entire cluster
+	TArray<TSharedPtr<FDisplayClusterViewport, ESPMode::ThreadSafe>> EntireClusterViewports;
+
+	// Viewports of the current render frame (viewports from the current node or from a special named list)
+	TArray<TSharedPtr<FDisplayClusterViewport, ESPMode::ThreadSafe>> CurrentRenderFrameViewports;
 
-	/** Render thread proxy manager. Deleted on render thread */
+	// Render thread proxy manager. Deleted on render thread
 	TSharedPtr<FDisplayClusterViewportManagerProxy, ESPMode::ThreadSafe> ViewportManagerProxy;
 
 	// Pointer to the current scene
diff --git a/Engine/Plugins/Runtime/nDisplay/Source/DisplayCluster/Private/Render/Viewport/DisplayClusterViewportManagerProxy.cpp b/Engine/Plugins/Runtime/nDisplay/Source/DisplayCluster/Private/Render/Viewport/DisplayClusterViewportManagerProxy.cpp
--- a/Engine/Plugins/Runtime/nDisplay/Source/DisplayCluster/Private/Render/Viewport/DisplayClusterViewportManagerProxy.cpp
+++ b/Engine/Plugins/Runtime/nDisplay/Source/DisplayCluster/Private/Render/Viewport/DisplayClusterViewportManagerProxy.cpp
@@ -88,8 +88,8 @@ void FDisplayClusterViewportManagerProxy::Release_RenderThread()
 	check(IsInRenderingThread());
 
 	// Delete viewport proxy objects
-	ViewportProxies.Empty();
-	ClusterNodeViewportProxies.Empty();
+	EntireClusterViewportProxies.Empty();
+	CurrentRenderFrameViewportProxies.Empty();
 
 	if (RenderTargetManager.IsValid())
 	{
@@ -127,7 +127,7 @@ void FDisplayClusterViewportManagerProxy::DeleteResource_RenderThread(FDisplayCl
 	if (InDeletedResourcePtr)
 	{
 		// Handle resource refs must be removed from all viewports
-		for (TSharedPtr<FDisplayClusterViewportProxy, ESPMode::ThreadSafe>& ViewportProxyIt : ViewportProxies)
+		for (const TSharedPtr<FDisplayClusterViewportProxy, ESPMode::ThreadSafe>& ViewportProxyIt : ImplGetEntireClusterViewportProxies_RenderThread())
 		{
 			if (ViewportProxyIt.IsValid())
 			{
@@ -140,15 +140,24 @@ void FDisplayClusterViewportManagerProxy::DeleteResource_RenderThread(FDisplayCl
 	}
 }
 
-void FDisplayClusterViewportManagerProxy::ImplUpdateClusterNodeViewportProxies()
+void FDisplayClusterViewportManagerProxy::ImplUpdateClusterNodeViewportProxies_RenderThread()
 {
-	ClusterNodeViewportProxies.Empty();
-	// Collect viewport proxies for rendered cluster node
-	for (TSharedPtr<FDisplayClusterViewportProxy, ESPMode::ThreadSafe>& ViewportProxyIt : ViewportProxies)
+	// Get viewports used to rendering of the current frame
+	if (RenderFrameSettings.ClusterNodeId.IsEmpty())
 	{
-		if (ViewportProxyIt.IsValid() && (RenderFrameSettings.ClusterNodeId.IsEmpty() || ViewportProxyIt->GetClusterNodeId() == RenderFrameSettings.ClusterNodeId))
+		// When a cluster node name is empty, we render without the cluster nodes
+		CurrentRenderFrameViewportProxies = ImplGetEntireClusterViewportProxies_RenderThread();
+	}
+	else
+	{
+		// Get viewport proxies for current cluster nodes
+		CurrentRenderFrameViewportProxies.Reset();
+		for (const TSharedPtr<FDisplayClusterViewportProxy, ESPMode::ThreadSafe>& ViewportProxyIt : ImplGetEntireClusterViewportProxies_RenderThread())
 		{
-			ClusterNodeViewportProxies.Add(ViewportProxyIt);
+			if (ViewportProxyIt.IsValid() && (ViewportProxyIt->GetClusterNodeId() == RenderFrameSettings.ClusterNodeId))
+			{
+				CurrentRenderFrameViewportProxies.Add(ViewportProxyIt);
+			}
 		}
 	}
 }
@@ -157,8 +166,9 @@ void FDisplayClusterViewportManagerProxy::CreateViewport_RenderThread(const TSha
 {
 	check(IsInRenderingThread());
 	
-	ViewportProxies.Add(InViewportProxy);
-	ImplUpdateClusterNodeViewportProxies();
+	EntireClusterViewportProxies.Add(InViewportProxy);
+
+	ImplUpdateClusterNodeViewportProxies_RenderThread();
 }
 
 void FDisplayClusterViewportManagerProxy::DeleteViewport_RenderThread(const TSharedPtr<FDisplayClusterViewportProxy, ESPMode::ThreadSafe>& InViewportProxy)
@@ -166,16 +176,16 @@ void FDisplayClusterViewportManagerProxy::DeleteViewport_RenderThread(const TSha
 	check(IsInRenderingThread());
 
 	// Remove viewport obj from manager
-	const int32 ViewportProxyIndex = ViewportProxies.Find(InViewportProxy);
+	const int32 ViewportProxyIndex = EntireClusterViewportProxies.Find(InViewportProxy);
 	if (ViewportProxyIndex != INDEX_NONE)
 	{
-		ViewportProxies.RemoveAt(ViewportProxyIndex);
+		EntireClusterViewportProxies.RemoveAt(ViewportProxyIndex);
 	}
 
-	const int32 ClusterViewportProxyIndex = ClusterNodeViewportProxies.Find(InViewportProxy);
+	const int32 ClusterViewportProxyIndex = CurrentRenderFrameViewportProxies.Find(InViewportProxy);
 	if (ClusterViewportProxyIndex != INDEX_NONE)
 	{
-		ClusterNodeViewportProxies.RemoveAt(ClusterViewportProxyIndex);
+		CurrentRenderFrameViewportProxies.RemoveAt(ClusterViewportProxyIndex);
 	}
 }
 
@@ -200,18 +210,18 @@ void FDisplayClusterViewportManagerProxy::ImplUpdateRenderFrameSettings(const FD
 		delete Settings;
 
 		// After updated settings we need update cluster node viewports
-		ViewportManagerProxy->ImplUpdateClusterNodeViewportProxies();
+		ViewportManagerProxy->ImplUpdateClusterNodeViewportProxies_RenderThread();
 	});
 }
 
 DECLARE_GPU_STAT_NAMED(nDisplay_ViewportManager_UpdateViewports, TEXT("nDisplay ViewportManager::UpdateViewports"));
 
-void FDisplayClusterViewportManagerProxy::ImplUpdateViewports(const TArray<TSharedPtr<FDisplayClusterViewport, ESPMode::ThreadSafe>>& InViewports)
+void FDisplayClusterViewportManagerProxy::ImplUpdateViewportProxies_GameThread(const TArray<TSharedPtr<FDisplayClusterViewport, ESPMode::ThreadSafe>>& InCurrentRenderFrameViewports)
 {
 	check(IsInGameThread());
 
 	TArray<FDisplayClusterViewportProxyData*> ViewportProxiesData;
-	for (const TSharedPtr<FDisplayClusterViewport, ESPMode::ThreadSafe>& ViewportIt : InViewports)
+	for (const TSharedPtr<FDisplayClusterViewport, ESPMode::ThreadSafe>& ViewportIt : InCurrentRenderFrameViewports)
 	{
 		if (ViewportIt.IsValid())
 		{
@@ -240,7 +250,7 @@ DECLARE_GPU_STAT_NAMED(nDisplay_ViewportManager_CrossGPUTransfer, TEXT("nDisplay
 DECLARE_GPU_STAT_NAMED(nDisplay_ViewportManager_UpdateDeferredResources, TEXT("nDisplay ViewportManager::UpdateDeferredResources"));
 DECLARE_GPU_STAT_NAMED(nDisplay_ViewportManager_WarpBlend, TEXT("nDisplay ViewportManager::WarpBlend"));
 
-void FDisplayClusterViewportManagerProxy::ImplRenderFrame(FViewport* InViewport)
+void FDisplayClusterViewportManagerProxy::ImplRenderFrame_GameThread(FViewport* InViewport)
 {
 	ENQUEUE_RENDER_COMMAND(DisplayClusterRenderFrame_Setup)(
 		[InViewportManagerProxy = SharedThis(this), InViewport](FRHICommandListImmediate& RHICmdList)
@@ -349,9 +359,9 @@ void FDisplayClusterViewportManagerProxy::UpdateDeferredResources_RenderThread(F
 	check(IsInRenderingThread());
 
 	TArray<TSharedPtr<FDisplayClusterViewportProxy, ESPMode::ThreadSafe>> OverriddenViewports;
-	OverriddenViewports.Reserve(ClusterNodeViewportProxies.Num());
+	OverriddenViewports.Reserve(ImplGetCurrentRenderFrameViewportProxies_RenderThread().Num());
 
-	for (const TSharedPtr<FDisplayClusterViewportProxy, ESPMode::ThreadSafe>& ViewportProxy : ClusterNodeViewportProxies)
+	for (const TSharedPtr<FDisplayClusterViewportProxy, ESPMode::ThreadSafe>& ViewportProxy : ImplGetCurrentRenderFrameViewportProxies_RenderThread())
 	{
 		if (!ViewportProxy->RenderSettings.IsViewportOverrided())
 		{
@@ -421,7 +431,7 @@ void FDisplayClusterViewportManagerProxy::UpdateFrameResources_RenderThread(FRHI
 	}
 
 	// Support viewport overlap order sorting:
-	TArray<TSharedPtr<FDisplayClusterViewportProxy, ESPMode::ThreadSafe>> SortedViewportProxy = ClusterNodeViewportProxies;
+	TArray<TSharedPtr<FDisplayClusterViewportProxy, ESPMode::ThreadSafe>> SortedViewportProxy = ImplGetCurrentRenderFrameViewportProxies_RenderThread();
 	SortedViewportProxy.Sort(
 		[](const TSharedPtr<FDisplayClusterViewportProxy, ESPMode::ThreadSafe>& VP1, const TSharedPtr<FDisplayClusterViewportProxy, ESPMode::ThreadSafe>& VP2)
 		{
@@ -530,7 +540,7 @@ void FDisplayClusterViewportManagerProxy::DoCrossGPUTransfers_RenderThread(FRHIC
 	// Copy the view render results to all GPUs that are native to the viewport.
 	TArray<FTransferResourceParams> TransferResources;
 
-	for (const TSharedPtr<FDisplayClusterViewportProxy, ESPMode::ThreadSafe>& ViewportProxyIt : ClusterNodeViewportProxies)
+	for (const TSharedPtr<FDisplayClusterViewportProxy, ESPMode::ThreadSafe>& ViewportProxyIt : ImplGetCurrentRenderFrameViewportProxies_RenderThread())
 	{
 		if (!ViewportProxyIt.IsValid())
 		{
@@ -611,7 +621,7 @@ bool FDisplayClusterViewportManagerProxy::GetFrameTargets_RenderThread(TArray<FR
 	check(IsInRenderingThread());
 
 	// Get any defined frame targets from first visible viewport
-	for (const TSharedPtr<FDisplayClusterViewportProxy, ESPMode::ThreadSafe>& ViewportProxyIt : ClusterNodeViewportProxies)
+	for (const TSharedPtr<FDisplayClusterViewportProxy, ESPMode::ThreadSafe>& ViewportProxyIt : ImplGetCurrentRenderFrameViewportProxies_RenderThread())
 	{
 		// Process only valid viewports with output frame resources
 		if (ViewportProxyIt.IsValid() && !ViewportProxyIt->OutputFrameTargetableResources.IsEmpty())
@@ -702,12 +712,12 @@ bool FDisplayClusterViewportManagerProxy::ResolveFrameTargetToBackBuffer_RenderT
 	return false;
 }
 
-FDisplayClusterViewportProxy* FDisplayClusterViewportManagerProxy::ImplFindViewport_RenderThread(const FString& ViewportId) const
+FDisplayClusterViewportProxy* FDisplayClusterViewportManagerProxy::ImplFindViewportProxy_RenderThread(const FString& ViewportId) const
 {
 	check(IsInRenderingThread());
 
 	// Ok, we have a request for a particular viewport. Let's find it.
-	TSharedPtr<FDisplayClusterViewportProxy, ESPMode::ThreadSafe> const* DesiredViewport = ViewportProxies.FindByPredicate([ViewportId](const TSharedPtr<FDisplayClusterViewportProxy, ESPMode::ThreadSafe>& ItemViewport)
+	TSharedPtr<FDisplayClusterViewportProxy, ESPMode::ThreadSafe> const* DesiredViewport = ImplGetEntireClusterViewportProxies_RenderThread().FindByPredicate([ViewportId](const TSharedPtr<FDisplayClusterViewportProxy, ESPMode::ThreadSafe>& ItemViewport)
 	{
 		return ViewportId.Equals(ItemViewport->GetId(), ESearchCase::IgnoreCase);
 	});
@@ -715,7 +725,7 @@ FDisplayClusterViewportProxy* FDisplayClusterViewportManagerProxy::ImplFindViewp
 	return (DesiredViewport != nullptr) ? DesiredViewport->Get() : nullptr;
 }
 
-FDisplayClusterViewportProxy* FDisplayClusterViewportManagerProxy::ImplFindViewport_RenderThread(const int32 StereoViewIndex, uint32* OutContextNum) const
+FDisplayClusterViewportProxy* FDisplayClusterViewportManagerProxy::ImplFindViewportProxy_RenderThread(const int32 StereoViewIndex, uint32* OutContextNum) const
 {
 	check(IsInRenderingThread());
 	
@@ -727,7 +737,7 @@ FDisplayClusterViewportProxy* FDisplayClusterViewportManagerProxy::ImplFindViewp
 		return nullptr;
 	}
 
-	for (const TSharedPtr<FDisplayClusterViewportProxy, ESPMode::ThreadSafe>&  ViewportProxyIt : ViewportProxies)
+	for (const TSharedPtr<FDisplayClusterViewportProxy, ESPMode::ThreadSafe>&  ViewportProxyIt : ImplGetEntireClusterViewportProxies_RenderThread())
 	{
 		if (ViewportProxyIt.IsValid() && ViewportProxyIt->FindContext_RenderThread(StereoViewIndex, OutContextNum))
 		{
diff --git a/Engine/Plugins/Runtime/nDisplay/Source/DisplayCluster/Private/Render/Viewport/DisplayClusterViewportManagerProxy.h b/Engine/Plugins/Runtime/nDisplay/Source/DisplayCluster/Private/Render/Viewport/DisplayClusterViewportManagerProxy.h
--- a/Engine/Plugins/Runtime/nDisplay/Source/DisplayCluster/Private/Render/Viewport/DisplayClusterViewportManagerProxy.h
+++ b/Engine/Plugins/Runtime/nDisplay/Source/DisplayCluster/Private/Render/Viewport/DisplayClusterViewportManagerProxy.h
@@ -47,17 +47,17 @@ public:
 
 	virtual IDisplayClusterViewportProxy* FindViewport_RenderThread(const FString& InViewportId) const override
 	{
-		return ImplFindViewport_RenderThread(InViewportId);
+		return ImplFindViewportProxy_RenderThread(InViewportId);
 	}
 
 	virtual IDisplayClusterViewportProxy* FindViewport_RenderThread(const int32 StereoViewIndex, uint32* OutContextNum = nullptr) const override
 	{
-		return ImplFindViewport_RenderThread(StereoViewIndex, OutContextNum);
+		return ImplFindViewportProxy_RenderThread(StereoViewIndex, OutContextNum);
 	}
 
 	virtual const TArrayView<TSharedPtr<IDisplayClusterViewportProxy, ESPMode::ThreadSafe>> GetViewports_RenderThread() const override
 	{
-		return TArrayView<TSharedPtr<IDisplayClusterViewportProxy, ESPMode::ThreadSafe>>((TSharedPtr<IDisplayClusterViewportProxy, ESPMode::ThreadSafe>*)(ClusterNodeViewportProxies.GetData()), ClusterNodeViewportProxies.Num());
+		return TArrayView<TSharedPtr<IDisplayClusterViewportProxy, ESPMode::ThreadSafe>>((TSharedPtr<IDisplayClusterViewportProxy, ESPMode::ThreadSafe>*)(CurrentRenderFrameViewportProxies.GetData()), CurrentRenderFrameViewportProxies.Num());
 	}
 
 	virtual bool GetFrameTargets_RenderThread(TArray<FRHITexture2D*>& OutFrameResources, TArray<FIntPoint>& OutTargetOffsets, TArray<FRHITexture2D*>* OutAdditionalFrameResources = nullptr) const override;
@@ -89,59 +89,63 @@ public:
 	void Initialize(FDisplayClusterViewportManager& InViewportManager);
 
 	/** Get LightCardManager proxy object. */
-	TSharedPtr<FDisplayClusterViewportLightCardManagerProxy, ESPMode::ThreadSafe> GetLightCardManagerProxy_RenderThread() const
+	inline TSharedPtr<FDisplayClusterViewportLightCardManagerProxy, ESPMode::ThreadSafe> GetLightCardManagerProxy_RenderThread() const
 	{ return LightCardManagerProxy; }
 
 	/** Get RenderFrameSettings on render thread. */
-	const FDisplayClusterRenderFrameSettings& GetRenderFrameSettings_RenderThread() const
+	inline  const FDisplayClusterRenderFrameSettings& GetRenderFrameSettings_RenderThread() const
 	{ return RenderFrameSettings; }
 
-	/** Get the viewports of the cluster node on the render thread. */
-	const TArray<TSharedPtr<FDisplayClusterViewportProxy, ESPMode::ThreadSafe>>& ImplGetViewportProxies_RenderThread() const
-	{ return ClusterNodeViewportProxies; }
+	/** Geting the viewport proxies of the current rendering frame. (viewports from the current node or from a special named list) */
+	inline const TArray<TSharedPtr<FDisplayClusterViewportProxy, ESPMode::ThreadSafe>>& ImplGetCurrentRenderFrameViewportProxies_RenderThread() const
+	{ return CurrentRenderFrameViewportProxies; }
+
+	/** Get the viewport proxies  of the whole cluster on the render thread. */
+	inline const TArray<TSharedPtr<FDisplayClusterViewportProxy, ESPMode::ThreadSafe>>& ImplGetEntireClusterViewportProxies_RenderThread() const
+	{ return EntireClusterViewportProxies; }
 
 	/** Find viewport by name. */
-	FDisplayClusterViewportProxy* ImplFindViewport_RenderThread(const FString& InViewportId) const;
+	FDisplayClusterViewportProxy* ImplFindViewportProxy_RenderThread(const FString& InViewportId) const;
 
 	/** Find viewport and context index by StereoViewIndex. */
-	FDisplayClusterViewportProxy* ImplFindViewport_RenderThread(const int32 StereoViewIndex, uint32* OutContextNum = nullptr) const;
+	FDisplayClusterViewportProxy* ImplFindViewportProxy_RenderThread(const int32 StereoViewIndex, uint32* OutContextNum = nullptr) const;
 
 	/** Copy RenderFrameSettings from game to render thread. */
 	void ImplUpdateRenderFrameSettings(const FDisplayClusterRenderFrameSettings& InRenderFrameSettings,
 		const TSharedPtr<FDisplayClusterViewportManagerViewExtension, ESPMode::ThreadSafe>& InViewportManagerViewExtension);
 
-	/** Copy Viewports data to Proxies from game to render thread. */
-	void ImplUpdateViewports(const TArray<TSharedPtr<FDisplayClusterViewport, ESPMode::ThreadSafe>>& InViewports);
+	/** Sending viewport data from the game stream to the viewport proxy that is in the rendering stream. */
+	void ImplUpdateViewportProxies_GameThread(const TArray<TSharedPtr<FDisplayClusterViewport, ESPMode::ThreadSafe>>& InCurrentRenderFrameViewports);
 
 	/** Rendering the final frame of nDisplay. Called after all viewports have been rendered in RTT. */
-	void ImplRenderFrame(FViewport* InViewport);
+	void ImplRenderFrame_GameThread(FViewport* InViewport);
 
 	/** Rendering black color in all frame textures. */
 	void ImplClearFrameTargets_RenderThread(FRHICommandListImmediate& RHICmdList) const;
 
 private:
 	/** Update ClusterNodeViewportProxies from ViewportProxies for the current cluster node. */
-	void ImplUpdateClusterNodeViewportProxies();
+	void ImplUpdateClusterNodeViewportProxies_RenderThread();
 
 private:
-	/** RenderTarget manager object. Used to remove related objects in order. */
+	// RenderTarget manager object. Used to remove related objects in order.
 	TSharedPtr<FDisplayClusterRenderTargetManager, ESPMode::ThreadSafe>        RenderTargetManager;
 
-	/** PostProcess manager object. Used to remove related objects in order. */
+	// PostProcess manager object. Used to remove related objects in order.
 	TSharedPtr<FDisplayClusterViewportPostProcessManager, ESPMode::ThreadSafe> PostProcessManager;
 
-	/** LightCard manager proxy object. */
+	// LightCard manager proxy object.
 	TSharedPtr<FDisplayClusterViewportLightCardManagerProxy, ESPMode::ThreadSafe> LightCardManagerProxy;
 
-	/** Copy of RenderFrameSettings for the render thread. */
+	// Copy of RenderFrameSettings for the render thread.
 	FDisplayClusterRenderFrameSettings RenderFrameSettings;
 
-	/** Proxies of the entire cluster. */
-	TArray<TSharedPtr<FDisplayClusterViewportProxy, ESPMode::ThreadSafe>> ViewportProxies;
+	// Viewport proxies of the entire cluster.
+	TArray<TSharedPtr<FDisplayClusterViewportProxy, ESPMode::ThreadSafe>> EntireClusterViewportProxies;
 
-	/** Proxies of the current cluster node. */
-	TArray<TSharedPtr<FDisplayClusterViewportProxy, ESPMode::ThreadSafe>> ClusterNodeViewportProxies;
+	// Viewport proxies of the current rendering frame (viewports from the current node or from a special named list).
+	TArray<TSharedPtr<FDisplayClusterViewportProxy, ESPMode::ThreadSafe>> CurrentRenderFrameViewportProxies;
 
-	/** nDisplay VE object. Used to remove related objects in order. */
+	// nDisplay VE object. Used to remove related objects in order.
 	TSharedPtr<class FDisplayClusterViewportManagerViewExtension, ESPMode::ThreadSafe> ViewportManagerViewExtension;
 };
diff --git a/Engine/Plugins/Runtime/nDisplay/Source/DisplayCluster/Private/Render/Viewport/DisplayClusterViewportManagerViewExtension.cpp b/Engine/Plugins/Runtime/nDisplay/Source/DisplayCluster/Private/Render/Viewport/DisplayClusterViewportManagerViewExtension.cpp
--- a/Engine/Plugins/Runtime/nDisplay/Source/DisplayCluster/Private/Render/Viewport/DisplayClusterViewportManagerViewExtension.cpp
+++ b/Engine/Plugins/Runtime/nDisplay/Source/DisplayCluster/Private/Render/Viewport/DisplayClusterViewportManagerViewExtension.cpp
@@ -130,7 +130,7 @@ FScreenPassTexture FDisplayClusterViewportManagerViewExtension::PostProcessPassA
 	if (const FDisplayClusterViewportManagerProxy* ViewportManagerProxy = GetViewportManagerProxy())
 	{
 		uint32 ContextNum = 0;
-		if (FDisplayClusterViewportProxy* ViewportProxyPtr = ViewportManagerProxy->ImplFindViewport_RenderThread(View.StereoViewIndex, &ContextNum))
+		if (FDisplayClusterViewportProxy* ViewportProxyPtr = ViewportManagerProxy->ImplFindViewportProxy_RenderThread(View.StereoViewIndex, &ContextNum))
 		{
 			if (ViewportProxyPtr->ShouldUsePostProcessPassAfterFXAA())
 			{
@@ -152,7 +152,7 @@ FScreenPassTexture FDisplayClusterViewportManagerViewExtension::PostProcessPassA
 	if (const FDisplayClusterViewportManagerProxy* ViewportManagerProxy = GetViewportManagerProxy())
 	{
 		uint32 ContextNum = 0;
-		if (FDisplayClusterViewportProxy* ViewportProxyPtr = ViewportManagerProxy->ImplFindViewport_RenderThread(View.StereoViewIndex, &ContextNum))
+		if (FDisplayClusterViewportProxy* ViewportProxyPtr = ViewportManagerProxy->ImplFindViewportProxy_RenderThread(View.StereoViewIndex, &ContextNum))
 		{
 			if (ViewportProxyPtr->ShouldUsePostProcessPassAfterSSRInput())
 			{
@@ -174,7 +174,7 @@ FScreenPassTexture FDisplayClusterViewportManagerViewExtension::PostProcessPassA
 	if (const FDisplayClusterViewportManagerProxy* ViewportManagerProxy = GetViewportManagerProxy())
 	{
 		uint32 ContextNum = 0;
-		if (FDisplayClusterViewportProxy* ViewportProxyPtr = ViewportManagerProxy->ImplFindViewport_RenderThread(View.StereoViewIndex, &ContextNum))
+		if (FDisplayClusterViewportProxy* ViewportProxyPtr = ViewportManagerProxy->ImplFindViewportProxy_RenderThread(View.StereoViewIndex, &ContextNum))
 		{
 			if (ViewportProxyPtr->ShouldUsePostProcessPassTonemap())
 			{
diff --git a/Engine/Plugins/Runtime/nDisplay/Source/DisplayCluster/Private/Render/Viewport/DisplayClusterViewportManager_EditorPreview.cpp b/Engine/Plugins/Runtime/nDisplay/Source/DisplayCluster/Private/Render/Viewport/DisplayClusterViewportManager_EditorPreview.cpp
--- a/Engine/Plugins/Runtime/nDisplay/Source/DisplayCluster/Private/Render/Viewport/DisplayClusterViewportManager_EditorPreview.cpp
+++ b/Engine/Plugins/Runtime/nDisplay/Source/DisplayCluster/Private/Render/Viewport/DisplayClusterViewportManager_EditorPreview.cpp
@@ -54,11 +54,11 @@ void FDisplayClusterViewportManager::ImplUpdatePreviewRTTResources()
 	TArray<FString>        PreviewViewportNames;
 	TArray<FTextureRHIRef> PreviewRenderTargetableTextures;
 
-	PreviewViewportNames.Reserve(Viewports.Num());
-	PreviewRenderTargetableTextures.Reserve(Viewports.Num());
+	PreviewViewportNames.Reserve(ImplGetEntireClusterViewports().Num());
+	PreviewRenderTargetableTextures.Reserve(ImplGetEntireClusterViewports().Num());
 
 	// Collect visible and enabled viewports for preview rendering:
-	for (const TSharedPtr<FDisplayClusterViewport, ESPMode::ThreadSafe>& ViewportIt : Viewports)
+	for (const TSharedPtr<FDisplayClusterViewport, ESPMode::ThreadSafe>& ViewportIt : ImplGetEntireClusterViewports())
 	{
 		// update only current cluster node
 		if (ViewportIt.IsValid() && (ClusterNodeId.IsEmpty() || ViewportIt->GetClusterNodeId() == ClusterNodeId))
@@ -98,7 +98,7 @@ void FDisplayClusterViewportManager::ImplUpdatePreviewRTTResources()
 void FDisplayClusterViewportManager::OnPreGarbageCollect()
 {
 	// The view state can reference materials from the world being cleaned up. (Example post process materials)
-	for (TSharedPtr<FDisplayClusterViewport, ESPMode::ThreadSafe>& ViewportIt : Viewports)
+	for (const TSharedPtr<FDisplayClusterViewport, ESPMode::ThreadSafe>& ViewportIt : ImplGetEntireClusterViewports())
 	{
 		if (ViewportIt.IsValid())
 		{
diff --git a/Engine/Plugins/Runtime/nDisplay/Source/DisplayCluster/Private/Render/Viewport/DisplayClusterViewportProxy.cpp b/Engine/Plugins/Runtime/nDisplay/Source/DisplayCluster/Private/Render/Viewport/DisplayClusterViewportProxy.cpp
--- a/Engine/Plugins/Runtime/nDisplay/Source/DisplayCluster/Private/Render/Viewport/DisplayClusterViewportProxy.cpp
+++ b/Engine/Plugins/Runtime/nDisplay/Source/DisplayCluster/Private/Render/Viewport/DisplayClusterViewportProxy.cpp
@@ -306,7 +306,7 @@ const FDisplayClusterViewportProxy& FDisplayClusterViewportProxy::GetRenderingVi
 	case EDisplayClusterViewportOverrideMode::InernalRTT:
 		if (FDisplayClusterViewportManagerProxy* ViewportManagerProxy = GetViewportManagerProxyImpl_RenderThread())
 		{
-			if (FDisplayClusterViewportProxy const* OverrideViewportProxy = ViewportManagerProxy->ImplFindViewport_RenderThread(RenderSettings.ViewportOverrideId))
+			if (FDisplayClusterViewportProxy const* OverrideViewportProxy = ViewportManagerProxy->ImplFindViewportProxy_RenderThread(RenderSettings.ViewportOverrideId))
 			{
 				return *OverrideViewportProxy;
 			}
diff --git a/Engine/Plugins/Runtime/nDisplay/Source/DisplayCluster/Private/Render/Viewport/DisplayClusterViewport_CustomPostProcessSettings.cpp b/Engine/Plugins/Runtime/nDisplay/Source/DisplayCluster/Private/Render/Viewport/DisplayClusterViewport_CustomPostProcessSettings.cpp
--- a/Engine/Plugins/Runtime/nDisplay/Source/DisplayCluster/Private/Render/Viewport/DisplayClusterViewport_CustomPostProcessSettings.cpp
+++ b/Engine/Plugins/Runtime/nDisplay/Source/DisplayCluster/Private/Render/Viewport/DisplayClusterViewport_CustomPostProcessSettings.cpp
@@ -5,10 +5,12 @@
 
 void FDisplayClusterViewport_CustomPostProcessSettings::AddCustomPostProcess(const ERenderPass InRenderPass, const FPostProcessSettings& InSettings, float BlendWeight, bool bSingleFrame)
 {
-	PostprocessAsset.Emplace(InRenderPass, FPostprocessData(InSettings, BlendWeight, bSingleFrame));
+	if (BlendWeight > 0.f) // Ignore PP with zero weights
+	{
+		PostprocessAsset.Emplace(InRenderPass, FPostprocessData(InSettings, BlendWeight, bSingleFrame));
+	}
 }
 
-
 void FDisplayClusterViewport_CustomPostProcessSettings::RemoveCustomPostProcess(const ERenderPass InRenderPass)
 {
 	if (PostprocessAsset.Contains(InRenderPass))
diff --git a/Engine/Plugins/Runtime/nDisplay/Source/DisplayCluster/Private/Render/Viewport/DisplayClusterViewport_CustomPostProcessSettings.h b/Engine/Plugins/Runtime/nDisplay/Source/DisplayCluster/Private/Render/Viewport/DisplayClusterViewport_CustomPostProcessSettings.h
--- a/Engine/Plugins/Runtime/nDisplay/Source/DisplayCluster/Private/Render/Viewport/DisplayClusterViewport_CustomPostProcessSettings.h
+++ b/Engine/Plugins/Runtime/nDisplay/Source/DisplayCluster/Private/Render/Viewport/DisplayClusterViewport_CustomPostProcessSettings.h
@@ -16,8 +16,9 @@ public:
 	virtual ~FDisplayClusterViewport_CustomPostProcessSettings() = default;
 
 public:
-	void AddCustomPostProcess(const ERenderPass InRenderPass, const FPostProcessSettings& InSettings, float BlendWeight, bool bSingleFrame);
-	void RemoveCustomPostProcess(const ERenderPass InRenderPass);
+	virtual void AddCustomPostProcess(const ERenderPass InRenderPass, const FPostProcessSettings& InSettings, float BlendWeight, bool bSingleFrame) override;
+	virtual void RemoveCustomPostProcess(const ERenderPass InRenderPass) override;
+
 	void FinalizeFrame();
 
 	virtual bool DoPostProcess(const ERenderPass InRenderPass, FPostProcessSettings* OutSettings, float* OutBlendWeight = nullptr) const override;
diff --git a/Engine/Plugins/Runtime/nDisplay/Source/DisplayCluster/Private/Render/Viewport/DisplayClusterViewport_ViewPoint.cpp b/Engine/Plugins/Runtime/nDisplay/Source/DisplayCluster/Private/Render/Viewport/DisplayClusterViewport_ViewPoint.cpp
--- a/Engine/Plugins/Runtime/nDisplay/Source/DisplayCluster/Private/Render/Viewport/DisplayClusterViewport_ViewPoint.cpp
+++ b/Engine/Plugins/Runtime/nDisplay/Source/DisplayCluster/Private/Render/Viewport/DisplayClusterViewport_ViewPoint.cpp
@@ -27,7 +27,7 @@ using namespace UE::DisplayClusterViewport;
 ///////////////////////////////////////////////////////////////////////////////////////
 //          FDisplayClusterViewport
 ///////////////////////////////////////////////////////////////////////////////////////
-UDisplayClusterCameraComponent* FDisplayClusterViewport::GetViewPointCamera() const
+UDisplayClusterCameraComponent* FDisplayClusterViewport::GetViewPointCameraComponent() const
 {
 	ADisplayClusterRootActor* RootActor = GetRootActor();
 	if (!RootActor)
@@ -53,7 +53,7 @@ UDisplayClusterCameraComponent* FDisplayClusterViewport::GetViewPointCamera() co
 	const FString& CameraId = GetRenderSettings().CameraId;
 	if (CameraId.Len() > 0)
 	{
-		UE_LOG(LogDisplayClusterViewport, Verbose, TEXT("Viewport '%s' has assigned camera '%s'"), *GetId(), *CameraId);
+		UE_LOG(LogDisplayClusterViewport, Verbose, TEXT("Viewport '%s' has assigned ViewPoint '%s'"), *GetId(), *CameraId);
 	}
 
 	// Get camera component assigned to the viewport (or default camera if nothing assigned)
@@ -64,14 +64,14 @@ UDisplayClusterCameraComponent* FDisplayClusterViewport::GetViewPointCamera() co
 		return ViewCamera;
 	}
 
-	UE_LOG(LogDisplayClusterViewport, Warning, TEXT("No camera found for viewport '%s'"), *GetId());
+	UE_LOG(LogDisplayClusterViewport, Warning, TEXT("ViewPoint '%s' is not found for viewport '%s'"), *CameraId, * GetId());
 
 	return nullptr;
 }
 
 bool FDisplayClusterViewport::SetupViewPoint(FMinimalViewInfo& InOutViewInfo)
 {
-	if (UDisplayClusterCameraComponent* ViewCamera = GetViewPointCamera())
+	if (UDisplayClusterCameraComponent* ViewCamera = GetViewPointCameraComponent())
 	{
 		// First try to get the viewpoint from the projection policy
 		if (!ProjectionPolicy->GetViewPoint(this, InOutViewInfo.Rotation, InOutViewInfo.Location))
@@ -91,7 +91,7 @@ float FDisplayClusterViewport::GetStereoEyeOffsetDistance(const uint32 InContext
 {
 	float StereoEyeOffsetDistance = 0.f;
 
-	if (UDisplayClusterCameraComponent* ViewCamera = GetViewPointCamera())
+	if (UDisplayClusterCameraComponent* ViewCamera = GetViewPointCameraComponent())
 	{
 		// First try to get the eye offset distance from the projection policy
 		if (!ProjectionPolicy->GetStereoEyeOffsetDistance(this, InContextNum, StereoEyeOffsetDistance))
diff --git a/Engine/Plugins/Runtime/nDisplay/Source/DisplayCluster/Public/Render/Containers/IDisplayClusterRender_MeshComponent.h b/Engine/Plugins/Runtime/nDisplay/Source/DisplayCluster/Public/Render/Containers/IDisplayClusterRender_MeshComponent.h
--- a/Engine/Plugins/Runtime/nDisplay/Source/DisplayCluster/Public/Render/Containers/IDisplayClusterRender_MeshComponent.h
+++ b/Engine/Plugins/Runtime/nDisplay/Source/DisplayCluster/Public/Render/Containers/IDisplayClusterRender_MeshComponent.h
@@ -16,7 +16,7 @@ class IDisplayClusterRender_MeshComponentProxy;
 struct FStaticMeshLODResources;
 struct FProcMeshSection;
 
-class IDisplayClusterRender_MeshComponent
+class DISPLAYCLUSTER_API IDisplayClusterRender_MeshComponent
 {
 public:
 	virtual ~IDisplayClusterRender_MeshComponent() = default;
diff --git a/Engine/Plugins/Runtime/nDisplay/Source/DisplayCluster/Public/Render/Containers/IDisplayClusterRender_MeshComponentProxy.h b/Engine/Plugins/Runtime/nDisplay/Source/DisplayCluster/Public/Render/Containers/IDisplayClusterRender_MeshComponentProxy.h
--- a/Engine/Plugins/Runtime/nDisplay/Source/DisplayCluster/Public/Render/Containers/IDisplayClusterRender_MeshComponentProxy.h
+++ b/Engine/Plugins/Runtime/nDisplay/Source/DisplayCluster/Public/Render/Containers/IDisplayClusterRender_MeshComponentProxy.h
@@ -7,7 +7,7 @@
 #include "RHI.h"
 #include "RHICommandList.h"
 
-class IDisplayClusterRender_MeshComponentProxy
+class DISPLAYCLUSTER_API IDisplayClusterRender_MeshComponentProxy
 {
 public:
 	virtual ~IDisplayClusterRender_MeshComponentProxy() = default;
diff --git a/Engine/Plugins/Runtime/nDisplay/Source/DisplayCluster/Public/Render/Projection/IDisplayClusterProjectionPolicy.h b/Engine/Plugins/Runtime/nDisplay/Source/DisplayCluster/Public/Render/Projection/IDisplayClusterProjectionPolicy.h
--- a/Engine/Plugins/Runtime/nDisplay/Source/DisplayCluster/Public/Render/Projection/IDisplayClusterProjectionPolicy.h
+++ b/Engine/Plugins/Runtime/nDisplay/Source/DisplayCluster/Public/Render/Projection/IDisplayClusterProjectionPolicy.h
@@ -94,7 +94,7 @@ public:
 	}
 
 	// Return true, if camera projection visible for this viewport geometry
-	// ICVFX Perforamance : if camera frame not visible on this node, disable render for this camera
+	// ICVFX Performance : if camera frame not visible on this node, disable render for this camera
 	virtual bool IsCameraProjectionVisible(const FRotator& InViewRotation, const FVector& InViewLocation, const FMatrix& InProjectionMatrix)
 	{
 		return true;
@@ -121,6 +121,10 @@ public:
 		return false;
 	}
 
+	/** Projection policy can override PP */
+	virtual void OverridePostProcessSettings(IDisplayClusterViewport* InViewport)
+	{ }
+
 	 /** Get the distance from the eye to the viewpoint location
 	 *
 	 * @param InContextNum - eye context of this viewport
diff --git a/Engine/Plugins/Runtime/nDisplay/Source/DisplayCluster/Public/Render/Viewport/IDisplayClusterViewport.h b/Engine/Plugins/Runtime/nDisplay/Source/DisplayCluster/Public/Render/Viewport/IDisplayClusterViewport.h
--- a/Engine/Plugins/Runtime/nDisplay/Source/DisplayCluster/Public/Render/Viewport/IDisplayClusterViewport.h
+++ b/Engine/Plugins/Runtime/nDisplay/Source/DisplayCluster/Public/Render/Viewport/IDisplayClusterViewport.h
@@ -47,6 +47,9 @@ public:
 	 */
 	virtual bool SetupViewPoint(struct FMinimalViewInfo& InOutViewInfo) = 0;
 
+	/** Return view point camera component for this viewport. */
+	virtual class UDisplayClusterCameraComponent* GetViewPointCameraComponent() const = 0;
+
 	/** Get the distance from the eye to the viewpoint location.
 	 *
 	 * @param InContextNum - eye context of this viewport
@@ -59,6 +62,7 @@ public:
 
 	// Override postprocess settings for this viewport
 	virtual const IDisplayClusterViewport_CustomPostProcessSettings& GetViewport_CustomPostProcessSettings() const = 0;
+	virtual IDisplayClusterViewport_CustomPostProcessSettings& GetViewport_CustomPostProcessSettings() = 0;
 
 	// Setup scene view for rendering specified Context
 	virtual void SetupSceneView(uint32 ContextNum, class UWorld* World, class FSceneViewFamily& InViewFamily, FSceneView& InView) const = 0;
diff --git a/Engine/Plugins/Runtime/nDisplay/Source/DisplayCluster/Public/Render/Viewport/IDisplayClusterViewport_CustomPostProcessSettings.h b/Engine/Plugins/Runtime/nDisplay/Source/DisplayCluster/Public/Render/Viewport/IDisplayClusterViewport_CustomPostProcessSettings.h
--- a/Engine/Plugins/Runtime/nDisplay/Source/DisplayCluster/Public/Render/Viewport/IDisplayClusterViewport_CustomPostProcessSettings.h
+++ b/Engine/Plugins/Runtime/nDisplay/Source/DisplayCluster/Public/Render/Viewport/IDisplayClusterViewport_CustomPostProcessSettings.h
@@ -23,4 +23,10 @@ public:
 	// Override posproces, if defined
 	// * return true, if processed
 	virtual bool DoPostProcess(const ERenderPass InRenderPass, struct FPostProcessSettings* OutSettings, float* OutBlendWeight = nullptr) const = 0;
+
+	/** Add custom postprocess. */
+	virtual void AddCustomPostProcess(const ERenderPass InRenderPass, const FPostProcessSettings& InSettings, float BlendWeight, bool bSingleFrame) = 0;
+
+	/** remove custom postprocess. */
+	virtual void RemoveCustomPostProcess(const ERenderPass InRenderPass) = 0;
 };
