nDisplay: Several new public classes have been added
- IDisplayClusterRender_CachedTexture - Texture object created from memory data. These texture objects are also cached, and will be reused if their name matches.
- TDisplayClusterDataCache - A template class that implements a cache for storing TSharedPtr<> objects. Also this cache is tickable and can delete unused objects after a timeout.
- FDisplayClusterTickableGameObject - a tickable object to which an event can be subscribed

(this subtask is part of the In-Frustum Fit for MPCDI/2D/3D/mesh)

diff --git a/Engine/Plugins/Runtime/nDisplay/Source/DisplayCluster/Private/Misc/DisplayClusterTickableGameObject.cpp b/Engine/Plugins/Runtime/nDisplay/Source/DisplayCluster/Private/Misc/DisplayClusterTickableGameObject.cpp
new file mode 100644
--- /dev/null
+++ b/Engine/Plugins/Runtime/nDisplay/Source/DisplayCluster/Private/Misc/DisplayClusterTickableGameObject.cpp
@@ -0,0 +1,11 @@
+// Copyright Epic Games, Inc. All Rights Reserved.
+
+#include "Misc/DisplayClusterTickableGameObject.h"
+
+void FDisplayClusterTickableGameObject::Tick(float DeltaTime)
+{
+	if (TickEvent.IsBound())
+	{
+		TickEvent.Broadcast(DeltaTime);
+	}
+}
diff --git a/Engine/Plugins/Runtime/nDisplay/Source/DisplayCluster/Private/Render/Containers/DisplayClusterRender_Texture.cpp b/Engine/Plugins/Runtime/nDisplay/Source/DisplayCluster/Private/Render/Containers/DisplayClusterRender_Texture.cpp
new file mode 100644
--- /dev/null
+++ b/Engine/Plugins/Runtime/nDisplay/Source/DisplayCluster/Private/Render/Containers/DisplayClusterRender_Texture.cpp
@@ -0,0 +1,105 @@
+// Copyright Epic Games, Inc. All Rights Reserved.
+
+#include "Render/Containers/DisplayClusterRender_Texture.h"
+#include "RenderingThread.h"
+
+int32 GDisplayClusterRender_TextureCacheEnable = 1;
+static FAutoConsoleVariableRef CVarDisplayClusterRender_TextureCacheEnable(
+	TEXT("nDisplay.cache.Texture.enable"),
+	GDisplayClusterRender_TextureCacheEnable,
+	TEXT("Enable the use of a named texture cache.\n"),
+	ECVF_Default
+);
+
+int32 GDisplayClusterRender_TextureCacheTimeOutInFrames = 5 * 60 * 60; // Timeout is 5 minutes (for 60 frames per second)
+static FAutoConsoleVariableRef CVarDisplayClusterRenderCached_TextureCacheTimeOutInFrames(
+	TEXT("nDisplay.cache.Texture.TimeOutInFrames"),
+	GDisplayClusterRender_TextureCacheTimeOutInFrames,
+	TEXT("The timeout value in frames for the named texture cache.\n")
+	TEXT("-1 - disable timeout.\n"),
+	ECVF_Default
+);
+
+//---------------------------------------------------------------------------------------
+// FDisplayClusterRender_Texture
+//---------------------------------------------------------------------------------------
+FDisplayClusterRender_Texture::FDisplayClusterRender_Texture(const FString& InUniqueTextureName)
+	: UniqueTextureName(InUniqueTextureName)
+{ }
+
+FDisplayClusterRender_Texture::~FDisplayClusterRender_Texture()
+{
+	if (PrivateResource.IsValid())
+	{
+		PrivateResource->ReleaseRenderResource();
+		PrivateResource.Reset();
+	}
+
+	if (PrivateResourceRenderThread.IsValid())
+	{
+		PrivateResourceRenderThread->ReleaseRenderResource();
+		PrivateResourceRenderThread.Reset();
+	}
+}
+
+void FDisplayClusterRender_Texture::CreateTexture(const void* InTextureData, const uint32 InComponentDepth, const uint32 InBytesPerComponent, uint32_t InWidth, uint32_t InHeight, bool bInHasCPUAccess)
+{
+	// Release the existing texture resource.
+	ReleaseResource();
+
+	//Dedicated servers have no texture internals
+	if (FApp::CanEverRender())
+	{
+		// Create a new texture resource.
+		TSharedPtr<FDisplayClusterRender_TextureResource, ESPMode::ThreadSafe> NewResource = MakeShared<FDisplayClusterRender_TextureResource, ESPMode::ThreadSafe>(InTextureData, InComponentDepth, InBytesPerComponent, InWidth, InHeight, bInHasCPUAccess);
+		if (NewResource.IsValid())
+		{
+			NewResource->InitializeRenderResource();
+		}
+
+		SetResource(NewResource);
+	}
+}
+
+const TSharedPtr<FDisplayClusterRender_TextureResource>& FDisplayClusterRender_Texture::GetResource() const
+{
+	if (IsInActualRenderingThread() || IsInRHIThread())
+	{
+		return PrivateResourceRenderThread;
+	}
+	return PrivateResource;
+}
+
+void FDisplayClusterRender_Texture::SetResource(const TSharedPtr<FDisplayClusterRender_TextureResource>& InResource)
+{
+	check(!IsInActualRenderingThread() && !IsInRHIThread());
+
+	// Each PrivateResource value must be updated in it's own thread because any
+	// rendering code trying to access the Resource
+	// crash if it suddenly sees nullptr or a new resource that has not had it's InitRHI called.
+
+	PrivateResource = InResource;
+	ENQUEUE_RENDER_COMMAND(DisplayCluster_SetResourceRenderThread)([This = SharedThis(this), InResource](FRHICommandListImmediate& RHICmdList)
+	{
+		This->PrivateResourceRenderThread = InResource;
+	});
+}
+
+void FDisplayClusterRender_Texture::ReleaseResource()
+{
+	if (PrivateResource)
+	{
+		// Free the resource.
+		SetResource(nullptr);
+	}
+}
+
+int32 FDisplayClusterRender_Texture::GetDataCacheTimeOutInFrames()
+{
+	return FMath::Max(0, GDisplayClusterRender_TextureCacheTimeOutInFrames);
+}
+
+bool FDisplayClusterRender_Texture::IsDataCacheEnabled()
+{
+	return GDisplayClusterRender_TextureCacheEnable != 0;
+}
diff --git a/Engine/Plugins/Runtime/nDisplay/Source/DisplayCluster/Private/Render/Containers/DisplayClusterRender_Texture.h b/Engine/Plugins/Runtime/nDisplay/Source/DisplayCluster/Private/Render/Containers/DisplayClusterRender_Texture.h
new file mode 100644
--- /dev/null
+++ b/Engine/Plugins/Runtime/nDisplay/Source/DisplayCluster/Private/Render/Containers/DisplayClusterRender_Texture.h
@@ -0,0 +1,110 @@
+// Copyright Epic Games, Inc. All Rights Reserved.
+
+#pragma once
+#include "Render/Containers/IDisplayClusterRender_Texture.h"
+#include "Render/Containers/DisplayClusterRender_TextureResource.h"
+
+/**
+ * Implementation for runtime texture object
+ * Can be cached with TDisplayClusterDataCache: Multiple objects can get the same texture by name
+ */
+class FDisplayClusterRender_Texture
+	: public IDisplayClusterRender_Texture
+	, public TSharedFromThis<FDisplayClusterRender_Texture, ESPMode::ThreadSafe>
+{
+public:
+	FDisplayClusterRender_Texture(const FString& InUniqueTextureName);
+	virtual ~FDisplayClusterRender_Texture();
+
+public:
+	virtual TSharedPtr<IDisplayClusterRender_Texture, ESPMode::ThreadSafe> ToSharedPtr() override
+	{
+		return AsShared();
+	}
+
+	virtual TSharedPtr<const IDisplayClusterRender_Texture, ESPMode::ThreadSafe> ToSharedPtr() const override
+	{
+		return AsShared();
+	}
+
+	virtual void CreateTexture(const void* InTextureData, const uint32 InComponentDepth, const uint32 InBitDepth, uint32_t InWidth, uint32_t InHeight, bool bInHasCPUAccess=false) override;
+
+	virtual bool IsEnabled() const override
+	{
+		const TSharedPtr<FDisplayClusterRender_TextureResource, ESPMode::ThreadSafe>& TextureResource = GetResource();
+		if (TextureResource.IsValid())
+		{
+			return TextureResource->GetWidth() > 0 && TextureResource->GetHeight() > 0;
+		}
+
+		return false;
+	}
+
+	virtual void* GetData() const override
+	{
+		const TSharedPtr<FDisplayClusterRender_TextureResource, ESPMode::ThreadSafe>& TextureResource = GetResource();
+		return TextureResource ? TextureResource->GetTextureData() : nullptr;
+	}
+
+	virtual uint32_t GetWidth() const override
+	{
+		const TSharedPtr<FDisplayClusterRender_TextureResource, ESPMode::ThreadSafe>& TextureResource = GetResource();
+		return TextureResource.IsValid() ? TextureResource->GetWidth() : 0;
+	}
+
+	virtual uint32_t GetHeight() const override
+	{
+		const TSharedPtr<FDisplayClusterRender_TextureResource, ESPMode::ThreadSafe>& TextureResource = GetResource();
+		return TextureResource.IsValid() ? TextureResource->GetHeight() : 0;
+	}
+
+	virtual EPixelFormat GetPixelFormat() const override
+	{
+		const TSharedPtr<FDisplayClusterRender_TextureResource, ESPMode::ThreadSafe>& TextureResource = GetResource();
+		return TextureResource.IsValid() ? TextureResource->GetPixelFormat() : PF_Unknown;
+	}
+
+	virtual FRHITexture* GetRHITexture() const override
+	{
+		const TSharedPtr<FDisplayClusterRender_TextureResource, ESPMode::ThreadSafe>& TextureResource = GetResource();
+		return TextureResource.IsValid() ? TextureResource->TextureRHI : nullptr;
+	}
+
+	//~ Begin TDisplayClusterDataCache
+	/** Return DataCache timeout in frames. */
+	static int32 GetDataCacheTimeOutInFrames();
+
+	/** Return true if DataCache is enabled. */
+	static bool IsDataCacheEnabled();
+
+	/** Returns the unique name of this texture for DataCache. */
+	inline const FString& GetDataCacheName() const
+	{
+		return UniqueTextureName;
+	}
+	// ~~ End TDisplayClusterDataCache
+
+protected:
+	/** Set texture's resource, can be NULL */
+	void SetResource(const TSharedPtr<FDisplayClusterRender_TextureResource, ESPMode::ThreadSafe>& InResource);
+
+	/** Get the texture's resource, can be NULL */
+	const TSharedPtr<FDisplayClusterRender_TextureResource, ESPMode::ThreadSafe>& GetResource() const;
+
+	/**
+	 * Resets the resource for the texture.
+	 */
+	void ReleaseResource();
+
+private:
+	// The texture's resource, can be NULL
+	TSharedPtr<FDisplayClusterRender_TextureResource, ESPMode::ThreadSafe> PrivateResource;
+
+	// Value updated and returned by the render-thread to allow
+	// fenceless update from the game-thread without causing
+	// potential crash in the render thread.
+	TSharedPtr<FDisplayClusterRender_TextureResource, ESPMode::ThreadSafe> PrivateResourceRenderThread;
+
+	// Unique texture name
+	FString UniqueTextureName;
+};
diff --git a/Engine/Plugins/Runtime/nDisplay/Source/DisplayCluster/Private/Render/Containers/DisplayClusterRender_TextureResource.cpp b/Engine/Plugins/Runtime/nDisplay/Source/DisplayCluster/Private/Render/Containers/DisplayClusterRender_TextureResource.cpp
new file mode 100644
--- /dev/null
+++ b/Engine/Plugins/Runtime/nDisplay/Source/DisplayCluster/Private/Render/Containers/DisplayClusterRender_TextureResource.cpp
@@ -0,0 +1,165 @@
+// Copyright Epic Games, Inc. All Rights Reserved.
+
+#include "Render/Containers/DisplayClusterRender_TextureResource.h"
+#include "Containers/ResourceArray.h"
+#include "RenderingThread.h"
+
+/**
+ * Container for texture data
+ */
+class FDisplayClusterResourceBulkData : public FResourceBulkDataInterface
+{
+public:
+	FDisplayClusterResourceBulkData(const void* InData, uint32_t InDataSize)
+		: Data(InData)
+		, DataSize(InDataSize)
+	{ }
+
+public:
+	virtual const void* GetResourceBulkData() const
+	{
+		return Data;
+	}
+
+	virtual uint32 GetResourceBulkDataSize() const
+	{
+		return DataSize;
+	}
+
+	virtual void Discard()
+	{ }
+
+private:
+	const void* Data;
+	uint32_t    DataSize;
+};
+
+//---------------------------------------------------------------------------------------
+// FDisplayClusterRender_TextureResource
+//---------------------------------------------------------------------------------------
+FDisplayClusterRender_TextureResource::FDisplayClusterRender_TextureResource(const void* InTextureData, const uint32 InComponentDepth, const uint32 InBitDepth, uint32_t InWidth, uint32_t InHeight, bool bInHasCPUAccess)
+{
+	check(InTextureData);
+	check(InWidth > 0);
+	check(InHeight > 0);
+
+	check(InComponentDepth > 0 && InComponentDepth <= 4);
+	check(InBitDepth == 8 || InBitDepth == 16 || InBitDepth == 32);
+
+	bHasCPUAccess = bInHasCPUAccess;
+
+	static const EPixelFormat DestPixelFormatMapping[4][4] =
+	{
+		{ PF_G8,       PF_G16,     PF_Unknown, PF_Unknown },
+		{ PF_R8G8,     PF_G16R16,  PF_Unknown, PF_Unknown },
+		{ PF_R8G8B8A8, PF_Unknown, PF_Unknown, PF_Unknown },
+		{ PF_R8G8B8A8, PF_Unknown, PF_Unknown, PF_A32B32G32R32F}
+	};
+
+	const EPixelFormat DestPixelFormat = DestPixelFormatMapping[InComponentDepth - 1][(InBitDepth>>3) - 1];
+	check(DestPixelFormat != PF_Unknown);
+
+	// Source data not equal to dest
+	const int32 DestNumComponents = GPixelFormats[DestPixelFormat].NumComponents;
+
+	const uint32 DataSize = CalculateImageBytes(InWidth, InHeight, 1, DestPixelFormat);
+	TextureData = FMemory::Malloc(DataSize);
+
+	const size_t SrcPixelSize = InComponentDepth * (InBitDepth >> 3);
+	const size_t DestPixelSize = DestNumComponents * (InBitDepth >> 3);
+
+	if (SrcPixelSize == DestPixelSize)
+	{
+		FMemory::Memcpy(TextureData, InTextureData, DataSize);
+	}
+	else
+	{
+		check(SrcPixelSize < DestPixelSize);
+
+		FMemory::Memset(TextureData, 0, DataSize);
+
+		const uint32 TotalPixels = InWidth * InHeight;
+		uint8* DestPixel = (uint8*)TextureData;
+		uint8* SrcPixel = (uint8*)InTextureData;
+
+		// adopt source data to dest texture format:
+		for (uint32 PixelIndex = 0; PixelIndex < TotalPixels; PixelIndex++)
+		{
+			FMemory::Memcpy(DestPixel, SrcPixel, SrcPixelSize);
+
+			DestPixel += DestPixelSize;
+			SrcPixel += SrcPixelSize;
+		}
+	}
+
+	// Create texture from data:
+	PixelFormat = DestPixelFormat;
+	Width = InWidth;
+	Height = InHeight;
+}
+
+FDisplayClusterRender_TextureResource::~FDisplayClusterRender_TextureResource()
+{
+	ReleaseTextureData();
+}
+
+
+void FDisplayClusterRender_TextureResource::InitializeRenderResource()
+{
+	// Initialize on render thread
+	ENQUEUE_RENDER_COMMAND(DisplayClusterRenderTextureResource_Initialize)([This = SharedThis(this)](FRHICommandListImmediate& RHICmdList)
+	{
+		This->InitResource();
+	});
+}
+
+void FDisplayClusterRender_TextureResource::ReleaseRenderResource()
+{
+	// Release on render thread
+	ENQUEUE_RENDER_COMMAND(DisplayClusterRenderTextureResource_Release)([This = SharedThis(this)](FRHICommandListImmediate& RHICmdList)
+	{
+		This->ReleaseResource();
+	});
+}
+
+void FDisplayClusterRender_TextureResource::ReleaseTextureData()
+{
+	if (TextureData != nullptr)
+	{
+		FMemory::Free(TextureData);
+		TextureData = nullptr;
+
+		Width = 0;
+		Height = 0;
+		PixelFormat = PF_Unknown;
+		bHasCPUAccess = false;
+	}
+}
+
+void FDisplayClusterRender_TextureResource::InitRHI()
+{
+	check(IsInRenderingThread());
+
+	if (TextureData == nullptr)
+	{
+		// Texture data is needed
+		return;
+	}
+
+	const uint32 DataSize = CalculateImageBytes(Width, Height, 1, PixelFormat);
+	FDisplayClusterResourceBulkData BulkDataInterface(TextureData, DataSize);
+
+	// Create a texture from the data in memory.
+	const FRHITextureCreateDesc Desc =
+		FRHITextureCreateDesc::Create2D(TEXT("DisplayClusterRender_TextureResource"), Width, Height, PixelFormat)
+		.SetFlags(ETextureCreateFlags::ShaderResource)
+		.SetBulkData(&BulkDataInterface);
+
+	TextureRHI = RHICreateTexture(Desc);
+
+	// CPU access not required, release from memory
+	if (bHasCPUAccess == false)
+	{
+		ReleaseTextureData();
+	}
+}
diff --git a/Engine/Plugins/Runtime/nDisplay/Source/DisplayCluster/Private/Render/Containers/DisplayClusterRender_TextureResource.h b/Engine/Plugins/Runtime/nDisplay/Source/DisplayCluster/Private/Render/Containers/DisplayClusterRender_TextureResource.h
new file mode 100644
--- /dev/null
+++ b/Engine/Plugins/Runtime/nDisplay/Source/DisplayCluster/Private/Render/Containers/DisplayClusterRender_TextureResource.h
@@ -0,0 +1,72 @@
+// Copyright Epic Games, Inc. All Rights Reserved.
+
+#pragma once
+#include "CoreMinimal.h"
+#include "TextureResource.h"
+#include "Templates/SharedPointer.h"
+
+/**
+ * Creating a texture resource from memory data
+ * Also, this texture data will be stored in memory if bHasCPUAccess is true.
+ */
+class FDisplayClusterRender_TextureResource
+	: public FTextureResource
+	, public TSharedFromThis<FDisplayClusterRender_TextureResource, ESPMode::ThreadSafe>
+{
+public:
+	FDisplayClusterRender_TextureResource(const void* InTextureData, const uint32 InComponentDepth, const uint32 InBitDepth, uint32_t InWidth, uint32_t InHeight, bool bInHasCPUAccess);
+	virtual ~FDisplayClusterRender_TextureResource();
+
+	//~FTextureResource
+	virtual void InitRHI() override;
+	//~~ FTextureResource
+
+	/** Release texture data from memory. */
+	void ReleaseTextureData();
+
+	/** Initializes the render resource. */
+	void InitializeRenderResource();
+
+	/** Prepares the resource for deletion. */
+	void ReleaseRenderResource();
+
+	/** Get ptr to resource data in memory. The variable bHasCPUAccess must be true. */
+	void* GetTextureData() const
+	{
+		return TextureData;
+	}
+
+	/** Get texture width. */
+	uint32_t GetWidth() const
+	{
+		return Width;
+	}
+
+	/** Get texture height. */
+	uint32_t GetHeight() const
+	{
+		return Height;
+	}
+
+	/** Get texture pixel format. */
+	EPixelFormat GetPixelFormat() const
+	{
+		return PixelFormat;
+	}
+
+private:
+	// RAW data of the texture. Mips not supported now
+	void* TextureData;
+	
+	// Width of the texture.
+	uint32_t Width;
+	
+	// Height of the texture.
+	uint32_t Height;
+	
+	// Format in which mip data is stored.
+	EPixelFormat PixelFormat;
+	
+	// has CPU access to texture data, dont release from memory
+	bool bHasCPUAccess;
+};
diff --git a/Engine/Plugins/Runtime/nDisplay/Source/DisplayCluster/Private/Render/DisplayClusterRenderManager.cpp b/Engine/Plugins/Runtime/nDisplay/Source/DisplayCluster/Private/Render/DisplayClusterRenderManager.cpp
--- a/Engine/Plugins/Runtime/nDisplay/Source/DisplayCluster/Private/Render/DisplayClusterRenderManager.cpp
+++ b/Engine/Plugins/Runtime/nDisplay/Source/DisplayCluster/Private/Render/DisplayClusterRenderManager.cpp
@@ -31,6 +31,7 @@
 #include "Render/Synchronization/IDisplayClusterRenderSyncPolicyFactory.h"
 
 #include "Render/Containers/DisplayClusterRender_MeshComponent.h"
+#include "Render/Containers/DisplayClusterRender_Texture.h"
 
 #include "Render/Presentation/DisplayClusterPresentationNative.h"
 
@@ -64,7 +65,36 @@ static TAutoConsoleVariable<int32> CVarSyncDiagnosticsVBlankMonitoring(
 	ECVF_ReadOnly
 );
 
+#include "Misc/DisplayClusterDataCache.h"
 
+/**
+ * The cache for DC texture objects. (Singleton)
+ * Allows you to reuse textures with the same unique name.
+ */
+class FDisplayClusterRenderTextureCache
+	: public TDisplayClusterDataCache<FDisplayClusterRender_Texture>
+{
+public:
+	static TSharedPtr<FDisplayClusterRender_Texture, ESPMode::ThreadSafe> GetOrCreateRenderTexture(const FString& InTextureName)
+	{
+		static FDisplayClusterRenderTextureCache TextureCacheSingleton;
+
+		const FString UniqueName = HashString(InTextureName);
+
+		TSharedPtr<FDisplayClusterRender_Texture, ESPMode::ThreadSafe> TextureRef = TextureCacheSingleton.Find(UniqueName);
+		if (!TextureRef.IsValid())
+		{
+			TextureRef = MakeShared<FDisplayClusterRender_Texture, ESPMode::ThreadSafe>(UniqueName);
+			TextureCacheSingleton.Add(TextureRef);
+		}
+
+		return TextureRef;
+	}
+};
+
+//---------------------------------------------------
+// FDisplayClusterRenderManager
+//---------------------------------------------------
 FDisplayClusterRenderManager::FDisplayClusterRenderManager()
 {
 	// Instantiate and register internal render device factory
@@ -504,6 +534,11 @@ TSharedPtr<IDisplayClusterRender_MeshComponent, ESPMode::ThreadSafe> FDisplayClu
 	return MakeShared<FDisplayClusterRender_MeshComponent, ESPMode::ThreadSafe>();
 }
 
+TSharedPtr<IDisplayClusterRender_Texture, ESPMode::ThreadSafe> FDisplayClusterRenderManager::GetOrCreateCachedTexture(const FString& InUniqueTextureName) const
+{
+	return FDisplayClusterRenderTextureCache::GetOrCreateRenderTexture(InUniqueTextureName);
+}
+
 IDisplayClusterViewportManager* FDisplayClusterRenderManager::GetViewportManager() const
 {
 	ADisplayClusterRootActor* RootActor = GDisplayCluster->GetGameMgr()->GetRootActor();
diff --git a/Engine/Plugins/Runtime/nDisplay/Source/DisplayCluster/Private/Render/DisplayClusterRenderManager.h b/Engine/Plugins/Runtime/nDisplay/Source/DisplayCluster/Private/Render/DisplayClusterRenderManager.h
--- a/Engine/Plugins/Runtime/nDisplay/Source/DisplayCluster/Private/Render/DisplayClusterRenderManager.h
+++ b/Engine/Plugins/Runtime/nDisplay/Source/DisplayCluster/Private/Render/DisplayClusterRenderManager.h
@@ -63,6 +63,7 @@ public:
 	virtual TSharedPtr<IDisplayClusterPostProcessFactory> GetPostProcessFactory(const FString& InPostProcessType) override;
 	virtual void GetRegisteredPostProcess(TArray<FString>& OutPostProcessIDs) const override;
 	virtual TSharedPtr<IDisplayClusterRender_MeshComponent, ESPMode::ThreadSafe> CreateMeshComponent() const override;
+	virtual TSharedPtr<IDisplayClusterRender_Texture, ESPMode::ThreadSafe> GetOrCreateCachedTexture(const FString& InUniqueTextureName) const override;
 
 	virtual TSharedPtr<IDisplayClusterVblankMonitor, ESPMode::ThreadSafe> GetVblankMonitor() override
 	{
diff --git a/Engine/Plugins/Runtime/nDisplay/Source/DisplayCluster/Public/Misc/DisplayClusterDataCache.h b/Engine/Plugins/Runtime/nDisplay/Source/DisplayCluster/Public/Misc/DisplayClusterDataCache.h
new file mode 100644
--- /dev/null
+++ b/Engine/Plugins/Runtime/nDisplay/Source/DisplayCluster/Public/Misc/DisplayClusterDataCache.h
@@ -0,0 +1,175 @@
+// Copyright Epic Games, Inc. All Rights Reserved.
+
+#pragma once
+#include "CoreMinimal.h"
+
+#include "Misc/DisplayClusterTickableGameObject.h"
+#include "Misc/SecureHash.h"
+
+/**
+ * Cache template for named data with automatic deletion of unused items after timeout.
+ */
+template<typename InDataType>
+class TDisplayClusterDataCache
+{
+public:
+	/** Data type of stored objects
+	 * 
+	 * The following methods must be implemented in DataType:
+	 * 
+	 *   // Return DataCache timeout in frames. A negative value disables the timeout.
+	 *   static int32 GetDataCacheTimeOutInFrames();
+	 * 
+	 *   // Return true if DataCache is enabled.
+	 *   static bool IsDataCacheEnabled();
+	 *
+	 *   // Returns the unique name of this texture for DataCache.
+	 *   inline const FString& GetDataCacheName() const;
+	 */
+	typedef InDataType DataType;
+
+	virtual ~TDisplayClusterDataCache()
+	{
+		Release();
+	}
+
+	/** Helper to perform the very common case of hashing an FString into a hex representation. */
+	static inline FString HashString(const FString& InString)
+	{
+		return FMD5::HashAnsiString(*InString.ToLower());
+	}
+
+	/** Retrieve an existing object by a unique name. */
+	TSharedPtr<DataType, ESPMode::ThreadSafe> Find(const FString& InUniqueDataName)
+	{
+		if (DataType::IsDataCacheEnabled())
+		{
+			if (FCachedObject* ExistCachedObject = CachedObjects.FindByPredicate([InUniqueDataName](const FCachedObject& CachedObjectIt)
+			{
+				return CachedObjectIt.IsNameEqual(InUniqueDataName);
+			}))
+			{
+				return ExistCachedObject->DataRef;
+			}
+		}
+
+		return nullptr;
+	}
+
+	/** Register a new object. */
+	void Add(const TSharedPtr<DataType, ESPMode::ThreadSafe>& InDataRef)
+	{
+		if (DataType::IsDataCacheEnabled())
+		{
+			CachedObjects.Add(FCachedObject(InDataRef));
+
+			UpdateTickableGameObject();
+		}
+	}
+
+protected:
+	/**
+	 * Release the cache
+	 */
+	void Release()
+	{
+		CachedObjects.Empty();
+		UpdateTickableGameObject();
+	}
+	 
+	 /**
+	 * When the number of DataRef references drops to 1, these items will be removed after a timeout.
+	 */
+	void Tick(float DeltaTime)
+	{
+		// If the cache is disabled, clear it
+		if (!DataType::IsDataCacheEnabled())
+		{
+			Release();
+
+			return;
+		}
+
+		const int32 TimeOutInFrames = DataType::GetDataCacheTimeOutInFrames();
+
+		// Updating the time for all objects in use
+		for (int32 Index = 0; Index < CachedObjects.Num(); Index++)
+		{
+			if (CachedObjects[Index].DataRef.GetSharedReferenceCount() >1 || TimeOutInFrames < 0)
+			{
+				// When this object used, reset this counter
+				CachedObjects[Index].FramesInCachedState = 0;
+			}
+			else
+			{
+				// After the data ref counter remains 1, this counter will be incremented every frame
+				CachedObjects[Index].FramesInCachedState++;
+
+				// Remove after timout
+				if (CachedObjects[Index].FramesInCachedState > TimeOutInFrames)
+				{
+					// Deleting an unused item after time has elapsed
+					CachedObjects[Index].DataRef.Reset();
+
+					// Delete the current element, and reuse the current array index,
+					// because the next element will move to the index of the current
+					CachedObjects.RemoveAt(Index--);
+				}
+			}
+		}
+
+		UpdateTickableGameObject();
+	}
+
+	/** Create or remove tickable game object. */
+	inline void UpdateTickableGameObject()
+	{
+		if (CachedObjects.IsEmpty())
+		{
+			// Unregister tick event
+			if (TickHandle.IsValid())
+			{
+				TickableGameObject->OnTick().Remove(TickHandle);
+			}
+
+			TickableGameObject.Reset();
+		}
+		else
+		{
+			if (!TickableGameObject.IsValid())
+			{
+				TickableGameObject = MakeUnique<FDisplayClusterTickableGameObject >();
+				TickHandle = TickableGameObject->OnTick().AddRaw(this, &TDisplayClusterDataCache<DataType>::Tick);
+			}
+		}
+	}
+
+private:
+	struct FCachedObject
+	{
+		FCachedObject(const TSharedPtr<DataType, ESPMode::ThreadSafe>& InDataRef)
+			: DataRef(InDataRef)
+		{ }
+
+		inline bool IsNameEqual(const FString& InName) const
+		{
+			return DataRef.IsValid() && DataRef->GetDataCacheName() == InName;
+		}
+
+		// reference to data
+		TSharedPtr<DataType, ESPMode::ThreadSafe> DataRef;
+
+		// After the data ref counter remains 1, this counter will be incremented every frame
+		int32 FramesInCachedState = 0;
+	};
+
+	// Cached objects map
+	TArray<FCachedObject> CachedObjects;
+
+	// When CachedObjects is not empty, this ticking object will be created.
+	// Also, this object will be deleted when CachedObjects becomes empty.
+	TUniquePtr<FDisplayClusterTickableGameObject> TickableGameObject;
+
+	// Event delegate container of the Tickable object
+	FDelegateHandle TickHandle;
+};
diff --git a/Engine/Plugins/Runtime/nDisplay/Source/DisplayCluster/Public/Misc/DisplayClusterTickableGameObject.h b/Engine/Plugins/Runtime/nDisplay/Source/DisplayCluster/Public/Misc/DisplayClusterTickableGameObject.h
new file mode 100644
--- /dev/null
+++ b/Engine/Plugins/Runtime/nDisplay/Source/DisplayCluster/Public/Misc/DisplayClusterTickableGameObject.h
@@ -0,0 +1,44 @@
+// Copyright Epic Games, Inc. All Rights Reserved.
+
+#pragma once
+
+#include "CoreMinimal.h"
+#include "Tickable.h"
+
+/**
+ * Tickable game object: has an OnTick() callback to start all subscribers
+ */
+class DISPLAYCLUSTER_API FDisplayClusterTickableGameObject : public FTickableGameObject
+{
+public:
+	FDisplayClusterTickableGameObject() = default;
+	virtual ~FDisplayClusterTickableGameObject() = default;
+
+public:
+	/** Callback on tick **/
+	DECLARE_EVENT_OneParam(FDisplayClusterTickableGameObject, FTickEvent, float);
+	FTickEvent& OnTick()
+	{
+		return TickEvent;
+	}
+
+	virtual bool IsTickableInEditor() const override
+	{
+		return true;
+	}
+
+	virtual void Tick(float DeltaTime) override;
+
+	virtual ETickableTickType GetTickableTickType() const override
+	{
+		return ETickableTickType::Always;
+	}
+
+	virtual TStatId GetStatId() const override
+	{
+		RETURN_QUICK_DECLARE_CYCLE_STAT(FDisplayClusterTickableGameObject, STATGROUP_Tickables);
+	}
+
+private:
+	FTickEvent TickEvent;
+};
diff --git a/Engine/Plugins/Runtime/nDisplay/Source/DisplayCluster/Public/Render/Containers/IDisplayClusterRender_Texture.h b/Engine/Plugins/Runtime/nDisplay/Source/DisplayCluster/Public/Render/Containers/IDisplayClusterRender_Texture.h
new file mode 100644
--- /dev/null
+++ b/Engine/Plugins/Runtime/nDisplay/Source/DisplayCluster/Public/Render/Containers/IDisplayClusterRender_Texture.h
@@ -0,0 +1,55 @@
+// Copyright Epic Games, Inc. All Rights Reserved.
+
+#pragma once
+
+#include "CoreMinimal.h"
+
+#include "RHI.h"
+#include "RHIResources.h"
+#include "RenderResource.h"
+
+/**
+ * A runtime texture that is created based on data in memory 
+  */
+class DISPLAYCLUSTER_API IDisplayClusterRender_Texture
+{
+public:
+	virtual ~IDisplayClusterRender_Texture() = default;
+
+public:
+	/** Get TSharePtr from this. */
+	virtual TSharedPtr<IDisplayClusterRender_Texture, ESPMode::ThreadSafe> ToSharedPtr() = 0;
+	virtual TSharedPtr<const IDisplayClusterRender_Texture, ESPMode::ThreadSafe> ToSharedPtr() const = 0;
+
+	/**
+	* Create texture from memory data
+	*
+	* @param InTextureData    - Pointer to texture data in memory
+	* @param InComponentDepth - Components per pixel
+	* @param InBitDepth       - Bit per component
+	* @param InWidth          - Texture width in pixels
+	* @param InHeight         - Texture height in pixels
+	* @param bInHasCPUAccess  - has CPU access to texture data, dont release from memory
+	*
+	* @return - true, if the texture resource was created successfully
+	*/
+	virtual void CreateTexture(const void* InTextureData, const uint32 InComponentDepth, const uint32 InBitDepth, uint32_t InWidth, uint32_t InHeight, bool bInHasCPUAccess = false) = 0;
+
+	/** True if the texture resource is initialized. */
+	virtual bool IsEnabled() const = 0;
+
+	/** Get ptr to resource data in memory (The texture must be created with bInHasCPUAccess set to true). */
+	virtual void* GetData() const = 0;
+
+	/** Get texture width. */
+	virtual uint32_t GetWidth() const = 0;
+
+	/** Get texture height. */
+	virtual uint32_t GetHeight() const = 0;
+	
+	/** Get texture pixel format. */
+	virtual EPixelFormat GetPixelFormat() const = 0;
+
+	/** Get the RHI resource. */
+	virtual FRHITexture* GetRHITexture() const = 0;
+};
diff --git a/Engine/Plugins/Runtime/nDisplay/Source/DisplayCluster/Public/Render/IDisplayClusterRenderManager.h b/Engine/Plugins/Runtime/nDisplay/Source/DisplayCluster/Public/Render/IDisplayClusterRenderManager.h
--- a/Engine/Plugins/Runtime/nDisplay/Source/DisplayCluster/Public/Render/IDisplayClusterRenderManager.h
+++ b/Engine/Plugins/Runtime/nDisplay/Source/DisplayCluster/Public/Render/IDisplayClusterRenderManager.h
@@ -8,6 +8,7 @@ class IDisplayClusterRenderDevice;
 class IDisplayClusterPostProcess;
 class IDisplayClusterPostProcessFactory;
 class IDisplayClusterRender_MeshComponent;
+class IDisplayClusterRender_Texture;
 class IDisplayClusterProjectionPolicy;
 class IDisplayClusterProjectionPolicyFactory;
 class IDisplayClusterRenderDeviceFactory;
@@ -177,4 +178,9 @@ public:
 	* @return - new mesh component object
 	*/
 	virtual TSharedPtr<IDisplayClusterRender_MeshComponent, ESPMode::ThreadSafe> CreateMeshComponent() const = 0;
+
+	/**
+	 * Get or create a cached texture with unique name
+	 */
+	virtual TSharedPtr<IDisplayClusterRender_Texture, ESPMode::ThreadSafe> GetOrCreateCachedTexture(const FString& InUniqueTextureName) const = 0;
 };
