nDisplay: Fixed possibility of keeping an external map texture reference.

Also better handled the case of reintancing a deleted component from the undo buffer, where it wouldn't enable ticking in such case.

diff --git a/Engine/Plugins/Runtime/nDisplay/Source/DisplayCluster/Private/Game/EngineClasses/Scene/DisplayClusterPreviewShareComponent.cpp b/Engine/Plugins/Runtime/nDisplay/Source/DisplayCluster/Private/Game/EngineClasses/Scene/DisplayClusterPreviewShareComponent.cpp
--- a/Engine/Plugins/Runtime/nDisplay/Source/DisplayCluster/Private/Game/EngineClasses/Scene/DisplayClusterPreviewShareComponent.cpp
+++ b/Engine/Plugins/Runtime/nDisplay/Source/DisplayCluster/Private/Game/EngineClasses/Scene/DisplayClusterPreviewShareComponent.cpp
@@ -60,6 +60,58 @@ namespace UE::DisplayClusterPreviewShare
 
 		return *ViewportPtr;
 	}
+
+	void RemoveExternalMapReferences(ADisplayClusterRootActor* RootActor)
+	{
+		// In all cases, we should make sure we never restore in illegal external reference,
+		// because if we do, it will be very hard for the user to fix this and save the map, since these
+		// references are not exposed in the details panel of the instance.
+
+		if (!RootActor)
+		{
+			return;
+		}
+
+		const UPackage* RootActorPackage = RootActor->GetPackage();
+
+		if (!RootActorPackage)
+		{
+			return;
+		}
+
+		for (UActorComponent* ActorComponent : RootActor->K2_GetComponentsByClass(UDisplayClusterPreviewComponent::StaticClass()))
+		{
+			UDisplayClusterPreviewComponent* PreviewComponent = Cast<UDisplayClusterPreviewComponent>(ActorComponent);
+
+			if (!PreviewComponent)
+			{
+				continue;
+			}
+
+			UDisplayClusterConfigurationViewport* Viewport = GetViewportFromDCRA(RootActor, PreviewComponent->GetClusterNodeId(), PreviewComponent->GetViewportId());
+
+			if (!Viewport)
+			{
+				continue;
+			}
+
+			if (Viewport->RenderSettings.Replace.SourceTexture)
+			{
+				const UPackage* TexturePackage = Viewport->RenderSettings.Replace.SourceTexture->GetPackage();
+
+				if (TexturePackage && TexturePackage->ContainsMap() && (TexturePackage != RootActorPackage))
+				{
+					UE_LOG(LogDisplayClusterGame, Warning,
+						TEXT("Removing unexpected external reference to texture package '%s', a map different to the wwning RootActor package '%s'"),
+						*TexturePackage->GetName(),
+						*RootActorPackage->GetName()
+					);
+
+					Viewport->RenderSettings.Replace.SourceTexture = nullptr;
+				}
+			}
+		}
+	}
 }
 
 #endif // WITH_EDITOR
@@ -226,9 +278,6 @@ void UDisplayClusterPreviewShareComponent::RestoreRootActorOriginalSettings()
 	
 	if (!RootActor)
 	{
-		OriginalSourceTextures.Empty();
-		OriginalTextureReplaces.Empty();
-
 		return;
 	}
 
@@ -260,9 +309,19 @@ void UDisplayClusterPreviewShareComponent::RestoreRootActorOriginalSettings()
 			Viewport->RenderSettings.Replace.bAllowReplace = *BoolPtr;
 		}
 
-		if (TObjectPtr<UTexture>* TexturePtr = OriginalSourceTextures.Find(ViewportKey))
+		if (TObjectPtr<UTexture>* OriginalTexturePtr = OriginalSourceTextures.Find(ViewportKey))
 		{
-			Viewport->RenderSettings.Replace.SourceTexture = *TexturePtr;
+			// It should not be possible to clear OriginalSourceTextures without restoring them first.
+			// Or assign textures without saving the original references.
+			// 
+			// * If we did, then we'd leave assigned the pulled ones (since there would be nothing to restore), 
+			//   potentially with external map references, which would not let the user save the map.
+			// 
+			// * In addition, it should not be possible to think that the currently assigned replace textures are
+			//   the original ones when they really aren't. Otherwise we'd be eventually be restoring textures
+			//   that belong to an external map.
+
+			Viewport->RenderSettings.Replace.SourceTexture = *OriginalTexturePtr;
 		}
 	}
 
@@ -270,6 +329,7 @@ void UDisplayClusterPreviewShareComponent::RestoreRootActorOriginalSettings()
 	OriginalTextureReplaces.Empty();
 }
 
+
 void UDisplayClusterPreviewShareComponent::CloseAllMedia()
 {
 	// Stop all media
@@ -814,8 +874,15 @@ void UDisplayClusterPreviewShareComponent::TickReceive()
 
 			// Remember original replacement textures
 
-			OriginalTextureReplaces.Add(ViewportKey, Viewport->RenderSettings.Replace.bAllowReplace);
-			OriginalSourceTextures.Add(ViewportKey, Viewport->RenderSettings.Replace.SourceTexture);
+			if (!OriginalTextureReplaces.Find(ViewportKey))
+			{
+				OriginalTextureReplaces.Add(ViewportKey, Viewport->RenderSettings.Replace.bAllowReplace);
+			}
+
+			if (!OriginalSourceTextures.Find(ViewportKey))
+			{
+				OriginalSourceTextures.Add(ViewportKey, Viewport->RenderSettings.Replace.SourceTexture);
+			}
 
 			// Start the player right away
 			MediaPlayer->PlayOnOpen = true;
@@ -939,6 +1006,14 @@ void UDisplayClusterPreviewShareComponent::DestroyComponent(bool bPromoteChildre
 	Super::DestroyComponent(bPromoteChildren);
 }
 
+void UDisplayClusterPreviewShareComponent::OnRegister()
+{
+	Super::OnRegister();
+
+	// Make sure it is in a functional state after registration. This is useful when component is re-created from the undo buffer
+	ModeChanged();
+}
+
 void UDisplayClusterPreviewShareComponent::HandleMapChanged(UWorld* InWorld, EMapChangeType InMapChangeType)
 {
 	// We remove the reference to the source actor preview texture to avoid this being interpreted as a 
@@ -955,6 +1030,9 @@ void UDisplayClusterPreviewShareComponent::HandleMapChanged(UWorld* InWorld, EMa
 void UDisplayClusterPreviewShareComponent::Serialize(FArchive& Ar)
 {
 #if WITH_EDITOR
+
+	using namespace UE::DisplayClusterPreviewShare;
+
 	if (Ar.IsSaving())
 	{
 		// Because for speed we may be directly referencing the source nDisplay actor preview textures,
@@ -963,6 +1041,7 @@ void UDisplayClusterPreviewShareComponent::Serialize(FArchive& Ar)
 		if (Mode == EDisplayClusterPreviewShareMode::PullActor)
 		{
 			RestoreRootActorOriginalSettings();
+			RemoveExternalMapReferences(Cast<ADisplayClusterRootActor>(GetOwner()));
 		}
 	}
 #endif // WITH_EDITOR
diff --git a/Engine/Plugins/Runtime/nDisplay/Source/DisplayCluster/Public/Components/DisplayClusterPreviewShareComponent.h b/Engine/Plugins/Runtime/nDisplay/Source/DisplayCluster/Public/Components/DisplayClusterPreviewShareComponent.h
--- a/Engine/Plugins/Runtime/nDisplay/Source/DisplayCluster/Public/Components/DisplayClusterPreviewShareComponent.h
+++ b/Engine/Plugins/Runtime/nDisplay/Source/DisplayCluster/Public/Components/DisplayClusterPreviewShareComponent.h
@@ -102,6 +102,16 @@ public:
 	UFUNCTION(BlueprintCallable, Category = "Sharing")
 	void SetUniqueName(const FString& NewUniqueName);
 
+protected:
+
+#if WITH_EDITOR // Bulk wrap with WITH_EDITOR until preview is supported in other modes.
+
+	//~ UActorComponent interface begin
+	virtual void OnRegister() override;
+	//~ UActorComponent interface end
+
+#endif // WITH_EDITOR
+
 public:
 
 #if WITH_EDITORONLY_DATA
