Fixed build break.

diff --git a/Engine/Plugins/Runtime/nDisplay/Source/DisplayCluster/Private/Render/Containers/DisplayClusterRender_TextureResource.cpp b/Engine/Plugins/Runtime/nDisplay/Source/DisplayCluster/Private/Render/Containers/DisplayClusterRender_TextureResource.cpp
--- a/Engine/Plugins/Runtime/nDisplay/Source/DisplayCluster/Private/Render/Containers/DisplayClusterRender_TextureResource.cpp
+++ b/Engine/Plugins/Runtime/nDisplay/Source/DisplayCluster/Private/Render/Containers/DisplayClusterRender_TextureResource.cpp
@@ -109,7 +109,7 @@ void FDisplayClusterRender_TextureResource::InitializeRenderResource()
 	// Initialize on render thread
 	ENQUEUE_RENDER_COMMAND(DisplayClusterRenderTextureResource_Initialize)([This = SharedThis(this)](FRHICommandListImmediate& RHICmdList)
 	{
-		This->InitResource();
+		This->InitResource(RHICmdList);
 	});
 }
 
