Add a modify to StageActorComponent so that it's changes are properly registered with the transaction system.

diff --git a/Engine/Plugins/Runtime/nDisplay/Source/DisplayCluster/Private/Game/EngineClasses/Scene/DisplayClusterLightCardActor.cpp b/Engine/Plugins/Runtime/nDisplay/Source/DisplayCluster/Private/Game/EngineClasses/Scene/DisplayClusterLightCardActor.cpp
--- a/Engine/Plugins/Runtime/nDisplay/Source/DisplayCluster/Private/Game/EngineClasses/Scene/DisplayClusterLightCardActor.cpp
+++ b/Engine/Plugins/Runtime/nDisplay/Source/DisplayCluster/Private/Game/EngineClasses/Scene/DisplayClusterLightCardActor.cpp
@@ -757,6 +757,9 @@ void ADisplayClusterLightCardActor::AddToRootActor(ADisplayClusterRootActor* InR
 		}
 	}
 
+#if WITH_EDITOR
+	StageActorComponent->Modify();
+#endif
 	StageActorComponent->SetRootActor(InRootActor);
 }
 
