- Adding warning message when sync policy is incompatible

diff --git a/Engine/Plugins/Runtime/nDisplay/Source/DisplayClusterMedia/Private/Capture/DisplayClusterMediaCaptureBase.cpp b/Engine/Plugins/Runtime/nDisplay/Source/DisplayClusterMedia/Private/Capture/DisplayClusterMediaCaptureBase.cpp
--- a/Engine/Plugins/Runtime/nDisplay/Source/DisplayClusterMedia/Private/Capture/DisplayClusterMediaCaptureBase.cpp
+++ b/Engine/Plugins/Runtime/nDisplay/Source/DisplayClusterMedia/Private/Capture/DisplayClusterMediaCaptureBase.cpp
@@ -57,16 +57,23 @@ bool FDisplayClusterMediaCaptureBase::StartCapture()
 	if (MediaOutput && !MediaCapture)
 	{
 		MediaCapture = MediaOutput->CreateMediaCapture();
-		if (MediaCapture)
+		if (IsValid(MediaCapture))
 		{
 			MediaCapture->SetMediaOutput(MediaOutput);
 
 			// Initialize and start capture synchronization
-			if (SyncPolicy && SyncPolicy->IsCaptureTypeSupported(MediaCapture))
+			if (IsValid(SyncPolicy))
 			{
-				if (!SyncPolicy->StartSynchronization(MediaCapture, GetMediaId()))
+				if (SyncPolicy->IsCaptureTypeSupported(MediaCapture))
 				{
-					UE_LOG(LogDisplayClusterMedia, Warning, TEXT("MediaCapture '%s': couldn't start synchronization."), *GetMediaId());
+					if (!SyncPolicy->StartSynchronization(MediaCapture, GetMediaId()))
+					{
+						UE_LOG(LogDisplayClusterMedia, Warning, TEXT("MediaCapture '%s': couldn't start synchronization."), *GetMediaId());
+					}
+				}
+				else
+				{
+					UE_LOG(LogDisplayClusterMedia, Warning, TEXT("MediaCapture '%s' is not compatible with media SyncPolicy '%s'."), *GetMediaId(), *SyncPolicy->GetName());
 				}
 			}
 
