nDisplay:
- Reworked logic in the FDisplayClusterRender_Texture
	* fixed an assert crash from the ~FTextureResource(): the ReleaseRenderResource() function should be called before calling the object destructor of the FDisplayClusterRender_TextureResource class.
- Added changes the FDisplayClusterRender_TextureResource

diff --git a/Engine/Plugins/Runtime/nDisplay/Source/DisplayCluster/Private/Render/Containers/DisplayClusterRender_Texture.cpp b/Engine/Plugins/Runtime/nDisplay/Source/DisplayCluster/Private/Render/Containers/DisplayClusterRender_Texture.cpp
--- a/Engine/Plugins/Runtime/nDisplay/Source/DisplayCluster/Private/Render/Containers/DisplayClusterRender_Texture.cpp
+++ b/Engine/Plugins/Runtime/nDisplay/Source/DisplayCluster/Private/Render/Containers/DisplayClusterRender_Texture.cpp
@@ -29,68 +29,71 @@ FDisplayClusterRender_Texture::FDisplayClusterRender_Texture(const FString& InUn
 
 FDisplayClusterRender_Texture::~FDisplayClusterRender_Texture()
 {
-	if (PrivateResource.IsValid())
-	{
-		PrivateResource->ReleaseRenderResource();
-		PrivateResource.Reset();
-	}
-
-	if (PrivateResourceRenderThread.IsValid())
-	{
-		PrivateResourceRenderThread->ReleaseRenderResource();
-		PrivateResourceRenderThread.Reset();
-	}
-}
-
-void FDisplayClusterRender_Texture::CreateTexture(const void* InTextureData, const uint32 InComponentDepth, const uint32 InBytesPerComponent, uint32_t InWidth, uint32_t InHeight, bool bInHasCPUAccess)
-{
-	// Release the existing texture resource.
-	ReleaseResource();
-
-	//Dedicated servers have no texture internals
-	if (FApp::CanEverRender())
-	{
-		// Create a new texture resource.
-		TSharedPtr<FDisplayClusterRender_TextureResource, ESPMode::ThreadSafe> NewResource = MakeShared<FDisplayClusterRender_TextureResource, ESPMode::ThreadSafe>(InTextureData, InComponentDepth, InBytesPerComponent, InWidth, InHeight, bInHasCPUAccess);
-		if (NewResource.IsValid())
-		{
-			NewResource->InitializeRenderResource();
-		}
-
-		SetResource(NewResource);
-	}
+	SetTextureResource(nullptr);
 }
 
-const TSharedPtr<FDisplayClusterRender_TextureResource>& FDisplayClusterRender_Texture::GetResource() const
+const TSharedPtr<FDisplayClusterRender_TextureResource, ESPMode::ThreadSafe>& FDisplayClusterRender_Texture::GetTextureResource() const
 {
 	if (IsInActualRenderingThread() || IsInRHIThread())
 	{
-		return PrivateResourceRenderThread;
+		return TextureResourceProxyPtr;
 	}
-	return PrivateResource;
+
+	return TextureResourcePtr;
 }
 
-void FDisplayClusterRender_Texture::SetResource(const TSharedPtr<FDisplayClusterRender_TextureResource>& InResource)
+void FDisplayClusterRender_Texture::SetTextureResource(const TSharedPtr<FDisplayClusterRender_TextureResource, ESPMode::ThreadSafe>& InTextureResourcePtr)
 {
 	check(!IsInActualRenderingThread() && !IsInRHIThread());
 
+	if (InTextureResourcePtr == TextureResourcePtr)
+	{
+		// ignore the same resoures
+		return;
+	}
+
 	// Each PrivateResource value must be updated in it's own thread because any
 	// rendering code trying to access the Resource
 	// crash if it suddenly sees nullptr or a new resource that has not had it's InitRHI called.
 
-	PrivateResource = InResource;
-	ENQUEUE_RENDER_COMMAND(DisplayCluster_SetResourceRenderThread)([This = SharedThis(this), InResource](FRHICommandListImmediate& RHICmdList)
+	if (TextureResourcePtr.IsValid())
 	{
-		This->PrivateResourceRenderThread = InResource;
+		// Always call ReleaseRenderResource before destructor
+		TextureResourcePtr->ReleaseRenderResource();
+		TextureResourcePtr.Reset();
+	}
+	TextureResourcePtr = InTextureResourcePtr;
+
+	// Update render thread
+	ENQUEUE_RENDER_COMMAND(FDisplayClusterRender_Texture_UpdateResource)([This = SharedThis(this), InTextureResourcePtr](FRHICommandListImmediate& RHICmdList)
+	{
+		if (This->TextureResourceProxyPtr.IsValid())
+		{
+			// Always call ReleaseRenderResource before destructor
+			This->TextureResourceProxyPtr->ReleaseRenderResource();
+			This->TextureResourceProxyPtr.Reset();
+		}
+
+		This->TextureResourceProxyPtr = InTextureResourcePtr;
 	});
 }
 
-void FDisplayClusterRender_Texture::ReleaseResource()
+void FDisplayClusterRender_Texture::CreateTexture(const void* InTextureData, const uint32 InComponentDepth, const uint32 InBytesPerComponent, uint32_t InWidth, uint32_t InHeight, bool bInHasCPUAccess)
 {
-	if (PrivateResource)
+	//Dedicated servers have no texture internals
+	if (FApp::CanEverRender())
 	{
-		// Free the resource.
-		SetResource(nullptr);
+		// Create a new texture resource.
+		TSharedPtr<FDisplayClusterRender_TextureResource, ESPMode::ThreadSafe> NewResource = MakeShared<FDisplayClusterRender_TextureResource, ESPMode::ThreadSafe>(InTextureData, InComponentDepth, InBytesPerComponent, InWidth, InHeight, bInHasCPUAccess);
+		if (NewResource.IsValid())
+		{
+			NewResource->InitializeRenderResource();
+			SetTextureResource(NewResource);
+		}
+		else
+		{
+			SetTextureResource(nullptr);
+		}
 	}
 }
 
diff --git a/Engine/Plugins/Runtime/nDisplay/Source/DisplayCluster/Private/Render/Containers/DisplayClusterRender_Texture.h b/Engine/Plugins/Runtime/nDisplay/Source/DisplayCluster/Private/Render/Containers/DisplayClusterRender_Texture.h
--- a/Engine/Plugins/Runtime/nDisplay/Source/DisplayCluster/Private/Render/Containers/DisplayClusterRender_Texture.h
+++ b/Engine/Plugins/Runtime/nDisplay/Source/DisplayCluster/Private/Render/Containers/DisplayClusterRender_Texture.h
@@ -31,42 +31,37 @@ public:
 
 	virtual bool IsEnabled() const override
 	{
-		const TSharedPtr<FDisplayClusterRender_TextureResource, ESPMode::ThreadSafe>& TextureResource = GetResource();
-		if (TextureResource.IsValid())
-		{
-			return TextureResource->GetWidth() > 0 && TextureResource->GetHeight() > 0;
-		}
-
-		return false;
+		const TSharedPtr<FDisplayClusterRender_TextureResource, ESPMode::ThreadSafe>& TextureResource = GetTextureResource();
+		return TextureResource.IsValid() && TextureResource->IsRenderResourceInitialized();
 	}
 
 	virtual void* GetData() const override
 	{
-		const TSharedPtr<FDisplayClusterRender_TextureResource, ESPMode::ThreadSafe>& TextureResource = GetResource();
+		const TSharedPtr<FDisplayClusterRender_TextureResource, ESPMode::ThreadSafe>& TextureResource = GetTextureResource();
 		return TextureResource ? TextureResource->GetTextureData() : nullptr;
 	}
 
 	virtual uint32_t GetWidth() const override
 	{
-		const TSharedPtr<FDisplayClusterRender_TextureResource, ESPMode::ThreadSafe>& TextureResource = GetResource();
+		const TSharedPtr<FDisplayClusterRender_TextureResource, ESPMode::ThreadSafe>& TextureResource = GetTextureResource();
 		return TextureResource.IsValid() ? TextureResource->GetWidth() : 0;
 	}
 
 	virtual uint32_t GetHeight() const override
 	{
-		const TSharedPtr<FDisplayClusterRender_TextureResource, ESPMode::ThreadSafe>& TextureResource = GetResource();
+		const TSharedPtr<FDisplayClusterRender_TextureResource, ESPMode::ThreadSafe>& TextureResource = GetTextureResource();
 		return TextureResource.IsValid() ? TextureResource->GetHeight() : 0;
 	}
 
 	virtual EPixelFormat GetPixelFormat() const override
 	{
-		const TSharedPtr<FDisplayClusterRender_TextureResource, ESPMode::ThreadSafe>& TextureResource = GetResource();
+		const TSharedPtr<FDisplayClusterRender_TextureResource, ESPMode::ThreadSafe>& TextureResource = GetTextureResource();
 		return TextureResource.IsValid() ? TextureResource->GetPixelFormat() : PF_Unknown;
 	}
 
 	virtual FRHITexture* GetRHITexture() const override
 	{
-		const TSharedPtr<FDisplayClusterRender_TextureResource, ESPMode::ThreadSafe>& TextureResource = GetResource();
+		const TSharedPtr<FDisplayClusterRender_TextureResource, ESPMode::ThreadSafe>& TextureResource = GetTextureResource();
 		return TextureResource.IsValid() ? TextureResource->TextureRHI : nullptr;
 	}
 
@@ -85,25 +80,20 @@ public:
 	// ~~ End TDisplayClusterDataCache
 
 protected:
-	/** Set texture's resource, can be NULL */
-	void SetResource(const TSharedPtr<FDisplayClusterRender_TextureResource, ESPMode::ThreadSafe>& InResource);
-
-	/** Get the texture's resource, can be NULL */
-	const TSharedPtr<FDisplayClusterRender_TextureResource, ESPMode::ThreadSafe>& GetResource() const;
+	/** Set new texture resource. */
+	void SetTextureResource(const TSharedPtr<FDisplayClusterRender_TextureResource, ESPMode::ThreadSafe>& InTextureResourcePtr);
 
-	/**
-	 * Resets the resource for the texture.
-	 */
-	void ReleaseResource();
+	/** Get texture resource: supports both the game thread and the rendering thread. */
+	const TSharedPtr<FDisplayClusterRender_TextureResource, ESPMode::ThreadSafe>& GetTextureResource() const;
 
 private:
 	// The texture's resource, can be NULL
-	TSharedPtr<FDisplayClusterRender_TextureResource, ESPMode::ThreadSafe> PrivateResource;
+	TSharedPtr<FDisplayClusterRender_TextureResource, ESPMode::ThreadSafe> TextureResourcePtr;
 
 	// Value updated and returned by the render-thread to allow
 	// fenceless update from the game-thread without causing
 	// potential crash in the render thread.
-	TSharedPtr<FDisplayClusterRender_TextureResource, ESPMode::ThreadSafe> PrivateResourceRenderThread;
+	TSharedPtr<FDisplayClusterRender_TextureResource, ESPMode::ThreadSafe> TextureResourceProxyPtr;
 
 	// Unique texture name
 	FString UniqueTextureName;
diff --git a/Engine/Plugins/Runtime/nDisplay/Source/DisplayCluster/Private/Render/Containers/DisplayClusterRender_TextureResource.cpp b/Engine/Plugins/Runtime/nDisplay/Source/DisplayCluster/Private/Render/Containers/DisplayClusterRender_TextureResource.cpp
--- a/Engine/Plugins/Runtime/nDisplay/Source/DisplayCluster/Private/Render/Containers/DisplayClusterRender_TextureResource.cpp
+++ b/Engine/Plugins/Runtime/nDisplay/Source/DisplayCluster/Private/Render/Containers/DisplayClusterRender_TextureResource.cpp
@@ -106,6 +106,8 @@ FDisplayClusterRender_TextureResource::~FDisplayClusterRender_TextureResource()
 
 void FDisplayClusterRender_TextureResource::InitializeRenderResource()
 {
+	bRenderResourceInitialized = true;
+
 	// Initialize on render thread
 	ENQUEUE_RENDER_COMMAND(DisplayClusterRenderTextureResource_Initialize)([This = SharedThis(this)](FRHICommandListImmediate& RHICmdList)
 	{
@@ -115,6 +117,8 @@ void FDisplayClusterRender_TextureResource::InitializeRenderResource()
 
 void FDisplayClusterRender_TextureResource::ReleaseRenderResource()
 {
+	bRenderResourceInitialized = false;
+
 	// Release on render thread
 	ENQUEUE_RENDER_COMMAND(DisplayClusterRenderTextureResource_Release)([This = SharedThis(this)](FRHICommandListImmediate& RHICmdList)
 	{
@@ -129,9 +133,6 @@ void FDisplayClusterRender_TextureResource::ReleaseTextureData()
 		FMemory::Free(TextureData);
 		TextureData = nullptr;
 
-		Width = 0;
-		Height = 0;
-		PixelFormat = PF_Unknown;
 		bHasCPUAccess = false;
 	}
 }
diff --git a/Engine/Plugins/Runtime/nDisplay/Source/DisplayCluster/Private/Render/Containers/DisplayClusterRender_TextureResource.h b/Engine/Plugins/Runtime/nDisplay/Source/DisplayCluster/Private/Render/Containers/DisplayClusterRender_TextureResource.h
--- a/Engine/Plugins/Runtime/nDisplay/Source/DisplayCluster/Private/Render/Containers/DisplayClusterRender_TextureResource.h
+++ b/Engine/Plugins/Runtime/nDisplay/Source/DisplayCluster/Private/Render/Containers/DisplayClusterRender_TextureResource.h
@@ -30,6 +30,18 @@ public:
 	/** Prepares the resource for deletion. */
 	void ReleaseRenderResource();
 
+	/** Whether this texture resource has data in CPU memory. */
+	bool HasCPUAccess() const
+	{
+		return bHasCPUAccess;
+	}
+
+	/** Get ptr to resource data in memory. The variable bHasCPUAccess must be true. */
+	bool IsRenderResourceInitialized() const
+	{
+		return bRenderResourceInitialized;
+	}
+
 	/** Get ptr to resource data in memory. The variable bHasCPUAccess must be true. */
 	void* GetTextureData() const
 	{
@@ -56,17 +68,20 @@ public:
 
 private:
 	// RAW data of the texture. Mips not supported now
-	void* TextureData;
+	void* TextureData = nullptr;
 	
 	// Width of the texture.
-	uint32_t Width;
+	uint32_t Width = 0;
 	
 	// Height of the texture.
-	uint32_t Height;
+	uint32_t Height = 0;
 	
 	// Format in which mip data is stored.
-	EPixelFormat PixelFormat;
+	EPixelFormat PixelFormat = EPixelFormat::PF_Unknown;
 	
 	// has CPU access to texture data, dont release from memory
-	bool bHasCPUAccess;
+	bool bHasCPUAccess = false;
+
+	// Is render resource initialized
+	bool bRenderResourceInitialized = false;
 };
