nDisplay: fixed CineCamera custom near clipping plane is not being used

diff --git a/Engine/Plugins/Runtime/nDisplay/Source/DisplayClusterProjection/Private/Policy/Camera/DisplayClusterProjectionCameraPolicy.cpp b/Engine/Plugins/Runtime/nDisplay/Source/DisplayClusterProjection/Private/Policy/Camera/DisplayClusterProjectionCameraPolicy.cpp
--- a/Engine/Plugins/Runtime/nDisplay/Source/DisplayClusterProjection/Private/Policy/Camera/DisplayClusterProjectionCameraPolicy.cpp
+++ b/Engine/Plugins/Runtime/nDisplay/Source/DisplayClusterProjection/Private/Policy/Camera/DisplayClusterProjectionCameraPolicy.cpp
@@ -86,6 +86,12 @@ bool FDisplayClusterProjectionCameraPolicy::CalculateView(IDisplayClusterViewpor
 	ZNear = NCP;
 	ZFar  = FCP;
 
+	if (CustomNearClippingPlane >= 0)
+	{
+		ZNear = CustomNearClippingPlane;
+		ZFar = (NCP == FCP) ? CustomNearClippingPlane : FCP;
+	}
+
 	return true;
 }
 
@@ -144,6 +150,8 @@ bool FDisplayClusterProjectionCameraPolicy::GetViewPoint(IDisplayClusterViewport
 	InOutViewLocation = FVector::ZeroVector;
 	InOutViewRotation = FRotator::ZeroRotator;
 
+	CustomNearClippingPlane = -1;
+
 	// Use assigned camera
 	if (UCameraComponent* CameraComponent = GetCameraComponent())
 	{
@@ -151,8 +159,7 @@ bool FDisplayClusterProjectionCameraPolicy::GetViewPoint(IDisplayClusterViewport
 		{
 			if (CineCameraComponent->bOverride_CustomNearClippingPlane)
 			{
-				ZNear = CineCameraComponent->CustomNearClippingPlane;
-				ZFar = ZNear;
+				CustomNearClippingPlane = CineCameraComponent->CustomNearClippingPlane;
 			}
 		}
 
diff --git a/Engine/Plugins/Runtime/nDisplay/Source/DisplayClusterProjection/Private/Policy/Camera/DisplayClusterProjectionCameraPolicy.h b/Engine/Plugins/Runtime/nDisplay/Source/DisplayClusterProjection/Private/Policy/Camera/DisplayClusterProjectionCameraPolicy.h
--- a/Engine/Plugins/Runtime/nDisplay/Source/DisplayClusterProjection/Private/Policy/Camera/DisplayClusterProjectionCameraPolicy.h
+++ b/Engine/Plugins/Runtime/nDisplay/Source/DisplayClusterProjection/Private/Policy/Camera/DisplayClusterProjectionCameraPolicy.h
@@ -53,4 +53,7 @@ private:
 	FDisplayClusterProjectionCameraPolicySettings CameraSettings;
 	float ZNear = 1.f;
 	float ZFar = 1.f;
+
+	// A value less than zero means that the variable is not used.
+	float CustomNearClippingPlane = -1;
 };
