- Changing media sync handling to use native classes and keep uobject only for settings / user configuration

diff --git a/Engine/Plugins/Runtime/nDisplay/Source/DisplayClusterConfiguration/Public/DisplayClusterConfigurationTypes_Media.h b/Engine/Plugins/Runtime/nDisplay/Source/DisplayClusterConfiguration/Public/DisplayClusterConfigurationTypes_Media.h
--- a/Engine/Plugins/Runtime/nDisplay/Source/DisplayClusterConfiguration/Public/DisplayClusterConfigurationTypes_Media.h
+++ b/Engine/Plugins/Runtime/nDisplay/Source/DisplayClusterConfiguration/Public/DisplayClusterConfigurationTypes_Media.h
@@ -5,6 +5,7 @@
 #include "CoreMinimal.h"
 
 #include "DisplayClusterConfigurationTypes_Base.h"
+#include "DisplayClusterConfigurationTypes_MediaSync.h"
 
 #include "MediaPlayer.h"
 #include "MediaSource.h"
@@ -14,33 +15,6 @@
 #include "DisplayClusterConfigurationTypes_Media.generated.h"
 
 
-/*
- * Base media output synchronization policy class
- */
-UCLASS(Abstract, editinlinenew, BlueprintType, hidecategories = (Object))
-class DISPLAYCLUSTERCONFIGURATION_API UDisplayClusterMediaOutputSynchronizationPolicy
-	: public UObject
-{
-	GENERATED_BODY()
-
-public:
-	/** Returns true if specified media capture type can be synchonized by the policy implementation */
-	virtual bool IsCaptureTypeSupported(UMediaCapture* MediaCapture) const
-	{
-		return true;
-	}
-
-	/** Starts synchronization of specific output stream (capture device). Returns false if failed. */
-	virtual bool StartSynchronization(UMediaCapture* MediaCapture, const FString& MediaId) PURE_VIRTUAL(UDisplayClusterMediaOutputSynchronizationPolicy::StartSynchronization, return false;);
-
-	/** Stops synchronization of specific output stream (capture device). */
-	virtual void StopSynchronization() PURE_VIRTUAL(UDisplayClusterMediaOutputSynchronizationPolicy::StopSynchronization, );
-
-	/** Returns true if currently synchronising a media output. */
-	virtual bool IsRunning() PURE_VIRTUAL(UDisplayClusterMediaOutputSynchronizationPolicy::IsRunning, return false; );
-};
-
-
 PRAGMA_DISABLE_DEPRECATION_WARNINGS
 
 /*
diff --git a/Engine/Plugins/Runtime/nDisplay/Source/DisplayClusterConfiguration/Public/DisplayClusterConfigurationTypes_MediaSync.h b/Engine/Plugins/Runtime/nDisplay/Source/DisplayClusterConfiguration/Public/DisplayClusterConfigurationTypes_MediaSync.h
new file mode 100644
--- /dev/null
+++ b/Engine/Plugins/Runtime/nDisplay/Source/DisplayClusterConfiguration/Public/DisplayClusterConfigurationTypes_MediaSync.h
@@ -0,0 +1,55 @@
+// Copyright Epic Games, Inc. All Rights Reserved.
+
+#pragma once
+
+#include "CoreTypes.h"
+
+#include "Templates/SharedPointer.h"
+#include "Templates/SubclassOf.h"
+#include "UObject/ObjectMacros.h"
+
+#include "DisplayClusterConfigurationTypes_MediaSync.generated.h"
+
+class UMediaCapture;
+
+/** Media output sync policy handler interface. All sync logic is handled in a derived native class to avoid UObject interaction from other threads. */
+class DISPLAYCLUSTERCONFIGURATION_API IDisplayClusterMediaOutputSynchronizationPolicyHandler : public TSharedFromThis<IDisplayClusterMediaOutputSynchronizationPolicyHandler>
+{
+public:
+	
+	/** Destructor */
+	virtual ~IDisplayClusterMediaOutputSynchronizationPolicyHandler() = default;
+
+	/** Returns true if specified media capture type can be synchonized by the policy implementation */
+	virtual bool IsCaptureTypeSupported(UMediaCapture* MediaCapture) const
+	{
+		return true;
+	}
+
+	/** Starts synchronization of specific output stream (capture device). Returns false if failed. */
+	virtual bool StartSynchronization(UMediaCapture* MediaCapture, const FString& MediaId) = 0;
+
+	/** Stops synchronization of specific output stream (capture device). */
+	virtual void StopSynchronization() = 0;
+
+	/** Returns true if currently synchronizing a media output. */
+	virtual bool IsRunning() = 0;
+
+	/** Returns the policy config class used to create this policy handler. */
+	virtual TSubclassOf<UDisplayClusterMediaOutputSynchronizationPolicy> GetPolicyClass() const = 0;
+};
+
+/*
+ * Base media output synchronization policy class settings
+ */
+UCLASS(Abstract, editinlinenew, BlueprintType, hidecategories = (Object))
+class DISPLAYCLUSTERCONFIGURATION_API UDisplayClusterMediaOutputSynchronizationPolicy
+	: public UObject
+{
+	GENERATED_BODY()
+
+public:
+
+	/** Returns the handler associated with this policy config object */
+	virtual TSharedPtr<IDisplayClusterMediaOutputSynchronizationPolicyHandler> GetHandler() PURE_VIRTUAL(UDisplayClusterMediaOutputSynchronizationPolicy::GetHandler, return nullptr;);
+};
diff --git a/Engine/Plugins/Runtime/nDisplay/Source/DisplayClusterMedia/Private/Capture/DisplayClusterMediaCaptureBase.cpp b/Engine/Plugins/Runtime/nDisplay/Source/DisplayClusterMedia/Private/Capture/DisplayClusterMediaCaptureBase.cpp
--- a/Engine/Plugins/Runtime/nDisplay/Source/DisplayClusterMedia/Private/Capture/DisplayClusterMediaCaptureBase.cpp
+++ b/Engine/Plugins/Runtime/nDisplay/Source/DisplayClusterMedia/Private/Capture/DisplayClusterMediaCaptureBase.cpp
@@ -7,6 +7,7 @@
 #include "IDisplayClusterCallbacks.h"
 
 #include "DisplayClusterConfigurationTypes_Media.h"
+#include "DisplayClusterConfigurationTypes_MediaSync.h"
 
 #include "MediaCapture.h"
 #include "MediaOutput.h"
@@ -64,16 +65,28 @@ bool FDisplayClusterMediaCaptureBase::StartCapture()
 			// Initialize and start capture synchronization
 			if (IsValid(SyncPolicy))
 			{
-				if (SyncPolicy->IsCaptureTypeSupported(MediaCapture))
+				SyncPolicyHandler = SyncPolicy->GetHandler();
+				if (SyncPolicyHandler)
 				{
-					if (!SyncPolicy->StartSynchronization(MediaCapture, GetMediaId()))
+					if (SyncPolicyHandler->IsCaptureTypeSupported(MediaCapture))
 					{
-						UE_LOG(LogDisplayClusterMedia, Warning, TEXT("MediaCapture '%s': couldn't start synchronization."), *GetMediaId());
+						if (SyncPolicyHandler->StartSynchronization(MediaCapture, GetMediaId()))
+						{
+							UE_LOG(LogDisplayClusterMedia, Log, TEXT("MediaCapture '%s' started synchronization type '%s'."), *GetMediaId(), *SyncPolicy->GetName());
+						}
+						else
+						{
+							UE_LOG(LogDisplayClusterMedia, Warning, TEXT("MediaCapture '%s': couldn't start synchronization."), *GetMediaId());
+						}
+					}
+					else
+					{
+						UE_LOG(LogDisplayClusterMedia, Warning, TEXT("MediaCapture '%s' is not compatible with media SyncPolicy '%s'."), *GetMediaId(), *SyncPolicy->GetName());
 					}
 				}
 				else
 				{
-					UE_LOG(LogDisplayClusterMedia, Warning, TEXT("MediaCapture '%s' is not compatible with media SyncPolicy '%s'."), *GetMediaId(), *SyncPolicy->GetName());
+					UE_LOG(LogDisplayClusterMedia, Warning, TEXT("Could not create media sync policy handler from '%s'."), *SyncPolicy->GetName());
 				}
 			}
 
@@ -88,9 +101,9 @@ bool FDisplayClusterMediaCaptureBase::StartCapture()
 void FDisplayClusterMediaCaptureBase::StopCapture()
 {
 	// Stop synchronization
-	if (SyncPolicy)
+	if (SyncPolicyHandler)
 	{
-		SyncPolicy->StopSynchronization();
+		SyncPolicyHandler->StopSynchronization();
 	}
 
 	// Stop capture
diff --git a/Engine/Plugins/Runtime/nDisplay/Source/DisplayClusterMedia/Private/Capture/DisplayClusterMediaCaptureBase.h b/Engine/Plugins/Runtime/nDisplay/Source/DisplayClusterMedia/Private/Capture/DisplayClusterMediaCaptureBase.h
--- a/Engine/Plugins/Runtime/nDisplay/Source/DisplayClusterMedia/Private/Capture/DisplayClusterMediaCaptureBase.h
+++ b/Engine/Plugins/Runtime/nDisplay/Source/DisplayClusterMedia/Private/Capture/DisplayClusterMediaCaptureBase.h
@@ -13,7 +13,7 @@ class FRDGBuilder;
 class UMediaCapture;
 class UMediaOutput;
 class UDisplayClusterMediaOutputSynchronizationPolicy;
-
+class IDisplayClusterMediaOutputSynchronizationPolicyHandler;
 
 /**
  * Base media capture class
@@ -70,9 +70,9 @@ private:
 
 private:
 	//~ Begin GC by AddReferencedObjects
-	TObjectPtr<UMediaOutput>  MediaOutput  = nullptr;
-	TObjectPtr<UMediaCapture> MediaCapture = nullptr;
-	TObjectPtr<UDisplayClusterMediaOutputSynchronizationPolicy> SyncPolicy = nullptr;
+	TObjectPtr<UMediaOutput>  MediaOutput;
+	TObjectPtr<UMediaCapture> MediaCapture;
+	TObjectPtr<UDisplayClusterMediaOutputSynchronizationPolicy> SyncPolicy;
 	//~ End GC by AddReferencedObjects
 
 	// Used to restart media capture in the case it falls in error
@@ -83,4 +83,7 @@ private:
 
 	// Last region size of the texture being exported. Used to restart the capture when in error.
 	std::atomic<FIntSize> LastSrcRegionSize { FIntSize(0,0) };
+
+	/** Sync policy handler to deal with synchronization logic */
+	TSharedPtr<IDisplayClusterMediaOutputSynchronizationPolicyHandler> SyncPolicyHandler;
 };
diff --git a/Engine/Plugins/Runtime/nDisplay/Source/DisplayClusterMedia/Private/Synchronization/SyncPolicies/DisplayClusterMediaOutputSynchronizationPolicyEthernetBarrier.cpp b/Engine/Plugins/Runtime/nDisplay/Source/DisplayClusterMedia/Private/Synchronization/SyncPolicies/DisplayClusterMediaOutputSynchronizationPolicyEthernetBarrier.cpp
--- a/Engine/Plugins/Runtime/nDisplay/Source/DisplayClusterMedia/Private/Synchronization/SyncPolicies/DisplayClusterMediaOutputSynchronizationPolicyEthernetBarrier.cpp
+++ b/Engine/Plugins/Runtime/nDisplay/Source/DisplayClusterMedia/Private/Synchronization/SyncPolicies/DisplayClusterMediaOutputSynchronizationPolicyEthernetBarrier.cpp
@@ -3,8 +3,29 @@
 #include "Synchronization/DisplayClusterMediaOutputSynchronizationPolicyEthernetBarrier.h"
 
 
-void UDisplayClusterMediaOutputSynchronizationPolicyEthernetBarrier::Synchronize()
+FDisplayClusterMediaOutputSynchronizationPolicyEthernetBarrierHandler::FDisplayClusterMediaOutputSynchronizationPolicyEthernetBarrierHandler(UDisplayClusterMediaOutputSynchronizationPolicyEthernetBarrier* InPolicyObject)
+	: Super(InPolicyObject)
+{
+
+}
+
+TSubclassOf<UDisplayClusterMediaOutputSynchronizationPolicy> FDisplayClusterMediaOutputSynchronizationPolicyEthernetBarrierHandler::GetPolicyClass() const
+{
+	return UDisplayClusterMediaOutputSynchronizationPolicyEthernetBarrier::StaticClass();
+}
+
+void FDisplayClusterMediaOutputSynchronizationPolicyEthernetBarrierHandler::Synchronize()
 {
 	// Just sync on the barrier
 	SyncThreadOnBarrier();
 }
+
+TSharedPtr<IDisplayClusterMediaOutputSynchronizationPolicyHandler> UDisplayClusterMediaOutputSynchronizationPolicyEthernetBarrier::GetHandler()
+{
+	if (!Handler)
+	{
+		Handler = MakeShared<FDisplayClusterMediaOutputSynchronizationPolicyEthernetBarrierHandler>(this);
+	}
+
+	return Handler;
+}
diff --git a/Engine/Plugins/Runtime/nDisplay/Source/DisplayClusterMedia/Private/Synchronization/SyncPolicies/DisplayClusterMediaOutputSynchronizationPolicyEthernetBarrierBase.cpp b/Engine/Plugins/Runtime/nDisplay/Source/DisplayClusterMedia/Private/Synchronization/SyncPolicies/DisplayClusterMediaOutputSynchronizationPolicyEthernetBarrierBase.cpp
--- a/Engine/Plugins/Runtime/nDisplay/Source/DisplayClusterMedia/Private/Synchronization/SyncPolicies/DisplayClusterMediaOutputSynchronizationPolicyEthernetBarrierBase.cpp
+++ b/Engine/Plugins/Runtime/nDisplay/Source/DisplayClusterMedia/Private/Synchronization/SyncPolicies/DisplayClusterMediaOutputSynchronizationPolicyEthernetBarrierBase.cpp
@@ -17,8 +17,16 @@
 #include "Cluster/IDisplayClusterClusterManager.h"
 #include "Game/IDisplayClusterGameManager.h"
 
+#include "Templates/SharedPointer.h"
 
-bool UDisplayClusterMediaOutputSynchronizationPolicyEthernetBarrierBase::StartSynchronization(UMediaCapture* MediaCapture, const FString& MediaId)
+
+FDisplayClusterMediaOutputSynchronizationPolicyEthernetBarrierBaseHandler::FDisplayClusterMediaOutputSynchronizationPolicyEthernetBarrierBaseHandler(UDisplayClusterMediaOutputSynchronizationPolicyEthernetBarrierBase* InPolicyObject)
+	: BarrierTimeoutMs(InPolicyObject->BarrierTimeoutMs)
+{
+
+}
+
+bool FDisplayClusterMediaOutputSynchronizationPolicyEthernetBarrierBaseHandler::StartSynchronization(UMediaCapture* MediaCapture, const FString& MediaId)
 {
 	// Cluster mode only
 	if (IDisplayCluster::Get().GetOperationMode() != EDisplayClusterOperationMode::Cluster)
@@ -57,8 +65,7 @@ bool UDisplayClusterMediaOutputSynchronizationPolicyEthernetBarrierBase::StartSy
 		return false;
 	}
 
-	// And subscribe for sync callbacks
-	CapturingDevice->OnOutputSynchronization.BindUObject(this, &UDisplayClusterMediaOutputSynchronizationPolicyEthernetBarrierBase::ProcessMediaSynchronizationCallback);
+	CapturingDevice->OnOutputSynchronization.BindSP(this, &FDisplayClusterMediaOutputSynchronizationPolicyEthernetBarrierBaseHandler::ProcessMediaSynchronizationCallback);
 
 	// Update state
 	bIsRunning = true;
@@ -66,7 +73,7 @@ bool UDisplayClusterMediaOutputSynchronizationPolicyEthernetBarrierBase::StartSy
 	return true;
 }
 
-void UDisplayClusterMediaOutputSynchronizationPolicyEthernetBarrierBase::StopSynchronization()
+void FDisplayClusterMediaOutputSynchronizationPolicyEthernetBarrierBaseHandler::StopSynchronization()
 {
 	if (bIsRunning)
 	{
@@ -85,17 +92,17 @@ void UDisplayClusterMediaOutputSynchronizationPolicyEthernetBarrierBase::StopSyn
 	}
 }
 
-bool UDisplayClusterMediaOutputSynchronizationPolicyEthernetBarrierBase::IsRunning()
+bool FDisplayClusterMediaOutputSynchronizationPolicyEthernetBarrierBaseHandler::IsRunning()
 {
 	return bIsRunning;
 }
 
-FString UDisplayClusterMediaOutputSynchronizationPolicyEthernetBarrierBase::GetMediaDeviceId() const
+FString FDisplayClusterMediaOutputSynchronizationPolicyEthernetBarrierBaseHandler::GetMediaDeviceId() const
 {
 	return MediaDeviceId;
 }
 
-void UDisplayClusterMediaOutputSynchronizationPolicyEthernetBarrierBase::SyncThreadOnBarrier()
+void FDisplayClusterMediaOutputSynchronizationPolicyEthernetBarrierBaseHandler::SyncThreadOnBarrier()
 {
 	// Sync on the barrier if everything is good
 	if (bIsRunning && EthernetBarrierClient && EthernetBarrierClient->IsConnected())
@@ -105,7 +112,7 @@ void UDisplayClusterMediaOutputSynchronizationPolicyEthernetBarrierBase::SyncThr
 	}
 }
 
-bool UDisplayClusterMediaOutputSynchronizationPolicyEthernetBarrierBase::InitializeBarrier(const FString& MediaId)
+bool FDisplayClusterMediaOutputSynchronizationPolicyEthernetBarrierBaseHandler::InitializeBarrier(const FString& MediaId)
 {
 	if (MediaId.IsEmpty())
 	{
@@ -153,7 +160,7 @@ bool UDisplayClusterMediaOutputSynchronizationPolicyEthernetBarrierBase::Initial
 	return true;
 }
 
-void UDisplayClusterMediaOutputSynchronizationPolicyEthernetBarrierBase::ReleaseBarrier()
+void FDisplayClusterMediaOutputSynchronizationPolicyEthernetBarrierBaseHandler::ReleaseBarrier()
 {
 	if (EthernetBarrierClient)
 	{
@@ -166,7 +173,7 @@ void UDisplayClusterMediaOutputSynchronizationPolicyEthernetBarrierBase::Release
 	}
 }
 
-FString UDisplayClusterMediaOutputSynchronizationPolicyEthernetBarrierBase::GenerateBarrierName() const
+FString FDisplayClusterMediaOutputSynchronizationPolicyEthernetBarrierBaseHandler::GenerateBarrierName() const
 {
 	// Currently we don't have any synchronization groups. This means all the sync policy instances of the same
 	// class use the same barrier. If we want to introduce sync groups in the future, the barrier ID should
@@ -177,10 +184,10 @@ FString UDisplayClusterMediaOutputSynchronizationPolicyEthernetBarrierBase::Gene
 	//
 	// However! All media captures are locked to UE rendering pipeline. This means all the captures will run
 	// with the same framerate. Therefore we don't need any sync groups so far.
-	return GetClass()->GetName();
+	return GetPolicyClass()->GetName();
 }
 
-void UDisplayClusterMediaOutputSynchronizationPolicyEthernetBarrierBase::GenerateListOfThreadMarkers(TArray<FString>& OutMarkers) const
+void FDisplayClusterMediaOutputSynchronizationPolicyEthernetBarrierBaseHandler::GenerateListOfThreadMarkers(TArray<FString>& OutMarkers) const
 {
 	OutMarkers.Empty();
 
@@ -202,7 +209,7 @@ void UDisplayClusterMediaOutputSynchronizationPolicyEthernetBarrierBase::Generat
 					if (MediaSettings.bEnable && MediaSettings.MediaOutput && MediaSettings.OutputSyncPolicy)
 					{
 						// Pick the same sync policy only
-						if (MediaSettings.OutputSyncPolicy->GetClass() == GetClass())
+						if (MediaSettings.OutputSyncPolicy->GetClass() == GetPolicyClass())
 						{
 							const FString BackbufferCaptureId = DisplayClusterMediaHelpers::MediaId::GenerateMediaId(
 								DisplayClusterMediaHelpers::MediaId::EMediaDeviceType::Output,
@@ -224,7 +231,7 @@ void UDisplayClusterMediaOutputSynchronizationPolicyEthernetBarrierBase::Generat
 						if (MediaSettings.bEnable && MediaSettings.MediaOutput && MediaSettings.OutputSyncPolicy)
 						{
 							// Pick the same sync policy only
-							if (MediaSettings.OutputSyncPolicy->GetClass() == GetClass())
+							if (MediaSettings.OutputSyncPolicy->GetClass() == GetPolicyClass())
 							{
 								const FString ViewportCaptureId = DisplayClusterMediaHelpers::MediaId::GenerateMediaId(
 									DisplayClusterMediaHelpers::MediaId::EMediaDeviceType::Output,
@@ -257,7 +264,7 @@ void UDisplayClusterMediaOutputSynchronizationPolicyEthernetBarrierBase::Generat
 						// Pick the same sync policy only
 						if (MediaOutputGroup.MediaOutput && MediaOutputGroup.OutputSyncPolicy)
 						{
-							if (MediaOutputGroup.OutputSyncPolicy->GetClass() == GetClass())
+							if (MediaOutputGroup.OutputSyncPolicy->GetClass() == GetPolicyClass())
 							{
 								for (const FString& NodeId : MediaOutputGroup.ClusterNodes.ItemNames)
 								{
@@ -283,7 +290,7 @@ void UDisplayClusterMediaOutputSynchronizationPolicyEthernetBarrierBase::Generat
 	}
 }
 
-void UDisplayClusterMediaOutputSynchronizationPolicyEthernetBarrierBase::ProcessMediaSynchronizationCallback()
+void FDisplayClusterMediaOutputSynchronizationPolicyEthernetBarrierBaseHandler::ProcessMediaSynchronizationCallback()
 {
 	UE_LOG(LogDisplayClusterMediaSync, VeryVerbose, TEXT("'%s': Synchronizing capture..."), *GetMediaDeviceId());
 
diff --git a/Engine/Plugins/Runtime/nDisplay/Source/DisplayClusterMedia/Private/Synchronization/SyncPolicies/DisplayClusterMediaOutputSynchronizationPolicyThresholdBase.cpp b/Engine/Plugins/Runtime/nDisplay/Source/DisplayClusterMedia/Private/Synchronization/SyncPolicies/DisplayClusterMediaOutputSynchronizationPolicyThresholdBase.cpp
--- a/Engine/Plugins/Runtime/nDisplay/Source/DisplayClusterMedia/Private/Synchronization/SyncPolicies/DisplayClusterMediaOutputSynchronizationPolicyThresholdBase.cpp
+++ b/Engine/Plugins/Runtime/nDisplay/Source/DisplayClusterMedia/Private/Synchronization/SyncPolicies/DisplayClusterMediaOutputSynchronizationPolicyThresholdBase.cpp
@@ -5,7 +5,14 @@
 #include "DisplayClusterMediaLog.h"
 
 
-void UDisplayClusterMediaOutputSynchronizationPolicyThresholdBase::Synchronize()
+FDisplayClusterMediaOutputSynchronizationPolicyThresholdBaseHandler::FDisplayClusterMediaOutputSynchronizationPolicyThresholdBaseHandler(UDisplayClusterMediaOutputSynchronizationPolicyThresholdBase* InPolicyObject)
+	: Super(InPolicyObject)
+	, MarginMs(InPolicyObject->MarginMs)
+{
+
+}
+
+void FDisplayClusterMediaOutputSynchronizationPolicyThresholdBaseHandler::Synchronize()
 {
 	// Ask the sync implementation about how much time we have before next synchronization timepoint
 	const double TimeLeftSeconds = GetTimeBeforeNextSyncPoint();
diff --git a/Engine/Plugins/Runtime/nDisplay/Source/DisplayClusterMedia/Private/Synchronization/SyncPolicies/DisplayClusterMediaOutputSynchronizationPolicyVblank.cpp b/Engine/Plugins/Runtime/nDisplay/Source/DisplayClusterMedia/Private/Synchronization/SyncPolicies/DisplayClusterMediaOutputSynchronizationPolicyVblank.cpp
--- a/Engine/Plugins/Runtime/nDisplay/Source/DisplayClusterMedia/Private/Synchronization/SyncPolicies/DisplayClusterMediaOutputSynchronizationPolicyVblank.cpp
+++ b/Engine/Plugins/Runtime/nDisplay/Source/DisplayClusterMedia/Private/Synchronization/SyncPolicies/DisplayClusterMediaOutputSynchronizationPolicyVblank.cpp
@@ -9,7 +9,18 @@
 #include "DisplayClusterMediaLog.h"
 
 
-bool UDisplayClusterMediaOutputSynchronizationPolicyVblank::StartSynchronization(UMediaCapture* MediaCapture, const FString& MediaId)
+FDisplayClusterMediaOutputSynchronizationPolicyVblankHandler::FDisplayClusterMediaOutputSynchronizationPolicyVblankHandler(UDisplayClusterMediaOutputSynchronizationPolicyVblank* InPolicyObject)
+	: Super(InPolicyObject)
+{
+
+}
+
+TSubclassOf<UDisplayClusterMediaOutputSynchronizationPolicy> FDisplayClusterMediaOutputSynchronizationPolicyVblankHandler::GetPolicyClass() const
+{
+	return UDisplayClusterMediaOutputSynchronizationPolicyVblank::StaticClass();
+}
+
+bool FDisplayClusterMediaOutputSynchronizationPolicyVblankHandler::StartSynchronization(UMediaCapture* MediaCapture, const FString& MediaId)
 {
 	// Get monitor interface if not yet available
 	if (!VblankMonitor)
@@ -29,7 +40,7 @@ bool UDisplayClusterMediaOutputSynchronizationPolicyVblank::StartSynchronization
 	return Super::StartSynchronization(MediaCapture, MediaId);
 }
 
-double UDisplayClusterMediaOutputSynchronizationPolicyVblank::GetTimeBeforeNextSyncPoint()
+double FDisplayClusterMediaOutputSynchronizationPolicyVblankHandler::GetTimeBeforeNextSyncPoint()
 {
 	// Compute time remaining
 	if(VblankMonitor->IsVblankTimeDataAvailable())
@@ -42,3 +53,13 @@ double UDisplayClusterMediaOutputSynchronizationPolicyVblank::GetTimeBeforeNextS
 	// Return 1sec (a huge duration) to prevent calling thread blocking
 	return 1;
 }
+
+TSharedPtr<IDisplayClusterMediaOutputSynchronizationPolicyHandler> UDisplayClusterMediaOutputSynchronizationPolicyVblank::GetHandler()
+{
+	if (!Handler)
+	{
+		Handler = MakeShared<FDisplayClusterMediaOutputSynchronizationPolicyVblankHandler>(this);
+	}
+
+	return Handler;
+}
diff --git a/Engine/Plugins/Runtime/nDisplay/Source/DisplayClusterMedia/Public/Synchronization/DisplayClusterMediaOutputSynchronizationPolicyEthernetBarrier.h b/Engine/Plugins/Runtime/nDisplay/Source/DisplayClusterMedia/Public/Synchronization/DisplayClusterMediaOutputSynchronizationPolicyEthernetBarrier.h
--- a/Engine/Plugins/Runtime/nDisplay/Source/DisplayClusterMedia/Public/Synchronization/DisplayClusterMediaOutputSynchronizationPolicyEthernetBarrier.h
+++ b/Engine/Plugins/Runtime/nDisplay/Source/DisplayClusterMedia/Public/Synchronization/DisplayClusterMediaOutputSynchronizationPolicyEthernetBarrier.h
@@ -2,11 +2,31 @@
 
 #pragma once
 
-#include "CoreMinimal.h"
 #include "Synchronization/DisplayClusterMediaOutputSynchronizationPolicyEthernetBarrierBase.h"
 
 #include "DisplayClusterMediaOutputSynchronizationPolicyEthernetBarrier.generated.h"
 
+/*
+ * Synchronization logic handler class for UDisplayClusterMediaOutputSynchronizationPolicyEthernetBarrier.
+ */
+class DISPLAYCLUSTERMEDIA_API FDisplayClusterMediaOutputSynchronizationPolicyEthernetBarrierHandler
+	: public FDisplayClusterMediaOutputSynchronizationPolicyEthernetBarrierBaseHandler
+{
+	using Super = FDisplayClusterMediaOutputSynchronizationPolicyEthernetBarrierBaseHandler;
+
+public:
+
+	FDisplayClusterMediaOutputSynchronizationPolicyEthernetBarrierHandler(UDisplayClusterMediaOutputSynchronizationPolicyEthernetBarrier* InPolicyObject);
+
+	//~ Begin IDisplayClusterMediaOutputSynchronizationPolicyHandler interface
+	virtual TSubclassOf<UDisplayClusterMediaOutputSynchronizationPolicy> GetPolicyClass() const override;
+	//~ End IDisplayClusterMediaOutputSynchronizationPolicyHandler interface
+
+protected:
+	//~ Begin FDisplayClusterMediaOutputSynchronizationPolicyEthernetBarrierBaseHandler interface
+	virtual void Synchronize() override;
+	//~ End FDisplayClusterMediaOutputSynchronizationPolicyEthernetBarrierBaseHandler interface
+};
 
 /*
  * EthernetBarrier media synchronization policy implementation
@@ -18,6 +38,8 @@ class DISPLAYCLUSTERMEDIA_API UDisplayClusterMediaOutputSynchronizationPolicyEth
 	GENERATED_BODY()
 
 protected:
-	/** Synchronization procedure implementation. */
-	virtual void Synchronize() override;
+	virtual TSharedPtr<IDisplayClusterMediaOutputSynchronizationPolicyHandler> GetHandler() override;
+
+protected:
+	TSharedPtr<IDisplayClusterMediaOutputSynchronizationPolicyHandler> Handler;
 };
diff --git a/Engine/Plugins/Runtime/nDisplay/Source/DisplayClusterMedia/Public/Synchronization/DisplayClusterMediaOutputSynchronizationPolicyEthernetBarrierBase.h b/Engine/Plugins/Runtime/nDisplay/Source/DisplayClusterMedia/Public/Synchronization/DisplayClusterMediaOutputSynchronizationPolicyEthernetBarrierBase.h
--- a/Engine/Plugins/Runtime/nDisplay/Source/DisplayClusterMedia/Public/Synchronization/DisplayClusterMediaOutputSynchronizationPolicyEthernetBarrierBase.h
+++ b/Engine/Plugins/Runtime/nDisplay/Source/DisplayClusterMedia/Public/Synchronization/DisplayClusterMediaOutputSynchronizationPolicyEthernetBarrierBase.h
@@ -2,52 +2,35 @@
 
 #pragma once
 
-#include "CoreMinimal.h"
-#include "DisplayClusterConfigurationTypes_Media.h"
+#include "DisplayClusterConfigurationTypes_MediaSync.h"
 
 #include "Cluster/IDisplayClusterGenericBarriersClient.h"
 
 #include "DisplayClusterMediaOutputSynchronizationPolicyEthernetBarrierBase.generated.h"
 
-
 /*
- * Base class for Ethernet barrier based media synchronization policies.
- * 
- * It encapsulates network barriers related logic.
+ * Synchronization logic handler class for UDisplayClusterMediaOutputSynchronizationPolicyEthernetBarrierBase.
  */
-UCLASS(Abstract, editinlinenew, BlueprintType, hidecategories = (Object))
-class DISPLAYCLUSTERMEDIA_API UDisplayClusterMediaOutputSynchronizationPolicyEthernetBarrierBase
-	: public UDisplayClusterMediaOutputSynchronizationPolicy
+class DISPLAYCLUSTERMEDIA_API FDisplayClusterMediaOutputSynchronizationPolicyEthernetBarrierBaseHandler
+								: public IDisplayClusterMediaOutputSynchronizationPolicyHandler
 {
-	GENERATED_BODY()
-
 public:
-	/** Starts synchronization of specific capture device. Returns false if failed. */
-	virtual bool StartSynchronization(UMediaCapture* MediaCapture, const FString& MediaId) override;
 
-	/** Stops synchronization of specific output stream (capture device). */
-	virtual void StopSynchronization() override;
+	FDisplayClusterMediaOutputSynchronizationPolicyEthernetBarrierBaseHandler(UDisplayClusterMediaOutputSynchronizationPolicyEthernetBarrierBase* InPolicyObject);
 
-	/** Returns true if currently synchronising a media output. */
+	//~ Begin IDisplayClusterMediaOutputSynchronizationPolicyHandler interface
+	virtual bool StartSynchronization(UMediaCapture* MediaCapture, const FString& MediaId) override;
+	virtual void StopSynchronization() override;
 	virtual bool IsRunning() override final;
-
-protected:
-	/** Synchronizes calling thread at the barrier. */
-	void SyncThreadOnBarrier();
+	//~ End IDisplayClusterMediaOutputSynchronizationPolicyHandler interface
 
 	/** Children implement their own sync approaches. */
-	virtual void Synchronize() PURE_VIRTUAL(UDisplayClusterMediaOutputSynchronizationPolicyEthernetBarrierBase::Synchronize, )
+	virtual void Synchronize() = 0;
 
 protected:
-	/** Capture device being used. */
-	UPROPERTY(Transient)
-	TObjectPtr<UMediaCapture> CapturingDevice;
-
-	/** Barrier timeout (ms) */
-	UPROPERTY(EditAnywhere, BlueprintReadWrite, Category = "Settings", meta = (DisplayName = "Barrier Timeout (ms)", ClampMin = "1", ClampMax = "10000", UIMin = "1", UIMax = "10000"))
-	int32 BarrierTimeoutMs = 3000;
+	/** Synchronizes calling thread at the barrier. */
+	void SyncThreadOnBarrier();
 
-protected:
 	/** Returns media device ID being synchronized */
 	FString GetMediaDeviceId() const;
 
@@ -67,6 +50,14 @@ private:
 	/** Handles media capture sync callbacks. */
 	void ProcessMediaSynchronizationCallback();
 
+protected:
+
+	/** Capture device being used. */
+	TObjectPtr<UMediaCapture> CapturingDevice;
+
+	/** Barrier timeout (ms) */
+	int32 BarrierTimeoutMs = 3000;
+
 private:
 	/** Is synchronization currently active. */
 	bool bIsRunning = false;
@@ -83,3 +74,21 @@ private:
 	/** Barrier sync client. */
 	TUniquePtr<IDisplayClusterGenericBarriersClient> EthernetBarrierClient;
 };
+
+/*
+ * Base class for Ethernet barrier based media synchronization policies.
+ * 
+ * It encapsulates network barriers related settings.
+ */
+UCLASS(Abstract, editinlinenew, BlueprintType, hidecategories = (Object))
+class DISPLAYCLUSTERMEDIA_API UDisplayClusterMediaOutputSynchronizationPolicyEthernetBarrierBase
+	: public UDisplayClusterMediaOutputSynchronizationPolicy
+{
+	GENERATED_BODY()
+
+public:
+
+	/** Barrier timeout (ms) */
+	UPROPERTY(EditAnywhere, BlueprintReadWrite, Category = "Settings", meta = (DisplayName = "Barrier Timeout (ms)", ClampMin = "1", ClampMax = "10000", UIMin = "1", UIMax = "10000"))
+	int32 BarrierTimeoutMs = 3000;
+};
diff --git a/Engine/Plugins/Runtime/nDisplay/Source/DisplayClusterMedia/Public/Synchronization/DisplayClusterMediaOutputSynchronizationPolicyThresholdBase.h b/Engine/Plugins/Runtime/nDisplay/Source/DisplayClusterMedia/Public/Synchronization/DisplayClusterMediaOutputSynchronizationPolicyThresholdBase.h
--- a/Engine/Plugins/Runtime/nDisplay/Source/DisplayClusterMedia/Public/Synchronization/DisplayClusterMediaOutputSynchronizationPolicyThresholdBase.h
+++ b/Engine/Plugins/Runtime/nDisplay/Source/DisplayClusterMedia/Public/Synchronization/DisplayClusterMediaOutputSynchronizationPolicyThresholdBase.h
@@ -2,11 +2,34 @@
 
 #pragma once
 
-#include "CoreMinimal.h"
 #include "Synchronization/DisplayClusterMediaOutputSynchronizationPolicyEthernetBarrierBase.h"
 
 #include "DisplayClusterMediaOutputSynchronizationPolicyThresholdBase.generated.h"
 
+/*
+ * Synchronization logic handler class for UDisplayClusterMediaOutputSynchronizationPolicyThresholdBase.
+ */
+class DISPLAYCLUSTERMEDIA_API FDisplayClusterMediaOutputSynchronizationPolicyThresholdBaseHandler
+	: public FDisplayClusterMediaOutputSynchronizationPolicyEthernetBarrierBaseHandler
+{
+	using Super = FDisplayClusterMediaOutputSynchronizationPolicyEthernetBarrierBaseHandler;
+public:
+	
+	FDisplayClusterMediaOutputSynchronizationPolicyThresholdBaseHandler(UDisplayClusterMediaOutputSynchronizationPolicyThresholdBase* InPolicyObject);
+
+protected:
+	//~ Begin FDisplayClusterMediaOutputSynchronizationPolicyEthernetBarrierBaseHandler interface
+	virtual void Synchronize() override;
+	//~ End FDisplayClusterMediaOutputSynchronizationPolicyEthernetBarrierBaseHandler interface
+
+protected:
+	/** Returns amount of time before next synchronization point. */
+	virtual double GetTimeBeforeNextSyncPoint() = 0;
+
+protected:
+	/** Synchronization margin (ms) */
+	int32 MarginMs = 5;
+};
 
 /*
  * Base class for threshold based media synchronization policies.
@@ -19,14 +42,7 @@ class DISPLAYCLUSTERMEDIA_API UDisplayClusterMediaOutputSynchronizationPolicyThr
 {
 	GENERATED_BODY()
 
-protected:
-	/** Synchronization procedure implementation. */
-	virtual void Synchronize() override;
-
-	/** Returns amount of time before next synchronization point. */
-	virtual double GetTimeBeforeNextSyncPoint() PURE_VIRTUAL(UDisplayClusterMediaOutputSynchronizationPolicyThresholdBase::GetTimeBeforeNextSyncPoint, return 0;)
-
-protected:
+public:
 	/** Synchronization margin (ms) */
 	UPROPERTY(EditAnywhere, BlueprintReadWrite, Category = "Settings", meta = (DisplayName = "Margin (ms)", ClampMin = "1", ClampMax = "20", UIMin = "1", UIMax = "20"))
 	int32 MarginMs = 5;
diff --git a/Engine/Plugins/Runtime/nDisplay/Source/DisplayClusterMedia/Public/Synchronization/DisplayClusterMediaOutputSynchronizationPolicyVblank.h b/Engine/Plugins/Runtime/nDisplay/Source/DisplayClusterMedia/Public/Synchronization/DisplayClusterMediaOutputSynchronizationPolicyVblank.h
--- a/Engine/Plugins/Runtime/nDisplay/Source/DisplayClusterMedia/Public/Synchronization/DisplayClusterMediaOutputSynchronizationPolicyVblank.h
+++ b/Engine/Plugins/Runtime/nDisplay/Source/DisplayClusterMedia/Public/Synchronization/DisplayClusterMediaOutputSynchronizationPolicyVblank.h
@@ -2,31 +2,50 @@
 
 #pragma once
 
-#include "CoreMinimal.h"
 #include "Synchronization/DisplayClusterMediaOutputSynchronizationPolicyThresholdBase.h"
 
 #include "DisplayClusterMediaOutputSynchronizationPolicyVblank.generated.h"
 
 class IDisplayClusterVblankMonitor;
 
-
 /*
- * Vblank media synchronization policy implementation
+ * Synchronization logic handler class for UDisplayClusterMediaOutputSynchronizationPolicyVblank.
  */
-UCLASS(editinlinenew, Blueprintable, meta = (DisplayName = "V-blank"))
-class DISPLAYCLUSTERMEDIA_API UDisplayClusterMediaOutputSynchronizationPolicyVblank
-	: public UDisplayClusterMediaOutputSynchronizationPolicyThresholdBase
+class DISPLAYCLUSTERMEDIA_API FDisplayClusterMediaOutputSynchronizationPolicyVblankHandler
+	: public FDisplayClusterMediaOutputSynchronizationPolicyThresholdBaseHandler
 {
-	GENERATED_BODY()
+	using Super = FDisplayClusterMediaOutputSynchronizationPolicyThresholdBaseHandler;
 
 public:
+
+	FDisplayClusterMediaOutputSynchronizationPolicyVblankHandler(UDisplayClusterMediaOutputSynchronizationPolicyVblank* InPolicyObject);
+
+	//~ Begin IDisplayClusterMediaOutputSynchronizationPolicyHandler interface
+	virtual TSubclassOf<UDisplayClusterMediaOutputSynchronizationPolicy> GetPolicyClass() const override;
 	virtual bool StartSynchronization(UMediaCapture* MediaCapture, const FString& MediaId) override;
+	//~ End IDisplayClusterMediaOutputSynchronizationPolicyHandler interface
 
 protected:
 	/** Returns amount of time before next synchronization point. */
 	virtual double GetTimeBeforeNextSyncPoint() override;
-
+	
 private:
 	// V-blank monitor
-	TSharedPtr<IDisplayClusterVblankMonitor, ESPMode::ThreadSafe> VblankMonitor;
+	TSharedPtr<IDisplayClusterVblankMonitor> VblankMonitor;
+};
+
+/*
+ * Vblank media synchronization policy config
+ */
+UCLASS(editinlinenew, Blueprintable, meta = (DisplayName = "V-blank"))
+class DISPLAYCLUSTERMEDIA_API UDisplayClusterMediaOutputSynchronizationPolicyVblank
+	: public UDisplayClusterMediaOutputSynchronizationPolicyThresholdBase
+{
+	GENERATED_BODY()
+
+public:
+	virtual TSharedPtr<IDisplayClusterMediaOutputSynchronizationPolicyHandler> GetHandler() override;
+
+protected:
+	TSharedPtr<IDisplayClusterMediaOutputSynchronizationPolicyHandler> Handler;
 };
diff --git a/Engine/Plugins/VirtualProduction/Rivermax/RivermaxCore/Source/RivermaxCore/Private/RivermaxOutputStream.cpp b/Engine/Plugins/VirtualProduction/Rivermax/RivermaxCore/Source/RivermaxCore/Private/RivermaxOutputStream.cpp
--- a/Engine/Plugins/VirtualProduction/Rivermax/RivermaxCore/Source/RivermaxCore/Private/RivermaxOutputStream.cpp
+++ b/Engine/Plugins/VirtualProduction/Rivermax/RivermaxCore/Source/RivermaxCore/Private/RivermaxOutputStream.cpp
@@ -367,12 +367,21 @@ namespace UE::RivermaxCore::Private
 						StreamData.FrameFieldTimeIntervalNs = 1E9 / Options.FrameRate.AsDecimal();
 						InitializeStreamTimingSettings();
 
-						UE_LOG(LogRivermax, Display, TEXT("Output stream started sending on stream %s:%d on interface %s%s")
+						TStringBuilder<512> StreamDescription;
+						StreamDescription.Appendf(TEXT("Output stream started sending on stream %s:%d using interface %s%s. ")
 							, *Options.StreamAddress
 							, Options.Port
 							, *Options.InterfaceAddress
 							, bUseGPUDirect ? TEXT(" using GPUDirect") : TEXT(""));
 
+						StreamDescription.Appendf(TEXT("Settings: Resolution = %dx%d, "), Options.AlignedResolution.X, Options.AlignedResolution.Y);
+						StreamDescription.Appendf(TEXT("FrameRate = %s, "), *Options.FrameRate.ToPrettyText().ToString());
+						StreamDescription.Appendf(TEXT("Pixel format = %s, "), LexToString(Options.PixelFormat));
+						StreamDescription.Appendf(TEXT("Alignment = %s, "), LexToString(Options.AlignmentMode));
+						StreamDescription.Appendf(TEXT("Framelocking = %s."), LexToString(Options.FrameLockingMode));
+						
+						UE_LOG(LogRivermax, Display, TEXT("%s"), *FString(StreamDescription));
+
 						bIsActive = true;
 						RivermaxThread.Reset(FRunnableThread::Create(this, TEXT("Rmax OutputStream Thread"), 128 * 1024, TPri_TimeCritical, FPlatformAffinity::GetPoolThreadMask()));
 					}
diff --git a/Engine/Plugins/VirtualProduction/Rivermax/RivermaxCore/Source/RivermaxCore/Public/RivermaxFormats.h b/Engine/Plugins/VirtualProduction/Rivermax/RivermaxCore/Source/RivermaxCore/Public/RivermaxFormats.h
--- a/Engine/Plugins/VirtualProduction/Rivermax/RivermaxCore/Source/RivermaxCore/Public/RivermaxFormats.h
+++ b/Engine/Plugins/VirtualProduction/Rivermax/RivermaxCore/Source/RivermaxCore/Public/RivermaxFormats.h
@@ -27,6 +27,81 @@ namespace UE::RivermaxCore
 		, RGB_16bitFloat
 	};
 
+	inline const TCHAR* LexToString(ESamplingType InType)
+	{
+		switch (InType)
+		{
+
+		case ESamplingType::YUV444_10bit:
+		{
+			return TEXT("YUV444 10bit");
+		}
+		case ESamplingType::RGB_10bit:
+		{
+			return TEXT("RGB 10bit");
+		}
+		case ESamplingType::YUV444_8bit:
+		{
+			return TEXT("YUV444 8bit");
+		}
+		case ESamplingType::RGB_8bit:
+		{
+			return TEXT("RGB 8bit");
+		}
+		case ESamplingType::YUV444_12bit:
+		{
+			return TEXT("YUV444 12bit");
+		}
+		case ESamplingType::RGB_12bit:
+		{
+			return TEXT("RGB 12bit");
+		}
+		case ESamplingType::YUV444_16bit:
+		{
+			return TEXT("YUV444 16bit");
+		}
+		case ESamplingType::YUV444_16bitFloat:
+		{
+			return TEXT("YUV444 16bit float");
+		}
+		case ESamplingType::RGB_16bit:
+		{
+			return TEXT("RGB 16bit");
+		}
+		case ESamplingType::RGB_16bitFloat:
+		{
+			return TEXT("RGB 16bit float");
+		}
+		case ESamplingType::YUV422_8bit:
+		{
+			return TEXT("YUV422 8bit");
+		}
+		case ESamplingType::YUV422_10bit:
+		{
+			return TEXT("YUV422 10bit");
+		}
+		case ESamplingType::YUV422_12bit:
+		{
+			return TEXT("YUV422 12bit");
+		}
+		case ESamplingType::YUV422_16bit:
+		{
+			return TEXT("YUV422 16bit");
+		}
+		case ESamplingType::YUV422_16bitFloat:
+		{
+			return TEXT("YUV422 16bit float");
+		}
+		default:
+		{
+			checkNoEntry();
+			break;
+		}
+		}
+
+		return TEXT("<Unknown ESamplingType>");;
+	}
+
 	struct RIVERMAXCORE_API FVideoFormatInfo
 	{
 		/** Sampling type of this format */
diff --git a/Engine/Plugins/VirtualProduction/Rivermax/RivermaxCore/Source/RivermaxCore/Public/RivermaxTypes.h b/Engine/Plugins/VirtualProduction/Rivermax/RivermaxCore/Source/RivermaxCore/Public/RivermaxTypes.h
--- a/Engine/Plugins/VirtualProduction/Rivermax/RivermaxCore/Source/RivermaxCore/Public/RivermaxTypes.h
+++ b/Engine/Plugins/VirtualProduction/Rivermax/RivermaxCore/Source/RivermaxCore/Public/RivermaxTypes.h
@@ -18,6 +18,27 @@ namespace UE::RivermaxCore
 		FrameCreation,
 	};
 
+	inline const TCHAR* LexToString(ERivermaxAlignmentMode InValue)
+	{
+		switch (InValue)
+		{
+			case ERivermaxAlignmentMode::AlignmentPoint:
+			{
+				return TEXT("Alignment point");
+			}
+			case ERivermaxAlignmentMode::FrameCreation:
+			{
+				return TEXT("Frame creation");
+			}
+			default:
+			{
+				checkNoEntry();
+			}
+		}
+
+		return TEXT("<Unknown ERivermaxAlignmentMode>");
+	}
+
 	enum class EFrameLockingMode : uint8
 	{
 		/** If no frame available, continue */
@@ -27,6 +48,27 @@ namespace UE::RivermaxCore
 		BlockOnReservation,
 	};
 
+	inline const TCHAR* LexToString(EFrameLockingMode InValue)
+	{
+		switch (InValue)
+		{
+			case EFrameLockingMode::FreeRun:
+			{
+				return TEXT("Freerun");
+			}
+			case EFrameLockingMode::BlockOnReservation:
+			{
+				return TEXT("Blocking");
+			}
+			default:
+			{
+				checkNoEntry();
+			}
+		}
+
+		return TEXT("<Unknown EFrameLockingMode>");
+	}
+
 	struct RIVERMAXCORE_API FRivermaxInputStreamOptions
 	{
 		/** Stream FrameRate */
diff --git a/Engine/Plugins/VirtualProduction/Rivermax/RivermaxSync/Source/RivermaxSync/Private/MediaOutputSynchronizationPolicyRivermax.cpp b/Engine/Plugins/VirtualProduction/Rivermax/RivermaxSync/Source/RivermaxSync/Private/MediaOutputSynchronizationPolicyRivermax.cpp
--- a/Engine/Plugins/VirtualProduction/Rivermax/RivermaxSync/Source/RivermaxSync/Private/MediaOutputSynchronizationPolicyRivermax.cpp
+++ b/Engine/Plugins/VirtualProduction/Rivermax/RivermaxSync/Source/RivermaxSync/Private/MediaOutputSynchronizationPolicyRivermax.cpp
@@ -8,7 +8,18 @@
 #include "RivermaxPTPUtils.h"
 
 
-bool UMediaOutputSynchronizationPolicyRivermax::IsCaptureTypeSupported(UMediaCapture* MediaCapture) const
+FMediaOutputSynchronizationPolicyRivermaxHandler::FMediaOutputSynchronizationPolicyRivermaxHandler(UMediaOutputSynchronizationPolicyRivermax* InPolicyObject)
+	: Super(InPolicyObject)
+{
+
+}
+
+TSubclassOf<UDisplayClusterMediaOutputSynchronizationPolicy> FMediaOutputSynchronizationPolicyRivermaxHandler::GetPolicyClass() const
+{
+	return UMediaOutputSynchronizationPolicyRivermax::StaticClass();
+}
+
+bool FMediaOutputSynchronizationPolicyRivermaxHandler::IsCaptureTypeSupported(UMediaCapture* MediaCapture) const
 {
 	// We need to make sure:
 	// - it's RivermaxCapture
@@ -37,7 +48,7 @@ bool UMediaOutputSynchronizationPolicyRivermax::IsCaptureTypeSupported(UMediaCap
 	return false;
 }
 
-double UMediaOutputSynchronizationPolicyRivermax::GetTimeBeforeNextSyncPoint()
+double FMediaOutputSynchronizationPolicyRivermaxHandler::GetTimeBeforeNextSyncPoint()
 {
 	if (URivermaxMediaCapture* RmaxCapture = Cast<URivermaxMediaCapture>(CapturingDevice))
 	{
@@ -68,3 +79,13 @@ double UMediaOutputSynchronizationPolicyRivermax::GetTimeBeforeNextSyncPoint()
 	// to prevent calling thread blocking. 1 second is more than any possible threshold.
 	return 1.f;
 }
+
+TSharedPtr<IDisplayClusterMediaOutputSynchronizationPolicyHandler> UMediaOutputSynchronizationPolicyRivermax::GetHandler()
+{
+	if (!Handler)
+	{
+		Handler = MakeShared<FMediaOutputSynchronizationPolicyRivermaxHandler>(this);
+	}
+
+	return Handler;
+}
diff --git a/Engine/Plugins/VirtualProduction/Rivermax/RivermaxSync/Source/RivermaxSync/Public/MediaOutputSynchronizationPolicyRivermax.h b/Engine/Plugins/VirtualProduction/Rivermax/RivermaxSync/Source/RivermaxSync/Public/MediaOutputSynchronizationPolicyRivermax.h
--- a/Engine/Plugins/VirtualProduction/Rivermax/RivermaxSync/Source/RivermaxSync/Public/MediaOutputSynchronizationPolicyRivermax.h
+++ b/Engine/Plugins/VirtualProduction/Rivermax/RivermaxSync/Source/RivermaxSync/Public/MediaOutputSynchronizationPolicyRivermax.h
@@ -2,11 +2,34 @@
 
 #pragma once
 
-#include "CoreMinimal.h"
 #include "Synchronization/DisplayClusterMediaOutputSynchronizationPolicyThresholdBase.h"
 
 #include "MediaOutputSynchronizationPolicyRivermax.generated.h"
 
+/*
+ * Synchronization logic handler class for UMediaOutputSynchronizationPolicyRivermax.
+ */
+class RIVERMAXSYNC_API FMediaOutputSynchronizationPolicyRivermaxHandler
+	: public FDisplayClusterMediaOutputSynchronizationPolicyThresholdBaseHandler
+{
+	using Super = FDisplayClusterMediaOutputSynchronizationPolicyThresholdBaseHandler;
+public:
+
+	FMediaOutputSynchronizationPolicyRivermaxHandler(UMediaOutputSynchronizationPolicyRivermax* InPolicyObject);
+
+	//~ Begin IDisplayClusterMediaOutputSynchronizationPolicyHandler interface
+	virtual TSubclassOf<UDisplayClusterMediaOutputSynchronizationPolicy> GetPolicyClass() const override;
+	//~ End IDisplayClusterMediaOutputSynchronizationPolicyHandler interface
+
+public:
+	/** Returns true if specified media capture type can be synchonized by the policy implementation */
+	virtual bool IsCaptureTypeSupported(UMediaCapture* MediaCapture) const override;
+
+protected:
+	/** Returns amount of time before next synchronization point. */
+	virtual double GetTimeBeforeNextSyncPoint() override;
+};
+
 
 /*
  * Rivermax media synchronization policy implementation
@@ -17,11 +40,10 @@ class RIVERMAXSYNC_API UMediaOutputSynchronizationPolicyRivermax
 {
 	GENERATED_BODY()
 
+
 public:
-	/** Returns true if specified media capture type can be synchonized by the policy implementation */
-	virtual bool IsCaptureTypeSupported(UMediaCapture* MediaCapture) const override;
+	virtual TSharedPtr<IDisplayClusterMediaOutputSynchronizationPolicyHandler> GetHandler() override;
 
 protected:
-	/** Returns amount of time before next synchronization point. */
-	virtual double GetTimeBeforeNextSyncPoint() override;
+	TSharedPtr<IDisplayClusterMediaOutputSynchronizationPolicyHandler> Handler;
 };
