nDisplay: Added IDisplayClusterWarpPolicy interface

diff --git a/Engine/Plugins/Runtime/nDisplay/Source/DisplayCluster/Private/Render/DisplayClusterRenderManager.cpp b/Engine/Plugins/Runtime/nDisplay/Source/DisplayCluster/Private/Render/DisplayClusterRenderManager.cpp
--- a/Engine/Plugins/Runtime/nDisplay/Source/DisplayCluster/Private/Render/DisplayClusterRenderManager.cpp
+++ b/Engine/Plugins/Runtime/nDisplay/Source/DisplayCluster/Private/Render/DisplayClusterRenderManager.cpp
@@ -396,6 +396,9 @@ TSharedPtr<IDisplayClusterRenderSyncPolicy> FDisplayClusterRenderManager::GetCur
 	return SyncPolicy;
 }
 
+//------------------------------------------------------------------------------------------------------------------------------
+// Projection Policy
+//------------------------------------------------------------------------------------------------------------------------------
 bool FDisplayClusterRenderManager::RegisterProjectionPolicyFactory(const FString& InProjectionType, TSharedPtr<IDisplayClusterProjectionPolicyFactory>& InFactory)
 {
 	UE_LOG(LogDisplayClusterRender, Log, TEXT("Registering factory for projection type: %s"), *InProjectionType);
@@ -462,7 +465,9 @@ void FDisplayClusterRenderManager::GetRegisteredProjectionPolicies(TArray<FStrin
 	ProjectionPolicyFactories.GetKeys(OutPolicyIDs);
 }
 
-
+//------------------------------------------------------------------------------------------------------------------------------
+// PostProcess
+//------------------------------------------------------------------------------------------------------------------------------
 bool FDisplayClusterRenderManager::RegisterPostProcessFactory(const FString& InPostProcessType, TSharedPtr<IDisplayClusterPostProcessFactory>& InFactory)
 {
 	UE_LOG(LogDisplayClusterRender, Log, TEXT("Registering factory for postprocess type: %s"), *InPostProcessType);
@@ -529,6 +534,78 @@ void FDisplayClusterRenderManager::GetRegisteredPostProcess(TArray<FString>& Out
 	PostProcessFactories.GetKeys(OutPostProcessIDs);
 }
 
+//------------------------------------------------------------------------------------------------------------------------------
+// Warp Policy
+//------------------------------------------------------------------------------------------------------------------------------
+bool FDisplayClusterRenderManager::RegisterWarpPolicyFactory(const FString& InWarpPolicyType, TSharedPtr<IDisplayClusterWarpPolicyFactory>& InFactory)
+{
+	UE_LOG(LogDisplayClusterRender, Log, TEXT("Registering factory for warp policy type: %s"), *InWarpPolicyType);
+
+	if (!InFactory.IsValid())
+	{
+		UE_LOG(LogDisplayClusterRender, Warning, TEXT("Invalid factory object"));
+		return false;
+	}
+
+	{
+		FScopeLock Lock(&CritSecInternals);
+
+		if (WarpPolicyFactories.Contains(InWarpPolicyType))
+		{
+			UE_LOG(LogDisplayClusterRender, Warning, TEXT("A new factory for '%s' warp policy was set"), *InWarpPolicyType);
+		}
+
+		WarpPolicyFactories.Emplace(InWarpPolicyType, InFactory);
+	}
+
+	UE_LOG(LogDisplayClusterRender, Log, TEXT("Registered factory for warp policy type: %s"), *InWarpPolicyType);
+
+	return true;
+}
+
+bool FDisplayClusterRenderManager::UnregisterWarpPolicyFactory(const FString& InWarpPolicyType)
+{
+	UE_LOG(LogDisplayClusterRender, Log, TEXT("Unregistering factory for warp policy: %s"), *InWarpPolicyType);
+
+	{
+		FScopeLock Lock(&CritSecInternals);
+
+		if (!WarpPolicyFactories.Contains(InWarpPolicyType))
+		{
+			UE_LOG(LogDisplayClusterRender, Warning, TEXT("A handler for '%s' warp policy type not found"), *InWarpPolicyType);
+			return false;
+		}
+
+		WarpPolicyFactories.Remove(InWarpPolicyType);
+	}
+
+	UE_LOG(LogDisplayClusterRender, Log, TEXT("Unregistered factory for warp policy: %s"), *InWarpPolicyType);
+
+	return true;
+}
+
+TSharedPtr<IDisplayClusterWarpPolicyFactory> FDisplayClusterRenderManager::GetWarpPolicyFactory(const FString& InWarpPolicyType)
+{
+	FScopeLock Lock(&CritSecInternals);
+
+	TSharedPtr<IDisplayClusterWarpPolicyFactory> Factory;
+	if (!DisplayClusterHelpers::map::template ExtractValue(WarpPolicyFactories, InWarpPolicyType, Factory))
+	{
+		UE_LOG(LogDisplayClusterRender, Warning, TEXT("No factory found for warp policy: %s"), *InWarpPolicyType);
+	}
+
+	return Factory;
+}
+
+void FDisplayClusterRenderManager::GetRegisteredWarpPolicies(TArray<FString>& OutWarpPolicyIDs) const
+{
+	FScopeLock Lock(&CritSecInternals);
+	WarpPolicyFactories.GetKeys(OutWarpPolicyIDs);
+}
+
+//------------------------------------------------------------------------------------------------------------------------------
+// Resources
+//------------------------------------------------------------------------------------------------------------------------------
 TSharedPtr<IDisplayClusterRender_MeshComponent, ESPMode::ThreadSafe> FDisplayClusterRenderManager::CreateMeshComponent() const
 {
 	return MakeShared<FDisplayClusterRender_MeshComponent, ESPMode::ThreadSafe>();
diff --git a/Engine/Plugins/Runtime/nDisplay/Source/DisplayCluster/Private/Render/DisplayClusterRenderManager.h b/Engine/Plugins/Runtime/nDisplay/Source/DisplayCluster/Private/Render/DisplayClusterRenderManager.h
--- a/Engine/Plugins/Runtime/nDisplay/Source/DisplayCluster/Private/Render/DisplayClusterRenderManager.h
+++ b/Engine/Plugins/Runtime/nDisplay/Source/DisplayCluster/Private/Render/DisplayClusterRenderManager.h
@@ -62,6 +62,12 @@ public:
 	virtual bool UnregisterPostProcessFactory(const FString& InPostProcessType) override;
 	virtual TSharedPtr<IDisplayClusterPostProcessFactory> GetPostProcessFactory(const FString& InPostProcessType) override;
 	virtual void GetRegisteredPostProcess(TArray<FString>& OutPostProcessIDs) const override;
+	// Warp policy
+	virtual bool RegisterWarpPolicyFactory(const FString& InWarpPolicyType, TSharedPtr<IDisplayClusterWarpPolicyFactory>& Factory) override;
+	virtual bool UnregisterWarpPolicyFactory(const FString& InWarpPolicyType) override;
+	virtual TSharedPtr<IDisplayClusterWarpPolicyFactory> GetWarpPolicyFactory(const FString& InWarpPolicyType) override;
+	virtual void GetRegisteredWarpPolicies(TArray<FString>& OutWarpPolicyIDs) const override;
+	// Resources
 	virtual TSharedPtr<IDisplayClusterRender_MeshComponent, ESPMode::ThreadSafe> CreateMeshComponent() const override;
 	virtual TSharedPtr<IDisplayClusterRender_Texture, ESPMode::ThreadSafe> GetOrCreateCachedTexture(const FString& InUniqueTextureName) const override;
 
@@ -111,6 +117,9 @@ private:
 	// Postprocess internals
 	TMap<FString, TSharedPtr<IDisplayClusterPostProcessFactory>> PostProcessFactories;
 
+	// Warp internals
+	TMap<FString, TSharedPtr<IDisplayClusterWarpPolicyFactory>> WarpPolicyFactories;
+
 private:
 	// Internal data access synchronization
 	mutable FCriticalSection CritSecInternals;
diff --git a/Engine/Plugins/Runtime/nDisplay/Source/DisplayCluster/Private/Render/Viewport/DisplayClusterViewportManager.cpp b/Engine/Plugins/Runtime/nDisplay/Source/DisplayCluster/Private/Render/Viewport/DisplayClusterViewportManager.cpp
--- a/Engine/Plugins/Runtime/nDisplay/Source/DisplayCluster/Private/Render/Viewport/DisplayClusterViewportManager.cpp
+++ b/Engine/Plugins/Runtime/nDisplay/Source/DisplayCluster/Private/Render/Viewport/DisplayClusterViewportManager.cpp
@@ -9,6 +9,7 @@
 #include "Render/IDisplayClusterRenderManager.h"
 #include "Render/Projection/IDisplayClusterProjectionPolicyFactory.h"
 #include "Render/Projection/IDisplayClusterProjectionPolicy.h"
+#include "Render/Warp/IDisplayClusterWarpPolicy.h"
 
 #include "Render/Viewport/DisplayClusterViewport.h"
 #include "Render/Viewport/DisplayClusterViewportProxy.h"
@@ -49,6 +50,82 @@ static FAutoConsoleVariableRef CVarDisplayClusterChromaKeyAllowNanite(
 	ECVF_RenderThreadSafe
 );
 
+namespace UE::DisplayCluster::ViewportManager
+{
+	/**
+	 * Container for the warp policy with its associated group of viewports.
+	 */
+	struct FDisplayClusterWarpPolicyViewportsGroup
+	{
+		/** HandleNewFrame() call for the warp policy interface. */
+		inline void HandleNewFrame() const
+		{
+			if (WarpPolicy.IsValid())
+			{
+				WarpPolicy->HandleNewFrame(Viewports);
+			}
+		}
+
+		// Warp polic instance
+		TSharedPtr<IDisplayClusterWarpPolicy, ESPMode::ThreadSafe> WarpPolicy;
+
+		// A group of viewports using the same warp policy
+		TArray<TSharedPtr<IDisplayClusterViewport, ESPMode::ThreadSafe>> Viewports;
+	};
+
+	/**
+	 * A helper class for handling the warp policies.
+	 */
+	struct FDisplayClusterWarpPolicyManager
+	{
+		/** Handle new frame for warp policies. */
+		inline void HandleNewFrame(FDisplayClusterViewportManager& ViewportManager)
+		{
+			// Collect warp policies for the viewports of the entire cluster.
+			// this requires that all viewports in the cluster exist at the moment.
+			for (const TSharedPtr<FDisplayClusterViewport, ESPMode::ThreadSafe>& Viewport : ViewportManager.ImplGetEntireClusterViewports())
+			{
+				if (IDisplayClusterWarpPolicy* ViewportWarpPolicyPtr = (Viewport.IsValid() && Viewport->GetProjectionPolicy().IsValid()) ? Viewport->GetProjectionPolicy()->GetWarpPolicy() : nullptr)
+				{
+					AddViewportToWarpPolicy(ViewportWarpPolicyPtr->ToSharedPtr(), Viewport);
+				}
+			}
+
+			// Processing a new frame event for all warp policies
+			for (const FDisplayClusterWarpPolicyViewportsGroup& ViewportsGroup : WarpPolicies)
+			{
+				ViewportsGroup.HandleNewFrame();
+			}
+		}
+
+	protected:
+		/** Adds a viewport to containers with warp policies and associated viewport groups. */
+		inline void AddViewportToWarpPolicy(const TSharedPtr<IDisplayClusterWarpPolicy, ESPMode::ThreadSafe>& InWarpPolicy, const TSharedPtr<FDisplayClusterViewport, ESPMode::ThreadSafe>& InViewport)
+		{
+			if (FDisplayClusterWarpPolicyViewportsGroup* ViewportsGroupPtr = WarpPolicies.FindByPredicate([InWarpPolicy](const FDisplayClusterWarpPolicyViewportsGroup& ViewportsGroupIt)
+				{
+					return ViewportsGroupIt.WarpPolicy == InWarpPolicy;
+				}))
+			{
+				ViewportsGroupPtr->Viewports.Add(InViewport);
+			}
+			else
+			{
+				FDisplayClusterWarpPolicyViewportsGroup ViewportsGroup;
+				ViewportsGroup.WarpPolicy = InWarpPolicy;
+				ViewportsGroup.Viewports.Add(InViewport);
+
+				WarpPolicies.Add(ViewportsGroup);
+			}
+		}
+
+	private:
+		// Warp policies viewports groups
+		TArray<FDisplayClusterWarpPolicyViewportsGroup> WarpPolicies;
+	};
+};
+using namespace UE::DisplayCluster::ViewportManager;
+
 ///////////////////////////////////////////////////////////////////////////////////////
 //          FDisplayClusterViewportManager
 ///////////////////////////////////////////////////////////////////////////////////////
@@ -666,6 +743,10 @@ void FDisplayClusterViewportManager::InitializeNewFrame()
 		}
 	}
 
+	// Handle new frame for warp policies
+	FDisplayClusterWarpPolicyManager WarpPolicyManager;
+	WarpPolicyManager.HandleNewFrame(*this);
+
 	// Send render frame settings to rendering thread
 	ViewportManagerProxy->ImplUpdateRenderFrameSettings(GetRenderFrameSettings(), ViewportManagerViewExtension);
 
diff --git a/Engine/Plugins/Runtime/nDisplay/Source/DisplayCluster/Public/Render/IDisplayClusterRenderManager.h b/Engine/Plugins/Runtime/nDisplay/Source/DisplayCluster/Public/Render/IDisplayClusterRenderManager.h
--- a/Engine/Plugins/Runtime/nDisplay/Source/DisplayCluster/Public/Render/IDisplayClusterRenderManager.h
+++ b/Engine/Plugins/Runtime/nDisplay/Source/DisplayCluster/Public/Render/IDisplayClusterRenderManager.h
@@ -11,6 +11,7 @@ class IDisplayClusterRender_MeshComponent;
 class IDisplayClusterRender_Texture;
 class IDisplayClusterProjectionPolicy;
 class IDisplayClusterProjectionPolicyFactory;
+class IDisplayClusterWarpPolicyFactory;
 class IDisplayClusterRenderDeviceFactory;
 class IDisplayClusterRenderSyncPolicy;
 class IDisplayClusterRenderSyncPolicyFactory;
@@ -164,6 +165,41 @@ public:
 	*/
 	virtual void GetRegisteredPostProcess(TArray<FString>& OutPostProcessIDs) const = 0;
 
+	/**
+	* Registers warp policy factory
+	*
+	* @param InWarpPolicyType - Type of warp policy
+	* @param Factory          - Factory instance
+	*
+	* @return - True if success
+	*/
+	virtual bool RegisterWarpPolicyFactory(const FString& InWarpPolicyType, TSharedPtr<IDisplayClusterWarpPolicyFactory>& Factory) = 0;
+
+	/**
+	* Unregisters warp policy factory
+	*`
+	* @param InWarpPolicyType - Type of warp policy
+	*
+	* @return - True if success
+	*/
+	virtual bool UnregisterWarpPolicyFactory(const FString& InWarpPolicyType) = 0;
+
+	/**
+	* Returns a warp policy factory of specified type (if it has been registered previously)
+	*
+	* @param InWarpPolicyType - Warp policy type
+	*
+	* @return - Warp policy factory pointer or nullptr if not registered
+	*/
+	virtual TSharedPtr<IDisplayClusterWarpPolicyFactory> GetWarpPolicyFactory(const FString& InWarpPolicyType) = 0;
+
+	/**
+	* Returns all registered warp policy types
+	*
+	* @param OutWarpPolicyIDs - (out) array to put registered IDs
+	*/
+	virtual void GetRegisteredWarpPolicies(TArray<FString>& OutWarpPolicyIDs) const = 0;
+
 	/**
 	* Returns V-blank monitor interface
 	*/
diff --git a/Engine/Plugins/Runtime/nDisplay/Source/DisplayCluster/Public/Render/Projection/IDisplayClusterProjectionPolicy.h b/Engine/Plugins/Runtime/nDisplay/Source/DisplayCluster/Public/Render/Projection/IDisplayClusterProjectionPolicy.h
--- a/Engine/Plugins/Runtime/nDisplay/Source/DisplayCluster/Public/Render/Projection/IDisplayClusterProjectionPolicy.h
+++ b/Engine/Plugins/Runtime/nDisplay/Source/DisplayCluster/Public/Render/Projection/IDisplayClusterProjectionPolicy.h
@@ -8,6 +8,7 @@ class FRHICommandListImmediate;
 class IDisplayClusterViewport;
 class IDisplayClusterViewportProxy;
 class IDisplayClusterWarpBlend;
+class IDisplayClusterWarpPolicy;
 class UMeshComponent;
 struct FDisplayClusterConfigurationProjection;
 
@@ -70,6 +71,30 @@ public:
 	virtual void HandleEndScene(IDisplayClusterViewport* InViewport)
 	{ }
 
+	/**
+	* Set warp policy for this projection
+	*
+	* @param InWarpPolicy - the warp policy instance
+	*/
+	virtual void SetWarpPolicy(IDisplayClusterWarpPolicy* InWarpPolicy)
+	{ }
+
+	/**
+	* Get warp policy
+	*/
+	virtual IDisplayClusterWarpPolicy* GetWarpPolicy() const
+	{
+		return nullptr;
+	}
+
+	/**
+	* Get warp policy on rendering thread
+	*/
+	virtual IDisplayClusterWarpPolicy* GetWarpPolicy_RenderThread() const
+	{
+		return nullptr;
+	}
+
 	// Handle request for additional render targetable resource inside viewport api for projection policy
 	virtual bool ShouldUseAdditionalTargetableResource() const
 	{ 
diff --git a/Engine/Plugins/Runtime/nDisplay/Source/DisplayCluster/Public/Render/Warp/IDisplayClusterWarpPolicy.h b/Engine/Plugins/Runtime/nDisplay/Source/DisplayCluster/Public/Render/Warp/IDisplayClusterWarpPolicy.h
new file mode 100644
--- /dev/null
+++ b/Engine/Plugins/Runtime/nDisplay/Source/DisplayCluster/Public/Render/Warp/IDisplayClusterWarpPolicy.h
@@ -0,0 +1,78 @@
+// Copyright Epic Games, Inc. All Rights Reserved.
+
+#pragma once
+
+#include "CoreMinimal.h"
+
+class IDisplayClusterViewport;
+class IDisplayClusterViewportManager;
+
+/**
+ * Warp policy interface
+ * Customize warp math for projection policies.
+ * Multiple projections can use the same warp policy.
+ */
+class IDisplayClusterWarpPolicy
+{
+public:
+	virtual ~IDisplayClusterWarpPolicy() = default;
+
+public:
+	/** Get TSharedPtr from self. */
+	virtual TSharedPtr<IDisplayClusterWarpPolicy, ESPMode::ThreadSafe> ToSharedPtr() = 0;
+	virtual TSharedPtr<const IDisplayClusterWarpPolicy, ESPMode::ThreadSafe> ToSharedPtr() const = 0;
+
+	/**
+	* Return warp policy instance ID
+	*/
+	virtual const FString& GetId() const = 0;
+
+	/**
+	* Return warp policy type
+	*/
+	virtual const FString& GetType() const = 0;
+
+
+	/** Returns true if this policy supports ICVFX rendering
+	 * 
+	 * @param InViewport - a owner viewport
+	 * 
+	 * @return - true if this warp policy supports ICVFX
+	 */
+	virtual bool ShouldSupportICVFX(IDisplayClusterViewport* InViewport) const
+	{
+		return false;
+	}
+
+	/**
+	* Handle new frame for viewports group.
+	* This function is called once per frame at the beginning, immediately after the viewport setup is complete.
+	* 
+	* @param InViewports - a group of viewports from the entire cluster that use this WarpInterface
+	*/
+	virtual void HandleNewFrame(const TArray<TSharedPtr<IDisplayClusterViewport, ESPMode::ThreadSafe>>& InViewports)
+	{ }
+
+	/**
+	* Tick function
+	* The positions of movable preview components may be updated on each frame.
+	* Use this function to handle these updates.
+	* 
+	* @param InViewportManager - Viewport manager interface
+	* @param DeltaSeconds      - delta time
+	*/
+	virtual void Tick(IDisplayClusterViewportManager* InViewportManager, float DeltaSeconds)
+	{ }
+
+	/** Override frustum for viewport context
+	 * 
+	 * @param InViewport - a owner viewport
+	 * @param ContextNum - viewport eye context index
+	 * 
+	 * @return - true if frustum overrided.
+	 */
+	virtual bool CalcFrustum(IDisplayClusterViewport* InViewport, const uint32 ContextNum)
+	{
+		return false;
+	}
+};
diff --git a/Engine/Plugins/Runtime/nDisplay/Source/DisplayCluster/Public/Render/Warp/IDisplayClusterWarpPolicyFactory.h b/Engine/Plugins/Runtime/nDisplay/Source/DisplayCluster/Public/Render/Warp/IDisplayClusterWarpPolicyFactory.h
new file mode 100644
--- /dev/null
+++ b/Engine/Plugins/Runtime/nDisplay/Source/DisplayCluster/Public/Render/Warp/IDisplayClusterWarpPolicyFactory.h
@@ -0,0 +1,27 @@
+// Copyright Epic Games, Inc. All Rights Reserved.
+
+#pragma once
+
+#include "CoreMinimal.h"
+
+class IDisplayClusterWarpPolicy;
+
+/**
+ * nDisplay warp policy factory interface
+ */
+class IDisplayClusterWarpPolicyFactory
+{
+public:
+	virtual ~IDisplayClusterWarpPolicyFactory() = default;
+
+public:
+	/**
+	* Creates a warp policy instance
+	*
+	* @param InWarpPolicyType - Warp policy type, same as specified on registration (useful if the same factory is registered for multiple warp policy types)
+	* @param InWarpPolicyName - Warp policy name
+	*
+	* @return - Warp policy instance
+	*/
+	virtual TSharedPtr<IDisplayClusterWarpPolicy, ESPMode::ThreadSafe> Create(const FString& InWarpPolicyType, const FString& InWarpPolicyName) = 0;
+};
