[nDisplay] Update component transforms while overriding DCRA from a config

diff --git a/Engine/Plugins/Runtime/nDisplay/Source/DisplayCluster/Private/Game/EngineClasses/Scene/DisplayClusterRootActor.cpp b/Engine/Plugins/Runtime/nDisplay/Source/DisplayCluster/Private/Game/EngineClasses/Scene/DisplayClusterRootActor.cpp
--- a/Engine/Plugins/Runtime/nDisplay/Source/DisplayCluster/Private/Game/EngineClasses/Scene/DisplayClusterRootActor.cpp
+++ b/Engine/Plugins/Runtime/nDisplay/Source/DisplayCluster/Private/Game/EngineClasses/Scene/DisplayClusterRootActor.cpp
@@ -345,6 +345,12 @@ void ADisplayClusterRootActor::OverrideFromConfig(UDisplayClusterConfigurationDa
 		}
 	}
 
+	// Update component transforms if any of them have been changed in the file
+	{
+		TUniquePtr<FDisplayClusterRootActorInitializer> Initializer = MakeUnique<FDisplayClusterRootActorInitializer>();
+		Initializer->UpdateComponentTransformsOnly(this, ConfigData);
+	}
+
 	// There is no sense to call BuildHierarchy because it works for non-BP root actors.
 	// On the other hand, OverwriteFromConfig method is called for BP root actors only by nature.
 
diff --git a/Engine/Plugins/Runtime/nDisplay/Source/DisplayCluster/Private/Game/EngineClasses/Scene/DisplayClusterRootActorInitializer.cpp b/Engine/Plugins/Runtime/nDisplay/Source/DisplayCluster/Private/Game/EngineClasses/Scene/DisplayClusterRootActorInitializer.cpp
--- a/Engine/Plugins/Runtime/nDisplay/Source/DisplayCluster/Private/Game/EngineClasses/Scene/DisplayClusterRootActorInitializer.cpp
+++ b/Engine/Plugins/Runtime/nDisplay/Source/DisplayCluster/Private/Game/EngineClasses/Scene/DisplayClusterRootActorInitializer.cpp
@@ -44,6 +44,41 @@ void FDisplayClusterRootActorInitializer::Initialize(ADisplayClusterRootActor* I
 	Cleanup();
 }
 
+void FDisplayClusterRootActorInitializer::UpdateComponentTransformsOnly(ADisplayClusterRootActor* Actor, UDisplayClusterConfigurationData* ConfigData)
+{
+	if (!(Actor && ConfigData))
+	{
+		return;
+	}
+
+	// Xform components
+	for (const TPair<FString, TObjectPtr<UDisplayClusterConfigurationSceneComponentXform>>& Xform : ConfigData->Scene->Xforms)
+	{
+		if (USceneComponent* const ActorComponent = Actor->GetComponentByName<USceneComponent>(Xform.Key))
+		{
+			ActorComponent->SetRelativeLocationAndRotation(Xform.Value->Location, Xform.Value->Rotation);
+		}
+	}
+
+	// Camera components
+	for (const TPair<FString, TObjectPtr<UDisplayClusterConfigurationSceneComponentCamera>>& Camera : ConfigData->Scene->Cameras)
+	{
+		if (USceneComponent* const ActorComponent = Actor->GetComponentByName<USceneComponent>(Camera.Key))
+		{
+			ActorComponent->SetRelativeLocationAndRotation(Camera.Value->Location, Camera.Value->Rotation);
+		}
+	}
+
+	// Screen components
+	for (const TPair<FString, TObjectPtr<UDisplayClusterConfigurationSceneComponentScreen>>& Screen : ConfigData->Scene->Screens)
+	{
+		if (USceneComponent* const ActorComponent = Actor->GetComponentByName<USceneComponent>(Screen.Key))
+		{
+			ActorComponent->SetRelativeLocationAndRotation(Screen.Value->Location, Screen.Value->Rotation);
+		}
+	}
+}
+
 void FDisplayClusterRootActorInitializer::ApplyConfigDataToComponents()
 {
 	// Configure xform components
@@ -176,7 +211,7 @@ void FDisplayClusterRootActorInitializer::SpawnComponents(const TMap<FString, TC
 		if (!AllComponents.Contains(Item.Key))
 		{
 			// Check if a component with the same name is already stored on the BP.
-			UActorComponent* const ExistingComponent = Cast<TComp>(ProcessingActor->GetComponentByName<UActorComponent>(Item.Key));
+			UActorComponent* const ExistingComponent = ProcessingActor->GetComponentByName<UActorComponent>(Item.Key);
 			if (!ExistingComponent)
 			{
 				// Create new component
diff --git a/Engine/Plugins/Runtime/nDisplay/Source/DisplayCluster/Private/Game/EngineClasses/Scene/DisplayClusterRootActorInitializer.h b/Engine/Plugins/Runtime/nDisplay/Source/DisplayCluster/Private/Game/EngineClasses/Scene/DisplayClusterRootActorInitializer.h
--- a/Engine/Plugins/Runtime/nDisplay/Source/DisplayCluster/Private/Game/EngineClasses/Scene/DisplayClusterRootActorInitializer.h
+++ b/Engine/Plugins/Runtime/nDisplay/Source/DisplayCluster/Private/Game/EngineClasses/Scene/DisplayClusterRootActorInitializer.h
@@ -34,6 +34,9 @@ public:
 	// Main initialization function that manages the whole initialization process
 	void Initialize(ADisplayClusterRootActor* Actor, UDisplayClusterConfigurationData* ConfigData);
 
+	// Updates component transforms only based on the config settings
+	void UpdateComponentTransformsOnly(ADisplayClusterRootActor* Actor, UDisplayClusterConfigurationData* ConfigData);
+
 private:
 	// Applies configuration data to the corresponding components
 	void ApplyConfigDataToComponents();
