nDisplay: fixed copyright boilerplate

diff --git a/Engine/Plugins/Runtime/nDisplay/Source/DisplayClusterReplication/Private/DisplayClusterReplicationModule.cpp b/Engine/Plugins/Runtime/nDisplay/Source/DisplayClusterReplication/Private/DisplayClusterReplicationModule.cpp
--- a/Engine/Plugins/Runtime/nDisplay/Source/DisplayClusterReplication/Private/DisplayClusterReplicationModule.cpp
+++ b/Engine/Plugins/Runtime/nDisplay/Source/DisplayClusterReplication/Private/DisplayClusterReplicationModule.cpp
@@ -1,3 +1,4 @@
+// Copyright Epic Games, Inc. All Rights Reserved.
 #include "Modules/ModuleManager.h"
 
 IMPLEMENT_MODULE(FDefaultModuleImpl, DisplayClusterReplication)
\ No newline at end of file
