- Fixing clang compilation issue with ndisplay media sync interface (Horde issue #320631)

diff --git a/Engine/Plugins/Runtime/nDisplay/Source/DisplayClusterConfiguration/Public/DisplayClusterConfigurationTypes_MediaSync.h b/Engine/Plugins/Runtime/nDisplay/Source/DisplayClusterConfiguration/Public/DisplayClusterConfigurationTypes_MediaSync.h
--- a/Engine/Plugins/Runtime/nDisplay/Source/DisplayClusterConfiguration/Public/DisplayClusterConfigurationTypes_MediaSync.h
+++ b/Engine/Plugins/Runtime/nDisplay/Source/DisplayClusterConfiguration/Public/DisplayClusterConfigurationTypes_MediaSync.h
@@ -21,10 +21,7 @@ public:
 	virtual ~IDisplayClusterMediaOutputSynchronizationPolicyHandler() = default;
 
 	/** Returns true if specified media capture type can be synchonized by the policy implementation */
-	virtual bool IsCaptureTypeSupported(UMediaCapture* MediaCapture) const
-	{
-		return true;
-	}
+	virtual bool IsCaptureTypeSupported(UMediaCapture* MediaCapture) const = 0;
 
 	/** Starts synchronization of specific output stream (capture device). Returns false if failed. */
 	virtual bool StartSynchronization(UMediaCapture* MediaCapture, const FString& MediaId) = 0;
diff --git a/Engine/Plugins/Runtime/nDisplay/Source/DisplayClusterMedia/Private/Synchronization/SyncPolicies/DisplayClusterMediaOutputSynchronizationPolicyEthernetBarrierBase.cpp b/Engine/Plugins/Runtime/nDisplay/Source/DisplayClusterMedia/Private/Synchronization/SyncPolicies/DisplayClusterMediaOutputSynchronizationPolicyEthernetBarrierBase.cpp
--- a/Engine/Plugins/Runtime/nDisplay/Source/DisplayClusterMedia/Private/Synchronization/SyncPolicies/DisplayClusterMediaOutputSynchronizationPolicyEthernetBarrierBase.cpp
+++ b/Engine/Plugins/Runtime/nDisplay/Source/DisplayClusterMedia/Private/Synchronization/SyncPolicies/DisplayClusterMediaOutputSynchronizationPolicyEthernetBarrierBase.cpp
@@ -97,6 +97,11 @@ bool FDisplayClusterMediaOutputSynchronizationPolicyEthernetBarrierBaseHandler::
 	return bIsRunning;
 }
 
+bool FDisplayClusterMediaOutputSynchronizationPolicyEthernetBarrierBaseHandler::IsCaptureTypeSupported(UMediaCapture* MediaCapture) const
+{
+	return true;
+}
+
 FString FDisplayClusterMediaOutputSynchronizationPolicyEthernetBarrierBaseHandler::GetMediaDeviceId() const
 {
 	return MediaDeviceId;
diff --git a/Engine/Plugins/Runtime/nDisplay/Source/DisplayClusterMedia/Public/Synchronization/DisplayClusterMediaOutputSynchronizationPolicyEthernetBarrierBase.h b/Engine/Plugins/Runtime/nDisplay/Source/DisplayClusterMedia/Public/Synchronization/DisplayClusterMediaOutputSynchronizationPolicyEthernetBarrierBase.h
--- a/Engine/Plugins/Runtime/nDisplay/Source/DisplayClusterMedia/Public/Synchronization/DisplayClusterMediaOutputSynchronizationPolicyEthernetBarrierBase.h
+++ b/Engine/Plugins/Runtime/nDisplay/Source/DisplayClusterMedia/Public/Synchronization/DisplayClusterMediaOutputSynchronizationPolicyEthernetBarrierBase.h
@@ -22,6 +22,7 @@ public:
 	virtual bool StartSynchronization(UMediaCapture* MediaCapture, const FString& MediaId) override;
 	virtual void StopSynchronization() override;
 	virtual bool IsRunning() override final;
+	virtual bool IsCaptureTypeSupported(UMediaCapture* MediaCapture) const override;
 	//~ End IDisplayClusterMediaOutputSynchronizationPolicyHandler interface
 
 	/** Children implement their own sync approaches. */
