[nDisplay] Allow multiple media outputs assigned to nDisplay entities (backbuffer, viewport, ICVFX cameras)

diff --git a/Engine/Plugins/Runtime/nDisplay/Source/DisplayCluster/Private/Game/EngineClasses/Scene/DisplayClusterICVFXCameraComponent.cpp b/Engine/Plugins/Runtime/nDisplay/Source/DisplayCluster/Private/Game/EngineClasses/Scene/DisplayClusterICVFXCameraComponent.cpp
--- a/Engine/Plugins/Runtime/nDisplay/Source/DisplayCluster/Private/Game/EngineClasses/Scene/DisplayClusterICVFXCameraComponent.cpp
+++ b/Engine/Plugins/Runtime/nDisplay/Source/DisplayCluster/Private/Game/EngineClasses/Scene/DisplayClusterICVFXCameraComponent.cpp
@@ -88,7 +88,7 @@ bool UDisplayClusterICVFXCameraComponent::IsICVFXEnabled() const
 		static const FString NodeId = GDisplayCluster->GetPrivateClusterMgr()->GetNodeId();
 
 		// First condition to render offscreen: it has media output assigned
-		const bool bUsesMediaOutput = CameraSettings.RenderSettings.Media.bEnable ? !!CameraSettings.RenderSettings.Media.GetMediaOutput(NodeId) : false;
+		const bool bUsesMediaOutput = (CameraSettings.RenderSettings.Media.bEnable && CameraSettings.RenderSettings.Media.IsMediaOutputAssigned(NodeId));
 
 		// Get backbuffer media settings
 		const FDisplayClusterConfigurationMedia* BackbufferMediaSettings = nullptr;
@@ -105,7 +105,9 @@ bool UDisplayClusterICVFXCameraComponent::IsICVFXEnabled() const
 
 		// Second condition to render offscreen: the backbuffer has media output assigned.
 		// This means the whole frame including ICVFX cameras need to be rendered.
-		const bool bIsBackbufferBeingCaptured = BackbufferMediaSettings ? !!BackbufferMediaSettings->MediaOutput : false;
+		const bool bIsBackbufferBeingCaptured = BackbufferMediaSettings ?
+			BackbufferMediaSettings->bEnable && BackbufferMediaSettings->IsMediaOutputAssigned() :
+			false;
 
 		// Finally make a decision if the camera should be rendered
 		return CameraSettings.bEnable && (bUsesMediaOutput || bIsBackbufferBeingCaptured);
diff --git a/Engine/Plugins/Runtime/nDisplay/Source/DisplayCluster/Private/Render/Viewport/Configuration/DisplayClusterViewportConfigurationHelpers.cpp b/Engine/Plugins/Runtime/nDisplay/Source/DisplayCluster/Private/Render/Viewport/Configuration/DisplayClusterViewportConfigurationHelpers.cpp
--- a/Engine/Plugins/Runtime/nDisplay/Source/DisplayCluster/Private/Render/Viewport/Configuration/DisplayClusterViewportConfigurationHelpers.cpp
+++ b/Engine/Plugins/Runtime/nDisplay/Source/DisplayCluster/Private/Render/Viewport/Configuration/DisplayClusterViewportConfigurationHelpers.cpp
@@ -223,10 +223,10 @@ void FDisplayClusterViewportConfigurationHelpers::UpdateBaseViewportSetting(FDis
 			if (MediaSettings.bEnable)
 			{
 				// Don't render the viewport if media input assigned
-				DstViewport.RenderSettings.bSkipSceneRenderingButLeaveResourcesAvailable = !!MediaSettings.MediaSource;
+				DstViewport.RenderSettings.bSkipSceneRenderingButLeaveResourcesAvailable = MediaSettings.IsMediaInputAssigned();
 
 				// Mark this viewport is going to be captured by a capture device
-				DstViewport.RenderSettings.bIsBeingCaptured = !!MediaSettings.MediaOutput;
+				DstViewport.RenderSettings.bIsBeingCaptured = MediaSettings.IsMediaOutputAssigned();
 			}
 		}
 	}
diff --git a/Engine/Plugins/Runtime/nDisplay/Source/DisplayClusterConfiguration/Private/DisplayClusterConfigurationTypes.cpp b/Engine/Plugins/Runtime/nDisplay/Source/DisplayClusterConfiguration/Private/DisplayClusterConfigurationTypes.cpp
--- a/Engine/Plugins/Runtime/nDisplay/Source/DisplayClusterConfiguration/Private/DisplayClusterConfigurationTypes.cpp
+++ b/Engine/Plugins/Runtime/nDisplay/Source/DisplayClusterConfiguration/Private/DisplayClusterConfigurationTypes.cpp
@@ -165,8 +165,8 @@ FDisplayClusterConfigurationProjection::FDisplayClusterConfigurationProjection()
 }
 
 
-const float UDisplayClusterConfigurationViewport::ViewportMinimumSize = 1.0f;
-const float UDisplayClusterConfigurationViewport::ViewportMaximumSize = 15360.0f;
+constexpr float UDisplayClusterConfigurationViewport::ViewportMinimumSize = 1.0f;
+constexpr float UDisplayClusterConfigurationViewport::ViewportMaximumSize = 15360.0f;
 
 UDisplayClusterConfigurationViewport::UDisplayClusterConfigurationViewport()
 {
@@ -176,6 +176,24 @@ UDisplayClusterConfigurationViewport::UDisplayClusterConfigurationViewport()
 #endif
 }
 
+void UDisplayClusterConfigurationViewport::PostLoad()
+{
+	Super::PostLoad();
+
+#if WITH_EDITOR
+	PRAGMA_DISABLE_DEPRECATION_WARNINGS
+		if (IsValid(RenderSettings.Media.MediaOutput_DEPRECATED))
+		{
+			RenderSettings.Media.MediaOutputs.Add(
+				{ RenderSettings.Media.MediaOutput_DEPRECATED , RenderSettings.Media.OutputSyncPolicy_DEPRECATED });
+
+			RenderSettings.Media.MediaOutput_DEPRECATED = nullptr;
+			RenderSettings.Media.OutputSyncPolicy_DEPRECATED = nullptr;
+		}
+	PRAGMA_ENABLE_DEPRECATION_WARNINGS
+#endif // WITH_EDITOR
+}
+
 #if WITH_EDITOR
 void UDisplayClusterConfigurationViewport::PreEditChange(FEditPropertyChain& PropertyAboutToChange)
 {
@@ -252,6 +270,24 @@ void UDisplayClusterConfigurationViewport::OnPreCompile(FCompilerResultsLog& Mes
 	}
 }
 
+void UDisplayClusterConfigurationClusterNode::PostLoad()
+{
+	Super::PostLoad();
+
+#if WITH_EDITOR
+	PRAGMA_DISABLE_DEPRECATION_WARNINGS
+		if (IsValid(Media.MediaOutput_DEPRECATED))
+		{
+			Media.MediaOutputs.Add(
+				{ Media.MediaOutput_DEPRECATED , Media.OutputSyncPolicy_DEPRECATED });
+
+			Media.MediaOutput_DEPRECATED = nullptr;
+			Media.OutputSyncPolicy_DEPRECATED = nullptr;
+		}
+	PRAGMA_ENABLE_DEPRECATION_WARNINGS
+#endif // WITH_EDITOR
+}
+
 void UDisplayClusterConfigurationClusterNode::PostEditChangeChainProperty(FPropertyChangedChainEvent& PropertyChangedEvent)
 {
 	Super::PostEditChangeChainProperty(PropertyChangedEvent);
diff --git a/Engine/Plugins/Runtime/nDisplay/Source/DisplayClusterConfiguration/Private/DisplayClusterConfigurationTypes_Media.cpp b/Engine/Plugins/Runtime/nDisplay/Source/DisplayClusterConfiguration/Private/DisplayClusterConfigurationTypes_Media.cpp
--- a/Engine/Plugins/Runtime/nDisplay/Source/DisplayClusterConfiguration/Private/DisplayClusterConfigurationTypes_Media.cpp
+++ b/Engine/Plugins/Runtime/nDisplay/Source/DisplayClusterConfiguration/Private/DisplayClusterConfigurationTypes_Media.cpp
@@ -4,70 +4,80 @@
 
 #include "DisplayClusterConfigurationTypes_Media.h"
 
+///////////////////////////////////////////////////
+// FDisplayClusterConfigurationMedia
 
-bool FDisplayClusterConfigurationMediaICVFX::IsMediaInputAssigned(const FString& NodeId) const
+bool FDisplayClusterConfigurationMedia::IsMediaInputAssigned() const
 {
-	return GetMediaSource(NodeId) != nullptr;
+	return IsValid(MediaInput.MediaSource);
 }
 
-bool FDisplayClusterConfigurationMediaICVFX::IsMediaOutputAssigned(const FString& NodeId) const
+bool FDisplayClusterConfigurationMedia::IsMediaOutputAssigned() const
 {
-	return GetMediaOutput(NodeId) != nullptr;
-}
-
-UMediaSource* FDisplayClusterConfigurationMediaICVFX::GetMediaSource(const FString& NodeId) const
-{
-	// Look up for a group that contains node ID specified
-	for (const FDisplayClusterConfigurationMediaInputGroup& MediaInputGroup : MediaInputGroups)
+	// Return true if we have at least one media output set
+	for (const FDisplayClusterConfigurationMediaOutput& MediaOutputItem : MediaOutputs)
 	{
-		const bool bFound = MediaInputGroup.ClusterNodes.ItemNames.ContainsByPredicate([NodeId](const FString& Item)
-			{
-				return Item.Equals(NodeId, ESearchCase::IgnoreCase);
-			});
-
-		if (bFound)
+		if (IsValid(MediaOutputItem.MediaOutput))
 		{
-			return MediaInputGroup.MediaSource;
+			return true;
 		}
 	}
 
-	return nullptr;
+	return false;
 }
 
-UMediaOutput* FDisplayClusterConfigurationMediaICVFX::GetMediaOutput(const FString& NodeId) const
+
+///////////////////////////////////////////////////
+// FDisplayClusterConfigurationMediaICVFX
+
+bool FDisplayClusterConfigurationMediaICVFX::IsMediaInputAssigned(const FString& NodeId) const
 {
-	// Look up for a group that contains node ID specified
-	for (const FDisplayClusterConfigurationMediaOutputGroup& MediaOutputGroup : MediaOutputGroups)
-	{
-		const bool bFound = MediaOutputGroup.ClusterNodes.ItemNames.ContainsByPredicate([NodeId](const FString& Item)
-			{
-				return Item.Equals(NodeId, ESearchCase::IgnoreCase);
-			});
+	return IsValid(GetMediaSource(NodeId));
+}
 
-		if (bFound)
+bool FDisplayClusterConfigurationMediaICVFX::IsMediaOutputAssigned(const FString& NodeId) const
+{
+	const TArray<FDisplayClusterConfigurationMediaOutputGroupICVFX> NodeOutputGroups = GetMediaOutputGroups(NodeId);
+
+	for (const FDisplayClusterConfigurationMediaOutputGroupICVFX& NodeOutputGroup : NodeOutputGroups)
+	{
+		if (IsValid(NodeOutputGroup.MediaOutput))
 		{
-			return MediaOutputGroup.MediaOutput;
+			return true;
 		}
 	}
 
-	return nullptr;
+	return false;
 }
 
-UDisplayClusterMediaOutputSynchronizationPolicy* FDisplayClusterConfigurationMediaICVFX::GetOutputSyncPolicy(const FString& NodeId) const
+
+//@note
+// The way how media is configured for ICVFX cameras technically allows to have multiple inputs assigned
+// to the same camera. Yes, this is something that contradicts to a single input concept. However, it provides
+// a very user-friendly GUI. So to follow the single input concept, we always return the first media input found.
+UMediaSource* FDisplayClusterConfigurationMediaICVFX::GetMediaSource(const FString& NodeId) const
 {
 	// Look up for a group that contains node ID specified
-	for (const FDisplayClusterConfigurationMediaOutputGroup& MediaOutputGroup : MediaOutputGroups)
+	for (const FDisplayClusterConfigurationMediaInputGroupICVFX& MediaInputGroup : MediaInputGroups)
 	{
-		const bool bFound = MediaOutputGroup.ClusterNodes.ItemNames.ContainsByPredicate([NodeId](const FString& Item)
+		const bool bNodeFound = MediaInputGroup.ClusterNodes.ItemNames.ContainsByPredicate([NodeId](const FString& Item)
 			{
 				return Item.Equals(NodeId, ESearchCase::IgnoreCase);
 			});
 
-		if (bFound)
+		if (bNodeFound && IsValid(MediaInputGroup.MediaSource))
 		{
-			return MediaOutputGroup.OutputSyncPolicy;
+			return MediaInputGroup.MediaSource;
 		}
 	}
 
 	return nullptr;
 }
+
+TArray<FDisplayClusterConfigurationMediaOutputGroupICVFX> FDisplayClusterConfigurationMediaICVFX::GetMediaOutputGroups(const FString& NodeId) const
+{
+	return MediaOutputGroups.FilterByPredicate([NodeId](const FDisplayClusterConfigurationMediaOutputGroupICVFX& Item)
+		{
+			return Item.ClusterNodes.ItemNames.Contains(NodeId);
+		});
+}
diff --git a/Engine/Plugins/Runtime/nDisplay/Source/DisplayClusterConfiguration/Public/DisplayClusterConfigurationTypes.h b/Engine/Plugins/Runtime/nDisplay/Source/DisplayClusterConfiguration/Public/DisplayClusterConfigurationTypes.h
--- a/Engine/Plugins/Runtime/nDisplay/Source/DisplayClusterConfiguration/Public/DisplayClusterConfigurationTypes.h
+++ b/Engine/Plugins/Runtime/nDisplay/Source/DisplayClusterConfiguration/Public/DisplayClusterConfigurationTypes.h
@@ -306,12 +306,14 @@ public:
 	UFUNCTION(BlueprintPure, Category = "NDisplay|Configuration")
 	void GetReferencedMeshNames(TArray<FString>& OutMeshNames) const;
 	
-private:
+	//~ Begin UObject interface
+public:
+	virtual void PostLoad() override;
+
 #if WITH_EDITOR
-	// UObject interface
 	virtual void PostEditChangeChainProperty(FPropertyChangedChainEvent& PropertyChangedEvent) override;
-	// End of UObject interface
 #endif
+	//~ End UObject interface
 
 public:
 	/** IP address of this specific cluster Node */
diff --git a/Engine/Plugins/Runtime/nDisplay/Source/DisplayClusterConfiguration/Public/DisplayClusterConfigurationTypes_Media.h b/Engine/Plugins/Runtime/nDisplay/Source/DisplayClusterConfiguration/Public/DisplayClusterConfigurationTypes_Media.h
--- a/Engine/Plugins/Runtime/nDisplay/Source/DisplayClusterConfiguration/Public/DisplayClusterConfigurationTypes_Media.h
+++ b/Engine/Plugins/Runtime/nDisplay/Source/DisplayClusterConfiguration/Public/DisplayClusterConfigurationTypes_Media.h
@@ -15,6 +15,40 @@
 #include "DisplayClusterConfigurationTypes_Media.generated.h"
 
 
+/*
+ * Media input item
+ */
+USTRUCT(Blueprintable)
+struct DISPLAYCLUSTERCONFIGURATION_API FDisplayClusterConfigurationMediaInput
+{
+	GENERATED_BODY()
+
+public:
+	/** Media source to use */
+	UPROPERTY(Instanced, EditAnywhere, BlueprintReadWrite, Category = "Media")
+	TObjectPtr<UMediaSource> MediaSource;
+};
+
+
+/*
+ * Media output item
+ */
+USTRUCT(Blueprintable)
+struct DISPLAYCLUSTERCONFIGURATION_API FDisplayClusterConfigurationMediaOutput
+{
+	GENERATED_BODY()
+
+public:
+	/** Media output to use */
+	UPROPERTY(Instanced, EditAnywhere, BlueprintReadWrite, Category = "Media")
+	TObjectPtr<UMediaOutput> MediaOutput;
+
+	/** Media output synchronization policy */
+	UPROPERTY(Instanced, EditAnywhere, BlueprintReadWrite, Category = "Media", meta = (DisplayName = "Capture Synchronization"))
+	TObjectPtr<UDisplayClusterMediaOutputSynchronizationPolicy> OutputSyncPolicy;
+};
+
+
 PRAGMA_DISABLE_DEPRECATION_WARNINGS
 
 /*
@@ -31,22 +65,36 @@ public:
 	bool bEnable = false;
 
 	/** Media source to use */
-	UPROPERTY(Instanced, EditAnywhere, BlueprintReadWrite, Category = "Media")
-	TObjectPtr<UMediaSource> MediaSource;
+	UPROPERTY(EditAnywhere, BlueprintReadWrite, Category = "Media")
+	FDisplayClusterConfigurationMediaInput MediaInput;
 
-	/** Media output to use */
-	UPROPERTY(Instanced, EditAnywhere, BlueprintReadWrite, Category = "Output")
-	TObjectPtr<UMediaOutput> MediaOutput;
+	/** Media outputs to use */
+	UPROPERTY(EditAnywhere, BlueprintReadWrite, Category = "Media")
+	TArray<FDisplayClusterConfigurationMediaOutput> MediaOutputs;
 
-	/** Media output synchronization policy */
-	UPROPERTY(Instanced, EditAnywhere, BlueprintReadWrite, Category = "Output", meta = (DisplayName = "Capture Synchronization"))
-	TObjectPtr<UDisplayClusterMediaOutputSynchronizationPolicy> OutputSyncPolicy;
+public:
+	/** Returns true if a media source assigned */
+	bool IsMediaInputAssigned() const;
+
+	/** Returns true if at least one media output assigned */
+	bool IsMediaOutputAssigned() const;
 
 protected:
 #if WITH_EDITORONLY_DATA
-	UE_DEPRECATED(5.2, "This property has been deprecated")
-	UPROPERTY()
+
+	UPROPERTY(meta = (DeprecatedProperty, DeprecationMessage = "This property has been deprecated"))
 	FString MediaSharingNode_DEPRECATED;
+
+public:
+	UPROPERTY(meta = (DeprecatedProperty, DeprecationMessage = "This property has been deprecated. Please refer new MediaInput property."))
+	TObjectPtr<UMediaSource> MediaSource_DEPRECATED;
+
+	UPROPERTY(meta = (DeprecatedProperty, DeprecationMessage = "This property has been deprecated. Please refer new MediaOutputs property."))
+	TObjectPtr<UMediaOutput> MediaOutput_DEPRECATED;
+
+	UPROPERTY(meta = (DeprecatedProperty, DeprecationMessage = "This property has been deprecated. Please refer new MediaOutputs property."))
+	TObjectPtr<UDisplayClusterMediaOutputSynchronizationPolicy> OutputSyncPolicy_DEPRECATED;
+
 #endif // WITH_EDITORONLY_DATA
 };
 
@@ -54,10 +102,11 @@ PRAGMA_ENABLE_DEPRECATION_WARNINGS
 
 
 /*
- * Input media group
+ * Media input group (ICVFX)
  */
 USTRUCT(Blueprintable)
-struct DISPLAYCLUSTERCONFIGURATION_API FDisplayClusterConfigurationMediaInputGroup
+struct DISPLAYCLUSTERCONFIGURATION_API FDisplayClusterConfigurationMediaInputGroupICVFX
+	: public FDisplayClusterConfigurationMediaInput
 {
 	GENERATED_BODY()
 
@@ -65,18 +114,15 @@ public:
 	/** Cluster nodes that use media source below */
 	UPROPERTY(EditAnywhere, BlueprintReadWrite, Category = Media, meta = (ClusterItemType = ClusterNodes))
 	FDisplayClusterConfigurationClusterItemReferenceList ClusterNodes;
-
-	/** Media source to use */
-	UPROPERTY(Instanced, EditAnywhere, BlueprintReadWrite, Category = "Media")
-	TObjectPtr<UMediaSource> MediaSource;
 };
 
 
 /*
- * Output media group
+ * Media output group (ICVFX)
  */
 USTRUCT(Blueprintable)
-struct DISPLAYCLUSTERCONFIGURATION_API FDisplayClusterConfigurationMediaOutputGroup
+struct DISPLAYCLUSTERCONFIGURATION_API FDisplayClusterConfigurationMediaOutputGroupICVFX
+	: public FDisplayClusterConfigurationMediaOutput
 {
 	GENERATED_BODY()
 
@@ -84,14 +130,6 @@ public:
 	/** Cluster nodes that export media via MediaOutput below */
 	UPROPERTY(EditAnywhere, BlueprintReadWrite, Category = Media, meta = (ClusterItemType = ClusterNodes))
 	FDisplayClusterConfigurationClusterItemReferenceList ClusterNodes;
-
-	/** Media output to use */
-	UPROPERTY(Instanced, EditAnywhere, BlueprintReadWrite, Category = "Media")
-	TObjectPtr<UMediaOutput> MediaOutput;
-
-	/** Media output synchronization policy */
-	UPROPERTY(Instanced, EditAnywhere, BlueprintReadWrite, Category = "Media", meta = (DisplayName = "Capture Synchronization"))
-	TObjectPtr<UDisplayClusterMediaOutputSynchronizationPolicy> OutputSyncPolicy;
 };
 
 
@@ -108,13 +146,13 @@ public:
 	UPROPERTY(EditAnywhere, BlueprintReadWrite, Category = "Media")
 	bool bEnable = false;
 
-	/** Media output mapping */
+	/** Media input mapping */
 	UPROPERTY(EditAnywhere, BlueprintReadWrite, Category = "Media")
-	TArray<FDisplayClusterConfigurationMediaOutputGroup> MediaOutputGroups;
+	TArray<FDisplayClusterConfigurationMediaInputGroupICVFX> MediaInputGroups;
 
-	/** Media input mapping */
+	/** Media output mapping */
 	UPROPERTY(EditAnywhere, BlueprintReadWrite, Category = "Media")
-	TArray<FDisplayClusterConfigurationMediaInputGroup> MediaInputGroups;
+	TArray<FDisplayClusterConfigurationMediaOutputGroupICVFX> MediaOutputGroups;
 
 public:
 	/** Returns true if a specific cluster node has media source assigned */
@@ -126,11 +164,27 @@ public:
 	/** Returns media source bound to a specific cluster node */
 	UMediaSource* GetMediaSource(const FString& NodeId) const;
 
-	/** Returns media output bound to a specific cluster node */
-	UMediaOutput* GetMediaOutput(const FString& NodeId) const;
+	/** Returns media outputs bound to a specific cluster node */
+	TArray<FDisplayClusterConfigurationMediaOutputGroupICVFX> GetMediaOutputGroups(const FString& NodeId) const;
 
-	/** Returns media output sync policy bound to a specific cluster node */
-	UDisplayClusterMediaOutputSynchronizationPolicy* GetOutputSyncPolicy(const FString& NodeId) const;
+public:
+	UE_DEPRECATED(5.3, "This function has beend deprecated. Please use GetMediaOutputGroups.")
+	TArray<UMediaOutput*> GetMediaOutputs(const FString& NodeId) const
+	{
+		return {};
+	}
+
+	UE_DEPRECATED(5.3, "This function has beend deprecated. Please use GetMediaOutputGroups.")
+	UDisplayClusterMediaOutputSynchronizationPolicy* GetOutputSyncPolicy(const FString& NodeId) const
+	{
+		return nullptr;
+	}
+
+	UE_DEPRECATED(5.3, "This function has beend deprecated. Please use GetMediaOutputGroups.")
+	UMediaOutput* GetMediaOutput(const FString& NodeId) const
+	{
+		return nullptr;
+	}
 
 	UE_DEPRECATED(5.2, "This function has beend deprecated.")
 	bool IsMediaSharingNode(const FString& InNodeId) const
diff --git a/Engine/Plugins/Runtime/nDisplay/Source/DisplayClusterConfiguration/Public/DisplayClusterConfigurationTypes_Viewport.h b/Engine/Plugins/Runtime/nDisplay/Source/DisplayClusterConfiguration/Public/DisplayClusterConfigurationTypes_Viewport.h
--- a/Engine/Plugins/Runtime/nDisplay/Source/DisplayClusterConfiguration/Public/DisplayClusterConfigurationTypes_Viewport.h
+++ b/Engine/Plugins/Runtime/nDisplay/Source/DisplayClusterConfiguration/Public/DisplayClusterConfigurationTypes_Viewport.h
@@ -163,14 +163,19 @@ public:
 	UDisplayClusterConfigurationViewport();
 
 public:
-	void GetReferencedMeshNames(TArray<FString>& OutMeshNames) const;
+	//~ Begin UObject interface
+	void PostLoad() override;
 
 #if WITH_EDITOR
-	// UObject interface
 	virtual void PreEditChange(FEditPropertyChain& PropertyAboutToChange) override;
 	virtual void PostEditChangeChainProperty(FPropertyChangedChainEvent& PropertyChangedEvent) override;
-	// End of UObject interface
+#endif
+	//~ End UObject interfacef
+
+public:
+	void GetReferencedMeshNames(TArray<FString>& OutMeshNames) const;
 
+#if WITH_EDITOR
 	/** Enable the preview texture. Only should be called by the object managing the preview texture state. */
 	void EnablePreviewTexture();
 	
diff --git a/Engine/Plugins/Runtime/nDisplay/Source/DisplayClusterMedia/Private/DisplayClusterMediaHelpers.cpp b/Engine/Plugins/Runtime/nDisplay/Source/DisplayClusterMedia/Private/DisplayClusterMediaHelpers.cpp
--- a/Engine/Plugins/Runtime/nDisplay/Source/DisplayClusterMedia/Private/DisplayClusterMediaHelpers.cpp
+++ b/Engine/Plugins/Runtime/nDisplay/Source/DisplayClusterMedia/Private/DisplayClusterMediaHelpers.cpp
@@ -13,7 +13,7 @@ namespace DisplayClusterMediaHelpers
 	namespace MediaId
 	{
 		// Generate media ID for a specific entity
-		FString GenerateMediaId(EMediaDeviceType DeviceType, EMediaOwnerType OwnerType, const FString& NodeId, const FString& DCRAName, const FString& OwnerName)
+		FString GenerateMediaId(EMediaDeviceType DeviceType, EMediaOwnerType OwnerType, const FString& NodeId, const FString& DCRAName, const FString& OwnerName, uint8 Index)
 		{
 			if (DeviceType == EMediaDeviceType::Input)
 			{
@@ -37,16 +37,16 @@ namespace DisplayClusterMediaHelpers
 				switch (OwnerType)
 				{
 				case EMediaOwnerType::Backbuffer:
-					return FString::Printf(TEXT("%s_%s_backbuffer_capture"), *NodeId, *DCRAName);
+					return FString::Printf(TEXT("%s_%s_backbuffer_capture_%u"), *NodeId, *DCRAName, Index);
 
 				case EMediaOwnerType::Viewport:
-					return FString::Printf(TEXT("%s_%s_%s_viewport_capture"), *NodeId, *DCRAName, *OwnerName);
+					return FString::Printf(TEXT("%s_%s_%s_viewport_capture_%u"), *NodeId, *DCRAName, *OwnerName, Index);
 
 				case EMediaOwnerType::ICVFXCamera:
-					return FString::Printf(TEXT("%s_%s_%s_icvfx_capture"), *NodeId, *DCRAName, *OwnerName);
+					return FString::Printf(TEXT("%s_%s_%s_icvfx_capture_%u"), *NodeId, *DCRAName, *OwnerName, Index);
 
 				default:
-					unimplemented();
+					checkNoEntry();
 				}
 			}
 
diff --git a/Engine/Plugins/Runtime/nDisplay/Source/DisplayClusterMedia/Private/DisplayClusterMediaHelpers.h b/Engine/Plugins/Runtime/nDisplay/Source/DisplayClusterMedia/Private/DisplayClusterMediaHelpers.h
--- a/Engine/Plugins/Runtime/nDisplay/Source/DisplayClusterMedia/Private/DisplayClusterMediaHelpers.h
+++ b/Engine/Plugins/Runtime/nDisplay/Source/DisplayClusterMedia/Private/DisplayClusterMediaHelpers.h
@@ -28,7 +28,7 @@ namespace DisplayClusterMediaHelpers
 		};
 
 		// Generate media ID for a specific entity
-		FString GenerateMediaId(EMediaDeviceType DeviceType, EMediaOwnerType OwnerType, const FString& NodeId, const FString& DCRAName, const FString& OwnerName);
+		FString GenerateMediaId(EMediaDeviceType DeviceType, EMediaOwnerType OwnerType, const FString& NodeId, const FString& DCRAName, const FString& OwnerName, uint8 Index);
 	}
 
 	// Generates internal ICVFX viewport IDs
diff --git a/Engine/Plugins/Runtime/nDisplay/Source/DisplayClusterMedia/Private/DisplayClusterMediaModule.cpp b/Engine/Plugins/Runtime/nDisplay/Source/DisplayClusterMedia/Private/DisplayClusterMediaModule.cpp
--- a/Engine/Plugins/Runtime/nDisplay/Source/DisplayClusterMedia/Private/DisplayClusterMediaModule.cpp
+++ b/Engine/Plugins/Runtime/nDisplay/Source/DisplayClusterMedia/Private/DisplayClusterMediaModule.cpp
@@ -102,34 +102,42 @@ void FDisplayClusterMediaModule::InitializeMedia()
 				if (MediaSettings.bEnable)
 				{
 					// Media input
-					if (MediaSettings.MediaSource)
+					if (MediaSettings.IsMediaInputAssigned())
 					{
 						const FString MediaInputId = DisplayClusterMediaHelpers::MediaId::GenerateMediaId(
 							DisplayClusterMediaHelpers::MediaId::EMediaDeviceType::Input,
 							DisplayClusterMediaHelpers::MediaId::EMediaOwnerType::Backbuffer,
-							*ClusterNodeId, *RootActor->GetName(), FString());
+							*ClusterNodeId, *RootActor->GetName(), FString(), 0);
 
 						UE_LOG(LogDisplayClusterMedia, Log, TEXT("Initializing backbuffer media input '%s'..."), *MediaInputId);
 
 						InputNode = MakeShared<FDisplayClusterMediaInputNode>(
 							MediaInputId, ClusterNodeId,
-							MediaSettings.MediaSource);
+							MediaSettings.MediaInput.MediaSource);
 					}
 
 					// Media capture
-					if (MediaSettings.MediaOutput)
+					uint8 CaptureIdx = 0;
+					for (const FDisplayClusterConfigurationMediaOutput& MediaOutputItem : MediaSettings.MediaOutputs)
 					{
-						const FString MediaCaptureId = DisplayClusterMediaHelpers::MediaId::GenerateMediaId(
-							DisplayClusterMediaHelpers::MediaId::EMediaDeviceType::Output,
-							DisplayClusterMediaHelpers::MediaId::EMediaOwnerType::Backbuffer,
-							*ClusterNodeId, *RootActor->GetName(), FString());
+						if (IsValid(MediaOutputItem.MediaOutput))
+						{
+							const FString MediaCaptureId = DisplayClusterMediaHelpers::MediaId::GenerateMediaId(
+								DisplayClusterMediaHelpers::MediaId::EMediaDeviceType::Output,
+								DisplayClusterMediaHelpers::MediaId::EMediaOwnerType::Backbuffer,
+								*ClusterNodeId, *RootActor->GetName(), FString(), CaptureIdx);
+
+							UE_LOG(LogDisplayClusterMedia, Log, TEXT("Initializing backbuffer media capture [%u]: '%s'"), CaptureIdx, *MediaCaptureId);
+
+							TSharedPtr<FDisplayClusterMediaCaptureNode> NewNodeCapture = MakeShared<FDisplayClusterMediaCaptureNode>(
+								MediaCaptureId, ClusterNodeId,
+								MediaOutputItem.MediaOutput,
+								MediaOutputItem.OutputSyncPolicy);
 
-						UE_LOG(LogDisplayClusterMedia, Log, TEXT("Initializing backbuffer media capture '%s'..."), *MediaCaptureId);
+							CaptureNode.Emplace(MediaCaptureId, NewNodeCapture);
+						}
 
-						CaptureNode = MakeShared<FDisplayClusterMediaCaptureNode>(
-							MediaCaptureId, ClusterNodeId,
-							MediaSettings.MediaOutput,
-							MediaSettings.OutputSyncPolicy);
+						++CaptureIdx;
 					}
 				}
 			}
@@ -145,40 +153,46 @@ void FDisplayClusterMediaModule::InitializeMedia()
 					if (MediaSettings.bEnable)
 					{
 						// Media input
-						if (MediaSettings.MediaSource)
+						if (MediaSettings.IsMediaInputAssigned())
 						{
 							const FString MediaInputId = DisplayClusterMediaHelpers::MediaId::GenerateMediaId(
 								DisplayClusterMediaHelpers::MediaId::EMediaDeviceType::Input,
 								DisplayClusterMediaHelpers::MediaId::EMediaOwnerType::Viewport,
-								*ClusterNodeId, *RootActor->GetName(), ViewportIt.Key);
+								*ClusterNodeId, *RootActor->GetName(), ViewportIt.Key, 0);
 
-							UE_LOG(LogDisplayClusterMedia, Log, TEXT("Initializing viewport media input '%s' for viewport '%s'..."), *MediaInputId, *ViewportIt.Key);
+							UE_LOG(LogDisplayClusterMedia, Log, TEXT("Initializing viewport media input '%s' for viewport '%s'"), *MediaInputId, *ViewportIt.Key);
 
 							TSharedPtr<FDisplayClusterMediaInputViewport> NewViewportInput = MakeShared<FDisplayClusterMediaInputViewport>(
 								MediaInputId, ClusterNodeId,
 								ViewportIt.Key,
-								MediaSettings.MediaSource);
+								MediaSettings.MediaInput.MediaSource);
 
 							InputViewports.Emplace(MediaInputId, MoveTemp(NewViewportInput));
 						}
 
 						// Media capture
-						if (MediaSettings.MediaOutput)
+						uint8 CaptureIdx = 0;
+						for (const FDisplayClusterConfigurationMediaOutput& MediaOutputItem : MediaSettings.MediaOutputs)
 						{
-							const FString MediaCaptureId = DisplayClusterMediaHelpers::MediaId::GenerateMediaId(
-								DisplayClusterMediaHelpers::MediaId::EMediaDeviceType::Output,
-								DisplayClusterMediaHelpers::MediaId::EMediaOwnerType::Viewport,
-								*ClusterNodeId, *RootActor->GetName(), ViewportIt.Key);
+							if (IsValid(MediaOutputItem.MediaOutput))
+							{
+								const FString MediaCaptureId = DisplayClusterMediaHelpers::MediaId::GenerateMediaId(
+									DisplayClusterMediaHelpers::MediaId::EMediaDeviceType::Output,
+									DisplayClusterMediaHelpers::MediaId::EMediaOwnerType::Viewport,
+									*ClusterNodeId, *RootActor->GetName(), ViewportIt.Key, CaptureIdx);
 
-							UE_LOG(LogDisplayClusterMedia, Log, TEXT("Initializing viewport capture '%s' for viewport '%s'..."), *MediaCaptureId, *ViewportIt.Key);
+								UE_LOG(LogDisplayClusterMedia, Log, TEXT("Initializing viewport capture [%u]: '%s' for viewport '%s'"), CaptureIdx, *MediaCaptureId, *ViewportIt.Key);
 
-							TSharedPtr<FDisplayClusterMediaCaptureViewport> NewViewportCapture = MakeShared<FDisplayClusterMediaCaptureViewport>(
-								MediaCaptureId, ClusterNodeId,
-								ViewportIt.Key,
-								MediaSettings.MediaOutput,
-								MediaSettings.OutputSyncPolicy);
+								TSharedPtr<FDisplayClusterMediaCaptureViewport> NewViewportCapture = MakeShared<FDisplayClusterMediaCaptureViewport>(
+									MediaCaptureId, ClusterNodeId,
+									ViewportIt.Key,
+									MediaOutputItem.MediaOutput,
+									MediaOutputItem.OutputSyncPolicy);
+
+								CaptureViewports.Emplace(MediaCaptureId, MoveTemp(NewViewportCapture));
+							}
 
-							CaptureViewports.Emplace(MediaCaptureId, MoveTemp(NewViewportCapture));
+							++CaptureIdx;
 						}
 					}
 				}
@@ -200,15 +214,15 @@ void FDisplayClusterMediaModule::InitializeMedia()
 						const FString ICVFXCameraName = ICVFXCameraComponent->GetName();
 						const FString ICVFXViewportId = DisplayClusterMediaHelpers::GenerateICVFXViewportName(ClusterNodeId, ICVFXCameraName);
 
-						// Media input only
+						// Media input
 						if (UMediaSource* MediaSource = MediaSettings.GetMediaSource(ClusterNodeId))
 						{
 							const FString MediaInputId = DisplayClusterMediaHelpers::MediaId::GenerateMediaId(
 								DisplayClusterMediaHelpers::MediaId::EMediaDeviceType::Input,
 								DisplayClusterMediaHelpers::MediaId::EMediaOwnerType::ICVFXCamera,
-								*ClusterNodeId, *RootActor->GetName(), ICVFXCameraName);
+								*ClusterNodeId, *RootActor->GetName(), ICVFXCameraName, 0);
 
-							UE_LOG(LogDisplayClusterMedia, Log, TEXT("Initializing ICVFX media input '%s' for camera '%s'..."), *MediaInputId, *ICVFXCameraName);
+							UE_LOG(LogDisplayClusterMedia, Log, TEXT("Initializing ICVFX media input '%s' for camera '%s'"), *MediaInputId, *ICVFXCameraName);
 
 							TSharedPtr<FDisplayClusterMediaInputViewport> NewICVFXInput = MakeShared<FDisplayClusterMediaInputViewport>(
 								MediaInputId, ClusterNodeId,
@@ -218,23 +232,30 @@ void FDisplayClusterMediaModule::InitializeMedia()
 							InputViewports.Emplace(MediaInputId, MoveTemp(NewICVFXInput));
 						}
 
-						// Media capture only
-						if (UMediaOutput* MediaOutput = MediaSettings.GetMediaOutput(ClusterNodeId))
+						// Media capture
+						const TArray<FDisplayClusterConfigurationMediaOutputGroupICVFX> MediaOutputItems = MediaSettings.GetMediaOutputGroups(ClusterNodeId);
+						uint8 CaptureIdx = 0;
+						for (const FDisplayClusterConfigurationMediaOutputGroupICVFX& MediaOutputItem : MediaOutputItems)
 						{
-							const FString MediaCaptureId = DisplayClusterMediaHelpers::MediaId::GenerateMediaId(
-								DisplayClusterMediaHelpers::MediaId::EMediaDeviceType::Output,
-								DisplayClusterMediaHelpers::MediaId::EMediaOwnerType::ICVFXCamera,
-								*ClusterNodeId, *RootActor->GetName(), ICVFXCameraName);
+							if (IsValid(MediaOutputItem.MediaOutput))
+							{
+								const FString MediaCaptureId = DisplayClusterMediaHelpers::MediaId::GenerateMediaId(
+									DisplayClusterMediaHelpers::MediaId::EMediaDeviceType::Output,
+									DisplayClusterMediaHelpers::MediaId::EMediaOwnerType::ICVFXCamera,
+									*ClusterNodeId, *RootActor->GetName(), ICVFXCameraName, CaptureIdx);
 
-							UE_LOG(LogDisplayClusterMedia, Log, TEXT("Initializing ICVFX capture '%s' for camera '%s'..."), *MediaCaptureId, *ICVFXCameraName);
+								UE_LOG(LogDisplayClusterMedia, Log, TEXT("Initializing ICVFX capture [%u]: '%s' for camera '%s'"), CaptureIdx, *MediaCaptureId, *ICVFXCameraName);
 
-							TSharedPtr<FDisplayClusterMediaCaptureViewport> NewICVFXCapture = MakeShared<FDisplayClusterMediaCaptureCamera>(
-								MediaCaptureId, ClusterNodeId,
-								ICVFXCameraName, ICVFXViewportId,
-								MediaOutput,
-								MediaSettings.GetOutputSyncPolicy(ClusterNodeId));
+								TSharedPtr<FDisplayClusterMediaCaptureViewport> NewICVFXCapture = MakeShared<FDisplayClusterMediaCaptureCamera>(
+									MediaCaptureId, ClusterNodeId,
+									ICVFXCameraName, ICVFXViewportId,
+									MediaOutputItem.MediaOutput,
+									MediaOutputItem.OutputSyncPolicy);
+
+								CaptureViewports.Emplace(MediaCaptureId, MoveTemp(NewICVFXCapture));
+							}
 
-							CaptureViewports.Emplace(MediaCaptureId, MoveTemp(NewICVFXCapture));
+							++CaptureIdx;
 						}
 					}
 				}
@@ -266,9 +287,9 @@ void FDisplayClusterMediaModule::StartCapture()
 	}
 
 	// Start backbuffer capture
-	if (CaptureNode)
+	for (TPair<FString, TSharedPtr<FDisplayClusterMediaCaptureNode>>& Capture : CaptureNode)
 	{
-		CaptureNode->StartCapture();
+		Capture.Value->StartCapture();
 	}
 }
 
@@ -281,9 +302,9 @@ void FDisplayClusterMediaModule::StopCapture()
 	}
 
 	// Stop backbuffer capture
-	if (CaptureNode)
+	for (TPair<FString, TSharedPtr<FDisplayClusterMediaCaptureNode>>& Capture : CaptureNode)
 	{
-		CaptureNode->StopCapture();
+		Capture.Value->StopCapture();
 	}
 }
 
diff --git a/Engine/Plugins/Runtime/nDisplay/Source/DisplayClusterMedia/Private/DisplayClusterMediaModule.h b/Engine/Plugins/Runtime/nDisplay/Source/DisplayClusterMedia/Private/DisplayClusterMediaModule.h
--- a/Engine/Plugins/Runtime/nDisplay/Source/DisplayClusterMedia/Private/DisplayClusterMediaModule.h
+++ b/Engine/Plugins/Runtime/nDisplay/Source/DisplayClusterMedia/Private/DisplayClusterMediaModule.h
@@ -51,7 +51,7 @@ private:
 
 private:
 	TMap<FString, TSharedPtr<FDisplayClusterMediaCaptureViewport>> CaptureViewports;
-	TSharedPtr<FDisplayClusterMediaCaptureNode>                    CaptureNode;
+	TMap<FString, TSharedPtr<FDisplayClusterMediaCaptureNode>>     CaptureNode;
 
 	TMap<FString, TSharedPtr<FDisplayClusterMediaInputViewport>>   InputViewports;
 	TSharedPtr<FDisplayClusterMediaInputNode>                      InputNode;
diff --git a/Engine/Plugins/Runtime/nDisplay/Source/DisplayClusterMedia/Private/Synchronization/SyncPolicies/DisplayClusterMediaOutputSynchronizationPolicyEthernetBarrierBase.cpp b/Engine/Plugins/Runtime/nDisplay/Source/DisplayClusterMedia/Private/Synchronization/SyncPolicies/DisplayClusterMediaOutputSynchronizationPolicyEthernetBarrierBase.cpp
--- a/Engine/Plugins/Runtime/nDisplay/Source/DisplayClusterMedia/Private/Synchronization/SyncPolicies/DisplayClusterMediaOutputSynchronizationPolicyEthernetBarrierBase.cpp
+++ b/Engine/Plugins/Runtime/nDisplay/Source/DisplayClusterMedia/Private/Synchronization/SyncPolicies/DisplayClusterMediaOutputSynchronizationPolicyEthernetBarrierBase.cpp
@@ -211,17 +211,27 @@ void FDisplayClusterMediaOutputSynchronizationPolicyEthernetBarrierBaseHandler::
 				// Backbuffer capture
 				{
 					const FDisplayClusterConfigurationMedia& MediaSettings = NodeIt.Value->Media;
-					if (MediaSettings.bEnable && MediaSettings.MediaOutput && MediaSettings.OutputSyncPolicy)
+
+					if (MediaSettings.bEnable)
 					{
-						// Pick the same sync policy only
-						if (MediaSettings.OutputSyncPolicy->GetClass() == GetPolicyClass())
+						uint8 CaptureIdx = 0;
+						for (const FDisplayClusterConfigurationMediaOutput& MediaOutputItem : MediaSettings.MediaOutputs)
 						{
-							const FString BackbufferCaptureId = DisplayClusterMediaHelpers::MediaId::GenerateMediaId(
-								DisplayClusterMediaHelpers::MediaId::EMediaDeviceType::Output,
-								DisplayClusterMediaHelpers::MediaId::EMediaOwnerType::Backbuffer,
-								NodeIt.Key, RootActor->GetName(), FString());
+							if (IsValid(MediaOutputItem.MediaOutput) && IsValid(MediaOutputItem.OutputSyncPolicy))
+							{
+								// Pick the same sync policy only
+								if (MediaOutputItem.OutputSyncPolicy->GetClass() == GetPolicyClass())
+								{
+									const FString BackbufferCaptureId = DisplayClusterMediaHelpers::MediaId::GenerateMediaId(
+										DisplayClusterMediaHelpers::MediaId::EMediaDeviceType::Output,
+										DisplayClusterMediaHelpers::MediaId::EMediaOwnerType::Backbuffer,
+										NodeIt.Key, RootActor->GetName(), FString(), CaptureIdx);
 
-							OutMarkers.Add(BackbufferCaptureId);
+									OutMarkers.Add(BackbufferCaptureId);
+								}
+							}
+
+							++CaptureIdx;
 						}
 					}
 				}
@@ -233,17 +243,27 @@ void FDisplayClusterMediaOutputSynchronizationPolicyEthernetBarrierBaseHandler::
 					for (const TPair<FString, TObjectPtr<UDisplayClusterConfigurationViewport>>& ViewportIt : NodeIt.Value->Viewports)
 					{
 						const FDisplayClusterConfigurationMedia& MediaSettings = ViewportIt.Value->RenderSettings.Media;
-						if (MediaSettings.bEnable && MediaSettings.MediaOutput && MediaSettings.OutputSyncPolicy)
+
+						if (MediaSettings.bEnable)
 						{
-							// Pick the same sync policy only
-							if (MediaSettings.OutputSyncPolicy->GetClass() == GetPolicyClass())
+							uint8 CaptureIdx = 0;
+							for (const FDisplayClusterConfigurationMediaOutput& MediaOutputItem : MediaSettings.MediaOutputs)
 							{
-								const FString ViewportCaptureId = DisplayClusterMediaHelpers::MediaId::GenerateMediaId(
-									DisplayClusterMediaHelpers::MediaId::EMediaDeviceType::Output,
-									DisplayClusterMediaHelpers::MediaId::EMediaOwnerType::Viewport,
-									NodeIt.Key, RootActor->GetName(), ViewportIt.Key);
+								if (IsValid(MediaOutputItem.MediaOutput) && IsValid(MediaOutputItem.OutputSyncPolicy))
+								{
+									// Pick the same sync policy only
+									if (MediaOutputItem.OutputSyncPolicy->GetClass() == GetPolicyClass())
+									{
+										const FString ViewportCaptureId = DisplayClusterMediaHelpers::MediaId::GenerateMediaId(
+											DisplayClusterMediaHelpers::MediaId::EMediaDeviceType::Output,
+											DisplayClusterMediaHelpers::MediaId::EMediaOwnerType::Viewport,
+											NodeIt.Key, RootActor->GetName(), ViewportIt.Key, CaptureIdx);
+
+										OutMarkers.Add(ViewportCaptureId);
+									}
+								}
 
-								OutMarkers.Add(ViewportCaptureId);
+								++CaptureIdx;
 							}
 						}
 					}
@@ -262,9 +282,11 @@ void FDisplayClusterMediaOutputSynchronizationPolicyEthernetBarrierBaseHandler::
 			for (UDisplayClusterICVFXCameraComponent* const ICVFXCameraComponent : ICVFXCameraComponents)
 			{
 				const FDisplayClusterConfigurationMediaICVFX& MediaSettings = ICVFXCameraComponent->CameraSettings.RenderSettings.Media;
+
 				if (MediaSettings.bEnable)
 				{
-					for (const FDisplayClusterConfigurationMediaOutputGroup& MediaOutputGroup : MediaSettings.MediaOutputGroups)
+					uint8 CaptureIdx = 0;
+					for (const FDisplayClusterConfigurationMediaOutputGroupICVFX& MediaOutputGroup : MediaSettings.MediaOutputGroups)
 					{
 						// Pick the same sync policy only
 						if (MediaOutputGroup.MediaOutput && MediaOutputGroup.OutputSyncPolicy)
@@ -276,12 +298,14 @@ void FDisplayClusterMediaOutputSynchronizationPolicyEthernetBarrierBaseHandler::
 									const FString ICVFXCaptureId = DisplayClusterMediaHelpers::MediaId::GenerateMediaId(
 										DisplayClusterMediaHelpers::MediaId::EMediaDeviceType::Output,
 										DisplayClusterMediaHelpers::MediaId::EMediaOwnerType::ICVFXCamera,
-										NodeId, RootActor->GetName(), ICVFXCameraComponent->GetName());
+										NodeId, RootActor->GetName(), ICVFXCameraComponent->GetName(), CaptureIdx);
 
 									OutMarkers.Add(ICVFXCaptureId);
 								}
 							}
 						}
+
+						++CaptureIdx;
 					}
 				}
 			}
