[nDisplay] Re-implemented minor things introduced with new network replication (26098376)

diff --git a/Engine/Plugins/Runtime/nDisplay/Source/DisplayClusterConfiguration/Private/DisplayClusterConfigurationTypes.cpp b/Engine/Plugins/Runtime/nDisplay/Source/DisplayClusterConfiguration/Private/DisplayClusterConfigurationTypes.cpp
--- a/Engine/Plugins/Runtime/nDisplay/Source/DisplayClusterConfiguration/Private/DisplayClusterConfigurationTypes.cpp
+++ b/Engine/Plugins/Runtime/nDisplay/Source/DisplayClusterConfiguration/Private/DisplayClusterConfigurationTypes.cpp
@@ -10,10 +10,8 @@
 #include "DisplayClusterConfigurationLog.h"
 #include "DisplayClusterConfigurationStrings.h"
 #include "DisplayClusterProjectionStrings.h"
-#include "Config/IDisplayClusterConfigManager.h"
 
 #include "Engine/StaticMesh.h"
-#include "IDisplayCluster.h"
 #include "UObject/Package.h"
 
 #if WITH_EDITOR
@@ -472,33 +470,23 @@ void UDisplayClusterConfigurationData::GetReferencedMeshNames(TArray<FString>& O
 	}
 }
 
-int32 UDisplayClusterConfigurationData::GetNumberOfClusterNodes() const
+uint32 UDisplayClusterConfigurationData::GetNumberOfClusterNodes() const
 {
-	if (Cluster == nullptr)
-	{
-		return 0;
-	}
-
-	return Cluster->Nodes.Num();
+	return Cluster ? static_cast<uint32>(Cluster->Nodes.Num()) : 0;
 }
 
 FString UDisplayClusterConfigurationData::GetPrimaryNodeAddress() const
 {
-	if (!IDisplayCluster::Get().GetConfigMgr())
-	{
-		return FString();
-	}
-
-	if (!IDisplayCluster::Get().GetConfigMgr()->GetConfig())
+	if (Cluster)
 	{
-		return FString();
+		const FString& PrimaryNodeId = Cluster->PrimaryNode.Id;
+		if (UDisplayClusterConfigurationClusterNode* const PrimaryNode = Cluster->GetNode(PrimaryNodeId))
+		{
+			return PrimaryNode->Host;
+		}
 	}
 
-	UDisplayClusterConfigurationData* DisplayClusterConfigurationData = IDisplayCluster::Get().GetConfigMgr()->GetConfig();
-	
-	const FString PrimaryNodeId = DisplayClusterConfigurationData->Cluster->PrimaryNode.Id;
-
-	return DisplayClusterConfigurationData->Cluster->GetNode(PrimaryNodeId)->Host;
+	return FString();
 }
 
 UDisplayClusterConfigurationData* UDisplayClusterConfigurationData::CreateNewConfigData(UObject* Owner, EObjectFlags ObjectFlags)
diff --git a/Engine/Plugins/Runtime/nDisplay/Source/DisplayClusterConfiguration/Public/DisplayClusterConfigurationTypes.h b/Engine/Plugins/Runtime/nDisplay/Source/DisplayClusterConfiguration/Public/DisplayClusterConfigurationTypes.h
--- a/Engine/Plugins/Runtime/nDisplay/Source/DisplayClusterConfiguration/Public/DisplayClusterConfigurationTypes.h
+++ b/Engine/Plugins/Runtime/nDisplay/Source/DisplayClusterConfiguration/Public/DisplayClusterConfigurationTypes.h
@@ -545,7 +545,7 @@ public:
 
 public:
 	/** Returns number of nodes in current cluster */
-	int32 GetNumberOfClusterNodes() const;
+	uint32 GetNumberOfClusterNodes() const;
 
 	/** Returns primary node IP address of current cluster */
 	FString GetPrimaryNodeAddress() const;
diff --git a/Engine/Plugins/Runtime/nDisplay/Source/DisplayClusterReplication/Private/DisplayClusterNetDriver.cpp b/Engine/Plugins/Runtime/nDisplay/Source/DisplayClusterReplication/Private/DisplayClusterNetDriver.cpp
--- a/Engine/Plugins/Runtime/nDisplay/Source/DisplayClusterReplication/Private/DisplayClusterNetDriver.cpp
+++ b/Engine/Plugins/Runtime/nDisplay/Source/DisplayClusterReplication/Private/DisplayClusterNetDriver.cpp
@@ -281,8 +281,11 @@ void UDisplayClusterNetDriver::TickFlush(float DeltaSeconds)
 	}
 
 	const UDisplayClusterConfigurationData* ConfigData = IDisplayCluster::Get().GetConfigMgr()->GetConfig();
+	check(ConfigData);
+	const uint32 ClusterNodesAmount = ConfigData->GetNumberOfClusterNodes();
+	check(ClusterNodesAmount > 0);
 
-	if ((World->GetNetMode() == NM_ListenServer) && (SyncConnections.Num() < (ConfigData->GetNumberOfClusterNodes() - 1)))
+	if ((World->GetNetMode() == NM_ListenServer) && (SyncConnections.Num() < (static_cast<int32>(ClusterNodesAmount) - 1)))
 	{
 		return;
 	}
