nDisplay: Freeze Editor Preview now behaves similar to Freeze Viewports.

diff --git a/Engine/Plugins/Runtime/nDisplay/Source/DisplayCluster/Private/Game/EngineClasses/Scene/DisplayClusterRootActorEditor.cpp b/Engine/Plugins/Runtime/nDisplay/Source/DisplayCluster/Private/Game/EngineClasses/Scene/DisplayClusterRootActorEditor.cpp
--- a/Engine/Plugins/Runtime/nDisplay/Source/DisplayCluster/Private/Game/EngineClasses/Scene/DisplayClusterRootActorEditor.cpp
+++ b/Engine/Plugins/Runtime/nDisplay/Source/DisplayCluster/Private/Game/EngineClasses/Scene/DisplayClusterRootActorEditor.cpp
@@ -501,11 +501,6 @@ void ADisplayClusterRootActor::ImplRenderPreview_Editor()
 
 	if (PreviewClusterNodeIndex < 0)
 	{
-		if (bFreezePreviewRender)
-		{
-			return;
-		}
-
 		// Allow preview render
 		PreviewClusterNodeIndex = 0;
 	}
@@ -555,7 +550,7 @@ void ADisplayClusterRootActor::ImplRenderPreview_Editor()
 		PreviewClusterNodeIndex++;
 		if (PreviewClusterNodeIndex >= ExistClusterNodesIDs.Num())
 		{
-			PreviewClusterNodeIndex = bFreezePreviewRender ? -1 : 0;
+			PreviewClusterNodeIndex = 0;
 
 			bOutputFrameToPostProcessRenderTarget = bOutputFrameToPostProcessRenderTarget ?
 				bPreviewEnablePostProcess :
diff --git a/Engine/Plugins/Runtime/nDisplay/Source/DisplayCluster/Private/Render/Viewport/Configuration/DisplayClusterViewportConfiguration.cpp b/Engine/Plugins/Runtime/nDisplay/Source/DisplayCluster/Private/Render/Viewport/Configuration/DisplayClusterViewportConfiguration.cpp
--- a/Engine/Plugins/Runtime/nDisplay/Source/DisplayCluster/Private/Render/Viewport/Configuration/DisplayClusterViewportConfiguration.cpp
+++ b/Engine/Plugins/Runtime/nDisplay/Source/DisplayCluster/Private/Render/Viewport/Configuration/DisplayClusterViewportConfiguration.cpp
@@ -236,6 +236,7 @@ bool FDisplayClusterViewportConfiguration::ImplUpdateConfiguration(EDisplayClust
 
 				RenderFrameSettings.bIsRenderingInEditor = true;
 				RenderFrameSettings.bIsPreviewRendering = true;
+				RenderFrameSettings.bFreezePreviewRender = InPreviewSettings->bFreezePreviewRender;
 				
 				if (InPreviewSettings->bIsPIE)
 				{
diff --git a/Engine/Plugins/Runtime/nDisplay/Source/DisplayCluster/Private/Render/Viewport/Configuration/DisplayClusterViewportConfigurationICVFX.cpp b/Engine/Plugins/Runtime/nDisplay/Source/DisplayCluster/Private/Render/Viewport/Configuration/DisplayClusterViewportConfigurationICVFX.cpp
--- a/Engine/Plugins/Runtime/nDisplay/Source/DisplayCluster/Private/Render/Viewport/Configuration/DisplayClusterViewportConfigurationICVFX.cpp
+++ b/Engine/Plugins/Runtime/nDisplay/Source/DisplayCluster/Private/Render/Viewport/Configuration/DisplayClusterViewportConfigurationICVFX.cpp
@@ -306,8 +306,10 @@ void FDisplayClusterViewportConfigurationICVFX::Update()
 	{
 		ImplBeginReallocateViewports(*ViewportManager);
 
+		const FDisplayClusterRenderFrameSettings& RenderFrameSettings = ViewportManager->GetRenderFrameSettings();
+		
 		// Disable ICVFX if cluster node rendering is not used
-		const FString& ClusterNodeId = ViewportManager->GetRenderFrameSettings().ClusterNodeId;
+		const FString& ClusterNodeId = RenderFrameSettings.ClusterNodeId;
 		if(ClusterNodeId.IsEmpty())
 		{
 			// The nDisplay viewport should now always have the name of the cluster node.
@@ -402,7 +404,11 @@ void FDisplayClusterViewportConfigurationICVFX::Update()
 				}
 			}
 
-			if (StageSettings.bFreezeRenderOuterViewports)
+			if (StageSettings.bFreezeRenderOuterViewports
+#if WITH_EDITOR
+				|| (RenderFrameSettings.bIsPreviewRendering && RenderFrameSettings.bFreezePreviewRender)
+#endif
+				)
 			{
 				// Freeze rendering for outer viewports
 				for (const TSharedPtr<FDisplayClusterViewport, ESPMode::ThreadSafe>& TargetIt : TargetViewports)
diff --git a/Engine/Plugins/Runtime/nDisplay/Source/DisplayCluster/Private/Render/Viewport/RenderFrame/DisplayClusterRenderFrameSettings.h b/Engine/Plugins/Runtime/nDisplay/Source/DisplayCluster/Private/Render/Viewport/RenderFrame/DisplayClusterRenderFrameSettings.h
--- a/Engine/Plugins/Runtime/nDisplay/Source/DisplayCluster/Private/Render/Viewport/RenderFrame/DisplayClusterRenderFrameSettings.h
+++ b/Engine/Plugins/Runtime/nDisplay/Source/DisplayCluster/Private/Render/Viewport/RenderFrame/DisplayClusterRenderFrameSettings.h
@@ -76,6 +76,9 @@ struct FDisplayClusterRenderFrameSettings
 	// Configuration used to render preview
 	bool bIsPreviewRendering = false;
 
+	// If the preview rendering should be frozen
+	bool bFreezePreviewRender = false;
+	
 	// Allow mGPU in editor mode
 	bool bAllowMultiGPURenderingInEditor = false;
 	int32 PreviewMinGPUIndex = 0;
