[nDisplay] Fixed DDC issue caused by Multi-MediaOutput feature integration (26104746)
- renamed FDisplayClusterConfigurationMediaInputGroupICVFX back to FDisplayClusterConfigurationMediaInputGroup
- renamed FDisplayClusterConfigurationMediaOutputGroupICVFX back to FDisplayClusterConfigurationMediaOutputGroup

diff --git a/Engine/Plugins/Runtime/nDisplay/Source/DisplayClusterConfiguration/Private/DisplayClusterConfigurationTypes_Media.cpp b/Engine/Plugins/Runtime/nDisplay/Source/DisplayClusterConfiguration/Private/DisplayClusterConfigurationTypes_Media.cpp
--- a/Engine/Plugins/Runtime/nDisplay/Source/DisplayClusterConfiguration/Private/DisplayClusterConfigurationTypes_Media.cpp
+++ b/Engine/Plugins/Runtime/nDisplay/Source/DisplayClusterConfiguration/Private/DisplayClusterConfigurationTypes_Media.cpp
@@ -37,9 +37,9 @@ bool FDisplayClusterConfigurationMediaICVFX::IsMediaInputAssigned(const FString&
 
 bool FDisplayClusterConfigurationMediaICVFX::IsMediaOutputAssigned(const FString& NodeId) const
 {
-	const TArray<FDisplayClusterConfigurationMediaOutputGroupICVFX> NodeOutputGroups = GetMediaOutputGroups(NodeId);
+	const TArray<FDisplayClusterConfigurationMediaOutputGroup> NodeOutputGroups = GetMediaOutputGroups(NodeId);
 
-	for (const FDisplayClusterConfigurationMediaOutputGroupICVFX& NodeOutputGroup : NodeOutputGroups)
+	for (const FDisplayClusterConfigurationMediaOutputGroup& NodeOutputGroup : NodeOutputGroups)
 	{
 		if (IsValid(NodeOutputGroup.MediaOutput))
 		{
@@ -58,7 +58,7 @@ bool FDisplayClusterConfigurationMediaICVFX::IsMediaOutputAssigned(const FString
 UMediaSource* FDisplayClusterConfigurationMediaICVFX::GetMediaSource(const FString& NodeId) const
 {
 	// Look up for a group that contains node ID specified
-	for (const FDisplayClusterConfigurationMediaInputGroupICVFX& MediaInputGroup : MediaInputGroups)
+	for (const FDisplayClusterConfigurationMediaInputGroup& MediaInputGroup : MediaInputGroups)
 	{
 		const bool bNodeFound = MediaInputGroup.ClusterNodes.ItemNames.ContainsByPredicate([NodeId](const FString& Item)
 			{
@@ -74,9 +74,9 @@ UMediaSource* FDisplayClusterConfigurationMediaICVFX::GetMediaSource(const FStri
 	return nullptr;
 }
 
-TArray<FDisplayClusterConfigurationMediaOutputGroupICVFX> FDisplayClusterConfigurationMediaICVFX::GetMediaOutputGroups(const FString& NodeId) const
+TArray<FDisplayClusterConfigurationMediaOutputGroup> FDisplayClusterConfigurationMediaICVFX::GetMediaOutputGroups(const FString& NodeId) const
 {
-	return MediaOutputGroups.FilterByPredicate([NodeId](const FDisplayClusterConfigurationMediaOutputGroupICVFX& Item)
+	return MediaOutputGroups.FilterByPredicate([NodeId](const FDisplayClusterConfigurationMediaOutputGroup& Item)
 		{
 			return Item.ClusterNodes.ItemNames.Contains(NodeId);
 		});
diff --git a/Engine/Plugins/Runtime/nDisplay/Source/DisplayClusterConfiguration/Public/DisplayClusterConfigurationTypes_Media.h b/Engine/Plugins/Runtime/nDisplay/Source/DisplayClusterConfiguration/Public/DisplayClusterConfigurationTypes_Media.h
--- a/Engine/Plugins/Runtime/nDisplay/Source/DisplayClusterConfiguration/Public/DisplayClusterConfigurationTypes_Media.h
+++ b/Engine/Plugins/Runtime/nDisplay/Source/DisplayClusterConfiguration/Public/DisplayClusterConfigurationTypes_Media.h
@@ -105,7 +105,7 @@ PRAGMA_ENABLE_DEPRECATION_WARNINGS
  * Media input group (ICVFX)
  */
 USTRUCT(Blueprintable)
-struct DISPLAYCLUSTERCONFIGURATION_API FDisplayClusterConfigurationMediaInputGroupICVFX
+struct DISPLAYCLUSTERCONFIGURATION_API FDisplayClusterConfigurationMediaInputGroup
 	: public FDisplayClusterConfigurationMediaInput
 {
 	GENERATED_BODY()
@@ -121,7 +121,7 @@ public:
  * Media output group (ICVFX)
  */
 USTRUCT(Blueprintable)
-struct DISPLAYCLUSTERCONFIGURATION_API FDisplayClusterConfigurationMediaOutputGroupICVFX
+struct DISPLAYCLUSTERCONFIGURATION_API FDisplayClusterConfigurationMediaOutputGroup
 	: public FDisplayClusterConfigurationMediaOutput
 {
 	GENERATED_BODY()
@@ -148,11 +148,11 @@ public:
 
 	/** Media input mapping */
 	UPROPERTY(EditAnywhere, BlueprintReadWrite, Category = "Media")
-	TArray<FDisplayClusterConfigurationMediaInputGroupICVFX> MediaInputGroups;
+	TArray<FDisplayClusterConfigurationMediaInputGroup> MediaInputGroups;
 
 	/** Media output mapping */
 	UPROPERTY(EditAnywhere, BlueprintReadWrite, Category = "Media")
-	TArray<FDisplayClusterConfigurationMediaOutputGroupICVFX> MediaOutputGroups;
+	TArray<FDisplayClusterConfigurationMediaOutputGroup> MediaOutputGroups;
 
 public:
 	/** Returns true if a specific cluster node has media source assigned */
@@ -165,7 +165,7 @@ public:
 	UMediaSource* GetMediaSource(const FString& NodeId) const;
 
 	/** Returns media outputs bound to a specific cluster node */
-	TArray<FDisplayClusterConfigurationMediaOutputGroupICVFX> GetMediaOutputGroups(const FString& NodeId) const;
+	TArray<FDisplayClusterConfigurationMediaOutputGroup> GetMediaOutputGroups(const FString& NodeId) const;
 
 public:
 	UE_DEPRECATED(5.3, "This function has beend deprecated. Please use GetMediaOutputGroups.")
diff --git a/Engine/Plugins/Runtime/nDisplay/Source/DisplayClusterMedia/Private/DisplayClusterMediaModule.cpp b/Engine/Plugins/Runtime/nDisplay/Source/DisplayClusterMedia/Private/DisplayClusterMediaModule.cpp
--- a/Engine/Plugins/Runtime/nDisplay/Source/DisplayClusterMedia/Private/DisplayClusterMediaModule.cpp
+++ b/Engine/Plugins/Runtime/nDisplay/Source/DisplayClusterMedia/Private/DisplayClusterMediaModule.cpp
@@ -233,9 +233,9 @@ void FDisplayClusterMediaModule::InitializeMedia()
 						}
 
 						// Media capture
-						const TArray<FDisplayClusterConfigurationMediaOutputGroupICVFX> MediaOutputItems = MediaSettings.GetMediaOutputGroups(ClusterNodeId);
+						const TArray<FDisplayClusterConfigurationMediaOutputGroup> MediaOutputItems = MediaSettings.GetMediaOutputGroups(ClusterNodeId);
 						uint8 CaptureIdx = 0;
-						for (const FDisplayClusterConfigurationMediaOutputGroupICVFX& MediaOutputItem : MediaOutputItems)
+						for (const FDisplayClusterConfigurationMediaOutputGroup& MediaOutputItem : MediaOutputItems)
 						{
 							if (IsValid(MediaOutputItem.MediaOutput))
 							{
diff --git a/Engine/Plugins/Runtime/nDisplay/Source/DisplayClusterMedia/Private/Synchronization/SyncPolicies/DisplayClusterMediaOutputSynchronizationPolicyEthernetBarrierBase.cpp b/Engine/Plugins/Runtime/nDisplay/Source/DisplayClusterMedia/Private/Synchronization/SyncPolicies/DisplayClusterMediaOutputSynchronizationPolicyEthernetBarrierBase.cpp
--- a/Engine/Plugins/Runtime/nDisplay/Source/DisplayClusterMedia/Private/Synchronization/SyncPolicies/DisplayClusterMediaOutputSynchronizationPolicyEthernetBarrierBase.cpp
+++ b/Engine/Plugins/Runtime/nDisplay/Source/DisplayClusterMedia/Private/Synchronization/SyncPolicies/DisplayClusterMediaOutputSynchronizationPolicyEthernetBarrierBase.cpp
@@ -286,7 +286,7 @@ void FDisplayClusterMediaOutputSynchronizationPolicyEthernetBarrierBaseHandler::
 				if (MediaSettings.bEnable)
 				{
 					uint8 CaptureIdx = 0;
-					for (const FDisplayClusterConfigurationMediaOutputGroupICVFX& MediaOutputGroup : MediaSettings.MediaOutputGroups)
+					for (const FDisplayClusterConfigurationMediaOutputGroup& MediaOutputGroup : MediaSettings.MediaOutputGroups)
 					{
 						// Pick the same sync policy only
 						if (MediaOutputGroup.MediaOutput && MediaOutputGroup.OutputSyncPolicy)
