nDisplay:
- Added changes to the UDisplayClusterCameraComponent
	* Added GetDesiredView() virtual function, which returns ViewPoint with PP and CustomNearClippingPlane values for this component. Inherited components can use references to external cameras, which will return their data through this function.

- Added changes to the IDisplayClusterWarpBlend interface:
	* Added new function ShouldSupportICVFX(). Reasons: only A3D profile supports ICVFX, custom warp policy can disable ICVFX, etc.

- Added changes in the IDisplayClusterProjectionPolicy interface:
	* Added new parameter to ShouldSupportICVFX(). Reason: It is used by a new function in the IDisplayClusterWarpBlend interface.
	* Renamed GetViewPoint() function to SetupProjectionViewPoint(). The reason: the new name is used in other components (which provide a viewpoint) inside nDisplay. As a result, the value of the new name will be the same for several different classes within nDisplay.
	* Renamed the function OverridePostProcessSettings() to UpdatePostProcessSettings(). Reason: The name OverridePostProcessSettings is already used in the UE core classes.

- Added changes in the IDisplayClusterViewport interface:
	* Removed the ViewOffset, NCP and FCP parameters from the CalculateView() function. The reason: this values is now calculated inside this function, which eliminates the duplication of this code in 3 places (DisplayClusterMoviePipelineViewportPass, DisplayClusterViewport_EditorPreview, DisplayClusterDeviceBase).
	* Added new function GetViewPointCameraEye(): Gets the position of the ViewPointCamera component used by this viewport. The ViewPointCamera component has a new virtual function GetDesiredView(), which can return the position from the cameras it refers to. But in warpblend math, you need the component's position (viewpoint), not the reference.
	* Added new auxiliary static functions GetCameraComponentView() and GetPlayerCameraView(). Reasons: Removing duplicate code, getting PP inside the viewinfo, and optional CustomNearClippingPlane value.

- Moved the call of the renamed UpdatePostProcessSettings() function (formerly called 'OverridePostProcessSettings') to the FDisplayClusterViewportConfigurationProjectionPolicy::Update() function
	* Also removed the PP setting for the camera projection policy in FDisplayClusterViewportConfigurationProjectionPolicy::UpdateCameraPolicy_Base(), since this is already done in IDisplayClusterProjectionPolicy::UpdatePostProcessSettings().

- Added changes to the projection policy 'camera'
	* A new parameter bUseCameraPostprocess (Allow to use camera postprocess) was added to the CameraPolicySettings structure.
	* Added UpdatePostProcessSettings() function - fixed bug: postprocess was not passed from camera (postprocess is overridden from projection policy, not directly for viewport. CL #25885471)
	* Renamed the GetViewPoint() function to SetupProjectionViewPoint(). The associated logic has also been reworked.

- Added changes to the projection policy 'mpcdi'
	* Changed ShouldSupportICVFX(): WarpBlendInterface can disable ICVFX (only A3D supports ICVFX)
	* Added GetViewPointCameraEye() to the CalculateView() function (Override viewpoint: MPCDI always expects the location of the viewpoint component)

- Added changes to the projection policy 'manual' and 'link'
	* Added SetupProjectionViewPoint() function.

- fixed spelling 'CamersSettings' to 'CamerasSettings'in DisplayClusterProjection module.

- added UE_DEPRECATED for changes in public interfaces.

-reworked Viewport/Configuration
	* added warppolicy support
	* code cleanup

- code cleanup for UDisplayClusterICVFXCameraComponent

diff --git a/Engine/Plugins/Runtime/nDisplay/Source/DisplayCluster/Private/Game/EngineClasses/Scene/DisplayClusterCameraComponent.cpp b/Engine/Plugins/Runtime/nDisplay/Source/DisplayCluster/Private/Game/EngineClasses/Scene/DisplayClusterCameraComponent.cpp
--- a/Engine/Plugins/Runtime/nDisplay/Source/DisplayCluster/Private/Game/EngineClasses/Scene/DisplayClusterCameraComponent.cpp
+++ b/Engine/Plugins/Runtime/nDisplay/Source/DisplayCluster/Private/Game/EngineClasses/Scene/DisplayClusterCameraComponent.cpp
@@ -30,6 +30,22 @@ UDisplayClusterCameraComponent::UDisplayClusterCameraComponent(const FObjectInit
 #endif
 }
 
+void UDisplayClusterCameraComponent::GetDesiredView(FMinimalViewInfo& InOutViewInfo, float* OutCustomNearClippingPlane)
+{
+	// Use this component as a camera
+	InOutViewInfo.Location = GetComponentLocation();
+	InOutViewInfo.Rotation = GetComponentRotation();
+
+	// Ignore PP, because this component has no such settings
+	InOutViewInfo.PostProcessBlendWeight = 0.f;
+
+	if (OutCustomNearClippingPlane)
+	{
+		// Value less than zero means: don't override the NCP value
+		*OutCustomNearClippingPlane = -1.f;
+	}
+}
+
 #if WITH_EDITOR
 void UDisplayClusterCameraComponent::SetVisualizationScale(float Scale)
 {
diff --git a/Engine/Plugins/Runtime/nDisplay/Source/DisplayCluster/Private/Game/EngineClasses/Scene/DisplayClusterICVFXCameraComponent.cpp b/Engine/Plugins/Runtime/nDisplay/Source/DisplayCluster/Private/Game/EngineClasses/Scene/DisplayClusterICVFXCameraComponent.cpp
--- a/Engine/Plugins/Runtime/nDisplay/Source/DisplayCluster/Private/Game/EngineClasses/Scene/DisplayClusterICVFXCameraComponent.cpp
+++ b/Engine/Plugins/Runtime/nDisplay/Source/DisplayCluster/Private/Game/EngineClasses/Scene/DisplayClusterICVFXCameraComponent.cpp
@@ -53,20 +53,52 @@ void UDisplayClusterICVFXCameraComponent::PostLoad()
 	PRAGMA_ENABLE_DEPRECATION_WARNINGS
 }
 
-void UDisplayClusterICVFXCameraComponent::GetDesiredView(FMinimalViewInfo& DesiredView)
+void UDisplayClusterICVFXCameraComponent::GetCameraView(float DeltaTime, FMinimalViewInfo& InOutViewInfo)
 {
-	if (ADisplayClusterRootActor* RootActor = Cast<ADisplayClusterRootActor>(GetOwner()))
+	if (CameraSettings.ExternalCameraActor.IsValid())
+	{
+		// Get ViewInfo from external CineCamera
+		CameraSettings.ExternalCameraActor->GetCineCameraComponent()->GetCameraView(DeltaTime, InOutViewInfo);
+	}
+	else
+	{
+		// Get ViewInfo from this component
+		UCineCameraComponent::GetCameraView(DeltaTime, InOutViewInfo);
+	}
+
+	// CameraSettings can disable posprocess from this camera
+	if(!CameraSettings.RenderSettings.bUseCameraComponentPostprocess)
+	{
+		InOutViewInfo.PostProcessSettings = FPostProcessSettings();
+		InOutViewInfo.PostProcessBlendWeight = 0.0f;
+	}
+
+	// Add postprocess blur settings to viewinfo PP
+	const FDisplayClusterConfigurationICVFX_CameraMotionBlurOverridePPS& OverrideMotionBlurPPS = CameraSettings.CameraMotionBlur.MotionBlurPPS;
+	if (OverrideMotionBlurPPS.bReplaceEnable)
 	{
-		UCineCameraComponent* const CineCameraComponent = CameraSettings.ExternalCameraActor.IsValid() ? CameraSettings.ExternalCameraActor->GetCineCameraComponent() : this;
+		// Send camera postprocess to override
+		InOutViewInfo.PostProcessBlendWeight = 1.0f;
+
+		InOutViewInfo.PostProcessSettings.MotionBlurAmount = OverrideMotionBlurPPS.MotionBlurAmount;
+		InOutViewInfo.PostProcessSettings.bOverride_MotionBlurAmount = true;
+
+		InOutViewInfo.PostProcessSettings.MotionBlurMax = OverrideMotionBlurPPS.MotionBlurMax;
+		InOutViewInfo.PostProcessSettings.bOverride_MotionBlurMax = true;
 
-		const float DeltaTime = RootActor->GetWorldDeltaSeconds();
-		CineCameraComponent->GetCameraView(DeltaTime, DesiredView);
+		InOutViewInfo.PostProcessSettings.MotionBlurPerObjectSize = OverrideMotionBlurPPS.MotionBlurPerObjectSize;
+		InOutViewInfo.PostProcessSettings.bOverride_MotionBlurPerObjectSize = true;
 	}
 }
 
-UCameraComponent* UDisplayClusterICVFXCameraComponent::GetCameraComponent()
+UCineCameraComponent* UDisplayClusterICVFXCameraComponent::GetActualCineCameraComponent()
 {
-	return CameraSettings.ExternalCameraActor.IsValid() ? CameraSettings.ExternalCameraActor->GetCineCameraComponent() : this;
+	if (UCineCameraComponent* ExternalCineCameraComponent = CameraSettings.ExternalCameraActor.IsValid() ? CameraSettings.ExternalCameraActor->GetCineCameraComponent() : nullptr)
+	{
+		return ExternalCineCameraComponent;
+	}
+
+	return this;
 }
 
 FString UDisplayClusterICVFXCameraComponent::GetCameraUniqueId() const
diff --git a/Engine/Plugins/Runtime/nDisplay/Source/DisplayCluster/Private/Game/EngineClasses/Scene/DisplayClusterPreviewShareComponent.cpp b/Engine/Plugins/Runtime/nDisplay/Source/DisplayCluster/Private/Game/EngineClasses/Scene/DisplayClusterPreviewShareComponent.cpp
--- a/Engine/Plugins/Runtime/nDisplay/Source/DisplayCluster/Private/Game/EngineClasses/Scene/DisplayClusterPreviewShareComponent.cpp
+++ b/Engine/Plugins/Runtime/nDisplay/Source/DisplayCluster/Private/Game/EngineClasses/Scene/DisplayClusterPreviewShareComponent.cpp
@@ -724,7 +724,7 @@ void UDisplayClusterPreviewShareComponent::SyncIcvxCamerasFromSourceActor(const
 		check(SrcIcvfxCameraComponent);
 
 		// This one sould be the actual source camera, e.g. the camera component of the referenced cine camera.
-		const UCameraComponent* SrcCameraComponent = SrcIcvfxCameraComponent->GetCameraComponent();
+		const UCineCameraComponent* SrcCameraComponent = SrcIcvfxCameraComponent->GetActualCineCameraComponent();
 		check(SrcCameraComponent);
 		
 		// Get the source camera transform with respect to the Source root actor
@@ -732,7 +732,7 @@ void UDisplayClusterPreviewShareComponent::SyncIcvxCamerasFromSourceActor(const
 
 		// Apply this transform to the destination camera
 
-		UCameraComponent* DstCameraComponent = DstIcvfxCameraComponent->GetCameraComponent();
+		UCineCameraComponent* DstCameraComponent = DstIcvfxCameraComponent->GetActualCineCameraComponent();
 
 		AActor* DstCameraOwnerActor = DstCameraComponent->GetOwner();
 		
diff --git a/Engine/Plugins/Runtime/nDisplay/Source/DisplayCluster/Private/Render/Containers/DisplayClusterRender_Texture.cpp b/Engine/Plugins/Runtime/nDisplay/Source/DisplayCluster/Private/Render/Containers/DisplayClusterRender_Texture.cpp
--- a/Engine/Plugins/Runtime/nDisplay/Source/DisplayCluster/Private/Render/Containers/DisplayClusterRender_Texture.cpp
+++ b/Engine/Plugins/Runtime/nDisplay/Source/DisplayCluster/Private/Render/Containers/DisplayClusterRender_Texture.cpp
@@ -28,9 +28,7 @@ FDisplayClusterRender_Texture::FDisplayClusterRender_Texture(const FString& InUn
 { }
 
 FDisplayClusterRender_Texture::~FDisplayClusterRender_Texture()
-{
-	SetTextureResource(nullptr);
-}
+{ }
 
 const TSharedPtr<FDisplayClusterRender_TextureResource, ESPMode::ThreadSafe>& FDisplayClusterRender_Texture::GetTextureResource() const
 {
diff --git a/Engine/Plugins/Runtime/nDisplay/Source/DisplayCluster/Private/Render/Containers/DisplayClusterRender_Texture.h b/Engine/Plugins/Runtime/nDisplay/Source/DisplayCluster/Private/Render/Containers/DisplayClusterRender_Texture.h
--- a/Engine/Plugins/Runtime/nDisplay/Source/DisplayCluster/Private/Render/Containers/DisplayClusterRender_Texture.h
+++ b/Engine/Plugins/Runtime/nDisplay/Source/DisplayCluster/Private/Render/Containers/DisplayClusterRender_Texture.h
@@ -77,6 +77,13 @@ public:
 	{
 		return UniqueTextureName;
 	}
+
+	/** Method for releasing a cached data item, called before its destructor. */
+	inline void ReleaseDataCacheItem()
+	{
+		SetTextureResource(nullptr);
+	}
+
 	// ~~ End TDisplayClusterDataCache
 
 protected:
diff --git a/Engine/Plugins/Runtime/nDisplay/Source/DisplayCluster/Private/Render/Device/DisplayClusterDeviceBase.cpp b/Engine/Plugins/Runtime/nDisplay/Source/DisplayCluster/Private/Render/Device/DisplayClusterDeviceBase.cpp
--- a/Engine/Plugins/Runtime/nDisplay/Source/DisplayCluster/Private/Render/Device/DisplayClusterDeviceBase.cpp
+++ b/Engine/Plugins/Runtime/nDisplay/Source/DisplayCluster/Private/Render/Device/DisplayClusterDeviceBase.cpp
@@ -260,20 +260,9 @@ void FDisplayClusterDeviceBase::CalculateStereoViewOffset(const int32 ViewIndex,
 	}
 
 	// The camera position has already been determined from the SetupViewPoint() function
-	// Obtaining the stereo eye offset for a given viewport
-	const float PassOffsetSwap = ViewportPtr->GetStereoEyeOffsetDistance(ViewportContextNum);
-
-	FVector ViewOffset = FVector::ZeroVector;
-	{
-		// Apply computed offset to the view location
-		const FQuat EyeQuat = ViewRotation.Quaternion();
-		ViewOffset = EyeQuat.RotateVector(FVector(0.0f, PassOffsetSwap, 0.0f));
-		ViewLocation += ViewOffset;
-	}
-
+	// Obtaining the offset of the stereo eye and the values of the projection clipping plane for the given viewport was moved inside CalculateView().
 	// Perform view calculations on a policy side
-	const float CfgNCP = GNearClippingPlane;
-	if (ViewportPtr->CalculateView(ViewportContextNum, ViewLocation, ViewRotation, ViewOffset, WorldToMeters, CfgNCP, CfgNCP) == false)
+	if (ViewportPtr->CalculateView(ViewportContextNum, ViewLocation, ViewRotation, WorldToMeters) == false)
 	{
 #if WITH_EDITOR
 		// Hide spam in logs when configuring VP in editor [UE-114493]
diff --git a/Engine/Plugins/Runtime/nDisplay/Source/DisplayCluster/Private/Render/Viewport/Configuration/DisplayClusterViewportConfiguration.cpp b/Engine/Plugins/Runtime/nDisplay/Source/DisplayCluster/Private/Render/Viewport/Configuration/DisplayClusterViewportConfiguration.cpp
--- a/Engine/Plugins/Runtime/nDisplay/Source/DisplayCluster/Private/Render/Viewport/Configuration/DisplayClusterViewportConfiguration.cpp
+++ b/Engine/Plugins/Runtime/nDisplay/Source/DisplayCluster/Private/Render/Viewport/Configuration/DisplayClusterViewportConfiguration.cpp
@@ -196,6 +196,7 @@ bool FDisplayClusterViewportConfiguration::ImplUpdateConfiguration(EDisplayClust
 {
 	check(IsInGameThread());
 
+
 	ADisplayClusterRootActor* RootActor = GetRootActor();
 	if (RootActor)
 	{
@@ -203,7 +204,7 @@ bool FDisplayClusterViewportConfiguration::ImplUpdateConfiguration(EDisplayClust
 		FDisplayClusterViewportManager* ViewportManager = GetViewportManager();
 		if (ConfigurationData && ViewportManager)
 		{
-			FDisplayClusterViewportConfigurationBase ConfigurationBase(*ViewportManager, *RootActor, *ConfigurationData);
+			FDisplayClusterViewportConfigurationBase ConfigurationBase(*ViewportManager, *RootActor, *ConfigurationData, RenderFrameSettings);
 			FDisplayClusterViewportConfigurationICVFX ConfigurationICVFX(*RootActor);
 			FDisplayClusterViewportConfigurationProjectionPolicy ConfigurationProjectionPolicy(*ViewportManager, *RootActor, *ConfigurationData);
 
@@ -211,7 +212,6 @@ bool FDisplayClusterViewportConfiguration::ImplUpdateConfiguration(EDisplayClust
 
 			// Set current rendering mode
 			RenderFrameSettings.RenderMode = InRenderMode;
-			RenderFrameSettings.ClusterNodeId = InClusterNodeId;
 
 			// Support alpha channel capture
 			RenderFrameSettings.AlphaChannelCaptureMode = GetAlphaChannelCaptureMode();
@@ -254,11 +254,11 @@ bool FDisplayClusterViewportConfiguration::ImplUpdateConfiguration(EDisplayClust
 
 			if (InViewportNames)
 			{
-				ConfigurationBase.Update(*InViewportNames, RenderFrameSettings);
+				ConfigurationBase.UpdateCustomViewports(*InViewportNames);
 			}
 			else
 			{
-				ConfigurationBase.Update(InClusterNodeId);
+				ConfigurationBase.UpdateClusterNodeViewports(InClusterNodeId);
 			}
 
 			ConfigurationICVFX.Update();
@@ -274,10 +274,10 @@ bool FDisplayClusterViewportConfiguration::ImplUpdateConfiguration(EDisplayClust
 
 			ImplUpdateConfigurationVisibility(*RootActor, *ConfigurationData);
 
-			if (!InClusterNodeId.IsEmpty())
+			if (!InViewportNames)
 			{
-				// support postprocess only for per-node render
-				ConfigurationBase.UpdateClusterNodePostProcess(InClusterNodeId, RenderFrameSettings);
+				// Update postprocess for current cluster node
+				ConfigurationBase.UpdateClusterNodePostProcess(InClusterNodeId);
 			}
 
 			ImplPostUpdateRenderFrameConfiguration();
@@ -302,32 +302,8 @@ bool FDisplayClusterViewportConfiguration::UpdateCustomConfiguration(EDisplayClu
 #if WITH_EDITOR
 bool FDisplayClusterViewportConfiguration::UpdatePreviewConfiguration(EDisplayClusterRenderFrameMode InRenderMode, const FString& InClusterNodeId, const FDisplayClusterPreviewSettings& InPreviewSettings)
 {
-	if (InClusterNodeId.Equals(DisplayClusterConfigurationStrings::gui::preview::PreviewNodeAll, ESearchCase::IgnoreCase))
-	{
-		check(!InPreviewSettings.bIsPIE);
-
-		// initialize all nodes
-		ADisplayClusterRootActor* RootActor = GetRootActor();
-		if (RootActor != nullptr)
-		{
-			const UDisplayClusterConfigurationData* ConfigurationData = RootActor->GetConfigData();
-			if (ConfigurationData != nullptr && ConfigurationData->Cluster != nullptr)
-			{
-				TArray<FString> ClusterNodesIDs;
-				ConfigurationData->Cluster->GetNodeIds(ClusterNodesIDs);
-				for (const FString& ClusterNodeIdIt : ClusterNodesIDs)
-				{
-					ImplUpdateConfiguration(InRenderMode, ClusterNodeIdIt, &InPreviewSettings, nullptr);
-				}
-
-				// all cluster nodes viewports updated
-				return true;
-			}
-		}
-
-		return false;
-	}
-
+	// Special cluster node names are now supported in the lower classes.
+	// Therefore, we don't need to go through the cluster nodes to initialize the whole cluster.
 	return ImplUpdateConfiguration(InRenderMode, InClusterNodeId, &InPreviewSettings, nullptr);
 }
 #endif
diff --git a/Engine/Plugins/Runtime/nDisplay/Source/DisplayCluster/Private/Render/Viewport/Configuration/DisplayClusterViewportConfigurationBase.cpp b/Engine/Plugins/Runtime/nDisplay/Source/DisplayCluster/Private/Render/Viewport/Configuration/DisplayClusterViewportConfigurationBase.cpp
--- a/Engine/Plugins/Runtime/nDisplay/Source/DisplayCluster/Private/Render/Viewport/Configuration/DisplayClusterViewportConfigurationBase.cpp
+++ b/Engine/Plugins/Runtime/nDisplay/Source/DisplayCluster/Private/Render/Viewport/Configuration/DisplayClusterViewportConfigurationBase.cpp
@@ -2,10 +2,13 @@
 
 #include "DisplayClusterViewportConfigurationBase.h"
 
+#include "Render/Projection/IDisplayClusterProjectionPolicy.h"
+
 #include "Render/Viewport/DisplayClusterViewport.h"
 #include "Render/Viewport/DisplayClusterViewportManager.h"
 #include "Render/Viewport/Postprocess/DisplayClusterViewportPostProcessManager.h"
 #include "Render/Viewport/Postprocess/DisplayClusterViewportPostProcessOutputRemap.h"
+#include "Render/Viewport/RenderFrame/DisplayClusterRenderFrameSettings.h"
 
 #include "DisplayClusterViewportConfigurationHelpers.h"
 #include "DisplayClusterConfigurationTypes.h"
@@ -17,6 +20,8 @@
 
 #include "Misc/DisplayClusterLog.h"
 
+#include "Components/DisplayClusterCameraComponent.h"
+
 ///////////////////////////////////////////////////////////////////
 // Copied from "TextureShareDisplayCluster/Misc/TextureShareDisplayClusterStrings.h"
 namespace TextureShareDisplayClusterStrings
@@ -32,147 +37,215 @@ namespace TextureShareDisplayClusterStrings
 	}
 };
 
+namespace UE::DisplayCluster::Viewport::Configuration
+{
+	/**
+	 * Getting a valid cluster node name.
+	 * An empty cluster node name means using the viewports of the entire cluster.
+	 */
+	static inline FString GetValidClusterNodeId(const FString& InClusterNodeId)
+	{
+		// when InClusterNodeId==PreviewNodeAll, means that it is an undefined cluster node
+		const FString ClusterNodeId = InClusterNodeId == DisplayClusterConfigurationStrings::gui::preview::PreviewNodeAll ? TEXT("") : InClusterNodeId;
+
+		return ClusterNodeId;
+	}
+};
+using namespace UE::DisplayCluster::Viewport::Configuration;
+
+TArray<FString> FDisplayClusterViewportConfigurationBase::DisabledPostprocessNames;
+
 ///////////////////////////////////////////////////////////////////
 // FDisplayClusterViewportConfigurationBase
 ///////////////////////////////////////////////////////////////////
-TArray<FString> FDisplayClusterViewportConfigurationBase::DisabledPostprocessNames;
-
-void FDisplayClusterViewportConfigurationBase::Update(const FString& ClusterNodeId)
+void FDisplayClusterViewportConfigurationBase::UpdateClusterNodeViewports(const FString& InClusterNodeId)
 {
-	TMap<FString, UDisplayClusterConfigurationViewport*> DesiredViewports;
+	// The input cluster node name may contain special values. Obtain a valid value.
+	const FString ClusterNodeId = GetValidClusterNodeId(InClusterNodeId);
 
-	// Get render viewports for cluster node
-	const UDisplayClusterConfigurationClusterNode* ClusterNodeConfiguration = ConfigurationData.Cluster->GetNode(ClusterNodeId);
-	if (ClusterNodeConfiguration)
+	// Initialize variable EntireClusterViewports from configuration
+	ImplInitializeEntireClusterViewportsList();
+
+	if (ClusterNodeId.IsEmpty())
+	{
+		// If the cluster node name is empty, we use the viewports of the entire cluster 
+		CurrentFrameViewports = EntireClusterViewports;
+	}
+	else
 	{
-		for (const TPair<FString, TObjectPtr<UDisplayClusterConfigurationViewport>>& ViewportIt : ClusterNodeConfiguration->Viewports)
+		// Otherwise, we only use the viewports for one cluster node
+		CurrentFrameViewports.Reset();
+		for (const FDisplayClusterViewportConfigurationInstanceData& ViewportData : EntireClusterViewports)
 		{
-			if (ViewportIt.Key.Len() && ViewportIt.Value)
+			if (ViewportData.ClusterNodeId == ClusterNodeId)
 			{
-				DesiredViewports.Add(ViewportIt.Key, ViewportIt.Value);
+				CurrentFrameViewports.Add(ViewportData);
 			}
 		}
 	}
 
-	// Collect unused viewports and delete
+	// Create or update viewports from CurrentFrameViewports
+	ImplUpdateViewports();
+
+	// Updates warp policies only if the cluster node name is not empty.
+	if (!ClusterNodeId.IsEmpty())
 	{
-		TArray<FString> ExistClusterNodesIDs;
-		ConfigurationData.Cluster->GetNodeIds(ExistClusterNodesIDs);
+		// when the cluster node name is empty, it means that the DCRA preview right now is initializing for the entire cluster.
+		ImplUpdateViewportsWarpPolicy();
+	}
 
-		TArray<TSharedPtr<FDisplayClusterViewport, ESPMode::ThreadSafe>> UnusedViewports;
-		for (const TSharedPtr<FDisplayClusterViewport, ESPMode::ThreadSafe>& ViewportIt : ViewportManager.ImplGetCurrentRenderFrameViewports())
-		{
-			if (ViewportIt.IsValid())
-			{
-				// ignore ICVFX internal resources
-				if (!EnumHasAllFlags(ViewportIt->GetRenderSettingsICVFX().RuntimeFlags, EDisplayClusterViewportRuntimeICVFXFlags::InternalResource))
-				{
-					// Only viewports from cluster node in render
-					if (ClusterNodeId.IsEmpty() || ViewportIt->GetClusterNodeId() == ClusterNodeId)
-					{
-						if (!DesiredViewports.Contains(ViewportIt->GetId()))
-						{
-							UnusedViewports.Add(ViewportIt);
-						}
-					}
-					else
-					{
-						if (ExistClusterNodesIDs.Find(ViewportIt->GetClusterNodeId()) == INDEX_NONE)
-						{
-							// also remove viewports for deleted cluster nodes
-							UnusedViewports.Add(ViewportIt);
-						}
-					}
-				}
-			}
-		}
+	// Use only valid values of cluster node id
+	RenderFrameSettings.ClusterNodeId = ClusterNodeId;
+}
+
+void FDisplayClusterViewportConfigurationBase::UpdateCustomViewports(const TArray<FString>& InViewportNames)
+{
+	// Initialize variable EntireClusterViewports from configuration
+	ImplInitializeEntireClusterViewportsList();
 
-		// Delete unused viewports
-		for (const TSharedPtr<FDisplayClusterViewport, ESPMode::ThreadSafe>& UnusedViewportIt : UnusedViewports)
+	// Get the list of viewport instance data for the specified InViewportNames
+	for (const FDisplayClusterViewportConfigurationInstanceData& ViewportData : EntireClusterViewports)
+	{
+		if (InViewportNames.Contains(ViewportData.ViewportId))
 		{
-			if (UnusedViewportIt.IsValid())
-			{
-				ViewportManager.ImplDeleteViewport(UnusedViewportIt);
-			}
+			CurrentFrameViewports.Add(ViewportData);
 		}
 	}
 
-	// Update and Create new viewports
-	for (TPair<FString, UDisplayClusterConfigurationViewport*>& CfgIt : DesiredViewports)
+	// Create or update viewports from CurrentFrameViewports
+	ImplUpdateViewports();
+
+	// Do not use the cluster node name for this pass type
+	RenderFrameSettings.ClusterNodeId.Empty();
+}
+
+void FDisplayClusterViewportConfigurationBase::ImplUpdateViewports()
+{
+	// Delete an existing viewport when it is removed from the configuration or its configuration state is set as disabled.
+	const TArray<TSharedPtr<FDisplayClusterViewport, ESPMode::ThreadSafe>> EntireClusterDCViewports = ViewportManager.ImplGetEntireClusterViewports();
+	for (const TSharedPtr<FDisplayClusterViewport, ESPMode::ThreadSafe>& Viewport : EntireClusterDCViewports)
 	{
-		if (const UDisplayClusterConfigurationViewport* ConfigurationViewport = CfgIt.Value)
+		// ignore internal resources
+		if (Viewport.IsValid() && !EnumHasAnyFlags(Viewport->GetRenderSettingsICVFX().RuntimeFlags, EDisplayClusterViewportRuntimeICVFXFlags::InternalResource))
 		{
-			const TSharedPtr<FDisplayClusterViewport, ESPMode::ThreadSafe> ExistViewport = ViewportManager.ImplFindViewport(CfgIt.Key);
-			if (ExistViewport.IsValid())
-			{
-				FDisplayClusterViewportConfigurationBase::UpdateViewportConfiguration(*ExistViewport, ViewportManager, RootActor, *ConfigurationViewport);
-			}
-			else
+			// Delete an existing viewport when it is removed from the configuration
+			if (!ImplFindViewportInEntireCluster(Viewport->GetId()))
 			{
-				ViewportManager.CreateViewport(CfgIt.Key, *ConfigurationViewport);
+				// we can safely remove viewports in a loop, because we use our own local array
+				ViewportManager.ImplDeleteViewport(Viewport);
 			}
 		}
 	}
+
+	// Create or update viewports for current rendering frame
+	for (FDisplayClusterViewportConfigurationInstanceData& ViewportData : CurrentFrameViewports)
+	{
+		ViewportData.CreateOrUpdateViewportInstance(RootActor, ViewportManager, RenderFrameSettings);
+	}
+
+	// Clear viewports warp policies for entire cluster
+	for (const TSharedPtr<FDisplayClusterViewport, ESPMode::ThreadSafe>& Viewport : ViewportManager.ImplGetEntireClusterViewports())
+	{
+		FDisplayClusterViewportConfigurationInstanceData::SetViewportWarpPolicy(Viewport, nullptr);
+	}
 }
 
-void FDisplayClusterViewportConfigurationBase::Update(const TArray<FString>& InViewportNames, FDisplayClusterRenderFrameSettings& InOutRenderFrameSettings)
+void FDisplayClusterViewportConfigurationBase::ImplUpdateViewportsWarpPolicy()
 {
-	// Collect unused viewports and delete
+	// Collect DC origin components used in current frame
+	TArray<UDisplayClusterCameraComponent*> CurrentFrameViewPointCameraComponent;
+	for (const FDisplayClusterViewportConfigurationInstanceData& ViewportData : CurrentFrameViewports)
+	{
+		if (ViewportData.ViewPointCameraComponent)
+		{
+			CurrentFrameViewPointCameraComponent.AddUnique(ViewportData.ViewPointCameraComponent);
+		}
+	}
+
+	// Create reffered viewports from other cluster nodes and update warp policy on it
+	for (UDisplayClusterCameraComponent* ViewPointComponent : CurrentFrameViewPointCameraComponent)
 	{
-		TArray<TSharedPtr<FDisplayClusterViewport, ESPMode::ThreadSafe>> UnusedViewports;
-		for (const TSharedPtr<FDisplayClusterViewport, ESPMode::ThreadSafe>& ViewportIt : ViewportManager.ImplGetCurrentRenderFrameViewports())
+		// The viewpoint camera component expects all viewports (from the entire cluster) that refer to them to be exists (created)  and updated.
+		// this is for cluster rendering
+		if (ViewPointComponent->ShouldUseEntireClusterViewports(&ViewportManager))
 		{
-			if (ViewportIt.IsValid())
+			if (IDisplayClusterWarpPolicy* WarpPolicy = ViewPointComponent->GetWarpPolicy(&ViewportManager))
 			{
-				// ignore ICVFX internal resources
-				if (!EnumHasAllFlags(ViewportIt->GetRenderSettingsICVFX().RuntimeFlags, EDisplayClusterViewportRuntimeICVFXFlags::InternalResource))
+				// Find the viewports that use this ViewPointComponent in the entire cluster.
+				TArray<FDisplayClusterViewportConfigurationInstanceData> ViewPointViewports;
+				const FString ViewPointComponentId = ViewPointComponent->GetName();
+				for (const FDisplayClusterViewportConfigurationInstanceData& ViewportData : EntireClusterViewports)
 				{
-					if (InViewportNames.Find(ViewportIt->GetId()) == INDEX_NONE)
+					if (ViewportData.Configuration.Camera == ViewPointComponentId)
 					{
-						UnusedViewports.Add(ViewportIt);
+						ViewPointViewports.Add(ViewportData);
 					}
 				}
-			}
-		}
 
-		// Delete unused viewports
-		for (const TSharedPtr<FDisplayClusterViewport, ESPMode::ThreadSafe>& DeleteIt : UnusedViewports)
-		{
-			ViewportManager.ImplDeleteViewport(DeleteIt);
+				// Initialize the warp policy for all found viewports:
+				for (FDisplayClusterViewportConfigurationInstanceData& ViewportData : ViewPointViewports)
+				{
+					// If the viewport already exists in the current rendering frame, we simply assign a warp policy
+					if (FDisplayClusterViewportConfigurationInstanceData const* CurrentFrameViewportData = ImplFindCurrentFrameViewports(ViewportData.ViewportId))
+					{
+						// Set the warp projection for the viewport from the current cluster node
+						FDisplayClusterViewportConfigurationInstanceData::SetViewportWarpPolicy(CurrentFrameViewportData->Viewport, WarpPolicy);
+					}
+					else
+					// when the viewport is on another cluster node, we must create it, because the warp policy requires viewports from the whole cluster.
+					{
+						// Create or update this viewport from another cluster node
+						ViewportData.CreateOrUpdateViewportInstance(RootActor, ViewportManager, RenderFrameSettings);
+
+						// Set the warp projection for the viewport from another cluster node
+						FDisplayClusterViewportConfigurationInstanceData::SetViewportWarpPolicy(ViewportData.Viewport, WarpPolicy);
+					}
+				}
+			}
 		}
 	}
+}
 
-	// Update and Create new viewports
-	for (const TPair<FString, TObjectPtr<UDisplayClusterConfigurationClusterNode>>& ClusterNodeConfigurationIt : ConfigurationData.Cluster->Nodes)
+void FDisplayClusterViewportConfigurationBase::ImplInitializeEntireClusterViewportsList()
+{
+	if (ConfigurationData.Cluster)
 	{
-		if (const UDisplayClusterConfigurationClusterNode* ClusterNodeConfiguration = ClusterNodeConfigurationIt.Value)
+		// Update and Create new viewports
+		for (const TPair<FString, TObjectPtr<UDisplayClusterConfigurationClusterNode>>& ClusterNodeConfigurationIt : ConfigurationData.Cluster->Nodes)
 		{
-			// Assign the cluster node name to viewport internals
-			InOutRenderFrameSettings.ClusterNodeId = ClusterNodeConfigurationIt.Key;
-
-			for (const TPair<FString, TObjectPtr<UDisplayClusterConfigurationViewport>>& ViewportIt : ClusterNodeConfiguration->Viewports)
+			if (const UDisplayClusterConfigurationClusterNode* ClusterNodeConfiguration = ClusterNodeConfigurationIt.Value)
 			{
-				if (const UDisplayClusterConfigurationViewport* ConfigurationViewport = ViewportIt.Value)
+				for (const TPair<FString, TObjectPtr<UDisplayClusterConfigurationViewport>>& ViewportIt : ClusterNodeConfiguration->Viewports)
 				{
-					if (ViewportIt.Key.Len() && ViewportIt.Value && InViewportNames.Find(ViewportIt.Key) != INDEX_NONE)
+					if (const UDisplayClusterConfigurationViewport* ConfigurationViewport = ViewportIt.Value)
 					{
-						const TSharedPtr<FDisplayClusterViewport, ESPMode::ThreadSafe> ExistViewport = ViewportManager.ImplFindViewport(ViewportIt.Key);
-						if (ExistViewport.IsValid())
+						const FString& ClusterNodeId = ClusterNodeConfigurationIt.Key;
+						const FString& ViewportId = ViewportIt.Key;
+						if (!ClusterNodeId.IsEmpty() && !ViewportId.IsEmpty() && ConfigurationViewport->IsViewportEnabled())
 						{
-							FDisplayClusterViewportConfigurationBase::UpdateViewportConfiguration(*ExistViewport, ViewportManager, RootActor, *ConfigurationViewport);
-						}
-						else
-						{
-							ViewportManager.CreateViewport(ViewportIt.Key, *ConfigurationViewport);
+							EntireClusterViewports.Add(FDisplayClusterViewportConfigurationInstanceData(ClusterNodeId, ViewportId, *ConfigurationViewport));
 						}
 					}
 				}
 			}
 		}
 	}
+}
 
-	// Do not use the cluster node name for this pass type
-	InOutRenderFrameSettings.ClusterNodeId.Empty();
+FDisplayClusterViewportConfigurationInstanceData const* FDisplayClusterViewportConfigurationBase::ImplFindViewportInEntireCluster(const FString& InViewportId) const
+{
+	return EntireClusterViewports.FindByPredicate([InViewportId](const FDisplayClusterViewportConfigurationInstanceData& ViewportItem)
+		{
+			return ViewportItem.ViewportId == InViewportId;
+		});
+}
+
+FDisplayClusterViewportConfigurationInstanceData const* FDisplayClusterViewportConfigurationBase::ImplFindCurrentFrameViewports(const FString& InViewportId) const
+{
+	return CurrentFrameViewports.FindByPredicate([InViewportId](const FDisplayClusterViewportConfigurationInstanceData& ViewportItem)
+		{
+			return ViewportItem.ViewportId == InViewportId;
+		});
 }
 
 void FDisplayClusterViewportConfigurationBase::AddInternalPostprocess(const FString& InPostprocessName)
@@ -183,18 +256,25 @@ void FDisplayClusterViewportConfigurationBase::AddInternalPostprocess(const FStr
 	}
 }
 
-void FDisplayClusterViewportConfigurationBase::UpdateClusterNodePostProcess(const FString& InClusterNodeId, const FDisplayClusterRenderFrameSettings& InRenderFrameSettings)
+void FDisplayClusterViewportConfigurationBase::UpdateClusterNodePostProcess(const FString& InClusterNodeId)
 {
-	check(!InClusterNodeId.IsEmpty());
+	// The input cluster node name may contain special values. Obtain a valid value.
+	const FString ClusterNodeId = GetValidClusterNodeId(InClusterNodeId);
 
-	const UDisplayClusterConfigurationClusterNode* ClusterNode = ConfigurationData.Cluster->GetNode(InClusterNodeId);
+	if (ClusterNodeId.IsEmpty())
+	{
+		// this function expects a exists cluster node name
+		return;
+	}
+
+	const UDisplayClusterConfigurationClusterNode* ClusterNode = ConfigurationData.Cluster->GetNode(ClusterNodeId);
 	if (ClusterNode)
 	{
 		TSharedPtr<FDisplayClusterViewportPostProcessManager, ESPMode::ThreadSafe> PPManager = ViewportManager.GetPostProcessManager();
 		if (PPManager.IsValid())
 		{
 			// Add TextureShare postprocess:
-			if (ClusterNode->bEnableTextureShare && !InRenderFrameSettings.bIsPreviewRendering)
+			if (ClusterNode->bEnableTextureShare && !RenderFrameSettings.bIsPreviewRendering)
 			{
 				AddInternalPostprocess(TextureShareDisplayClusterStrings::Postprocess::TextureShare);
 			}
@@ -241,7 +321,7 @@ void FDisplayClusterViewportConfigurationBase::UpdateClusterNodePostProcess(cons
 						// Can't create... Disable this postprocess
 						DisabledPostprocessNames.AddUnique(InternalPostprocessId);
 
-						UE_LOG(LogDisplayClusterViewport, Error, TEXT("Can't create postprocess '%s' required by cluster node '%s': Disabled"), *InternalPostprocessId, *InClusterNodeId);
+						UE_LOG(LogDisplayClusterViewport, Error, TEXT("Can't create postprocess '%s' required by cluster node '%s': Disabled"), *InternalPostprocessId, *ClusterNodeId);
 					}
 				}
 			}
@@ -313,13 +393,3 @@ void FDisplayClusterViewportConfigurationBase::UpdateClusterNodePostProcess(cons
 	}
 }
 
-// Assign new configuration to this viewport <Runtime>
-bool FDisplayClusterViewportConfigurationBase::UpdateViewportConfiguration(FDisplayClusterViewport& DstViewport, FDisplayClusterViewportManager& ViewportManager, ADisplayClusterRootActor& RootActor, const UDisplayClusterConfigurationViewport& ConfigurationViewport)
-{
-	check(IsInGameThread());
-
-	FDisplayClusterViewportConfigurationHelpers::UpdateBaseViewportSetting(DstViewport, RootActor, ConfigurationViewport);
-	FDisplayClusterViewportConfigurationHelpers::UpdateProjectionPolicy(DstViewport, &(ConfigurationViewport.ProjectionPolicy));
-
-	return true;
-}
diff --git a/Engine/Plugins/Runtime/nDisplay/Source/DisplayCluster/Private/Render/Viewport/Configuration/DisplayClusterViewportConfigurationBase.h b/Engine/Plugins/Runtime/nDisplay/Source/DisplayCluster/Private/Render/Viewport/Configuration/DisplayClusterViewportConfigurationBase.h
--- a/Engine/Plugins/Runtime/nDisplay/Source/DisplayCluster/Private/Render/Viewport/Configuration/DisplayClusterViewportConfigurationBase.h
+++ b/Engine/Plugins/Runtime/nDisplay/Source/DisplayCluster/Private/Render/Viewport/Configuration/DisplayClusterViewportConfigurationBase.h
@@ -3,33 +3,41 @@
 #pragma once
 
 #include "CoreMinimal.h"
-#include "Render/Viewport/RenderFrame/DisplayClusterRenderFrameSettings.h"
-
-class FDisplayClusterViewportManager;
-class FDisplayClusterViewport;
+#include "Render/Viewport/Configuration/DisplayClusterViewportConfigurationInstanceData.h"
 
 class UDisplayClusterConfigurationData;
-class UDisplayClusterConfigurationViewport;
-
-class ADisplayClusterRootActor;
 
+/**
+ * A helper class that creates/updates/deletes viewports for the current rendering frame.
+ */
 struct FDisplayClusterViewportConfigurationBase
 {
 public:
-	FDisplayClusterViewportConfigurationBase(FDisplayClusterViewportManager& InViewportManager,
-		ADisplayClusterRootActor& InRootActor, const UDisplayClusterConfigurationData& InConfigurationData)
+	FDisplayClusterViewportConfigurationBase(FDisplayClusterViewportManager& InViewportManager, ADisplayClusterRootActor& InRootActor, const UDisplayClusterConfigurationData& InConfigurationData, FDisplayClusterRenderFrameSettings& InRenderFrameSettings)
 		: RootActor(InRootActor)
 		, ViewportManager(InViewportManager)
+		, RenderFrameSettings(InRenderFrameSettings)
 		, ConfigurationData(InConfigurationData)
 	{}
 
 public:
-	void Update(const FString& ClusterNodeId);
-	void Update(const TArray<FString>& InViewportNames, FDisplayClusterRenderFrameSettings& InOutRenderFrameSettings);
-	void UpdateClusterNodePostProcess(const FString& ClusterNodeId, const FDisplayClusterRenderFrameSettings& InRenderFrameSettings);
+	/** Updates the list of viewports for the specified cluster node name, using ConfigurationData.
+	 * 
+	 * @param ClusterNodeId - cluster node name (special names are supported)
+	 */
+	void UpdateClusterNodeViewports(const FString& ClusterNodeId);
 
-public:
-	static bool UpdateViewportConfiguration(FDisplayClusterViewport& DstViewport, FDisplayClusterViewportManager& ViewportManager, ADisplayClusterRootActor& RootActor, const UDisplayClusterConfigurationViewport& ConfigurationViewport);
+	/** Updates only the viewports for the specified list with the names of the viewports using ConfigurationData.
+	 * 
+	 * @param InViewportNames - a list of viewport names to create or update.
+	 */
+	void UpdateCustomViewports(const TArray<FString>& InViewportNames);
+
+	/** Update postprocess for cluster node.
+	 * 
+	 * @param ClusterNodeId - cluster node name (special names are supported)
+	 */
+	void UpdateClusterNodePostProcess(const FString& ClusterNodeId);
 
 protected:
 	/**
@@ -42,9 +50,26 @@ protected:
 	 */
 	void AddInternalPostprocess(const FString& InPostprocessName);
 
+private:
+	/** Update viewports instances in ViewportManager. */
+	void ImplUpdateViewports();
+
+	/** Update all viewports warp policies. */
+	void ImplUpdateViewportsWarpPolicy();
+
+	/** Initialize variable EntireClusterViewports from configuration. */
+	void ImplInitializeEntireClusterViewportsList();
+
+	/** Find the data of the viewport instance by name in the entire cluster. */
+	FDisplayClusterViewportConfigurationInstanceData const* ImplFindViewportInEntireCluster(const FString& InViewportId) const;
+
+	/** Find the viewport instance data by name in the current rendering frame. */
+	FDisplayClusterViewportConfigurationInstanceData const* ImplFindCurrentFrameViewports(const FString& InViewportId) const;
+
 private:
 	ADisplayClusterRootActor& RootActor;
 	FDisplayClusterViewportManager& ViewportManager;
+	FDisplayClusterRenderFrameSettings& RenderFrameSettings;
 	const UDisplayClusterConfigurationData& ConfigurationData;
 
 	/**
@@ -60,4 +85,10 @@ private:
 	 * This post-processing list is updated at runtime. See AddInternalPostprocess()
 	 */
 	TArray<FString> InternalPostprocessNames;
+
+	// The entire cluster viewports
+	TArray<FDisplayClusterViewportConfigurationInstanceData> EntireClusterViewports;
+
+	// the viewports of the current rendering frame (determined from the cluster node name or from the user's list of viewport names)
+	TArray<FDisplayClusterViewportConfigurationInstanceData> CurrentFrameViewports;
 };
diff --git a/Engine/Plugins/Runtime/nDisplay/Source/DisplayCluster/Private/Render/Viewport/Configuration/DisplayClusterViewportConfigurationHelpers.cpp b/Engine/Plugins/Runtime/nDisplay/Source/DisplayCluster/Private/Render/Viewport/Configuration/DisplayClusterViewportConfigurationHelpers.cpp
--- a/Engine/Plugins/Runtime/nDisplay/Source/DisplayCluster/Private/Render/Viewport/Configuration/DisplayClusterViewportConfigurationHelpers.cpp
+++ b/Engine/Plugins/Runtime/nDisplay/Source/DisplayCluster/Private/Render/Viewport/Configuration/DisplayClusterViewportConfigurationHelpers.cpp
@@ -243,23 +243,11 @@ void FDisplayClusterViewportConfigurationHelpers::UpdateBaseViewportSetting(FDis
 
 void FDisplayClusterViewportConfigurationHelpers::UpdateProjectionPolicy(FDisplayClusterViewport& DstViewport, const FDisplayClusterConfigurationProjection* InConfigurationProjectionPolicy)
 {
-	bool bNeedUpdateProjectionPolicy = false;
-
 	// Runtime update projection policy
-	if (InConfigurationProjectionPolicy)
-	{
-		if (DstViewport.ProjectionPolicy.IsValid())
-		{
-			// Current projection policy valid
-			bNeedUpdateProjectionPolicy = DstViewport.ProjectionPolicy->IsConfigurationChanged(InConfigurationProjectionPolicy);
-		}
-		else
-		if (DstViewport.UninitializedProjectionPolicy.IsValid())
-		{
-			// Current projection policy valid
-			bNeedUpdateProjectionPolicy = DstViewport.UninitializedProjectionPolicy->IsConfigurationChanged(InConfigurationProjectionPolicy);
-		}
-	}
+	const bool bNeedUpdateProjectionPolicy = InConfigurationProjectionPolicy && (
+			(DstViewport.ProjectionPolicy.IsValid() && DstViewport.ProjectionPolicy->IsConfigurationChanged(InConfigurationProjectionPolicy))
+		||  (DstViewport.UninitializedProjectionPolicy.IsValid() && DstViewport.UninitializedProjectionPolicy->IsConfigurationChanged(InConfigurationProjectionPolicy))
+		);
 
 	if (bNeedUpdateProjectionPolicy)
 	{
@@ -271,23 +259,12 @@ void FDisplayClusterViewportConfigurationHelpers::UpdateProjectionPolicy(FDispla
 
 		// Create new projection type interface
 		DstViewport.UninitializedProjectionPolicy = FDisplayClusterViewportManager::CreateProjectionPolicy(DstViewport.GetId(), InConfigurationProjectionPolicy);
-		DstViewport.HandleStartScene();
 	}
-	else
-	{
-		if (!DstViewport.ProjectionPolicy.IsValid())
-		{
-			if (DstViewport.IsSceneOpened())
+
+	// If the projection policy is not initialized and the scene is open, initialize it immediately
+	if (!DstViewport.ProjectionPolicy.IsValid() && DstViewport.IsSceneOpened())
 			{
 				// Try initialize proj policy every tick (mesh deferred load, etc)
 				DstViewport.HandleStartScene();
 			}
 		}
-	}
-
-	// Override PostProcess from projection policy
-	if (DstViewport.ProjectionPolicy.IsValid())
-	{
-		DstViewport.ProjectionPolicy->OverridePostProcessSettings(&DstViewport);
-	}
-}
diff --git a/Engine/Plugins/Runtime/nDisplay/Source/DisplayCluster/Private/Render/Viewport/Configuration/DisplayClusterViewportConfigurationHelpers_ICVFX.cpp b/Engine/Plugins/Runtime/nDisplay/Source/DisplayCluster/Private/Render/Viewport/Configuration/DisplayClusterViewportConfigurationHelpers_ICVFX.cpp
--- a/Engine/Plugins/Runtime/nDisplay/Source/DisplayCluster/Private/Render/Viewport/Configuration/DisplayClusterViewportConfigurationHelpers_ICVFX.cpp
+++ b/Engine/Plugins/Runtime/nDisplay/Source/DisplayCluster/Private/Render/Viewport/Configuration/DisplayClusterViewportConfigurationHelpers_ICVFX.cpp
@@ -473,17 +473,44 @@ bool FDisplayClusterViewportConfigurationHelpers_ICVFX::GetCameraContext(ADispla
 	}
 
 	// Initialize camera policy with camera component and settings
-	if (!ImplUpdateCameraProjectionSettings(CameraProjectionPolicy, RootActor, CameraSettings, InCameraComponent.GetCameraComponent()))
+	if (!ImplUpdateCameraProjectionSettings(CameraProjectionPolicy, RootActor, CameraSettings, &InCameraComponent))
 	{
 		return false;
 	}
 
 	// Get camera pos-rot-prj from policy
-	const float WorldToMeters = 100.f;
-	const float CfgNCP = 1.0f;
+
+	// Use default clipping planes
+	const float NCP = GNearClippingPlane;
+	const float FCP = NCP; // nDisplay does not use the far plane of the clipping
+
+	// Applying the correct sequence of steps to use the projection policy math:
+	// SetupProjectionViewPoint()->CalculateView()->GetProjectionMatrix()
+	FMinimalViewInfo CameraViewInfo;
+	float CustomNearClippingPlane = -1; // a value less than zero means ignoring.
+	CameraProjectionPolicy->SetupProjectionViewPoint(nullptr, RootActor.GetWorldDeltaSeconds(), CameraViewInfo, &CustomNearClippingPlane);
+
+	OutCameraContext.ViewLocation = CameraViewInfo.Location;
+	OutCameraContext.ViewRotation = CameraViewInfo.Rotation;
+
+	// Todo: Here we need to calculate the correct ViewOffset so that ICVFX can support stereo rendering.
 	const FVector ViewOffset = FVector::ZeroVector;
 
-	if (CameraProjectionPolicy->CalculateView(nullptr, 0, OutCameraContext.ViewLocation, OutCameraContext.ViewRotation, ViewOffset, WorldToMeters, CfgNCP, CfgNCP) &&
+	// Todo: maybe it makes sense to get the scaling value of WorldToMeters the right way?
+	// By default, this value is taken from the PlayerController in the LocalPlayer file.
+	// But ICVFX can use another data source for this value.
+	float WorldToMeters = 100.f;
+
+	// Supports custom near clipping plane
+	float ZNear = NCP;
+	float ZFar = FCP;
+	if (CustomNearClippingPlane >= 0)
+	{
+		ZNear = CustomNearClippingPlane;
+		ZFar = (NCP == FCP) ? ZNear : ZFar;
+	}
+
+	if (CameraProjectionPolicy->CalculateView(nullptr, 0, OutCameraContext.ViewLocation, OutCameraContext.ViewRotation, ViewOffset, WorldToMeters, ZNear, ZFar) &&
 		CameraProjectionPolicy->GetProjectionMatrix(nullptr, 0, OutCameraContext.PrjMatrix))
 	{
 		return true;
@@ -510,7 +537,7 @@ void FDisplayClusterViewportConfigurationHelpers_ICVFX::UpdateCameraViewportSett
 	}
 
 	// Update camera viewport projection policy settings
-	ImplUpdateCameraProjectionSettings(DstViewport.ProjectionPolicy, RootActor, CameraSettings, InCameraComponent.GetCameraComponent());
+	ImplUpdateCameraProjectionSettings(DstViewport.ProjectionPolicy, RootActor, CameraSettings, &InCameraComponent);
 
 	// Update OCIO for Camera Viewport
 	FDisplayClusterViewportConfigurationHelpers_OpenColorIO::UpdateCameraViewport(DstViewport, RootActor, InCameraComponent);
diff --git a/Engine/Plugins/Runtime/nDisplay/Source/DisplayCluster/Private/Render/Viewport/Configuration/DisplayClusterViewportConfigurationHelpers_Postprocess.cpp b/Engine/Plugins/Runtime/nDisplay/Source/DisplayCluster/Private/Render/Viewport/Configuration/DisplayClusterViewportConfigurationHelpers_Postprocess.cpp
--- a/Engine/Plugins/Runtime/nDisplay/Source/DisplayCluster/Private/Render/Viewport/Configuration/DisplayClusterViewportConfigurationHelpers_Postprocess.cpp
+++ b/Engine/Plugins/Runtime/nDisplay/Source/DisplayCluster/Private/Render/Viewport/Configuration/DisplayClusterViewportConfigurationHelpers_Postprocess.cpp
@@ -33,7 +33,7 @@ static bool ImplUpdatePerViewportColorGrading(FDisplayClusterViewport& DstViewpo
 	const FDisplayClusterConfigurationICVFX_StageSettings& StageSettings = RootActor.GetStageSettings();
 
 	// enable entire cluster only when global settings is on
-	bool bUseEntireClusterPostProcess = StageSettings.EntireClusterColorGrading.bEnableEntireClusterColorGrading & PerViewportColorGrading.bIsEntireClusterEnabled;
+	const bool bUseEntireClusterPostProcess = StageSettings.EntireClusterColorGrading.bEnableEntireClusterColorGrading && PerViewportColorGrading.bIsEntireClusterEnabled;
 	
 	FDisplayClusterConfigurationViewport_CustomPostprocessSettings FinalPerViewportColorGrading;
 	FinalPerViewportColorGrading.bIsEnabled = true;
@@ -90,8 +90,8 @@ static bool ImplUpdateIncameraPerNodeColorGrading(FDisplayClusterViewport& DstVi
 	FinalPerNodeColorGrading.bIsOneFrame = true;
 	FinalPerNodeColorGrading.BlendWeight = 1;
 
-	const bool bIncludeUseEntireClusterPostProcess = StageSettings.EntireClusterColorGrading.bEnableEntireClusterColorGrading & PerNodeColorGrading.bEntireClusterColorGrading;
-	const bool bIncludeAllNodesColorGrading = AllNodesColorGrading.bEnableInnerFrustumAllNodesColorGrading & PerNodeColorGrading.bAllNodesColorGrading;
+	const bool bIncludeUseEntireClusterPostProcess = StageSettings.EntireClusterColorGrading.bEnableEntireClusterColorGrading && PerNodeColorGrading.bEntireClusterColorGrading;
+	const bool bIncludeAllNodesColorGrading = AllNodesColorGrading.bEnableInnerFrustumAllNodesColorGrading && PerNodeColorGrading.bAllNodesColorGrading;
 
 	if (bIncludeUseEntireClusterPostProcess)
 	{
@@ -147,7 +147,7 @@ static bool ImplUpdateIncameraAllNodesColorGrading(FDisplayClusterViewport& DstV
 	const FDisplayClusterConfigurationICVFX_StageSettings& StageSettings = RootActor.GetStageSettings();
 
 	// enable entire cluster only when global settings is on
-	const bool bEnableEntireClusterColorGrading = StageSettings.EntireClusterColorGrading.bEnableEntireClusterColorGrading & AllNodesColorGrading.bEnableEntireClusterColorGrading;
+	const bool bEnableEntireClusterColorGrading = StageSettings.EntireClusterColorGrading.bEnableEntireClusterColorGrading && AllNodesColorGrading.bEnableEntireClusterColorGrading;
 
 	if (bEnableEntireClusterColorGrading)
 	{
@@ -286,41 +286,27 @@ bool FDisplayClusterViewportConfigurationHelpers_Postprocess::ImplUpdateViewport
 
 void FDisplayClusterViewportConfigurationHelpers_Postprocess::UpdateCameraPostProcessSettings(FDisplayClusterViewport& DstViewport, ADisplayClusterRootActor& RootActor, UDisplayClusterICVFXCameraComponent& InCameraComponent)
 {
-	FDisplayClusterConfigurationViewport_CustomPostprocessSettings CameraPPS;
-	CameraPPS.bIsOneFrame = true;
-	CameraPPS.BlendWeight = 1.f;
-
 	const FDisplayClusterConfigurationICVFX_CameraSettings& CameraSettings = InCameraComponent.GetCameraSettingsICVFX();
-	if (CameraSettings.RenderSettings.bUseCameraComponentPostprocess)
-	{
-		FMinimalViewInfo DesiredView;
-		InCameraComponent.GetDesiredView(DesiredView);
-
-		// Send camera postprocess to override
-		CameraPPS.bIsEnabled = true;
-		CameraPPS.PostProcessSettings = DesiredView.PostProcessSettings;
-	}
-
-	const FDisplayClusterConfigurationICVFX_CameraMotionBlurOverridePPS& OverrideMotionBlurPPS = CameraSettings.CameraMotionBlur.MotionBlurPPS;
-	if (OverrideMotionBlurPPS.bReplaceEnable)
-	{
-		// Send camera postprocess to override
-		CameraPPS.bIsEnabled = true;
-
-		CameraPPS.PostProcessSettings.MotionBlurAmount = OverrideMotionBlurPPS.MotionBlurAmount;
-		CameraPPS.PostProcessSettings.bOverride_MotionBlurAmount = true;
-
-		CameraPPS.PostProcessSettings.MotionBlurMax = OverrideMotionBlurPPS.MotionBlurMax;
-		CameraPPS.PostProcessSettings.bOverride_MotionBlurMax = true;
-		
-		CameraPPS.PostProcessSettings.MotionBlurPerObjectSize = OverrideMotionBlurPPS.MotionBlurPerObjectSize;
-		CameraPPS.PostProcessSettings.bOverride_MotionBlurPerObjectSize = true;
-	}
-
 	const FDisplayClusterConfigurationICVFX_StageSettings& StageSettings = RootActor.GetStageSettings();
 	// check if frustum color grading is enabled	
 	if (StageSettings.EnableColorGrading && CameraSettings.EnableInnerFrustumColorGrading)
 	{
+		FDisplayClusterConfigurationViewport_CustomPostprocessSettings CameraPPS;
+		CameraPPS.bIsOneFrame = true;
+		CameraPPS.BlendWeight = 1.f;
+
+		const bool bUseCameraPostprocess = true; // use internal rules of UDisplayClusterICVFXCameraComponent
+
+		// All logic was moved to the UDisplayClusterICVFXCameraComponent::GetCameraView() virtual function.
+		FMinimalViewInfo DesiredView;
+		IDisplayClusterViewport::GetCameraComponentView(&InCameraComponent, RootActor.GetWorldDeltaSeconds(), bUseCameraPostprocess, DesiredView);
+		if (DesiredView.PostProcessBlendWeight > 0)
+		{
+			// Send camera postprocess to override
+			CameraPPS.bIsEnabled = true;
+			CameraPPS.PostProcessSettings = DesiredView.PostProcessSettings;
+		}
+
 		ImplUpdateCustomPostprocess(DstViewport, CameraPPS.bIsEnabled, CameraPPS, IDisplayClusterViewport_CustomPostProcessSettings::ERenderPass::Override);
 
 		if (!ImplUpdateInnerFrustumColorGrading(DstViewport, RootActor, InCameraComponent))
diff --git a/Engine/Plugins/Runtime/nDisplay/Source/DisplayCluster/Private/Render/Viewport/Configuration/DisplayClusterViewportConfigurationICVFX.cpp b/Engine/Plugins/Runtime/nDisplay/Source/DisplayCluster/Private/Render/Viewport/Configuration/DisplayClusterViewportConfigurationICVFX.cpp
--- a/Engine/Plugins/Runtime/nDisplay/Source/DisplayCluster/Private/Render/Viewport/Configuration/DisplayClusterViewportConfigurationICVFX.cpp
+++ b/Engine/Plugins/Runtime/nDisplay/Source/DisplayCluster/Private/Render/Viewport/Configuration/DisplayClusterViewportConfigurationICVFX.cpp
@@ -593,7 +593,7 @@ EDisplayClusterViewportICVFXFlags FDisplayClusterViewportConfigurationICVFX::Imp
 			//Raise new projection target if possible
 			if (ViewportIt->RenderSettings.bEnable && EnumHasAnyFlags(ViewportIt->RenderSettingsICVFX.Flags, EDisplayClusterViewportICVFXFlags::Enable))
 			{
-				if (ViewportIt->ProjectionPolicy.IsValid() && ViewportIt->ProjectionPolicy->ShouldSupportICVFX())
+				if (ViewportIt->ProjectionPolicy.IsValid() && ViewportIt->ProjectionPolicy->ShouldSupportICVFX(ViewportIt.Get()))
 				{
 					// Collect this viewport ICVFX target
 					OutTargets.Add(ViewportIt);
diff --git a/Engine/Plugins/Runtime/nDisplay/Source/DisplayCluster/Private/Render/Viewport/Configuration/DisplayClusterViewportConfigurationInstanceData.cpp b/Engine/Plugins/Runtime/nDisplay/Source/DisplayCluster/Private/Render/Viewport/Configuration/DisplayClusterViewportConfigurationInstanceData.cpp
new file mode 100644
--- /dev/null
+++ b/Engine/Plugins/Runtime/nDisplay/Source/DisplayCluster/Private/Render/Viewport/Configuration/DisplayClusterViewportConfigurationInstanceData.cpp
@@ -0,0 +1,70 @@
+// Copyright Epic Games, Inc. All Rights Reserved.
+
+#include "DisplayClusterViewportConfigurationInstanceData.h"
+
+#include "Render/Projection/IDisplayClusterProjectionPolicy.h"
+
+#include "Render/Viewport/DisplayClusterViewport.h"
+#include "Render/Viewport/DisplayClusterViewportManager.h"
+#include "Render/Viewport/RenderFrame/DisplayClusterRenderFrameSettings.h"
+
+#include "DisplayClusterRootActor.h"
+
+#include "Render/Viewport/Configuration/DisplayClusterViewportConfigurationHelpers.h"
+
+void FDisplayClusterViewportConfigurationInstanceData::CreateOrUpdateViewportInstance(ADisplayClusterRootActor& RootActor, FDisplayClusterViewportManager& ViewportManager, FDisplayClusterRenderFrameSettings& RenderFrameSettings)
+{
+	// Store current cluster node
+	const FString CurrentClusterNodeID = RenderFrameSettings.ClusterNodeId;
+
+	// This function can create viewports from other nodes in the cluster
+	RenderFrameSettings.ClusterNodeId = ClusterNodeId;
+
+	const TSharedPtr<FDisplayClusterViewport, ESPMode::ThreadSafe> ExistViewport = ViewportManager.ImplFindViewport(ViewportId);
+	if (ExistViewport.IsValid())
+	{
+		UpdateViewportConfiguration(*ExistViewport, ViewportManager, RootActor, Configuration);
+
+		// Store viewpoint of exist viewport
+		ViewPointCameraComponent = ExistViewport->GetViewPointCameraComponent();
+
+		// Store Viewport
+		Viewport = ExistViewport;
+	}
+	else
+	{
+		if (FDisplayClusterViewport* NewViewport = ViewportManager.CreateViewport(ViewportId, Configuration))
+		{
+			// Store viewpoint of new viewport
+			ViewPointCameraComponent = NewViewport->GetViewPointCameraComponent();
+
+			// Store Viewport
+			Viewport = NewViewport->AsShared();
+		}
+	}
+
+	// Restore current cluster node id
+	RenderFrameSettings.ClusterNodeId = CurrentClusterNodeID;
+}
+
+bool FDisplayClusterViewportConfigurationInstanceData::UpdateViewportConfiguration(FDisplayClusterViewport& DstViewport, FDisplayClusterViewportManager& ViewportManager, ADisplayClusterRootActor& RootActor, const UDisplayClusterConfigurationViewport& ConfigurationViewport)
+{
+	check(IsInGameThread());
+
+	FDisplayClusterViewportConfigurationHelpers::UpdateBaseViewportSetting(DstViewport, RootActor, ConfigurationViewport);
+	FDisplayClusterViewportConfigurationHelpers::UpdateProjectionPolicy(DstViewport, &(ConfigurationViewport.ProjectionPolicy));
+
+	return true;
+}
+
+void FDisplayClusterViewportConfigurationInstanceData::SetViewportWarpPolicy(const TSharedPtr<FDisplayClusterViewport, ESPMode::ThreadSafe>& InViewport, IDisplayClusterWarpPolicy* InWarpPolicy)
+{
+	// ignore internal resources
+	if (InViewport.IsValid() && !EnumHasAnyFlags(InViewport->GetRenderSettingsICVFX().RuntimeFlags, EDisplayClusterViewportRuntimeICVFXFlags::InternalResource))
+	{
+		if (InViewport->GetProjectionPolicy().IsValid())
+		{
+			InViewport->GetProjectionPolicy()->SetWarpPolicy(InWarpPolicy);
+		}
+	}
+}
diff --git a/Engine/Plugins/Runtime/nDisplay/Source/DisplayCluster/Private/Render/Viewport/Configuration/DisplayClusterViewportConfigurationInstanceData.h b/Engine/Plugins/Runtime/nDisplay/Source/DisplayCluster/Private/Render/Viewport/Configuration/DisplayClusterViewportConfigurationInstanceData.h
new file mode 100644
--- /dev/null
+++ b/Engine/Plugins/Runtime/nDisplay/Source/DisplayCluster/Private/Render/Viewport/Configuration/DisplayClusterViewportConfigurationInstanceData.h
@@ -0,0 +1,58 @@
+// Copyright Epic Games, Inc. All Rights Reserved.
+
+#pragma once
+
+#include "CoreMinimal.h"
+
+class UDisplayClusterConfigurationViewport;
+class UDisplayClusterCameraComponent;
+
+class ADisplayClusterRootActor;
+class FDisplayClusterViewportManager;
+class FDisplayClusterViewport;
+class IDisplayClusterWarpPolicy;
+
+struct FDisplayClusterRenderFrameSettings;
+
+/**
+ * Container with information related to viewport: configuration and references to some class instances.
+ */
+struct FDisplayClusterViewportConfigurationInstanceData
+{
+	FDisplayClusterViewportConfigurationInstanceData(const FString& InClusterNodeId, const FString& InViewportId, const UDisplayClusterConfigurationViewport& InConfiguration)
+		: ClusterNodeId(InClusterNodeId)
+		, ViewportId(InViewportId)
+		, Configuration(InConfiguration)
+	{ }
+
+	/**
+	 * Create or update DC viewport instance from configuration
+	 */
+	void CreateOrUpdateViewportInstance(ADisplayClusterRootActor& RootActor, FDisplayClusterViewportManager& ViewportManager, FDisplayClusterRenderFrameSettings& RenderFrameSettings);
+
+	/**
+	 * Update exists viewport runtime settings from configuration
+	 */
+	static bool UpdateViewportConfiguration(FDisplayClusterViewport& DstViewport, FDisplayClusterViewportManager& ViewportManager, ADisplayClusterRootActor& RootActor, const UDisplayClusterConfigurationViewport& InConfigurationViewport);
+
+	/**
+	 * Assign warp policy instance to given viewport
+	 */
+	static void SetViewportWarpPolicy(const TSharedPtr<FDisplayClusterViewport, ESPMode::ThreadSafe>& InViewport, IDisplayClusterWarpPolicy* InWarpPolicy);
+
+public:
+	// Cluster name, who own this viewport
+	const FString& ClusterNodeId;
+
+	// Viewport unique name
+	const FString& ViewportId;
+
+	// Reference to the configuration used to create or update the viewport instance
+	const UDisplayClusterConfigurationViewport& Configuration;
+
+	// The viewport instance
+	TSharedPtr<FDisplayClusterViewport, ESPMode::ThreadSafe> Viewport;
+
+	// Pointer to the view origin component, used by this viewport instance.
+	UDisplayClusterCameraComponent* ViewPointCameraComponent = nullptr;
+};
diff --git a/Engine/Plugins/Runtime/nDisplay/Source/DisplayCluster/Private/Render/Viewport/Configuration/DisplayClusterViewportConfigurationProjectionPolicy.cpp b/Engine/Plugins/Runtime/nDisplay/Source/DisplayCluster/Private/Render/Viewport/Configuration/DisplayClusterViewportConfigurationProjectionPolicy.cpp
--- a/Engine/Plugins/Runtime/nDisplay/Source/DisplayCluster/Private/Render/Viewport/Configuration/DisplayClusterViewportConfigurationProjectionPolicy.cpp
+++ b/Engine/Plugins/Runtime/nDisplay/Source/DisplayCluster/Private/Render/Viewport/Configuration/DisplayClusterViewportConfigurationProjectionPolicy.cpp
@@ -34,9 +34,16 @@ void FDisplayClusterViewportConfigurationProjectionPolicy::Update()
 		if (ViewportIt.IsValid() && !EnumHasAnyFlags(ViewportIt->GetRenderSettingsICVFX().RuntimeFlags, EDisplayClusterViewportRuntimeICVFXFlags::InternalResource))
 		{
 			// Support advanced logic for 'camera' projection policy
-			if (ViewportIt->ProjectionPolicy.IsValid() && ViewportIt->ProjectionPolicy->GetType().Compare(DisplayClusterProjectionStrings::projection::Camera) == 0)
+			if (ViewportIt->ProjectionPolicy.IsValid())
 			{
-				UpdateCameraPolicy(*ViewportIt);
+				if (ViewportIt->ProjectionPolicy->GetType().Compare(DisplayClusterProjectionStrings::projection::Camera) == 0)
+				{
+					UpdateCameraPolicy(*ViewportIt);
+				}
+
+				// Projection policies can override postprocess settings
+				// The camera PP override code has been moved from FDisplayClusterViewportConfigurationProjectionPolicy::UpdateCameraPolicy_Base() to the projection policy new function:
+				ViewportIt->ProjectionPolicy->UpdatePostProcessSettings(ViewportIt.Get());
 			}
 		}
 	}
@@ -94,11 +101,7 @@ bool FDisplayClusterViewportConfigurationProjectionPolicy::UpdateCameraPolicy_Ba
 	check(InCameraComponent);
 
 	// add camera's post processing materials
-	FMinimalViewInfo DesiredView;
-	const float DeltaTime = RootActor.GetWorldDeltaSeconds();
-	InCameraComponent->GetCameraView(DeltaTime, DesiredView);
-
-	DstViewport.CustomPostProcessSettings.AddCustomPostProcess(IDisplayClusterViewport_CustomPostProcessSettings::ERenderPass::Override, DesiredView.PostProcessSettings, 1.0, true);
+	// Moved to IDisplayClusterProjectionPolicy::UpdatePostProcessSettings()
 
 	FDisplayClusterProjectionCameraPolicySettings PolicyCameraSettings;
 	PolicyCameraSettings.FOVMultiplier = FOVMultiplier;
diff --git a/Engine/Plugins/Runtime/nDisplay/Source/DisplayCluster/Private/Render/Viewport/DisplayClusterViewport.h b/Engine/Plugins/Runtime/nDisplay/Source/DisplayCluster/Private/Render/Viewport/DisplayClusterViewport.h
--- a/Engine/Plugins/Runtime/nDisplay/Source/DisplayCluster/Private/Render/Viewport/DisplayClusterViewport.h
+++ b/Engine/Plugins/Runtime/nDisplay/Source/DisplayCluster/Private/Render/Viewport/DisplayClusterViewport.h
@@ -99,14 +99,17 @@ public:
 		Contexts.Append(InContexts);
 	}
 
+	virtual FVector2D GetClippingPlanes() const override;
+
 	virtual void CalculateProjectionMatrix(const uint32 InContextNum, float Left, float Right, float Top, float Bottom, float ZNear, float ZFar, bool bIsAnglesInput) override;
 
-	virtual bool    CalculateView(const uint32 InContextNum, FVector& InOutViewLocation, FRotator& InOutViewRotation, const FVector& ViewOffset, const float WorldToMeters, const float NCP, const float FCP) override;
+	virtual bool    CalculateView(const uint32 InContextNum, FVector& InOutViewLocation, FRotator& InOutViewRotation, const float WorldToMeters) override;
 	virtual bool    GetProjectionMatrix(const uint32 InContextNum, FMatrix& OutPrjMatrix)  override;
 
 	virtual bool SetupViewPoint(FMinimalViewInfo& InOutViewInfo) override;
 	virtual float GetStereoEyeOffsetDistance(const uint32 InContextNum) override;
 	virtual class UDisplayClusterCameraComponent* GetViewPointCameraComponent() const override;
+	virtual bool GetViewPointCameraEye(const uint32 InContextNum, FVector& OutViewLocation, FRotator& OutViewRotation, FVector& OutViewOffset) override;
 
 	virtual const FDisplayClusterViewport_RenderSettingsICVFX& GetRenderSettingsICVFX() const override
 	{
@@ -339,4 +342,9 @@ protected:
 
 private:
 	bool bProjectionPolicyCalculateViewWarningOnce = false;
+
+	// Near clipping plane value (obtained from the GetDesiredView() functions).
+	// If the value is less than zero, it does not apply to this viewport.
+	// This value is changed in the SetupViewPoint() function, called at the beginning from LocalPlayer.
+	float CustomNearClippingPlane = -1;
 };
diff --git a/Engine/Plugins/Runtime/nDisplay/Source/DisplayCluster/Private/Render/Viewport/DisplayClusterViewportManager.cpp b/Engine/Plugins/Runtime/nDisplay/Source/DisplayCluster/Private/Render/Viewport/DisplayClusterViewportManager.cpp
--- a/Engine/Plugins/Runtime/nDisplay/Source/DisplayCluster/Private/Render/Viewport/DisplayClusterViewportManager.cpp
+++ b/Engine/Plugins/Runtime/nDisplay/Source/DisplayCluster/Private/Render/Viewport/DisplayClusterViewportManager.cpp
@@ -21,7 +21,7 @@
 #include "Render/Viewport/Postprocess/DisplayClusterViewportPostProcessManager.h"
 #include "Render/Viewport/Postprocess/DisplayClusterViewportPostProcessOutputRemap.h"
 #include "Render/Viewport/Configuration/DisplayClusterViewportConfiguration.h"
-#include "Render/Viewport/Configuration/DisplayClusterViewportConfigurationBase.h"
+#include "Render/Viewport/Configuration/DisplayClusterViewportConfigurationInstanceData.h"
 #include "Render/Viewport/DisplayClusterViewportStrings.h"
 
 #include "WarpBlend/IDisplayClusterWarpBlend.h"
@@ -956,7 +956,7 @@ void FDisplayClusterViewportManager::RenderFrame(FViewport* InViewport)
 	ViewportManagerProxy->ImplRenderFrame_GameThread(InViewport);
 }
 
-bool FDisplayClusterViewportManager::CreateViewport(const FString& InViewportId, const class UDisplayClusterConfigurationViewport& ConfigurationViewport)
+FDisplayClusterViewport* FDisplayClusterViewportManager::CreateViewport(const FString& InViewportId, const class UDisplayClusterConfigurationViewport& ConfigurationViewport)
 {
 	check(IsInGameThread());
 
@@ -964,7 +964,7 @@ bool FDisplayClusterViewportManager::CreateViewport(const FString& InViewportId,
 	if (InViewportId.IsEmpty())
 	{
 		UE_LOG(LogDisplayClusterViewport, Warning, TEXT("Wrong viewport ID"));
-		return false;
+		return nullptr;
 	}
 
 	// ID must be unique
@@ -972,7 +972,7 @@ bool FDisplayClusterViewportManager::CreateViewport(const FString& InViewportId,
 	{
 		UE_LOG(LogDisplayClusterViewport, Warning, TEXT("Viewport '%s' already exists"), *InViewportId);
 
-		return false;
+		return nullptr;
 	}
 
 	ADisplayClusterRootActor* RootActorPtr = GetRootActor();
@@ -980,7 +980,7 @@ bool FDisplayClusterViewportManager::CreateViewport(const FString& InViewportId,
 	{
 		UE_LOG(LogDisplayClusterViewport, Error, TEXT("Can't create viewport '%s': DCRA not defined"), *InViewportId);
 
-		return false;
+		return nullptr;
 	}
 
 	// Create projection policy for viewport
@@ -991,15 +991,15 @@ bool FDisplayClusterViewportManager::CreateViewport(const FString& InViewportId,
 		TSharedPtr<FDisplayClusterViewport, ESPMode::ThreadSafe> NewViewport = ImplCreateViewport(InViewportId, NewProjectionPolicy);
 		if (NewViewport.IsValid())
 		{
-			FDisplayClusterViewportConfigurationBase::UpdateViewportConfiguration(*NewViewport , *this, *RootActorPtr, ConfigurationViewport);
+			FDisplayClusterViewportConfigurationInstanceData::UpdateViewportConfiguration(*NewViewport , *this, *RootActorPtr, ConfigurationViewport);
 
-			return true;
+			return NewViewport.Get();
 		}
 	}
 
 	UE_LOG(LogDisplayClusterViewport, Error, TEXT("Viewports '%s' not created."), *InViewportId);
 
-	return false;
+	return nullptr;
 }
 
 IDisplayClusterViewport* FDisplayClusterViewportManager::FindViewport(const FString& InViewportId) const
@@ -1021,7 +1021,7 @@ TSharedPtr<FDisplayClusterViewport, ESPMode::ThreadSafe> FDisplayClusterViewport
 	return (DesiredViewport && DesiredViewport->IsValid()) ? *DesiredViewport : nullptr;
 }
 
-IDisplayClusterViewport* FDisplayClusterViewportManager::CreateViewport(const FString& ViewportId, const TSharedPtr<IDisplayClusterProjectionPolicy, ESPMode::ThreadSafe>& InProjectionPolicy)
+FDisplayClusterViewport* FDisplayClusterViewportManager::CreateViewport(const FString& ViewportId, const TSharedPtr<IDisplayClusterProjectionPolicy, ESPMode::ThreadSafe>& InProjectionPolicy)
 {
 	check(IsInGameThread());
 
diff --git a/Engine/Plugins/Runtime/nDisplay/Source/DisplayCluster/Private/Render/Viewport/DisplayClusterViewportManager.h b/Engine/Plugins/Runtime/nDisplay/Source/DisplayCluster/Private/Render/Viewport/DisplayClusterViewportManager.h
--- a/Engine/Plugins/Runtime/nDisplay/Source/DisplayCluster/Private/Render/Viewport/DisplayClusterViewportManager.h
+++ b/Engine/Plugins/Runtime/nDisplay/Source/DisplayCluster/Private/Render/Viewport/DisplayClusterViewportManager.h
@@ -97,11 +97,16 @@ public:
 	virtual IDisplayClusterViewport* FindViewport(const FString& InViewportId) const override;
 	virtual IDisplayClusterViewport* FindViewport(const int32 ViewIndex, uint32* OutContextNum = nullptr) const override;
 
-	virtual const TArrayView<TSharedPtr<IDisplayClusterViewport, ESPMode::ThreadSafe>> GetViewports() const override
+	virtual const TArrayView<TSharedPtr<IDisplayClusterViewport, ESPMode::ThreadSafe>> GetCurrentRenderFrameViewports() const override
 	{
 		return TArrayView<TSharedPtr<IDisplayClusterViewport, ESPMode::ThreadSafe>>((TSharedPtr<IDisplayClusterViewport, ESPMode::ThreadSafe>*)(CurrentRenderFrameViewports.GetData()), CurrentRenderFrameViewports.Num());
 	}
 
+	virtual const TArrayView<TSharedPtr<IDisplayClusterViewport, ESPMode::ThreadSafe>> GetEntireClusterViewports() const override
+	{
+		return TArrayView<TSharedPtr<IDisplayClusterViewport, ESPMode::ThreadSafe>>((TSharedPtr<IDisplayClusterViewport, ESPMode::ThreadSafe>*)(EntireClusterViewports.GetData()), EntireClusterViewports.Num());
+	}
+
 	virtual void MarkComponentGeometryDirty(const FName InComponentName = NAME_None) override;
 	virtual void AddReferencedObjects(FReferenceCollector& Collector) override;
 
@@ -121,8 +126,8 @@ public:
 	}
 
 	// internal use only
-	bool CreateViewport(const FString& ViewportId, const class UDisplayClusterConfigurationViewport& ConfigurationViewport);
-	IDisplayClusterViewport* CreateViewport(const FString& ViewportId, const TSharedPtr<IDisplayClusterProjectionPolicy, ESPMode::ThreadSafe>& InProjectionPolicy);
+	FDisplayClusterViewport* CreateViewport(const FString& ViewportId, const class UDisplayClusterConfigurationViewport& ConfigurationViewport);
+	FDisplayClusterViewport* CreateViewport(const FString& ViewportId, const TSharedPtr<IDisplayClusterProjectionPolicy, ESPMode::ThreadSafe>& InProjectionPolicy);
 	bool                     DeleteViewport(const FString& ViewportId);
 
 	TSharedPtr<FDisplayClusterViewport, ESPMode::ThreadSafe> ImplCreateViewport(const FString& ViewportId, const TSharedPtr<IDisplayClusterProjectionPolicy, ESPMode::ThreadSafe>& InProjectionPolicy);
diff --git a/Engine/Plugins/Runtime/nDisplay/Source/DisplayCluster/Private/Render/Viewport/DisplayClusterViewport_EditorPreview.cpp b/Engine/Plugins/Runtime/nDisplay/Source/DisplayCluster/Private/Render/Viewport/DisplayClusterViewport_EditorPreview.cpp
--- a/Engine/Plugins/Runtime/nDisplay/Source/DisplayCluster/Private/Render/Viewport/DisplayClusterViewport_EditorPreview.cpp
+++ b/Engine/Plugins/Runtime/nDisplay/Source/DisplayCluster/Private/Render/Viewport/DisplayClusterViewport_EditorPreview.cpp
@@ -222,18 +222,9 @@ bool FDisplayClusterViewport::ImplPreview_CalculateStereoViewOffset(const uint32
 	ViewLocation = ViewInfo.Location;
 	ViewRotation = ViewInfo.Rotation;
 
-	const float PassOffsetSwap = GetStereoEyeOffsetDistance(InContextNum);
-
-	FVector ViewOffset = FVector::ZeroVector;
-
-	// Apply computed offset to the view location
-	const FQuat EyeQuat = ViewRotation.Quaternion();
-	ViewOffset = EyeQuat.RotateVector(FVector(0.0f, PassOffsetSwap, 0.0f));
-	ViewLocation += ViewOffset;
-
+	// Obtaining the offset of the stereo eye and the values of the projection clipping plane for the given viewport was moved inside CalculateView().
 	// Perform view calculations on a policy side
-	const float CfgNCP = GNearClippingPlane;
-	if (!CalculateView(InContextNum, ViewLocation, ViewRotation, ViewOffset, WorldToMeters, CfgNCP, CfgNCP))
+	if (!CalculateView(InContextNum, ViewLocation, ViewRotation, WorldToMeters))
 	{
 		if (!bProjectionPolicyCalculateViewWarningOnce)
 		{
diff --git a/Engine/Plugins/Runtime/nDisplay/Source/DisplayCluster/Private/Render/Viewport/DisplayClusterViewport_Math.cpp b/Engine/Plugins/Runtime/nDisplay/Source/DisplayCluster/Private/Render/Viewport/DisplayClusterViewport_Math.cpp
--- a/Engine/Plugins/Runtime/nDisplay/Source/DisplayCluster/Private/Render/Viewport/DisplayClusterViewport_Math.cpp
+++ b/Engine/Plugins/Runtime/nDisplay/Source/DisplayCluster/Private/Render/Viewport/DisplayClusterViewport_Math.cpp
@@ -14,6 +14,8 @@
 
 #include "Render/Viewport/Containers/DisplayClusterViewport_PostRenderSettings.h"
 
+#include "Components/DisplayClusterCameraComponent.h"
+
 #include "EngineUtils.h"
 #include "SceneView.h"
 
@@ -21,11 +23,67 @@
 
 #include "Misc/DisplayClusterLog.h"
 
+namespace UE::DisplayCluster::Viewport::Math
+{
+	/**
+	 * Calculates the ViewOffset for the eye from the view location and changes the value of the view location to the eye position.
+	 * 
+	 * @param PassOffsetSwap - Distance to the eye from the midpoint between the eyes.
+	 * @param InOutViewLocation - (in, out) view location
+	 * @param InViewRotation (in) - rotation of the view (from this rotator we get the direction to the eye)
+	 * 
+	 * @return - the distance to the eye from the original ViewLocation
+	 */
+	static inline FVector ImplGetViewOffset(const float PassOffsetSwap, FVector& InOutViewLocation, const FRotator& InViewRotation)
+	{
+		// Apply computed offset to the view location
+		const FQuat EyeQuat = InViewRotation.Quaternion();
+		FVector ViewOffset = EyeQuat.RotateVector(FVector(0.0f, PassOffsetSwap, 0.0f));
+		
+		InOutViewLocation += ViewOffset;
+
+		return ViewOffset;
+	}
+};
+using namespace UE::DisplayCluster::Viewport::Math;
 ///////////////////////////////////////////////////////////////////////////////////////
 //          FDisplayClusterViewport
 ///////////////////////////////////////////////////////////////////////////////////////
+FVector2D FDisplayClusterViewport::GetClippingPlanes() const
+{
+	const float NCP = GNearClippingPlane;
+	const float FCP = NCP; // nDisplay does not use the far plane of the clipping
 
-bool FDisplayClusterViewport::CalculateView(const uint32 InContextNum, FVector& InOutViewLocation, FRotator& InOutViewRotation, const FVector& ViewOffset, const float WorldToMeters, const float NCP, const float FCP)
+	// Supports custom near clipping plane
+	float ZNear = NCP;
+	float ZFar = FCP;
+	if (CustomNearClippingPlane >= 0)
+	{
+		ZNear = CustomNearClippingPlane;
+		ZFar = (NCP == FCP) ? ZNear : ZFar;
+	}
+
+	return FVector2D(ZNear, ZFar);
+}
+
+bool FDisplayClusterViewport::GetViewPointCameraEye(const uint32 InContextNum, FVector& OutViewLocation, FRotator& OutViewRotation, FVector& OutViewOffset)
+{
+	// Here we use the ViewPoint component as the eye position
+	if (UDisplayClusterCameraComponent* ViewPoint = GetViewPointCameraComponent())
+	{
+		OutViewLocation = ViewPoint->GetComponentLocation();
+		OutViewRotation = ViewPoint->GetComponentRotation();
+
+		// Calculate stereo ViewOffset:
+		OutViewOffset = ImplGetViewOffset(GetStereoEyeOffsetDistance(InContextNum), OutViewLocation, OutViewRotation);
+
+		return true;
+	}
+
+	return false;
+}
+
+bool FDisplayClusterViewport::CalculateView(const uint32 InContextNum, FVector& InOutViewLocation, FRotator& InOutViewRotation, const float WorldToMeters)
 {
 	if (Contexts.IsValidIndex(InContextNum))
 	{
@@ -46,7 +104,11 @@ bool FDisplayClusterViewport::CalculateView(const uint32 InContextNum, FVector&
 
 				return true;
 			}
-			else if (ProjectionPolicy.IsValid() && ProjectionPolicy->CalculateView(this, InContextNum, InOutViewLocation, InOutViewRotation, ViewOffset, WorldToMeters, NCP, FCP))
+
+			// Calculate stereo ViewOffset:
+			const FVector ViewOffset = ImplGetViewOffset(GetStereoEyeOffsetDistance(InContextNum), InOutViewLocation, InOutViewRotation);
+			const FVector2D ClipingPlanes = GetClippingPlanes();
+			if (ProjectionPolicy.IsValid() && ProjectionPolicy->CalculateView(this, InContextNum, InOutViewLocation, InOutViewRotation, ViewOffset, WorldToMeters, ClipingPlanes.X, ClipingPlanes.Y))
 			{
 				Contexts[InContextNum].WorldToMeters = WorldToMeters;
 
@@ -57,11 +119,9 @@ bool FDisplayClusterViewport::CalculateView(const uint32 InContextNum, FVector&
 
 				return true;
 			}
-			else
-			{
-				// ProjectionPolicy->CalculateView() returns false, this view is invalid
-				EnumAddFlags(Contexts[InContextNum].ContextState, EDisplayClusterViewportContextState::InvalidViewPoint);
-			}
+
+			// ProjectionPolicy->CalculateView() returns false, this view is invalid
+			EnumAddFlags(Contexts[InContextNum].ContextState, EDisplayClusterViewportContextState::InvalidViewPoint);
 		}
 	}
 
diff --git a/Engine/Plugins/Runtime/nDisplay/Source/DisplayCluster/Private/Render/Viewport/DisplayClusterViewport_ViewPoint.cpp b/Engine/Plugins/Runtime/nDisplay/Source/DisplayCluster/Private/Render/Viewport/DisplayClusterViewport_ViewPoint.cpp
--- a/Engine/Plugins/Runtime/nDisplay/Source/DisplayCluster/Private/Render/Viewport/DisplayClusterViewport_ViewPoint.cpp
+++ b/Engine/Plugins/Runtime/nDisplay/Source/DisplayCluster/Private/Render/Viewport/DisplayClusterViewport_ViewPoint.cpp
@@ -7,12 +7,17 @@
 
 #include "DisplayClusterRootActor.h"
 #include "Components/DisplayClusterCameraComponent.h"
+#include "Components/DisplayClusterICVFXCameraComponent.h"
+#include "CineCameraComponent.h"
 
 #include "SceneView.h"
 
 #include "Misc/DisplayClusterLog.h"
 
-namespace UE::DisplayClusterViewport
+#include "GameFramework/PlayerController.h"
+#include "Camera/CameraComponent.h"
+
+namespace UE::DisplayCluster::Viewport::ViewPoint
 {
 	enum EDisplayClusterEyeType : int32
 	{
@@ -22,7 +27,83 @@ namespace UE::DisplayClusterViewport
 		COUNT
 	};
 };
-using namespace UE::DisplayClusterViewport;
+using namespace UE::DisplayCluster::Viewport::ViewPoint;
+
+///////////////////////////////////////////////////////////////////////////////////////
+// IDisplayClusterViewport
+///////////////////////////////////////////////////////////////////////////////////////
+bool IDisplayClusterViewport::GetCameraComponentView(UCameraComponent* InCameraComponent, const float InDeltaTime, const bool bUseCameraPostprocess, FMinimalViewInfo& InOutViewInfo, float* OutCustomNearClippingPlane)
+{
+	if (!InCameraComponent)
+	{
+		// Required camera component
+		return false;
+	}
+
+	InCameraComponent->GetCameraView(InDeltaTime, InOutViewInfo);
+
+	if(!bUseCameraPostprocess)
+	{
+		InOutViewInfo.PostProcessSettings = FPostProcessSettings();
+		InOutViewInfo.PostProcessBlendWeight = 0.0f;
+	}
+
+	// Get custom NCP from CineCamera component:
+	if (OutCustomNearClippingPlane)
+	{
+		*OutCustomNearClippingPlane = -1;
+
+		// Get settings from this cinecamera component
+		if (UCineCameraComponent* CineCameraComponent = Cast<UCineCameraComponent>(InCameraComponent))
+		{
+			// Supports ICVFX camera component as input
+			if (UDisplayClusterICVFXCameraComponent* ICVFXCameraComponent = Cast<UDisplayClusterICVFXCameraComponent>(CineCameraComponent))
+			{
+				// Getting settings from the actual CineCamera component
+				CineCameraComponent = ICVFXCameraComponent->GetActualCineCameraComponent();
+			}
+
+			if (CineCameraComponent && CineCameraComponent->bOverride_CustomNearClippingPlane)
+			{
+				*OutCustomNearClippingPlane = CineCameraComponent->CustomNearClippingPlane;
+			}
+		}
+	}
+
+	return true;
+}
+
+bool IDisplayClusterViewport::GetPlayerCameraView(UWorld* InWorld, const bool bUseCameraPostprocess, FMinimalViewInfo& InOutViewInfo)
+{
+	if (InWorld)
+	{
+		if (APlayerController* const CurPlayerController = InWorld->GetFirstPlayerController())
+		{
+			if (APlayerCameraManager* const CurPlayerCameraManager = CurPlayerController->PlayerCameraManager)
+			{
+				InOutViewInfo = CurPlayerCameraManager->GetCameraCacheView(); // Get desired view with postprocess from player camera
+
+				if (!bUseCameraPostprocess)
+				{
+					InOutViewInfo.PostProcessSettings = FPostProcessSettings();
+					InOutViewInfo.PostProcessBlendWeight = 0.0f;
+				}
+
+				InOutViewInfo.FOV = CurPlayerCameraManager->GetFOVAngle();
+				CurPlayerCameraManager->GetCameraViewPoint(/*out*/ InOutViewInfo.Location, /*out*/ InOutViewInfo.Rotation);
+
+				if (!bUseCameraPostprocess)
+				{
+					InOutViewInfo.PostProcessBlendWeight = 0.0f;
+				}
+
+				return true;
+			}
+		}
+	}
+
+	return false;
+}
 
 ///////////////////////////////////////////////////////////////////////////////////////
 //          FDisplayClusterViewport
@@ -73,12 +154,18 @@ bool FDisplayClusterViewport::SetupViewPoint(FMinimalViewInfo& InOutViewInfo)
 {
 	if (UDisplayClusterCameraComponent* ViewCamera = GetViewPointCameraComponent())
 	{
-		// First try to get the viewpoint from the projection policy
-		if (!ProjectionPolicy->GetViewPoint(this, InOutViewInfo.Rotation, InOutViewInfo.Location))
+		// Get ViewPoint from DCRA component
+		ViewCamera->GetDesiredView(InOutViewInfo, &CustomNearClippingPlane);
+
+		// The projection policy can override these ViewPoint data.
+		if (ProjectionPolicy.IsValid())
 		{
-			// if unsuccessful, get the viewport from the viewport camera
-			InOutViewInfo.Location = ViewCamera->GetComponentLocation();
-			InOutViewInfo.Rotation = ViewCamera->GetComponentRotation();
+			float DeltaTime = 0.0f;
+			if (ADisplayClusterRootActor* RootActor = GetRootActor())
+			{
+				DeltaTime = RootActor->GetWorldDeltaSeconds();
+			}
+			ProjectionPolicy->SetupProjectionViewPoint(this, DeltaTime, InOutViewInfo, &CustomNearClippingPlane);
 		}
 
 		return true;
@@ -93,47 +180,41 @@ float FDisplayClusterViewport::GetStereoEyeOffsetDistance(const uint32 InContext
 
 	if (UDisplayClusterCameraComponent* ViewCamera = GetViewPointCameraComponent())
 	{
-		// First try to get the eye offset distance from the projection policy
-		if (!ProjectionPolicy->GetStereoEyeOffsetDistance(this, InContextNum, StereoEyeOffsetDistance))
-		{
-			// if failed, get from the viewport camera
-
-			// Calculate eye offset considering the world scale
-			const float CfgEyeDist = ViewCamera->GetInterpupillaryDistance();
-			const float EyeOffset = CfgEyeDist / 2.f;
-			const float EyeOffsetValues[] = { -EyeOffset, 0.f, EyeOffset };
+		// Calculate eye offset considering the world scale
+		const float CfgEyeDist = ViewCamera->GetInterpupillaryDistance();
+		const float EyeOffset = CfgEyeDist / 2.f;
+		const float EyeOffsetValues[] = { -EyeOffset, 0.f, EyeOffset };
 
-			// Decode current eye type
-			const EDisplayClusterEyeType EyeType = (Contexts.Num() == 1)
-				? EDisplayClusterEyeType::Mono
-				: (InContextNum == 0) ? EDisplayClusterEyeType::StereoLeft : EDisplayClusterEyeType::StereoRight;
+		// Decode current eye type
+		const EDisplayClusterEyeType EyeType = (Contexts.Num() == 1)
+			? EDisplayClusterEyeType::Mono
+			: (InContextNum == 0) ? EDisplayClusterEyeType::StereoLeft : EDisplayClusterEyeType::StereoRight;
 
-			float PassOffset = 0.f;
+		float PassOffset = 0.f;
 
-			if (EyeType == EDisplayClusterEyeType::Mono)
-			{
-				// For monoscopic camera let's check if the "force offset" feature is used
-				// * Force left (-1) ==> 0 left eye
-				// * Force right (1) ==> 2 right eye
-				// * Default (0) ==> 1 mono
-				const EDisplayClusterEyeStereoOffset CfgEyeOffset = ViewCamera->GetStereoOffset();
-				const int32 EyeOffsetIdx =
-					(CfgEyeOffset == EDisplayClusterEyeStereoOffset::None ? 0 :
-						(CfgEyeOffset == EDisplayClusterEyeStereoOffset::Left ? -1 : 1));
-
-				PassOffset = EyeOffsetValues[EyeOffsetIdx + 1];
-				// Eye swap is not available for monoscopic so just save the value
-				StereoEyeOffsetDistance = PassOffset;
-			}
-			else
-			{
-				// For stereo camera we can only swap eyes if required (no "force offset" allowed)
-				PassOffset = EyeOffsetValues[EyeType];
+		if (EyeType == EDisplayClusterEyeType::Mono)
+		{
+			// For monoscopic camera let's check if the "force offset" feature is used
+			// * Force left (-1) ==> 0 left eye
+			// * Force right (1) ==> 2 right eye
+			// * Default (0) ==> 1 mono
+			const EDisplayClusterEyeStereoOffset CfgEyeOffset = ViewCamera->GetStereoOffset();
+			const int32 EyeOffsetIdx =
+				(CfgEyeOffset == EDisplayClusterEyeStereoOffset::None ? 0 :
+					(CfgEyeOffset == EDisplayClusterEyeStereoOffset::Left ? -1 : 1));
+
+			PassOffset = EyeOffsetValues[EyeOffsetIdx + 1];
+			// Eye swap is not available for monoscopic so just save the value
+			StereoEyeOffsetDistance = PassOffset;
+		}
+		else
+		{
+			// For stereo camera we can only swap eyes if required (no "force offset" allowed)
+			PassOffset = EyeOffsetValues[EyeType];
 
-				// Apply eye swap
-				const bool  CfgEyeSwap = ViewCamera->GetSwapEyes();
-				StereoEyeOffsetDistance = (CfgEyeSwap ? -PassOffset : PassOffset);
-			}
+			// Apply eye swap
+			const bool  CfgEyeSwap = ViewCamera->GetSwapEyes();
+			StereoEyeOffsetDistance = (CfgEyeSwap ? -PassOffset : PassOffset);
 		}
 	}
 
diff --git a/Engine/Plugins/Runtime/nDisplay/Source/DisplayCluster/Public/Components/DisplayClusterCameraComponent.h b/Engine/Plugins/Runtime/nDisplay/Source/DisplayCluster/Public/Components/DisplayClusterCameraComponent.h
--- a/Engine/Plugins/Runtime/nDisplay/Source/DisplayCluster/Public/Components/DisplayClusterCameraComponent.h
+++ b/Engine/Plugins/Runtime/nDisplay/Source/DisplayCluster/Public/Components/DisplayClusterCameraComponent.h
@@ -10,7 +10,8 @@
 
 class UBillboardComponent;
 class UTexture2D;
-
+class IDisplayClusterViewportManager;
+class IDisplayClusterWarpPolicy;
 
 UENUM()
 enum class EDisplayClusterEyeStereoOffset : uint8
@@ -34,6 +35,31 @@ class DISPLAYCLUSTER_API UDisplayClusterCameraComponent
 public:
 	UDisplayClusterCameraComponent(const FObjectInitializer& ObjectInitializer);
 
+	/** Return ViewPoint for this component
+	 * If the component logic supports postprocess, it will also be in the ViewInfo structure.
+	 * 
+	 * @param InOutViewInfo - ViewInfo data
+	 * @param OutCustomNearClippingPlane - Custom NCP, or a value less than zero if not defined.
+	 */
+	virtual void GetDesiredView(FMinimalViewInfo& InOutViewInfo, float* OutCustomNearClippingPlane = nullptr);
+
+	/**
+	 * All cluster viewports that reference this component will be created in the background on the current cluster node if the function returns true.
+	 */
+	virtual bool ShouldUseEntireClusterViewports(IDisplayClusterViewportManager* InViewportManager) const
+	{
+		return false;
+	}
+
+	/**
+	 * Get the warp policy instance used by this compoenent.
+	 * From the DC ViewportManager, these policies will be assigned to the viewports that use this viewpoint component.
+	 */
+	virtual IDisplayClusterWarpPolicy* GetWarpPolicy(IDisplayClusterViewportManager* InViewportManager)
+	{
+		return nullptr;
+	}
+
 public:
 	/**
 	* Get interpupillary distance
diff --git a/Engine/Plugins/Runtime/nDisplay/Source/DisplayCluster/Public/Components/DisplayClusterICVFXCameraComponent.h b/Engine/Plugins/Runtime/nDisplay/Source/DisplayCluster/Public/Components/DisplayClusterICVFXCameraComponent.h
--- a/Engine/Plugins/Runtime/nDisplay/Source/DisplayCluster/Public/Components/DisplayClusterICVFXCameraComponent.h
+++ b/Engine/Plugins/Runtime/nDisplay/Source/DisplayCluster/Public/Components/DisplayClusterICVFXCameraComponent.h
@@ -48,6 +48,12 @@ public:
 public:
 	FDisplayClusterViewport_CameraMotionBlur GetMotionBlurParameters();
 
+	/**
+	 * Return the actual source camera, e.g. the camera component of the referenced cine camera.
+	 * Use GetCameraView() function to get viewinfo with actual camera position, postprocess and ICVFX postprocess
+	 */
+	UCineCameraComponent* GetActualCineCameraComponent();
+
 	// Returns true if this camera is active
 	bool IsICVFXEnabled() const;
 
@@ -59,8 +65,10 @@ public:
 		return CameraSettings;
 	}
 
-	UCameraComponent* GetCameraComponent();
-	void GetDesiredView(FMinimalViewInfo& DesiredView);
+	/** Obtaining view information for the actual camera, such as the camera component to which the cine-camera is referencing.
+	 * The data from the CameraSettings variable is used in postprocess settings (EnableCameraPP, OverrideMotionBlur, etc.).
+	 */
+	virtual void GetCameraView(float DeltaTime, FMinimalViewInfo& InOutViewInfo) override;
 
 	// UActorComponent interface
 	virtual void OnRegister() override;
diff --git a/Engine/Plugins/Runtime/nDisplay/Source/DisplayCluster/Public/Misc/DisplayClusterDataCache.h b/Engine/Plugins/Runtime/nDisplay/Source/DisplayCluster/Public/Misc/DisplayClusterDataCache.h
--- a/Engine/Plugins/Runtime/nDisplay/Source/DisplayCluster/Public/Misc/DisplayClusterDataCache.h
+++ b/Engine/Plugins/Runtime/nDisplay/Source/DisplayCluster/Public/Misc/DisplayClusterDataCache.h
@@ -23,6 +23,9 @@ public:
 	 *   // Return true if DataCache is enabled.
 	 *   static bool IsDataCacheEnabled();
 	 *
+	 *   // Method for releasing a cached data item, called before its destructor
+	 *   void ReleaseDataCacheItem();
+	 *
 	 *   // Returns the unique name of this texture for DataCache.
 	 *   inline const FString& GetDataCacheName() const;
 	 */
@@ -151,6 +154,16 @@ private:
 			: DataRef(InDataRef)
 		{ }
 
+		~FCachedObject()
+		{
+			// Calling the release function for referenced data before the destructor
+			if (DataRef.IsValid())
+			{
+				DataRef->ReleaseDataCacheItem();
+				DataRef.Reset();
+			}
+		}
+
 		inline bool IsNameEqual(const FString& InName) const
 		{
 			return DataRef.IsValid() && DataRef->GetDataCacheName() == InName;
diff --git a/Engine/Plugins/Runtime/nDisplay/Source/DisplayCluster/Public/Render/Projection/IDisplayClusterProjectionPolicy.h b/Engine/Plugins/Runtime/nDisplay/Source/DisplayCluster/Public/Render/Projection/IDisplayClusterProjectionPolicy.h
--- a/Engine/Plugins/Runtime/nDisplay/Source/DisplayCluster/Public/Render/Projection/IDisplayClusterProjectionPolicy.h
+++ b/Engine/Plugins/Runtime/nDisplay/Source/DisplayCluster/Public/Render/Projection/IDisplayClusterProjectionPolicy.h
@@ -10,8 +10,9 @@ class IDisplayClusterViewportProxy;
 class IDisplayClusterWarpBlend;
 class IDisplayClusterWarpPolicy;
 class UMeshComponent;
+class USceneComponent;
 struct FDisplayClusterConfigurationProjection;
-
+struct FMinimalViewInfo;
 
 /**
  * nDisplay projection policy
@@ -35,7 +36,7 @@ public:
 	/**
 	* Return projection policy type
 	*/
-	UE_DEPRECATED(5.1, "This function has beend deprecated. Please use 'GetType'.")
+	UE_DEPRECATED(5.1, "This function has been deprecated. Please use 'GetType'.")
 	virtual const FString GetTypeId() const
 	{
 		return GetType();
@@ -46,6 +47,16 @@ public:
 	*/
 	virtual const TMap<FString, FString>& GetParameters() const = 0;
 
+	/**
+	 * Return Origin point component used by this viewport
+	 * This component is used to convert from the local DCRA space to the local projection policy space,
+	 * which contains the calibrated geometry data used for the warp.
+	 */
+	virtual const USceneComponent* const GetOriginComponent() const
+	{
+		return nullptr;
+	}
+
 	/**
 	* Send projection policy game thread data to render thread proxy
 	* called once per frame from FDisplayClusterViewportManager::FinalizeNewFrame
@@ -113,11 +124,21 @@ public:
 	}
 
 	// This policy can support ICVFX rendering
+	UE_DEPRECATED(5.3, "This function has been deprecated. Please use 'ShouldSupportICVFX(IDisplayClusterViewport*)'.")
 	virtual bool ShouldSupportICVFX() const
 	{
 		return false;
 	}
 
+	/** Returns true if this policy supports ICVFX rendering
+	 * 
+	 * @param InViewport - a owner viewport
+	 */
+	virtual bool ShouldSupportICVFX(IDisplayClusterViewport* InViewport) const
+	{
+		return false;
+	}
+
 	// Return true, if camera projection visible for this viewport geometry
 	// ICVFX Performance : if camera frame not visible on this node, disable render for this camera
 	virtual bool IsCameraProjectionVisible(const FRotator& InViewRotation, const FVector& InViewLocation, const FMatrix& InProjectionMatrix)
@@ -134,34 +155,20 @@ public:
 	*/
 	virtual bool IsConfigurationChanged(const struct FDisplayClusterConfigurationProjection* InConfigurationProjectionPolicy) const = 0;
 
-	/** Get viewpoint for this projection policy
-	 *
-	 * @param OutViewRotation - viewpoint rotation
-	 * @param OutViewLocation - viewpoint location
+	/** Override view from this projection policy
 	 *
-	 * @return - true, if the viewpoint values have changed
+	 * @param InViewport                 - the viewport of this projection policy
+	 * @param InDeltaTime                - delta time in current frame
+	 * @param InOutViewInfo              - ViewInfo data
+	 * @param OutCustomNearClippingPlane - Custom NCP, or a value less than zero if not defined.
 	 */
-	virtual bool GetViewPoint(IDisplayClusterViewport* InViewport, FRotator& InOutViewRotation, FVector& InOutViewLocation)
-	{
-		return false;
-	}
+	virtual void SetupProjectionViewPoint(IDisplayClusterViewport* InViewport, const float InDeltaTime, FMinimalViewInfo& InOutViewInfo, float* OutCustomNearClippingPlane = nullptr)
+	{ }
 
 	/** Projection policy can override PP */
-	virtual void OverridePostProcessSettings(IDisplayClusterViewport* InViewport)
+	virtual void UpdatePostProcessSettings(IDisplayClusterViewport* InViewport)
 	{ }
 
-	 /** Get the distance from the eye to the viewpoint location
-	 *
-	 * @param InContextNum - eye context of this viewport
-	 * @param OutStereoEyeOffsetDistance - new eye distance
-	 *
-	 * @return - true, if the offset distance of the stereo eye has changed
-	 */
-	virtual bool GetStereoEyeOffsetDistance(IDisplayClusterViewport* InViewport, const uint32 InContextNum, float& InOutStereoEyeOffsetDistance)
-	{
-		return false;
-	}
-
 	/**
 	* @param ViewIdx           - Index of view that is being processed for this viewport
 	* @param InOutViewLocation - (in/out) View location with ViewOffset (i.e. left eye pre-computed location)
@@ -262,6 +269,9 @@ public:
 
 	/**
 	* Build preview mesh
+	* This MeshComponent cannot be moved freely.
+	* This MeshComponent is attached to the geometry from the real world via Origin.
+	* When the screen geometry in the real world changes position, the position of this component must also be changed.
 	*
 	* @param InViewport - Projection specific parameters.
 	* @param bOutIsRootActorComponent - return true, if used custom root actor component. return false, if created unique temporary component
@@ -270,5 +280,26 @@ public:
 	{
 		return nullptr;
 	}
+
+	/**
+	* Ask projection policy instance if it has any movable mesh based preview
+	*
+	* @return - True if mesh based preview is available
+	*/
+	virtual bool HasPreviewMovableMesh()
+	{
+		return false;
+	}
+
+	/**
+	* Build preview movable mesh
+	* This MeshComponent is a copy of the preview mesh and can be moved freely with the UI visualization.
+	*
+	* @param InViewport - Projection specific parameters.
+	*/
+	virtual UMeshComponent* GetOrCreatePreviewMovableMeshComponent(IDisplayClusterViewport* InViewport)
+	{
+		return nullptr;
+	}
 #endif
 };
diff --git a/Engine/Plugins/Runtime/nDisplay/Source/DisplayCluster/Public/Render/Viewport/IDisplayClusterViewport.h b/Engine/Plugins/Runtime/nDisplay/Source/DisplayCluster/Public/Render/Viewport/IDisplayClusterViewport.h
--- a/Engine/Plugins/Runtime/nDisplay/Source/DisplayCluster/Public/Render/Viewport/IDisplayClusterViewport.h
+++ b/Engine/Plugins/Runtime/nDisplay/Source/DisplayCluster/Public/Render/Viewport/IDisplayClusterViewport.h
@@ -14,6 +14,10 @@
 
 #include "Engine/EngineTypes.h"
 
+class UCameraComponent;
+class UWorld;
+struct FMinimalViewInfo;
+
 /**
  * nDisplay: Viewport (interface for GameThread)
  */
@@ -34,9 +38,57 @@ public:
 	virtual const FDisplayClusterViewport_RenderSettingsICVFX& GetRenderSettingsICVFX() const = 0;
 	virtual const FDisplayClusterViewport_PostRenderSettings&  GetPostRenderSettings() const = 0;
 
-	// math wrappers: support overscan, etc
+	/** Calculate projection matrices for viewport context
+	 * (also created overscan matrices, etc)
+	 * 
+	 * @param InContextNum   - viewport eye context index
+	 * @param Left           - frustum left angle
+	 * @param Right          - frustum right angle
+	 * @param Top            - frustum top angle
+	 * @param Bottom         - frustum bottom angle
+	 * @param NCP            - Near clipping plane
+	 * @param FCP            - Far  clipping plane
+	 * @param bIsAnglesInput - true, if projection angles in degrees
+	 */
 	virtual void CalculateProjectionMatrix(const uint32 InContextNum, float Left, float Right, float Top, float Bottom, float ZNear, float ZFar, bool bIsAnglesInput) = 0;
-	virtual bool CalculateView(const uint32 InContextNum, FVector& InOutViewLocation, FRotator& InOutViewRotation, const FVector& ViewOffset, const float WorldToMeters, const float NCP, const float FCP) = 0;
+
+	/** Calculate viewport projection (This function has been deprecated.)
+	 *
+	 * @param InContextNum      - viewport eye context index
+	 * @param InOutViewLocation - (in, out) View location
+	 * @param InOutViewRotation - (in,out) View rotator
+	 * @param ViewOffset         - eye offset from view location
+	 * @param WorldToMeters     - UE world scale
+	 * @param NCP - Near clipping plane
+	 * @param FCP - Far clipping plane
+	 *
+	 * @return - true, if the calculation of the viewport projection is successful.
+	 */
+	UE_DEPRECATED(5.3, "This function has been deprecated. Please use 'CalculateView()'.")
+	virtual bool CalculateView(const uint32 InContextNum, FVector& InOutViewLocation, FRotator& InOutViewRotation, const FVector& ViewOffset, const float WorldToMeters, const float NCP, const float FCP)
+	{
+		return false;
+	}
+
+	/** Calculate viewport projection
+	 * 
+	 * @param InContextNum      - viewport eye context index
+	 * @param InOutViewLocation - (in, out) View location
+	 * @param InOutViewRotation - (in,out) View rotator
+	 * @param WorldToMeters     - UE world scale
+	 * 
+	 * @return - true, if the calculation of the viewport projection is successful.
+	 */
+	virtual bool CalculateView(const uint32 InContextNum, FVector& InOutViewLocation, FRotator& InOutViewRotation, const float WorldToMeters) = 0;
+
+
+	/** Get projection matrix for viewport context
+	 * 
+	 * @param InContextNum - viewport eye context index
+	 * @param OutPrjMatrix - (out) Projection  matrix
+	 * 
+	 * @return true, if projection matrix valid.
+	 */
 	virtual bool GetProjectionMatrix(const uint32 InContextNum, FMatrix& OutPrjMatrix) = 0;
 
 	/** Setup viewpoint for this viewport
@@ -50,6 +102,12 @@ public:
 	/** Return view point camera component for this viewport. */
 	virtual class UDisplayClusterCameraComponent* GetViewPointCameraComponent() const = 0;
 
+	/** Retrieves the view position from the ViewPoint component used by this viewport.
+	 * The CalculateView() function receives arguments that may not match this position because they can be overridden.
+	 * Therefore, projection policies that expect values from the ViewPoint component must use this function to get the correct values.
+	 */
+	virtual bool GetViewPointCameraEye(const uint32 InContextNum, FVector& OutViewLocation, FRotator& OutViewRotation, FVector& OutViewOffset) = 0;
+
 	/** Get the distance from the eye to the viewpoint location.
 	 *
 	 * @param InContextNum - eye context of this viewport
@@ -60,6 +118,12 @@ public:
 
 	virtual const TArray<FDisplayClusterViewport_Context>& GetContexts() const = 0;
 
+	/** Get clipping planes used by this viewport
+	 * 
+	 * @return Clipping planes as FVector2D(ZNear, ZFar)
+	 */
+	virtual FVector2D GetClippingPlanes() const = 0;
+
 	// Override postprocess settings for this viewport
 	virtual const IDisplayClusterViewport_CustomPostProcessSettings& GetViewport_CustomPostProcessSettings() const = 0;
 	virtual IDisplayClusterViewport_CustomPostProcessSettings& GetViewport_CustomPostProcessSettings() = 0;
@@ -67,6 +131,14 @@ public:
 	// Setup scene view for rendering specified Context
 	virtual void SetupSceneView(uint32 ContextNum, class UWorld* World, class FSceneViewFamily& InViewFamily, FSceneView& InView) const = 0;
 
+	UE_DEPRECATED(5.3, "This function has beend deprecate. Please use 'GetViewportManager'.")
+	virtual class IDisplayClusterViewportManager& GetOwner() const
+	{
+		check(false);
+
+		return *GetViewportManager();
+	}
+
 	/** Return the viewport manager that owns this viewport */
 	virtual class IDisplayClusterViewportManager* GetViewportManager() const = 0;
 
@@ -89,4 +161,26 @@ public:
 	virtual void SetContexts(TArray<FDisplayClusterViewport_Context>& InContexts) = 0;
 
 	static FMatrix MakeProjectionMatrix(float InLeft, float InRight, float InTop, float InBottom, float ZNear, float ZFar);
+
+	/** Get View from CameraComponent
+	*
+	* @param InCameraComponent          - Ptr to the camera component used for rendering. if this ptr is nullptr, then the active player camera is used
+	* @param InDeltaTime                - delta time in current frame
+	* @param bUseCameraPostprocess      - If false, then InOutViewInfo.PostProcessBlendWeight is set to 0.
+	* @param InOutViewInfo              - ViewPoint data
+	* @param OutCustomNearClippingPlane - Custom NCP, or a value less than zero if not defined.
+	*
+	* @return false if the viewpoint data cannot be retrieved.
+	*/
+	static bool GetCameraComponentView(UCameraComponent* InCameraComponent, const float InDeltaTime, const bool bUseCameraPostprocess, FMinimalViewInfo& InOutViewInfo, float* OutCustomNearClippingPlane = nullptr);
+
+	/** Get View from current player camera
+	*
+	* @param InWorld - In this world, an active game camera will be searched for.
+	* @param bUseCameraPostprocess - if false, then InOutViewInfo.PostProcessBlendWeight is set to 0.
+	* @param InOutViewInfo         - (in,out) ViewPoint data
+	*
+	* @return false if the viewpoint data cannot be retrieved.
+	*/
+	static bool GetPlayerCameraView(UWorld* InWorld, const bool bUseCameraPostprocess, FMinimalViewInfo& InOutViewInfo);
 };
diff --git a/Engine/Plugins/Runtime/nDisplay/Source/DisplayCluster/Public/Render/Viewport/IDisplayClusterViewportManager.h b/Engine/Plugins/Runtime/nDisplay/Source/DisplayCluster/Public/Render/Viewport/IDisplayClusterViewportManager.h
--- a/Engine/Plugins/Runtime/nDisplay/Source/DisplayCluster/Public/Render/Viewport/IDisplayClusterViewportManager.h
+++ b/Engine/Plugins/Runtime/nDisplay/Source/DisplayCluster/Public/Render/Viewport/IDisplayClusterViewportManager.h
@@ -176,13 +176,27 @@ public:
 	*/
 	virtual IDisplayClusterViewport* FindViewport(const int32 StereoViewIndex, uint32* OutContextNum = nullptr) const = 0;
 	
+	UE_DEPRECATED(5.3, "This function has been deprecated. Please use 'GetCurrentRenderFrameViewports'.")
+	virtual const TArrayView<IDisplayClusterViewport*> GetViewports() const
+	{
+		return TArrayView<IDisplayClusterViewport*>();
+	}
+
+	/**
+	* Return all exist viewports objects for current cluster node
+	* [Game thread func]
+	*
+	* @return - arrays with viewport objects refs
+	*/
+	virtual const TArrayView<TSharedPtr<IDisplayClusterViewport, ESPMode::ThreadSafe>> GetCurrentRenderFrameViewports() const = 0;
+
 	/**
-	* Return all exist viewports objects
+	* Return entire cluster viewports objects
 	* [Game thread func]
 	*
 	* @return - arrays with viewport objects refs
 	*/
-	virtual const TArrayView<TSharedPtr<IDisplayClusterViewport, ESPMode::ThreadSafe>> GetViewports() const = 0;
+	virtual const TArrayView<TSharedPtr<IDisplayClusterViewport, ESPMode::ThreadSafe>> GetEntireClusterViewports() const = 0;
 
 	/**
 	* Mark the geometry of the referenced component(s) as dirty (ProceduralMesh, etc)
diff --git a/Engine/Plugins/Runtime/nDisplay/Source/DisplayCluster/Public/Render/Warp/IDisplayClusterWarpPolicy.h b/Engine/Plugins/Runtime/nDisplay/Source/DisplayCluster/Public/Render/Warp/IDisplayClusterWarpPolicy.h
--- a/Engine/Plugins/Runtime/nDisplay/Source/DisplayCluster/Public/Render/Warp/IDisplayClusterWarpPolicy.h
+++ b/Engine/Plugins/Runtime/nDisplay/Source/DisplayCluster/Public/Render/Warp/IDisplayClusterWarpPolicy.h
@@ -64,15 +64,44 @@ public:
 	virtual void Tick(IDisplayClusterViewportManager* InViewportManager, float DeltaSeconds)
 	{ }
 
+	/** Should override frustum for viewport context
+	 * 
+	 * @param InViewport - a owner viewport
+	 * @param ContextNum - viewport eye context index
+	 * 
+	 * @return - true, if the CalcFrustumOverrideFunc() function will be used
+	 */
+	virtual bool ShouldOverrideCalcFrustum(IDisplayClusterViewport* InViewport)
+	{
+		return false;
+	}
+
 	/** Override frustum for viewport context
+	 * This function is called only when the ShouldOverrideCalcFrustum() function returns true
 	 * 
 	 * @param InViewport - a owner viewport
 	 * @param ContextNum - viewport eye context index
 	 * 
 	 * @return - true if frustum overrided.
 	 */
-	virtual bool CalcFrustum(IDisplayClusterViewport* InViewport, const uint32 ContextNum)
+	virtual bool OverrideCalcFrustum(IDisplayClusterViewport* InViewport, const uint32 ContextNum)
 	{
 		return false;
 	}
+
+	/** Call before CalcFrustum()
+	 *
+	 * @param InViewport - a owner viewport
+	 * @param ContextNum - viewport eye context index
+	 */
+	virtual void BeginCalcFrustum(IDisplayClusterViewport* InViewport, const uint32 ContextNum)
+	{ }
+
+	/** Call after CalcFrustum()
+	 *
+	 * @param InViewport - a owner viewport
+	 * @param ContextNum - viewport eye context index
+	 */
+	virtual void EndCalcFrustum(IDisplayClusterViewport* InViewport, const uint32 ContextNum)
+	{ }
 };
diff --git a/Engine/Plugins/Runtime/nDisplay/Source/DisplayClusterConfiguration/Public/DisplayClusterConfigurationTypes_Viewport.h b/Engine/Plugins/Runtime/nDisplay/Source/DisplayClusterConfiguration/Public/DisplayClusterConfigurationTypes_Viewport.h
--- a/Engine/Plugins/Runtime/nDisplay/Source/DisplayClusterConfiguration/Public/DisplayClusterConfigurationTypes_Viewport.h
+++ b/Engine/Plugins/Runtime/nDisplay/Source/DisplayClusterConfiguration/Public/DisplayClusterConfigurationTypes_Viewport.h
@@ -175,6 +175,12 @@ public:
 public:
 	void GetReferencedMeshNames(TArray<FString>& OutMeshNames) const;
 
+	/** The viewport object will only be created or updated if this function returns true.*/
+	bool IsViewportEnabled() const
+	{
+		return bAllowRendering;
+	}
+
 #if WITH_EDITOR
 	/** Enable the preview texture. Only should be called by the object managing the preview texture state. */
 	void EnablePreviewTexture();
diff --git a/Engine/Plugins/Runtime/nDisplay/Source/DisplayClusterMoviePipeline/Private/DisplayClusterMoviePipelineViewportPass.cpp b/Engine/Plugins/Runtime/nDisplay/Source/DisplayClusterMoviePipeline/Private/DisplayClusterMoviePipelineViewportPass.cpp
--- a/Engine/Plugins/Runtime/nDisplay/Source/DisplayClusterMoviePipeline/Private/DisplayClusterMoviePipelineViewportPass.cpp
+++ b/Engine/Plugins/Runtime/nDisplay/Source/DisplayClusterMoviePipeline/Private/DisplayClusterMoviePipelineViewportPass.cpp
@@ -560,23 +560,13 @@ IDisplayClusterViewport* UDisplayClusterMoviePipelineViewportPassBase::GetAndCal
 			OutView.ViewRotation = ViewInfo.Rotation;
 			OutView.ViewLocation = ViewInfo.Location;
 			
-			const float PassOffsetSwap = DCViewport->GetStereoEyeOffsetDistance(InContextNum);
-
-			FVector ViewOffset = FVector::ZeroVector;
-			const float CfgNCP = GNearClippingPlane;
-			{
-				// Apply computed offset to the view location
-				const FQuat EyeQuat = OutView.ViewRotation.Quaternion();
-				ViewOffset = EyeQuat.RotateVector(FVector(0.0f, PassOffsetSwap, 0.0f));
-				OutView.ViewLocation += ViewOffset;
-			}
-
 			const AWorldSettings* WorldSettings = CurrentWorld->GetWorldSettings();
 			const float InWorldToMeters = (WorldSettings) ? WorldSettings->WorldToMeters : 100.f;
 
 			bool bResult = false;
+			// Obtaining the offset of the stereo eye and the values of the projection clipping plane for the given viewport was moved inside CalculateView().
 			// Perform view calculations on a policy side
-			if (DCViewport->CalculateView(InContextNum, OutView.ViewLocation, OutView.ViewRotation, ViewOffset, InWorldToMeters, CfgNCP, CfgNCP))
+			if (DCViewport->CalculateView(InContextNum, OutView.ViewLocation, OutView.ViewRotation, InWorldToMeters))
 			{
 				OutView.ProjectionMatrix = FMatrix::Identity;
 
diff --git a/Engine/Plugins/Runtime/nDisplay/Source/DisplayClusterProjection/Private/DisplayClusterProjectionModule.cpp b/Engine/Plugins/Runtime/nDisplay/Source/DisplayClusterProjection/Private/DisplayClusterProjectionModule.cpp
--- a/Engine/Plugins/Runtime/nDisplay/Source/DisplayClusterProjection/Private/DisplayClusterProjectionModule.cpp
+++ b/Engine/Plugins/Runtime/nDisplay/Source/DisplayClusterProjection/Private/DisplayClusterProjectionModule.cpp
@@ -139,13 +139,13 @@ TSharedPtr<IDisplayClusterProjectionPolicyFactory> FDisplayClusterProjectionModu
 	return nullptr;
 }
 
-bool FDisplayClusterProjectionModule::CameraPolicySetCamera(const TSharedPtr<IDisplayClusterProjectionPolicy, ESPMode::ThreadSafe>& InPolicy, UCameraComponent* const NewCamera, const FDisplayClusterProjectionCameraPolicySettings& CamersSettings)
+bool FDisplayClusterProjectionModule::CameraPolicySetCamera(const TSharedPtr<IDisplayClusterProjectionPolicy, ESPMode::ThreadSafe>& InPolicy, UCameraComponent* const NewCamera, const FDisplayClusterProjectionCameraPolicySettings& CameraSettings)
 {
 	if (InPolicy.IsValid())
 	{
 		if (FDisplayClusterProjectionCameraPolicy* CameraPolicyInstance = static_cast<FDisplayClusterProjectionCameraPolicy*>(InPolicy.Get()))
 		{
-			CameraPolicyInstance->SetCamera(NewCamera, CamersSettings);
+			CameraPolicyInstance->SetCamera(NewCamera, CameraSettings);
 			return true;
 		}
 	}
diff --git a/Engine/Plugins/Runtime/nDisplay/Source/DisplayClusterProjection/Private/DisplayClusterProjectionModule.h b/Engine/Plugins/Runtime/nDisplay/Source/DisplayClusterProjection/Private/DisplayClusterProjectionModule.h
--- a/Engine/Plugins/Runtime/nDisplay/Source/DisplayClusterProjection/Private/DisplayClusterProjectionModule.h
+++ b/Engine/Plugins/Runtime/nDisplay/Source/DisplayClusterProjection/Private/DisplayClusterProjectionModule.h
@@ -31,7 +31,7 @@ public:
 	virtual void GetSupportedProjectionTypes(TArray<FString>& OutProjectionTypes) override;
 	virtual TSharedPtr<IDisplayClusterProjectionPolicyFactory> GetProjectionFactory(const FString& InProjectionType) override;
 
-	virtual bool CameraPolicySetCamera(const TSharedPtr<IDisplayClusterProjectionPolicy, ESPMode::ThreadSafe>& InPolicy, UCameraComponent* const NewCamera, const FDisplayClusterProjectionCameraPolicySettings& CamersSettings) override;
+	virtual bool CameraPolicySetCamera(const TSharedPtr<IDisplayClusterProjectionPolicy, ESPMode::ThreadSafe>& InPolicy, UCameraComponent* const NewCamera, const FDisplayClusterProjectionCameraPolicySettings& CameraSettings) override;
 
 private:
 	// Available factories
diff --git a/Engine/Plugins/Runtime/nDisplay/Source/DisplayClusterProjection/Private/Policy/Camera/DisplayClusterProjectionCameraPolicy.cpp b/Engine/Plugins/Runtime/nDisplay/Source/DisplayClusterProjection/Private/Policy/Camera/DisplayClusterProjectionCameraPolicy.cpp
--- a/Engine/Plugins/Runtime/nDisplay/Source/DisplayClusterProjection/Private/Policy/Camera/DisplayClusterProjectionCameraPolicy.cpp
+++ b/Engine/Plugins/Runtime/nDisplay/Source/DisplayClusterProjection/Private/Policy/Camera/DisplayClusterProjectionCameraPolicy.cpp
@@ -8,7 +8,6 @@
 #include "Game/IDisplayClusterGameManager.h"
 
 #include "Engine/World.h"
-#include "GameFramework/PlayerController.h"
 #include "Camera/CameraComponent.h"
 #include "CineCameraComponent.h"
 #include "ComposurePostMoves.h"
@@ -16,21 +15,44 @@
 #include "Render/Viewport/IDisplayClusterViewport.h"
 #include "Render/Viewport/IDisplayClusterViewportManager.h"
 
+#include "DisplayClusterRootActor.h"
 
+
+//////////////////////////////////////////////////////////////////////////////////////////////
+// FDisplayClusterProjectionPolicy
+//////////////////////////////////////////////////////////////////////////////////////////////
 FDisplayClusterProjectionCameraPolicy::FDisplayClusterProjectionCameraPolicy(const FString& ProjectionPolicyId, const FDisplayClusterConfigurationProjection* InConfigurationProjectionPolicy)
 	: FDisplayClusterProjectionPolicyBase(ProjectionPolicyId, InConfigurationProjectionPolicy)
 {
 }
 
-//////////////////////////////////////////////////////////////////////////////////////////////
-// IDisplayClusterProjectionPolicy
-//////////////////////////////////////////////////////////////////////////////////////////////
 const FString& FDisplayClusterProjectionCameraPolicy::GetType() const
 {
 	static const FString Type(DisplayClusterProjectionStrings::projection::Camera);
 	return Type;
 }
 
+void FDisplayClusterProjectionCameraPolicy::UpdatePostProcessSettings(IDisplayClusterViewport* InViewport)
+{
+	// Override viewport PP from camera, except ICVFX inner camera
+	// because the ICVFX InCamera postprocess is overridden separately in FDisplayClusterViewportConfigurationHelpers_Postprocess::UpdateCameraPostProcessSettings();
+	// The ICVFX camera postprocess is more detailed and is defined in the FDisplayClusterConfigurationICVFX_CameraSettings structure.
+	if (InViewport && !EnumHasAnyFlags(InViewport->GetRenderSettingsICVFX().RuntimeFlags, EDisplayClusterViewportRuntimeICVFXFlags::InCamera))
+	{
+		float DeltaTime = 0.0f;
+		if (ADisplayClusterRootActor* RootActor = InViewport->GetRootActor())
+		{
+			DeltaTime = RootActor->GetWorldDeltaSeconds();
+		}
+
+		FMinimalViewInfo ViewInfo;
+		if (ImplSetupProjectionViewPoint(InViewport, DeltaTime, ViewInfo) && ViewInfo.PostProcessBlendWeight > 0.0f)
+		{
+			InViewport->GetViewport_CustomPostProcessSettings().AddCustomPostProcess(IDisplayClusterViewport_CustomPostProcessSettings::ERenderPass::Override, ViewInfo.PostProcessSettings, ViewInfo.PostProcessBlendWeight, true);
+		}
+	}
+}
+
 bool FDisplayClusterProjectionCameraPolicy::HandleStartScene(IDisplayClusterViewport* InViewport)
 {
 	check(IsInGameThread());
@@ -45,73 +67,59 @@ void FDisplayClusterProjectionCameraPolicy::HandleEndScene(IDisplayClusterViewpo
 	CameraRef.ResetSceneComponent();
 }
 
-APlayerCameraManager* const GetCurPlayerCameraManager(IDisplayClusterViewport* InViewport)
+void FDisplayClusterProjectionCameraPolicy::SetupProjectionViewPoint(IDisplayClusterViewport* InViewport, const float InDeltaTime, FMinimalViewInfo& InOutViewInfo, float* OutCustomNearClippingPlane)
 {
-	if (InViewport)
-	{
-		if (UWorld* World = InViewport->GetCurrentWorld())
-		{
-			if (APlayerController* const CurPlayerController = World->GetFirstPlayerController())
-			{
-				return CurPlayerController->PlayerCameraManager;
-			}
-		}
-	}
+	// Get camera position
+	ImplSetupProjectionViewPoint(InViewport, InDeltaTime, InOutViewInfo, OutCustomNearClippingPlane);
 
-	return nullptr;
+	// Saving camera parameters
+	CameraFOVDegrees = InOutViewInfo.FOV;
+	CameraAspectRatio = InOutViewInfo.AspectRatio;
 }
 
-UCameraComponent* FDisplayClusterProjectionCameraPolicy::GetCameraComponent()
+UCameraComponent* FDisplayClusterProjectionCameraPolicy::GetCameraComponent() const
 {
-	USceneComponent* SceneComponent = CameraRef.GetOrFindSceneComponent();
-	if (SceneComponent)
+	if (USceneComponent* SceneComponent = CameraRef.GetOrFindSceneComponent())
 	{
-		UCameraComponent* CameraComponent = Cast<UCameraComponent>(SceneComponent);
-		if (CameraComponent)
-		{
-			return CameraComponent;
-		}
+		return Cast<UCameraComponent>(SceneComponent);
 	}
-	
-	return  nullptr;
+
+	return nullptr;
 }
 
 bool FDisplayClusterProjectionCameraPolicy::CalculateView(IDisplayClusterViewport* InViewport, const uint32 InContextNum, FVector& InOutViewLocation, FRotator& InOutViewRotation, const FVector& ViewOffset, const float WorldToMeters, const float NCP, const float FCP)
 {
 	check(IsInGameThread());
 
-	// The camera position has already been determined from the GetViewPoint() function
+	// Camera location already defined from SetupProjectionViewPoint()
 
-	// Save Z values
+	// Saving clipping planes
 	ZNear = NCP;
-	ZFar  = FCP;
-
-	if (CustomNearClippingPlane >= 0)
-	{
-		ZNear = CustomNearClippingPlane;
-		ZFar = (NCP == FCP) ? CustomNearClippingPlane : FCP;
-	}
+	ZFar = FCP;
 
 	return true;
 }
 
-bool FDisplayClusterProjectionCameraPolicy::ImplGetProjectionMatrix(const float InCameraFOV, const float InCameraAspectRatio, IDisplayClusterViewport* InViewport, const uint32 InContextNum, FMatrix& OutPrjMatrix)
+
+bool FDisplayClusterProjectionCameraPolicy::GetProjectionMatrix(IDisplayClusterViewport* InViewport, const uint32 InContextNum, FMatrix& OutPrjMatrix)
 {
+	check(IsInGameThread());
+
 	// The horizontal field of view (in degrees)
-	const float ScaledCameraFOV = InCameraFOV * CameraSettings.FOVMultiplier;
+	const float ScaledCameraFOVDegrees = CameraFOVDegrees * CameraSettings.FOVMultiplier;
 
 	// Clamp camera fov to valid range [1.f, 178.f]
-	const float ClampedCameraFOV = FMath::Clamp(ScaledCameraFOV, 1.f, 178.f);
-	if (ClampedCameraFOV != ScaledCameraFOV && !IsEditorOperationMode(InViewport))
+	const float ClampedCameraFOVDegrees = FMath::Clamp(ScaledCameraFOVDegrees, 1.f, 178.f);
+	if (ClampedCameraFOVDegrees != ScaledCameraFOVDegrees && !IsEditorOperationMode(InViewport))
 	{
-		UE_LOG(LogDisplayClusterProjectionCamera, Warning, TEXT("CameraFOV clamped: '%f' -> '%f'. (FieldOfView='%f', FOVMultiplier='%f'"), ScaledCameraFOV, ClampedCameraFOV, InCameraFOV, CameraSettings.FOVMultiplier);
+		UE_LOG(LogDisplayClusterProjectionCamera, Warning, TEXT("CameraFOV clamped: '%f' -> '%f'. (FieldOfView='%f', FOVMultiplier='%f'"), ScaledCameraFOVDegrees, ClampedCameraFOVDegrees, CameraFOVDegrees, CameraSettings.FOVMultiplier);
 	}
 
 	if (InViewport)
 	{
 		// Support inner camera custom frustum
-		const float HalfFOVH = 0.5 * ClampedCameraFOV;
-		const float HalfFOVV = HalfFOVH / InCameraAspectRatio;
+		const float HalfFOVH = 0.5 * ClampedCameraFOVDegrees;
+		const float HalfFOVV = HalfFOVH / CameraAspectRatio;
 
 		InViewport->CalculateProjectionMatrix(InContextNum, -HalfFOVH, HalfFOVH, HalfFOVV, -HalfFOVV, ZNear, ZFar, true);
 		OutPrjMatrix = InViewport->GetContexts()[InContextNum].ProjectionMatrix;
@@ -119,82 +127,42 @@ bool FDisplayClusterProjectionCameraPolicy::ImplGetProjectionMatrix(const float
 	else
 	{
 		FComposurePostMoveSettings ComposureSettings;
-		OutPrjMatrix = ComposureSettings.GetProjectionMatrix(ClampedCameraFOV, InCameraAspectRatio);
+		OutPrjMatrix = ComposureSettings.GetProjectionMatrix(ClampedCameraFOVDegrees, CameraAspectRatio);
 	}
 
 	return true;
 }
 
-bool FDisplayClusterProjectionCameraPolicy::GetProjectionMatrix(IDisplayClusterViewport* InViewport, const uint32 InContextNum, FMatrix& OutPrjMatrix)
+//////////////////////////////////////////////////////////////////////////////////////////////
+// FDisplayClusterProjectionCameraPolicy
+//////////////////////////////////////////////////////////////////////////////////////////////
+bool FDisplayClusterProjectionCameraPolicy::ImplSetupProjectionViewPoint(IDisplayClusterViewport* InViewport, const float InDeltaTime, FMinimalViewInfo& InOutViewInfo, float* OutCustomNearClippingPlane) const
 {
-	check(IsInGameThread());
-
-	if (UCameraComponent* CameraComponent = GetCameraComponent())
+	if (OutCustomNearClippingPlane)
 	{
-		return ImplGetProjectionMatrix(CameraComponent->FieldOfView, CameraComponent->AspectRatio, InViewport, InContextNum, OutPrjMatrix);
+		*OutCustomNearClippingPlane = -1;
 	}
-	else
-	{
-		APlayerCameraManager* const CurPlayerCameraManager = GetCurPlayerCameraManager(InViewport);
-		if (CurPlayerCameraManager)
-		{
-			return ImplGetProjectionMatrix(CurPlayerCameraManager->GetFOVAngle(), CurPlayerCameraManager->DefaultAspectRatio, InViewport, InContextNum, OutPrjMatrix);
-		}
-	}
-
-	return false;
-}
-
-bool FDisplayClusterProjectionCameraPolicy::GetViewPoint(IDisplayClusterViewport* InViewport, FRotator& InOutViewRotation, FVector& InOutViewLocation)
-{
-	InOutViewLocation = FVector::ZeroVector;
-	InOutViewRotation = FRotator::ZeroRotator;
-
-	CustomNearClippingPlane = -1;
 
-	// Use assigned camera
+	bool bResult = false;
 	if (UCameraComponent* CameraComponent = GetCameraComponent())
 	{
-		if (UCineCameraComponent* CineCameraComponent = Cast<UCineCameraComponent>(CameraComponent))
-		{
-			if (CineCameraComponent->bOverride_CustomNearClippingPlane)
-			{
-				CustomNearClippingPlane = CineCameraComponent->CustomNearClippingPlane;
-			}
-		}
-
-		InOutViewLocation = CameraComponent->GetComponentLocation();
-		InOutViewRotation = CameraComponent->GetComponentRotation();
+		// Use assigned camera component
+		// Store CustomNearClippingPlane locally, and then use that value for the projection matrix
+		bResult = IDisplayClusterViewport::GetCameraComponentView(CameraComponent, InDeltaTime, CameraSettings.bUseCameraPostprocess, InOutViewInfo, OutCustomNearClippingPlane);
 	}
-	// Otherwise default UE camera is used
-	else
+	else if (UWorld* CurrentWorld = InViewport ? InViewport->GetCurrentWorld() : nullptr)
 	{
-		APlayerCameraManager* const CurPlayerCameraManager = GetCurPlayerCameraManager(InViewport);
-		if (CurPlayerCameraManager)
-		{
-			InOutViewLocation = CurPlayerCameraManager->GetCameraLocation();
-			InOutViewRotation = CurPlayerCameraManager->GetCameraRotation();
-		}
+		// Get active player camera
+		bResult = IDisplayClusterViewport::GetPlayerCameraView(CurrentWorld, CameraSettings.bUseCameraPostprocess, InOutViewInfo);
 	}
 
 	// Fix camera lens deffects (prototype)
-	InOutViewLocation += CameraSettings.FrustumOffset;
-	InOutViewRotation += CameraSettings.FrustumRotation;
+	InOutViewInfo.Location += CameraSettings.FrustumOffset;
+	InOutViewInfo.Rotation += CameraSettings.FrustumRotation;
 
-	return true;
+	return bResult;
 }
 
-bool FDisplayClusterProjectionCameraPolicy::GetStereoEyeOffsetDistance(IDisplayClusterViewport* InViewport, const uint32 InContextNum, float& InOutStereoEyeOffsetDistance)
-{
-	// Stereo is not supported by this policy
-	InOutStereoEyeOffsetDistance = 0.f;
-
-	return true;
-}
-
-//////////////////////////////////////////////////////////////////////////////////////////////
-// FDisplayClusterProjectionCameraPolicy
-//////////////////////////////////////////////////////////////////////////////////////////////
 void FDisplayClusterProjectionCameraPolicy::SetCamera(UCameraComponent* NewCamera, const FDisplayClusterProjectionCameraPolicySettings& InCameraSettings)
 {
 	if(CameraSettings.bCameraOverrideDefaults && InCameraSettings.bCameraOverrideDefaults == false)
diff --git a/Engine/Plugins/Runtime/nDisplay/Source/DisplayClusterProjection/Private/Policy/Camera/DisplayClusterProjectionCameraPolicy.h b/Engine/Plugins/Runtime/nDisplay/Source/DisplayClusterProjection/Private/Policy/Camera/DisplayClusterProjectionCameraPolicy.h
--- a/Engine/Plugins/Runtime/nDisplay/Source/DisplayClusterProjection/Private/Policy/Camera/DisplayClusterProjectionCameraPolicy.h
+++ b/Engine/Plugins/Runtime/nDisplay/Source/DisplayClusterProjection/Private/Policy/Camera/DisplayClusterProjectionCameraPolicy.h
@@ -29,9 +29,9 @@ public:
 
 	virtual bool CalculateView(IDisplayClusterViewport* InViewport, const uint32 InContextNum, FVector& InOutViewLocation, FRotator& InOutViewRotation, const FVector& ViewOffset, const float WorldToMeters, const float NCP, const float FCP) override;
 	virtual bool GetProjectionMatrix(IDisplayClusterViewport* InViewport, const uint32 InContextNum, FMatrix& OutPrjMatrix) override;
+	virtual void SetupProjectionViewPoint(IDisplayClusterViewport* InViewport, const float InDeltaTime, FMinimalViewInfo& InOutViewInfo, float* OutCustomNearClippingPlane = nullptr) override;
 
-	virtual bool GetViewPoint(IDisplayClusterViewport* InViewport, FRotator& InOutViewRotation, FVector& InOutViewLocation) override;
-	virtual bool GetStereoEyeOffsetDistance(IDisplayClusterViewport* InViewport, const uint32 InContextNum, float& InOutStereoEyeOffsetDistance) override;
+	virtual void UpdatePostProcessSettings(IDisplayClusterViewport* InViewport) override;
 
 	virtual bool ShouldUseSourceTextureWithMips() const override
 	{
@@ -42,18 +42,25 @@ public:
 	void SetCamera(UCameraComponent* const NewCamera, const FDisplayClusterProjectionCameraPolicySettings& InCameraSettings);
 
 private:
-	bool ImplGetProjectionMatrix(const float CameraFOV, const float CameraAspectRatio, IDisplayClusterViewport* InViewport, const uint32 InContextNum, FMatrix& OutPrjMatrix);
+	bool ImplSetupProjectionViewPoint(IDisplayClusterViewport* InViewport, const float InDeltaTime, FMinimalViewInfo& InOutViewInfo, float* OutCustomNearClippingPlane = nullptr) const;
 
 protected:
-	UCameraComponent* GetCameraComponent();
+	UCameraComponent* GetCameraComponent() const;
 
 private:
 	// Camera to use for rendering
 	FDisplayClusterSceneComponentRef CameraRef;
+
+	// Camera settings
 	FDisplayClusterProjectionCameraPolicySettings CameraSettings;
+
+	// Values of the clipping planes (saved from CalculateView())
 	float ZNear = 1.f;
 	float ZFar = 1.f;
 
-	// A value less than zero means that the variable is not used.
-	float CustomNearClippingPlane = -1;
+	// camera FOV value in degrees (saved from the SetupProjectionViewPoint() function)
+	float CameraFOVDegrees = 90;
+
+	// camera aspect ratio value (saved from the SetupProjectionViewPoint() function)
+	float CameraAspectRatio = 1;
 };
diff --git a/Engine/Plugins/Runtime/nDisplay/Source/DisplayClusterProjection/Private/Policy/DisplayClusterProjectionPolicyBase.h b/Engine/Plugins/Runtime/nDisplay/Source/DisplayClusterProjection/Private/Policy/DisplayClusterProjectionPolicyBase.h
--- a/Engine/Plugins/Runtime/nDisplay/Source/DisplayClusterProjection/Private/Policy/DisplayClusterProjectionPolicyBase.h
+++ b/Engine/Plugins/Runtime/nDisplay/Source/DisplayClusterProjection/Private/Policy/DisplayClusterProjectionPolicyBase.h
@@ -40,6 +40,11 @@ public:
 		PolicyOriginComponentRef.SetSceneComponent(OriginComp);
 	}
 
+	virtual const USceneComponent* const GetOriginComponent() const override
+	{
+		return GetOriginComp();
+	}
+
 	virtual const TMap<FString, FString>& GetParameters() const override
 	{
 		return Parameters;
diff --git a/Engine/Plugins/Runtime/nDisplay/Source/DisplayClusterProjection/Private/Policy/Link/DisplayClusterProjectionLinkPolicy.cpp b/Engine/Plugins/Runtime/nDisplay/Source/DisplayClusterProjection/Private/Policy/Link/DisplayClusterProjectionLinkPolicy.cpp
--- a/Engine/Plugins/Runtime/nDisplay/Source/DisplayClusterProjection/Private/Policy/Link/DisplayClusterProjectionLinkPolicy.cpp
+++ b/Engine/Plugins/Runtime/nDisplay/Source/DisplayClusterProjection/Private/Policy/Link/DisplayClusterProjectionLinkPolicy.cpp
@@ -27,6 +27,21 @@ const FString& FDisplayClusterProjectionLinkPolicy::GetType() const
 	return Type;
 }
 
+void FDisplayClusterProjectionLinkPolicy::SetupProjectionViewPoint(IDisplayClusterViewport* InViewport, const float InDeltaTime, FMinimalViewInfo& InOutViewInfo, float* OutCustomNearClippingPlane)
+{
+	// Getting the right data from a parent
+	if (IDisplayClusterViewportManager* ViewportManager = InViewport->GetViewportManager())
+	{
+		if (IDisplayClusterViewport* ParentViewport = ViewportManager->FindViewport(InViewport->GetRenderSettings().GetParentViewportId()))
+		{
+			if (ParentViewport->GetProjectionPolicy().IsValid())
+			{
+				ParentViewport->GetProjectionPolicy()->SetupProjectionViewPoint(ParentViewport, InDeltaTime, InOutViewInfo, OutCustomNearClippingPlane);
+			}
+		}
+	}
+}
+
 bool FDisplayClusterProjectionLinkPolicy::CalculateView(IDisplayClusterViewport* InViewport, const uint32 InContextNum, FVector& InOutViewLocation, FRotator& InOutViewRotation, const FVector& ViewOffset, const float WorldToMeters, const float NCP, const float FCP)
 {
 	check(IsInGameThread());
diff --git a/Engine/Plugins/Runtime/nDisplay/Source/DisplayClusterProjection/Private/Policy/Link/DisplayClusterProjectionLinkPolicy.h b/Engine/Plugins/Runtime/nDisplay/Source/DisplayClusterProjection/Private/Policy/Link/DisplayClusterProjectionLinkPolicy.h
--- a/Engine/Plugins/Runtime/nDisplay/Source/DisplayClusterProjection/Private/Policy/Link/DisplayClusterProjectionLinkPolicy.h
+++ b/Engine/Plugins/Runtime/nDisplay/Source/DisplayClusterProjection/Private/Policy/Link/DisplayClusterProjectionLinkPolicy.h
@@ -26,6 +26,7 @@ public:
 	// Important note: before doing this, the parent viewports must be updated.
 	virtual bool CalculateView(IDisplayClusterViewport* InViewport, const uint32 InContextNum, FVector& InOutViewLocation, FRotator& InOutViewRotation, const FVector& ViewOffset, const float WorldToMeters, const float NCP, const float FCP) override;
 	virtual bool GetProjectionMatrix(IDisplayClusterViewport* InViewport, const uint32 InContextNum, FMatrix& OutPrjMatrix) override;
+	virtual void SetupProjectionViewPoint(IDisplayClusterViewport* InViewport, const float InDeltaTime, FMinimalViewInfo& InOutViewInfo, float* OutCustomNearClippingPlane = nullptr) override;
 
 	virtual bool ShouldUseSourceTextureWithMips() const override
 	{
diff --git a/Engine/Plugins/Runtime/nDisplay/Source/DisplayClusterProjection/Private/Policy/MPCDI/DisplayClusterProjectionMPCDIPolicy.cpp b/Engine/Plugins/Runtime/nDisplay/Source/DisplayClusterProjection/Private/Policy/MPCDI/DisplayClusterProjectionMPCDIPolicy.cpp
--- a/Engine/Plugins/Runtime/nDisplay/Source/DisplayClusterProjection/Private/Policy/MPCDI/DisplayClusterProjectionMPCDIPolicy.cpp
+++ b/Engine/Plugins/Runtime/nDisplay/Source/DisplayClusterProjection/Private/Policy/MPCDI/DisplayClusterProjectionMPCDIPolicy.cpp
@@ -25,6 +25,7 @@
 #include "WarpBlend/IDisplayClusterWarpBlend.h"
 #include "WarpBlend/IDisplayClusterWarpBlendManager.h"
 
+#include "Render/Warp/IDisplayClusterWarpPolicy.h"
 #include "Render/Viewport/Containers/DisplayClusterViewport_RenderSettingsICVFX.h"
 
 #include "ShaderParameters/DisplayClusterShaderParameters_WarpBlend.h"
@@ -66,6 +67,47 @@ void FDisplayClusterProjectionMPCDIPolicy::UpdateProxyData(IDisplayClusterViewpo
 	});
 }
 
+bool FDisplayClusterProjectionMPCDIPolicy::ShouldSupportICVFX(IDisplayClusterViewport* InViewport) const
+{
+	check(IsInGameThread());
+	check(InViewport);
+
+	if (!WarpBlendInterface.IsValid() || !WarpBlendInterface->ShouldSupportICVFX(InViewport))
+	{
+		// WarpBlend does not support the ICVFX pipeline (mpcdi 2D,3D,SL profiles)
+		return false;
+	}
+
+	if (WarpPolicyInterface.IsValid() && !WarpPolicyInterface->ShouldSupportICVFX(InViewport))
+	{
+		// The warp policy is used for this projection and does not support the ICVFX pipeline
+		return false;
+	}
+
+	return true;
+}
+
+void FDisplayClusterProjectionMPCDIPolicy::SetWarpPolicy(IDisplayClusterWarpPolicy* InWarpPolicy)
+{
+	check(IsInGameThread());
+
+	WarpPolicyInterface = InWarpPolicy ? InWarpPolicy->ToSharedPtr() : nullptr;
+}
+
+IDisplayClusterWarpPolicy* FDisplayClusterProjectionMPCDIPolicy::GetWarpPolicy() const
+{
+	check(IsInGameThread());
+
+	return WarpPolicyInterface.Get();
+}
+
+IDisplayClusterWarpPolicy* FDisplayClusterProjectionMPCDIPolicy::GetWarpPolicy_RenderThread() const
+{
+	check(IsInRenderingThread());
+
+	return WarpPolicyInterface_Proxy.Get();
+}
+
 //////////////////////////////////////////////////////////////////////////////////////////////
 // IDisplayClusterProjectionPolicy
 //////////////////////////////////////////////////////////////////////////////////////////////
@@ -155,7 +197,7 @@ bool FDisplayClusterProjectionMPCDIPolicy::GetWarpBlendInterface_RenderThread(TS
 	return false;
 }
 
-bool FDisplayClusterProjectionMPCDIPolicy::CalculateView(IDisplayClusterViewport* InViewport, const uint32 InContextNum, FVector& InOutViewLocation, FRotator& InOutViewRotation, const FVector& ViewOffset, const float WorldToMeters, const float NCP, const float FCP)
+bool FDisplayClusterProjectionMPCDIPolicy::CalculateView(IDisplayClusterViewport* InViewport, const uint32 InContextNum, FVector& InOutViewLocation, FRotator& InOutViewRotation, const FVector& InViewOffset, const float WorldToMeters, const float NCP, const float FCP)
 {
 	check(IsInGameThread());
 
@@ -169,6 +211,14 @@ bool FDisplayClusterProjectionMPCDIPolicy::CalculateView(IDisplayClusterViewport
 		return false;
 	}
 
+	// Override viewpoint
+	// MPCDI always expects the location of the viewpoint component (eye location from the real world)
+	FVector ViewOffset = FVector::ZeroVector;
+	if (!InViewport || !InViewport->GetViewPointCameraEye(InContextNum, InOutViewLocation, InOutViewRotation, ViewOffset))
+	{
+		return false;
+	}
+
 	// World scale multiplier
 	const float WorldScale = WorldToMeters / 100.f;
 
diff --git a/Engine/Plugins/Runtime/nDisplay/Source/DisplayClusterProjection/Private/Policy/MPCDI/DisplayClusterProjectionMPCDIPolicy.h b/Engine/Plugins/Runtime/nDisplay/Source/DisplayClusterProjection/Private/Policy/MPCDI/DisplayClusterProjectionMPCDIPolicy.h
--- a/Engine/Plugins/Runtime/nDisplay/Source/DisplayClusterProjection/Private/Policy/MPCDI/DisplayClusterProjectionMPCDIPolicy.h
+++ b/Engine/Plugins/Runtime/nDisplay/Source/DisplayClusterProjection/Private/Policy/MPCDI/DisplayClusterProjectionMPCDIPolicy.h
@@ -33,15 +33,10 @@ public:
 	}
 
 	// This policy can support ICVFX rendering
-	virtual bool ShouldSupportICVFX() const
-	{
-		return true;
-	}
+	virtual bool ShouldSupportICVFX(IDisplayClusterViewport* InViewport) const override;
 
 public:
-	//////////////////////////////////////////////////////////////////////////////////////////////
-	// IDisplayClusterProjectionPolicy
-	//////////////////////////////////////////////////////////////////////////////////////////////
+	//~Begin IDisplayClusterProjectionPolicy
 	virtual const FString& GetType() const override;
 
 	virtual bool HandleStartScene(IDisplayClusterViewport* InViewport) override;
@@ -70,15 +65,23 @@ public:
 
 	virtual void UpdateProxyData(IDisplayClusterViewport* InViewport) override;
 
+	virtual void SetWarpPolicy(IDisplayClusterWarpPolicy* InWarpPolicy) override;
+	virtual IDisplayClusterWarpPolicy* GetWarpPolicy() const override;
+	virtual IDisplayClusterWarpPolicy* GetWarpPolicy_RenderThread() const override;
+
+	//~~End IDisplayClusterProjectionPolicy
+
 protected:
 	bool CreateWarpBlendFromConfig(IDisplayClusterViewport* InViewport);
 	void ImplRelease();
 
 protected:
 	TSharedPtr<IDisplayClusterWarpBlend, ESPMode::ThreadSafe> WarpBlendInterface;
+	TSharedPtr<IDisplayClusterWarpPolicy, ESPMode::ThreadSafe> WarpPolicyInterface;
 	TArray<FDisplayClusterWarpContext> WarpBlendContexts;
 
 	TSharedPtr<IDisplayClusterWarpBlend, ESPMode::ThreadSafe> WarpBlendInterface_Proxy;
+	TSharedPtr<IDisplayClusterWarpPolicy, ESPMode::ThreadSafe> WarpPolicyInterface_Proxy;
 	TArray<FDisplayClusterWarpContext> WarpBlendContexts_Proxy;
 
 private:
diff --git a/Engine/Plugins/Runtime/nDisplay/Source/DisplayClusterProjection/Private/Policy/Manual/DisplayClusterProjectionManualPolicy.cpp b/Engine/Plugins/Runtime/nDisplay/Source/DisplayClusterProjection/Private/Policy/Manual/DisplayClusterProjectionManualPolicy.cpp
--- a/Engine/Plugins/Runtime/nDisplay/Source/DisplayClusterProjection/Private/Policy/Manual/DisplayClusterProjectionManualPolicy.cpp
+++ b/Engine/Plugins/Runtime/nDisplay/Source/DisplayClusterProjection/Private/Policy/Manual/DisplayClusterProjectionManualPolicy.cpp
@@ -126,13 +126,18 @@ bool FDisplayClusterProjectionManualPolicy::HandleStartScene(IDisplayClusterView
 	return true;
 }
 
+void FDisplayClusterProjectionManualPolicy::SetupProjectionViewPoint(IDisplayClusterViewport* InViewport, const float InDeltaTime, FMinimalViewInfo& InOutViewInfo, float* OutCustomNearClippingPlane)
+{
+	// Add local rotation specified in config
+	InOutViewInfo.Rotation += ViewRotation;
+}
+
 bool FDisplayClusterProjectionManualPolicy::CalculateView(IDisplayClusterViewport* InViewport, const uint32 InContextNum, FVector& InOutViewLocation, FRotator& InOutViewRotation, const FVector& InViewOffset, const float InWorldToMeters, const float InNCP, const float InFCP)
 {
 	check(IsInGameThread());
 	check(InContextNum < 2);
 
-	// Add local rotation specified in config
-	InOutViewRotation += ViewRotation;
+	// View changes are already defined from SetupProjectionViewPoint()
 
 	// Store culling data
 	NCP = InNCP;
diff --git a/Engine/Plugins/Runtime/nDisplay/Source/DisplayClusterProjection/Private/Policy/Manual/DisplayClusterProjectionManualPolicy.h b/Engine/Plugins/Runtime/nDisplay/Source/DisplayClusterProjection/Private/Policy/Manual/DisplayClusterProjectionManualPolicy.h
--- a/Engine/Plugins/Runtime/nDisplay/Source/DisplayClusterProjection/Private/Policy/Manual/DisplayClusterProjectionManualPolicy.h
+++ b/Engine/Plugins/Runtime/nDisplay/Source/DisplayClusterProjection/Private/Policy/Manual/DisplayClusterProjectionManualPolicy.h
@@ -34,6 +34,7 @@ public:
 	virtual bool HandleStartScene(IDisplayClusterViewport* InViewport) override;
 	virtual bool CalculateView(IDisplayClusterViewport* InViewport, const uint32 InContextNum, FVector& InOutViewLocation, FRotator& InOutViewRotation, const FVector& InViewOffset, const float InWorldToMeters, const float InNCP, const float InFCP) override;
 	virtual bool GetProjectionMatrix(IDisplayClusterViewport* InViewport, const uint32 InContextNum, FMatrix& OutPrjMatrix) override;
+	virtual void SetupProjectionViewPoint(IDisplayClusterViewport* InViewport, const float InDeltaTime, FMinimalViewInfo& InOutViewInfo, float* OutCustomNearClippingPlane = nullptr) override;
 
 protected:
 	struct FFrustumAngles
diff --git a/Engine/Plugins/Runtime/nDisplay/Source/DisplayClusterProjection/Public/Containers/DisplayClusterProjectionCameraPolicySettings.h b/Engine/Plugins/Runtime/nDisplay/Source/DisplayClusterProjection/Public/Containers/DisplayClusterProjectionCameraPolicySettings.h
--- a/Engine/Plugins/Runtime/nDisplay/Source/DisplayClusterProjection/Public/Containers/DisplayClusterProjectionCameraPolicySettings.h
+++ b/Engine/Plugins/Runtime/nDisplay/Source/DisplayClusterProjection/Public/Containers/DisplayClusterProjectionCameraPolicySettings.h
@@ -16,4 +16,7 @@ struct FDisplayClusterProjectionCameraPolicySettings
 
 	// When enabled, ignore all next updates, until null cam ref not assigned (UE-137222)
 	bool bCameraOverrideDefaults = false;
+
+	// Allow to use camera postprocess
+	bool bUseCameraPostprocess = true;
 };
diff --git a/Engine/Plugins/Runtime/nDisplay/Source/DisplayClusterProjection/Public/IDisplayClusterProjection.h b/Engine/Plugins/Runtime/nDisplay/Source/DisplayClusterProjection/Public/IDisplayClusterProjection.h
--- a/Engine/Plugins/Runtime/nDisplay/Source/DisplayClusterProjection/Public/IDisplayClusterProjection.h
+++ b/Engine/Plugins/Runtime/nDisplay/Source/DisplayClusterProjection/Public/IDisplayClusterProjection.h
@@ -61,5 +61,5 @@ public:
 	* Set camera policy camera
 	*
 	*/
-	virtual bool CameraPolicySetCamera(const TSharedPtr<class IDisplayClusterProjectionPolicy, ESPMode::ThreadSafe>& InPolicy, class UCameraComponent* const NewCamera, const struct FDisplayClusterProjectionCameraPolicySettings& CamersSettings) = 0;
+	virtual bool CameraPolicySetCamera(const TSharedPtr<class IDisplayClusterProjectionPolicy, ESPMode::ThreadSafe>& InPolicy, class UCameraComponent* const NewCamera, const struct FDisplayClusterProjectionCameraPolicySettings& CameraSettings) = 0;
 };
diff --git a/Engine/Plugins/Runtime/nDisplay/Source/DisplayClusterShaders/Private/WarpBlend/DisplayClusterWarpBlend.h b/Engine/Plugins/Runtime/nDisplay/Source/DisplayClusterShaders/Private/WarpBlend/DisplayClusterWarpBlend.h
--- a/Engine/Plugins/Runtime/nDisplay/Source/DisplayClusterShaders/Private/WarpBlend/DisplayClusterWarpBlend.h
+++ b/Engine/Plugins/Runtime/nDisplay/Source/DisplayClusterShaders/Private/WarpBlend/DisplayClusterWarpBlend.h
@@ -50,6 +50,12 @@ public:
 		return GeometryContext.ProfileType;
 	}
 
+	virtual bool ShouldSupportICVFX(IDisplayClusterViewport* InViewport) const override
+	{
+		// only A3D supports ICVFX
+		return GeometryContext.ProfileType == EDisplayClusterWarpProfileType::warp_A3D;
+	}
+
 	// return current warp resource type
 	virtual EDisplayClusterWarpGeometryType GetWarpGeometryType() const override
 	{
diff --git a/Engine/Plugins/Runtime/nDisplay/Source/DisplayClusterShaders/Public/WarpBlend/IDisplayClusterWarpBlend.h b/Engine/Plugins/Runtime/nDisplay/Source/DisplayClusterShaders/Public/WarpBlend/IDisplayClusterWarpBlend.h
--- a/Engine/Plugins/Runtime/nDisplay/Source/DisplayClusterShaders/Public/WarpBlend/IDisplayClusterWarpBlend.h
+++ b/Engine/Plugins/Runtime/nDisplay/Source/DisplayClusterShaders/Public/WarpBlend/IDisplayClusterWarpBlend.h
@@ -7,6 +7,7 @@
 #include "DisplayClusterWarpEnums.h"
 #include "DisplayClusterWarpContext.h"
 
+class IDisplayClusterViewport;
 class IDisplayClusterRenderTexture;
 class UStaticMeshComponent;
 class USceneComponent;
@@ -26,7 +27,7 @@ public:
 	*
 	* @return - true if the context calculated successfully
 	*/
-	virtual bool CalcFrustumContext(class IDisplayClusterViewport* InViewport, const uint32 InContextNum, const FDisplayClusterWarpEye& InEye, FDisplayClusterWarpContext& OutWarpContext) = 0;
+	virtual bool CalcFrustumContext(IDisplayClusterViewport* InViewport, const uint32 InContextNum, const FDisplayClusterWarpEye& InEye, FDisplayClusterWarpContext& OutWarpContext) = 0;
 
 	// Access to resources
 	virtual class FRHITexture* GetTexture(EDisplayClusterWarpBlendTextureType TextureType) const = 0;
@@ -50,4 +51,7 @@ public:
 	* @return - true, if there is a marked.
 	*/
 	virtual bool MarkWarpGeometryComponentDirty(const FName& InComponentName) = 0;
+
+	/** Returns true if this warpblend supports ICVFX rendering */
+	virtual bool ShouldSupportICVFX(IDisplayClusterViewport* InViewport) const = 0;
 };
diff --git a/Engine/Plugins/VirtualProduction/TextureShare/Source/TextureShareDisplayCluster/Private/PostProcess/TextureSharePostprocess_Frame.cpp b/Engine/Plugins/VirtualProduction/TextureShare/Source/TextureShareDisplayCluster/Private/PostProcess/TextureSharePostprocess_Frame.cpp
--- a/Engine/Plugins/VirtualProduction/TextureShare/Source/TextureShareDisplayCluster/Private/PostProcess/TextureSharePostprocess_Frame.cpp
+++ b/Engine/Plugins/VirtualProduction/TextureShare/Source/TextureShareDisplayCluster/Private/PostProcess/TextureSharePostprocess_Frame.cpp
@@ -25,7 +25,7 @@ void FTextureSharePostprocess::UpdateSupportedViews(IDisplayClusterViewportManag
 	if (InViewportManager)
 	{
 		// Get all existing viewports on DC node
-		for (const TSharedPtr<IDisplayClusterViewport, ESPMode::ThreadSafe>& ViewportIt : InViewportManager->GetViewports())
+		for (const TSharedPtr<IDisplayClusterViewport, ESPMode::ThreadSafe>& ViewportIt : InViewportManager->GetCurrentRenderFrameViewports())
 		{
 			if (ViewportIt.IsValid())
 			{
@@ -44,7 +44,7 @@ void FTextureSharePostprocess::UpdateManualProjectionPolicy(IDisplayClusterViewp
 {
 	if (InViewportManager)
 	{
-		for (const TSharedPtr<IDisplayClusterViewport, ESPMode::ThreadSafe>& ViewportIt : InViewportManager->GetViewports())
+		for (const TSharedPtr<IDisplayClusterViewport, ESPMode::ThreadSafe>& ViewportIt : InViewportManager->GetCurrentRenderFrameViewports())
 		{
 			if (ViewportIt.IsValid())
 			{
@@ -80,7 +80,7 @@ void FTextureSharePostprocess::UpdateViews(IDisplayClusterViewportManager* InVie
 	if (InViewportManager)
 	{
 		//Register all viewport for scene eye data capture by default
-		for (const TSharedPtr<IDisplayClusterViewport, ESPMode::ThreadSafe>& ViewportIt : InViewportManager->GetViewports())
+		for (const TSharedPtr<IDisplayClusterViewport, ESPMode::ThreadSafe>& ViewportIt : InViewportManager->GetCurrentRenderFrameViewports())
 		{
 			if (ViewportIt.IsValid())
 			{
