[nDisplay] Fixed DCRA asset loading and instantiation in packaged builds

diff --git a/Engine/Plugins/Runtime/nDisplay/Source/DisplayCluster/Private/Game/DisplayClusterGameManager.cpp b/Engine/Plugins/Runtime/nDisplay/Source/DisplayCluster/Private/Game/DisplayClusterGameManager.cpp
--- a/Engine/Plugins/Runtime/nDisplay/Source/DisplayCluster/Private/Game/DisplayClusterGameManager.cpp
+++ b/Engine/Plugins/Runtime/nDisplay/Source/DisplayCluster/Private/Game/DisplayClusterGameManager.cpp
@@ -110,25 +110,34 @@ bool FDisplayClusterGameManager::StartScene(UWorld* InWorld)
 			}
 
 			// 2. Spawn the DCRA BP from a corresponding asset
-			UObject* ActorToSpawn = Cast<UObject>(StaticLoadObject(UObject::StaticClass(), NULL, *ConfigData->Info.AssetPath));
-			if (ActorToSpawn)
+
+			// We need to use generated class as it's the only one available in packaged buidls
+			const FString AssetPath = (ConfigData->Info.AssetPath.EndsWith(TEXT("_C")) ?
+				ConfigData->Info.AssetPath :
+				ConfigData->Info.AssetPath + TEXT("_C"));
+
+			if (UClass* ActorClass = Cast<UClass>(StaticLoadObject(UClass::StaticClass(), nullptr, *AssetPath)))
 			{
-				UBlueprint* GeneratedBP = Cast<UBlueprint>(ActorToSpawn);
-				UClass* ClassToSpawn = ActorToSpawn->StaticClass();
-				if (ClassToSpawn && GeneratedBP)
+				// Spawn the actor
+				if (AActor* NewActor = CurrentWorld->SpawnActor<AActor>(ActorClass, StartLocation, StartRotation, FActorSpawnParameters()))
 				{
-					// Spawn an asset
-					AActor* NewActor = CurrentWorld->SpawnActor<AActor>(GeneratedBP->GeneratedClass, StartLocation, StartRotation, FActorSpawnParameters());
 					RootActor = Cast<ADisplayClusterRootActor>(NewActor);
 
-					// Override actor settings in case the config file contains some updates
-					RootActor->OverrideFromConfig(ConfigData);
+					if (RootActor)
+					{
+						UE_LOG(LogDisplayClusterGame, Log, TEXT("Instantiated DCRA from asset '%s'"), *ConfigData->Info.AssetPath);
+
+						// Override actor settings in case the config file contains some updates
+						RootActor->OverrideFromConfig(ConfigData);
+					}
 				}
 			}
 
 			// 3. Still no root actor exists? Spawn the DCRA and initialize it with the config data
 			if (!RootActor)
 			{
+				UE_LOG(LogDisplayClusterGame, Warning, TEXT("Couldn't find a configuration asset '%s'. Initializing default DCRA..."), *ConfigData->Info.AssetPath);
+
 				RootActor = CurrentWorld->SpawnActor<ADisplayClusterRootActor>(ADisplayClusterRootActor::StaticClass(), StartLocation, StartRotation, FActorSpawnParameters());
 				if (RootActor)
 				{
