nDisplay: fixed Listen Server crash caused by nullptr accessigng non existend config manager

diff --git a/Engine/Plugins/Runtime/nDisplay/Source/DisplayClusterReplication/Private/DisplayClusterNetDriver.cpp b/Engine/Plugins/Runtime/nDisplay/Source/DisplayClusterReplication/Private/DisplayClusterNetDriver.cpp
--- a/Engine/Plugins/Runtime/nDisplay/Source/DisplayClusterReplication/Private/DisplayClusterNetDriver.cpp
+++ b/Engine/Plugins/Runtime/nDisplay/Source/DisplayClusterReplication/Private/DisplayClusterNetDriver.cpp
@@ -280,14 +280,19 @@ void UDisplayClusterNetDriver::TickFlush(float DeltaSeconds)
 		return;
 	}
 
-	const UDisplayClusterConfigurationData* ConfigData = IDisplayCluster::Get().GetConfigMgr()->GetConfig();
-	check(ConfigData);
-	const uint32 ClusterNodesAmount = ConfigData->GetNumberOfClusterNodes();
-	check(ClusterNodesAmount > 0);
-
-	if ((World->GetNetMode() == NM_ListenServer) && (SyncConnections.Num() < (static_cast<int32>(ClusterNodesAmount) - 1)))
+	if (World->GetNetMode() == NM_ListenServer)
 	{
-		return;
+		IDisplayClusterConfigManager* ConfigMgr = IDisplayCluster::Get().GetConfigMgr();
+
+		if (ConfigMgr != nullptr)
+		{
+			const UDisplayClusterConfigurationData* ConfigData = ConfigMgr->GetConfig();
+
+			if (SyncConnections.Num() < static_cast<int32>(ConfigData->GetNumberOfClusterNodes() - 1))
+			{
+				return;
+			}
+		}
 	}
 
 	bLastBunchWasAcked = false;
