[nDisplay]
* Fixed nopch compile error

diff --git a/Engine/Plugins/Runtime/nDisplay/Source/DisplayCluster/Public/Components/DisplayClusterCameraComponent.h b/Engine/Plugins/Runtime/nDisplay/Source/DisplayCluster/Public/Components/DisplayClusterCameraComponent.h
--- a/Engine/Plugins/Runtime/nDisplay/Source/DisplayCluster/Public/Components/DisplayClusterCameraComponent.h
+++ b/Engine/Plugins/Runtime/nDisplay/Source/DisplayCluster/Public/Components/DisplayClusterCameraComponent.h
@@ -12,6 +12,7 @@ class UBillboardComponent;
 class UTexture2D;
 class IDisplayClusterViewportManager;
 class IDisplayClusterWarpPolicy;
+struct FMinimalViewInfo;
 
 UENUM()
 enum class EDisplayClusterEyeStereoOffset : uint8
