nDisplay: reworked the DCViewport texture resource implementation
- cleanup the code and changing the names according to the code standard
- Rework the DC resource implementation with TSharedPtr.
- New resource types added (retargeted by internal rules)
	* Inside the nDisplay rendering pipeline, internal resources are used differently, depending on configuration and access times.
	* As a result, the user does not know what type of internal resource stores the desired content.
	* Therefore, we need the ability to externally query the resource by content (e.g., beforewarp, afterwarped, output, etc.).
-  resource containers were redesigned into a single class of viewport resource repository.

diff --git a/Engine/Plugins/Runtime/nDisplay/Source/DisplayCluster/Private/Game/EngineClasses/Scene/DisplayClusterPreviewComponent.cpp b/Engine/Plugins/Runtime/nDisplay/Source/DisplayCluster/Private/Game/EngineClasses/Scene/DisplayClusterPreviewComponent.cpp
--- a/Engine/Plugins/Runtime/nDisplay/Source/DisplayCluster/Private/Game/EngineClasses/Scene/DisplayClusterPreviewComponent.cpp
+++ b/Engine/Plugins/Runtime/nDisplay/Source/DisplayCluster/Private/Game/EngineClasses/Scene/DisplayClusterPreviewComponent.cpp
@@ -18,6 +18,7 @@
 #include "Render/Viewport/Containers/DisplayClusterViewportReadPixels.h"
 #include "Render/Viewport/RenderFrame/DisplayClusterRenderFrameSettings.h"
 #include "Render/Viewport/DisplayClusterViewportHelpers.h"
+#include "Render/Viewport/Resource/DisplayClusterViewportResourceSettings.h"
 
 #include "Materials/Material.h"
 #include "Materials/MaterialInstanceDynamic.h"
@@ -418,9 +419,12 @@ bool UDisplayClusterPreviewComponent::GetPreviewTextureSettings(FIntPoint& OutSi
 			// The viewport size is already capped for RenderSettings
 			if (!Viewport->GetContexts().IsEmpty())
 			{
-				DisplayClusterViewportHelpers::GetPreviewRenderTargetDesc_Editor(ViewportManager->GetRenderFrameSettings(), OutTextureFormat, OutGamma, bOutSRGB);
+				FDisplayClusterViewportResourceSettings DefaultResourceSettings(ViewportManager->GetRenderFrameSettings(), nullptr);
 
 				OutSize = Viewport->GetContexts()[0].FrameTargetRect.Size();
+				OutTextureFormat = DefaultResourceSettings.GetFormat();
+				OutGamma = DefaultResourceSettings.GetDisplayGamma();
+				bOutSRGB = EnumHasAnyFlags(DefaultResourceSettings.GetResourceFlags(), EDisplayClusterViewportResourceSettingsFlags::ShouldUseSRGB);
 
 				check(OutSize.X > 0);
 				check(OutSize.Y > 0);
diff --git a/Engine/Plugins/Runtime/nDisplay/Source/DisplayCluster/Private/Render/Containers/DisplayClusterRender_MeshGeometry.cpp b/Engine/Plugins/Runtime/nDisplay/Source/DisplayCluster/Private/Render/Containers/DisplayClusterRender_MeshGeometry.cpp
--- a/Engine/Plugins/Runtime/nDisplay/Source/DisplayCluster/Private/Render/Containers/DisplayClusterRender_MeshGeometry.cpp
+++ b/Engine/Plugins/Runtime/nDisplay/Source/DisplayCluster/Private/Render/Containers/DisplayClusterRender_MeshGeometry.cpp
@@ -7,7 +7,7 @@
 #include "Misc/FileHelper.h"
 #include "Misc/DisplayClusterLog.h"
 
-namespace ObjMeshStrings
+namespace UE::DisplayCluster::MeshGeometryStrings
 {
 	static constexpr auto Vertex = TEXT("v ");
 	static constexpr auto VertexNormal = TEXT("vn ");
@@ -20,8 +20,11 @@ namespace ObjMeshStrings
 		static constexpr auto Face = TEXT("/");
 	}
 };
+using namespace UE::DisplayCluster;
 
-
+/**
+ * OBJ mesh loader class
+ */
 class FDisplayCluster_MeshGeometryLoaderOBJ
 {
 public:
@@ -206,19 +209,19 @@ bool FDisplayCluster_MeshGeometryLoaderOBJ::SaveToTarget()
 
 bool FDisplayCluster_MeshGeometryLoaderOBJ::ParseLine(const FString& Line)
 {
-	if (Line.StartsWith(ObjMeshStrings::Vertex))
+	if (Line.StartsWith(MeshGeometryStrings::Vertex))
 	{
 		return ExtractVertex(Line);
 	}
-	else if (bImportVertexNormal && Line.StartsWith(ObjMeshStrings::VertexNormal))
+	else if (bImportVertexNormal && Line.StartsWith(MeshGeometryStrings::VertexNormal))
 	{
 		return ExtractVertexNormal(Line);
 	}
-	else if (Line.StartsWith(ObjMeshStrings::UV))
+	else if (Line.StartsWith(MeshGeometryStrings::UV))
 	{
 		return ExtractUV(Line);
 	}
-	else if (Line.StartsWith(ObjMeshStrings::Face))
+	else if (Line.StartsWith(MeshGeometryStrings::Face))
 	{
 		return ExtractFace(Line);
 	}
@@ -230,7 +233,7 @@ bool FDisplayCluster_MeshGeometryLoaderOBJ::ExtractVertex(const FString& Line)
 {
 	TArray<FString> Data;
 
-	if (Line.ParseIntoArray(Data, ObjMeshStrings::delims::Values) == 4)
+	if (Line.ParseIntoArray(Data, MeshGeometryStrings::delims::Values) == 4)
 	{
 		const float X = FCString::Atof(*Data[1]);
 		const float Y = FCString::Atof(*Data[2]);
@@ -248,7 +251,7 @@ bool FDisplayCluster_MeshGeometryLoaderOBJ::ExtractVertexNormal(const FString& L
 {
 	TArray<FString> Data;
 
-	if (Line.ParseIntoArray(Data, ObjMeshStrings::delims::Values) == 4)
+	if (Line.ParseIntoArray(Data, MeshGeometryStrings::delims::Values) == 4)
 	{
 		const float X = FCString::Atof(*Data[1]);
 		const float Y = FCString::Atof(*Data[2]);
@@ -265,7 +268,7 @@ bool FDisplayCluster_MeshGeometryLoaderOBJ::ExtractVertexNormal(const FString& L
 bool FDisplayCluster_MeshGeometryLoaderOBJ::ExtractUV(const FString& Line)
 {
 	TArray<FString> Data;
-	int32 Count = Line.ParseIntoArray(Data, ObjMeshStrings::delims::Values);
+	int32 Count = Line.ParseIntoArray(Data, MeshGeometryStrings::delims::Values);
 	if (Count > 2)
 	{
 		const float U = FCString::Atof(*Data[1]);
@@ -283,7 +286,7 @@ bool FDisplayCluster_MeshGeometryLoaderOBJ::ExtractUV(const FString& Line)
 bool FDisplayCluster_MeshGeometryLoaderOBJ::ExtractFace(const FString& Line)
 {
 	TArray<FString> Data;
-	Line.ParseIntoArray(Data, ObjMeshStrings::delims::Values);
+	Line.ParseIntoArray(Data, MeshGeometryStrings::delims::Values);
 
 	if (Data.Num() > 3)
 	{
@@ -326,7 +329,7 @@ int32 FDisplayCluster_MeshGeometryLoaderOBJ::ExtractFaceVertex(const FString& Li
 	FFaceIdx OutFaceIdx;
 
 	TArray<FString> Data;
-	if (Line.ParseIntoArray(Data, ObjMeshStrings::delims::Face) > 1)
+	if (Line.ParseIntoArray(Data, MeshGeometryStrings::delims::Face) > 1)
 	{
 		const int32 InVertexIndex = FCString::Atoi(*Data[0]) - 1;
 		if (InVertexIndex < 0 || InVertexIndex >= InVertex.Num())
diff --git a/Engine/Plugins/Runtime/nDisplay/Source/DisplayCluster/Private/Render/Viewport/Configuration/DisplayClusterViewportConfigurationHelpers_ICVFX.cpp b/Engine/Plugins/Runtime/nDisplay/Source/DisplayCluster/Private/Render/Viewport/Configuration/DisplayClusterViewportConfigurationHelpers_ICVFX.cpp
--- a/Engine/Plugins/Runtime/nDisplay/Source/DisplayCluster/Private/Render/Viewport/Configuration/DisplayClusterViewportConfigurationHelpers_ICVFX.cpp
+++ b/Engine/Plugins/Runtime/nDisplay/Source/DisplayCluster/Private/Render/Viewport/Configuration/DisplayClusterViewportConfigurationHelpers_ICVFX.cpp
@@ -56,7 +56,7 @@ static FAutoConsoleVariableRef CVarDisplayClusterPreviewEnableReuseViewportInClu
 );
 
 ////////////////////////////////////////////////////////////////////////////////
-namespace DisplayClusterViewportConfigurationHelpers_ICVFX_Impl
+namespace UE::DisplayCluster::Viewport::ConfigurationHelpers_ICVFX
 {
 	// Initialize camera policy with camera component and settings
 	static bool ImplUpdateCameraProjectionSettings(TSharedPtr<IDisplayClusterProjectionPolicy, ESPMode::ThreadSafe>& InOutCameraProjection, ADisplayClusterRootActor& RootActor, const FDisplayClusterConfigurationICVFX_CameraSettings& CameraSettings, UCameraComponent* const CameraComponent)
@@ -97,8 +97,7 @@ namespace DisplayClusterViewportConfigurationHelpers_ICVFX_Impl
 	}
 #endif
 };
-
-using namespace DisplayClusterViewportConfigurationHelpers_ICVFX_Impl;
+using namespace UE::DisplayCluster::Viewport::ConfigurationHelpers_ICVFX;
 
 ////////////////////////////////////////////////////////////////////////////////
 // FDisplayClusterViewportConfigurationHelpers_ICVFX
@@ -221,7 +220,7 @@ void FDisplayClusterViewportConfigurationHelpers_ICVFX::ReuseUVLightCardViewport
 	{
 		for (const TSharedPtr<FDisplayClusterViewport, ESPMode::ThreadSafe>& ViewportIt : ViewportManager->ImplGetEntireClusterViewports())
 		{
-			if (ViewportIt.IsValid() && ViewportIt != InUVLightCardViewport.AsShared() && !ViewportIt->RenderSettings.IsViewportOverrided()
+			if (ViewportIt.IsValid() && ViewportIt != InUVLightCardViewport.AsShared() && !ViewportIt->RenderSettings.IsViewportOverridden()
 				&& EnumHasAnyFlags(ViewportIt->RenderSettingsICVFX.RuntimeFlags, EDisplayClusterViewportRuntimeICVFXFlags::UVLightcard))
 			{
 				if (ViewportIt->IsOpenColorIOEquals(InUVLightCardViewport))
@@ -251,8 +250,8 @@ TArray<TSharedPtr<FDisplayClusterViewport, ESPMode::ThreadSafe>> FDisplayCluster
 		{
 			if (ViewportIt.IsValid()
 				&& EnumHasAnyFlags(ViewportIt->RenderSettingsICVFX.RuntimeFlags, InRuntimeFlagsMask)
-				&& (ViewportIt->InputShaderResources.Num() > 0 && ViewportIt->InputShaderResources[0] != nullptr && ViewportIt->Contexts.Num() > 0)
-				&& (!ViewportIt->RenderSettings.IsViewportOverrided()))
+				&& (ViewportIt->Resources[EDisplayClusterViewportResource::InputShaderResources].Num() > 0 && ViewportIt->Resources[EDisplayClusterViewportResource::InputShaderResources][0] != nullptr && ViewportIt->Contexts.Num() > 0)
+				&& (!ViewportIt->RenderSettings.IsViewportOverridden()))
 			{
 				// this is incamera viewport. Check by name
 				const FString RequiredViewportId = ImplGetNameICVFX(ViewportIt->GetClusterNodeId(), ICVFXCameraId, ViewportTypeId);
diff --git a/Engine/Plugins/Runtime/nDisplay/Source/DisplayCluster/Private/Render/Viewport/Containers/DisplayClusterViewportProxyData.cpp b/Engine/Plugins/Runtime/nDisplay/Source/DisplayCluster/Private/Render/Viewport/Containers/DisplayClusterViewportProxyData.cpp
--- a/Engine/Plugins/Runtime/nDisplay/Source/DisplayCluster/Private/Render/Viewport/Containers/DisplayClusterViewportProxyData.cpp
+++ b/Engine/Plugins/Runtime/nDisplay/Source/DisplayCluster/Private/Render/Viewport/Containers/DisplayClusterViewportProxyData.cpp
@@ -29,21 +29,8 @@ FDisplayClusterViewportProxyData::FDisplayClusterViewportProxyData(const TShared
 	ProjectionPolicy = SrcViewport->ProjectionPolicy;
 	Contexts         = SrcViewport->Contexts;
 
-	// Save resources ptrs into container
-	RenderTargets    = SrcViewport->RenderTargets;
-
-	OutputFrameTargetableResources     = SrcViewport->OutputFrameTargetableResources;
-	AdditionalFrameTargetableResources = SrcViewport->AdditionalFrameTargetableResources;
-
-
-#if WITH_EDITOR
-	OutputPreviewTargetableResource = SrcViewport->OutputPreviewTargetableResource;
+	Resources = SrcViewport->Resources;
 	ViewStates = SrcViewport->ViewStates;
-#endif
-
-	InputShaderResources = SrcViewport->InputShaderResources;
-	AdditionalTargetableResources = SrcViewport->AdditionalTargetableResources;
-	MipsShaderResources = SrcViewport->MipsShaderResources;
 }
 
 void FDisplayClusterViewportProxyData::UpdateProxy_RenderThread() const
@@ -78,17 +65,6 @@ void FDisplayClusterViewportProxyData::UpdateProxy_RenderThread() const
 	}
 
 	// Update viewport proxy resources from container
-	DstViewportProxy->RenderTargets    = RenderTargets;
-
-	DstViewportProxy->OutputFrameTargetableResources = OutputFrameTargetableResources;
-	DstViewportProxy->AdditionalFrameTargetableResources = AdditionalFrameTargetableResources;
-
-#if WITH_EDITOR
-	DstViewportProxy->OutputPreviewTargetableResource = OutputPreviewTargetableResource;
+	DstViewportProxy->Resources = Resources;
 	DstViewportProxy->ViewStates = ViewStates;
-#endif
-
-	DstViewportProxy->InputShaderResources = InputShaderResources;
-	DstViewportProxy->AdditionalTargetableResources = AdditionalTargetableResources;
-	DstViewportProxy->MipsShaderResources = MipsShaderResources;
 }
diff --git a/Engine/Plugins/Runtime/nDisplay/Source/DisplayCluster/Private/Render/Viewport/Containers/DisplayClusterViewportProxyData.h b/Engine/Plugins/Runtime/nDisplay/Source/DisplayCluster/Private/Render/Viewport/Containers/DisplayClusterViewportProxyData.h
--- a/Engine/Plugins/Runtime/nDisplay/Source/DisplayCluster/Private/Render/Viewport/Containers/DisplayClusterViewportProxyData.h
+++ b/Engine/Plugins/Runtime/nDisplay/Source/DisplayCluster/Private/Render/Viewport/Containers/DisplayClusterViewportProxyData.h
@@ -2,14 +2,14 @@
 
 #pragma once
 
+#include "Render/Viewport/DisplayClusterViewportResources.h"
 #include "Render/Viewport/Containers/DisplayClusterViewport_RenderSettings.h"
 #include "Render/Viewport/Containers/DisplayClusterViewport_RenderSettingsICVFX.h"
 #include "Render/Viewport/Containers/DisplayClusterViewport_PostRenderSettings.h"
 #include "Render/Viewport/Containers/DisplayClusterViewport_OverscanSettings.h"
 
 class IDisplayClusterProjectionPolicy;
-class FDisplayClusterViewportTextureResource;
-class FDisplayClusterViewportRenderTargetResource;
+class FDisplayClusterViewportResource;
 class FDisplayClusterViewport;
 class FDisplayClusterViewportProxy;
 class FDisplayClusterViewport_Context;
@@ -49,21 +49,9 @@ private:
 	// Viewport contexts (left/center/right eyes)
 	TArray<FDisplayClusterViewport_Context>      Contexts;
 
-	// View family render to this resources
-	TArray<FDisplayClusterViewportRenderTargetResource*> RenderTargets;
+	// Unified repository of viewport resources
+	FDisplayClusterViewportResources Resources;
 
-	// Projection policy output resources
-	TArray<FDisplayClusterViewportTextureResource*> OutputFrameTargetableResources;
-	TArray<FDisplayClusterViewportTextureResource*> AdditionalFrameTargetableResources;
-
-#if WITH_EDITOR
-	FTextureRHIRef OutputPreviewTargetableResource;
+	// Used ViewStates
 	TArray<TSharedPtr<FSceneViewStateReference, ESPMode::ThreadSafe>> ViewStates;
-#endif
-
-	// unique viewport resources
-	TArray<FDisplayClusterViewportTextureResource*> InputShaderResources;
-	TArray<FDisplayClusterViewportTextureResource*> AdditionalTargetableResources;
-	TArray<FDisplayClusterViewportTextureResource*> MipsShaderResources;
 };
-
diff --git a/Engine/Plugins/Runtime/nDisplay/Source/DisplayCluster/Private/Render/Viewport/DisplayClusterViewport.cpp b/Engine/Plugins/Runtime/nDisplay/Source/DisplayCluster/Private/Render/Viewport/DisplayClusterViewport.cpp
--- a/Engine/Plugins/Runtime/nDisplay/Source/DisplayCluster/Private/Render/Viewport/DisplayClusterViewport.cpp
+++ b/Engine/Plugins/Runtime/nDisplay/Source/DisplayCluster/Private/Render/Viewport/DisplayClusterViewport.cpp
@@ -307,16 +307,12 @@ void FDisplayClusterViewport::HandleEndScene()
 		ProjectionPolicy.Reset();
 	}
 
-#if WITH_EDITOR
 	CleanupViewState();
-#endif
 }
 
 void FDisplayClusterViewport::AddReferencedObjects(FReferenceCollector& Collector)
 {
-#if WITH_EDITOR
 	// ViewStates released on rendering thread from viewport proxy object
-#endif
 }
 
 bool FDisplayClusterViewport::ShouldUseAdditionalTargetableResource() const
@@ -437,8 +433,8 @@ inline void AdjustRect(FIntRect& InOutRect, const float multX, const float multY
 FIntRect FDisplayClusterViewport::GetValidRect(const FIntRect& InRect, const TCHAR* DbgSourceName)
 {
 	// The target always needs be within GMaxTextureDimensions, larger dimensions are not supported by the engine
-	const int32 MaxTextureSize = DisplayClusterViewportHelpers::GetMaxTextureDimension();
-	const int32 MinTextureSize = DisplayClusterViewportHelpers::GetMinTextureDimension();
+	const int32 MaxTextureSize = FDisplayClusterViewportHelpers::GetMaxTextureDimension();
+	const int32 MinTextureSize = FDisplayClusterViewportHelpers::GetMinTextureDimension();
 
 	int32 Width  = FMath::Max(MinTextureSize, InRect.Width());
 	int32 Height = FMath::Max(MinTextureSize, InRect.Height());
@@ -458,7 +454,7 @@ FIntRect FDisplayClusterViewport::GetValidRect(const FIntRect& InRect, const TCH
 	OutRect.Min.X = FMath::Min(OutRect.Min.X, MaxTextureSize);
 	OutRect.Min.Y = FMath::Min(OutRect.Min.Y, MaxTextureSize);
 
-	const FIntPoint ScaledRectMax = DisplayClusterViewportHelpers::ScaleTextureSize(OutRect.Max, RectScale);
+	const FIntPoint ScaledRectMax = FDisplayClusterViewportHelpers::ScaleTextureSize(OutRect.Max, RectScale);
 
 	OutRect.Max.X = FMath::Clamp(ScaledRectMax.X, OutRect.Min.X, MaxTextureSize);
 	OutRect.Max.Y = FMath::Clamp(ScaledRectMax.Y, OutRect.Min.Y, MaxTextureSize);
@@ -489,13 +485,13 @@ FIntPoint FDisplayClusterViewport::GetDesiredContextSize(const FIntPoint& InSize
 	const float ClusterRenderTargetRatioMult = GetClusterRenderTargetRatioMult(InFrameSettings);
 
 	// Check size multipliers in order bellow:
-	const float RenderTargetAdaptRatio = DisplayClusterViewportHelpers::GetValidSizeMultiplier(InSize, RenderSettings.RenderTargetAdaptRatio, ClusterRenderTargetRatioMult * RenderSettings.RenderTargetRatio);
-	const float RenderTargetRatio = DisplayClusterViewportHelpers::GetValidSizeMultiplier(InSize, RenderSettings.RenderTargetRatio, ClusterRenderTargetRatioMult * RenderTargetAdaptRatio);
-	const float ClusterMult = DisplayClusterViewportHelpers::GetValidSizeMultiplier(InSize, ClusterRenderTargetRatioMult, RenderTargetRatio * RenderTargetAdaptRatio);
+	const float RenderTargetAdaptRatio = FDisplayClusterViewportHelpers::GetValidSizeMultiplier(InSize, RenderSettings.RenderTargetAdaptRatio, ClusterRenderTargetRatioMult * RenderSettings.RenderTargetRatio);
+	const float RenderTargetRatio = FDisplayClusterViewportHelpers::GetValidSizeMultiplier(InSize, RenderSettings.RenderTargetRatio, ClusterRenderTargetRatioMult * RenderTargetAdaptRatio);
+	const float ClusterMult = FDisplayClusterViewportHelpers::GetValidSizeMultiplier(InSize, ClusterRenderTargetRatioMult, RenderTargetRatio * RenderTargetAdaptRatio);
 
-	FIntPoint DesiredContextSize = DisplayClusterViewportHelpers::ScaleTextureSize(InSize, FMath::Max(RenderTargetAdaptRatio * RenderTargetRatio * ClusterMult, 0.f));
+	FIntPoint DesiredContextSize = FDisplayClusterViewportHelpers::ScaleTextureSize(InSize, FMath::Max(RenderTargetAdaptRatio * RenderTargetRatio * ClusterMult, 0.f));
 
-	const int32 MaxTextureSize = DisplayClusterViewportHelpers::GetMaxTextureDimension();
+	const int32 MaxTextureSize = FDisplayClusterViewportHelpers::GetMaxTextureDimension();
 	DesiredContextSize.X = FMath::Min(DesiredContextSize.X, MaxTextureSize);
 	DesiredContextSize.Y = FMath::Min(DesiredContextSize.Y, MaxTextureSize);
 
@@ -525,23 +521,7 @@ float FDisplayClusterViewport::GetCustomBufferRatio(const FDisplayClusterRenderF
 
 void FDisplayClusterViewport::ResetFrameContexts()
 {
-
-#if WITH_EDITOR
-	OutputPreviewTargetableResource.SafeRelease();
-#endif
-
-	// Discard resources that are not used in frame composition
-	RenderTargets.Empty();
-	OutputFrameTargetableResources.Empty();
-	AdditionalFrameTargetableResources.Empty();
-
-	// Release old contexts
-	Contexts.Empty();
-
-	// Free internal resources
-	InputShaderResources.Empty();
-	AdditionalTargetableResources.Empty();
-	MipsShaderResources.Empty();
+	Resources.Release();
 }
 
 bool FDisplayClusterViewport::UpdateFrameContexts(const uint32 InStereoViewIndex, const FDisplayClusterRenderFrameSettings& InFrameSettings)
@@ -591,24 +571,18 @@ bool FDisplayClusterViewport::UpdateFrameContexts(const uint32 InStereoViewIndex
 	// Special case mono->stereo
 	const uint32 ViewportContextAmount = RenderSettings.bForceMono ? 1 : FrameTargetsAmount;
 
-#if WITH_EDITOR
-	OutputPreviewTargetableResource.SafeRelease();
-#endif
-
-	// Discard resources that are not used in frame composition
-	RenderTargets.Empty();
-	OutputFrameTargetableResources.Empty();
-	AdditionalFrameTargetableResources.Empty();
+	// Release only part of the resources, leaving resources that can be used by other viewports (viewport override feature)
+	Resources.Release(false);
 
 	// Freeze the image in the viewport only after the frame has been rendered
 	if (RenderSettings.bFreezeRendering && RenderSettings.bEnable)
 	{
 		// Freeze only when all resources valid
-		if (Contexts.Num() > 0 && Contexts.Num() == InputShaderResources.Num())
+		if (Contexts.Num() > 0 && Contexts.Num() == Resources[EDisplayClusterViewportResource::InputShaderResources].Num())
 		{
-			DisplayClusterViewportHelpers::FreezeRenderingOfViewportTextureResources(InputShaderResources);
-			DisplayClusterViewportHelpers::FreezeRenderingOfViewportTextureResources(AdditionalTargetableResources);
-			DisplayClusterViewportHelpers::FreezeRenderingOfViewportTextureResources(MipsShaderResources);
+			Resources.FreezeRendering(EDisplayClusterViewportResource::InputShaderResources);
+			Resources.FreezeRendering(EDisplayClusterViewportResource::AdditionalTargetableResources);
+			Resources.FreezeRendering(EDisplayClusterViewportResource::MipsShaderResources);
 
 			// Update context links for freezed viewport
 			for (FDisplayClusterViewport_Context& ContextIt : Contexts)
@@ -626,10 +600,8 @@ bool FDisplayClusterViewport::UpdateFrameContexts(const uint32 InStereoViewIndex
 	// Release old contexts
 	Contexts.Empty();
 
-	// Free internal resources
-	InputShaderResources.Empty();
-	AdditionalTargetableResources.Empty();
-	MipsShaderResources.Empty();
+	// Free all resources
+	Resources.Release();
 
 	if (RenderSettings.bEnable == false)
 	{
@@ -665,7 +637,7 @@ bool FDisplayClusterViewport::UpdateFrameContexts(const uint32 InStereoViewIndex
 	{
 	case EDisplayClusterRenderFrameMode::PreviewInScene:
 		{
-			DesiredContextSize = DisplayClusterViewportHelpers::GetTextureSizeLessThanMax(DesiredContextSize, InFrameSettings.PreviewMaxTextureDimension);
+			DesiredContextSize = FDisplayClusterViewportHelpers::GetTextureSizeLessThanMax(DesiredContextSize, InFrameSettings.PreviewMaxTextureDimension);
 		}
 		break;
 	default:
@@ -694,7 +666,7 @@ bool FDisplayClusterViewport::UpdateFrameContexts(const uint32 InStereoViewIndex
 	const float BaseCustomBufferRatio = GetCustomBufferRatio(InFrameSettings);
 
 	// Fix buffer ratio value vs MaxTextureSize:
-	const float CustomBufferRatio = DisplayClusterViewportHelpers::GetValidSizeMultiplier(RenderTargetRect.Size(), BaseCustomBufferRatio, 1.f);
+	const float CustomBufferRatio = FDisplayClusterViewportHelpers::GetValidSizeMultiplier(RenderTargetRect.Size(), BaseCustomBufferRatio, 1.f);
 
 	// Setup resource usage logic:
 	bool bDisableRender = false;
@@ -710,9 +682,9 @@ bool FDisplayClusterViewport::UpdateFrameContexts(const uint32 InStereoViewIndex
 		bDisableRender = true;
 	}
 
-	if (RenderSettings.IsViewportOverrided())
+	if (RenderSettings.IsViewportOverridden())
 	{
-		switch (RenderSettings.ViewportOverrideMode)
+		switch (RenderSettings.GetViewportOverrideMode())
 		{
 		case EDisplayClusterViewportOverrideMode::InernalRTT:
 			bDisableRender = true;
@@ -828,27 +800,33 @@ bool FDisplayClusterViewport::UpdateFrameContexts(const uint32 InStereoViewIndex
 	}
 
 	// Reserve for resources
+	if (InFrameSettings.RenderMode == EDisplayClusterRenderFrameMode::PreviewInScene)
+	{
+		// reserve preview texture resource
+		Resources[EDisplayClusterViewportResource::OutputPreviewTargetableResources].AddZeroed(FrameTargetsAmount);
+	}
+
 	if (!bDisableRender)
 	{
-		RenderTargets.AddZeroed(FrameTargetsAmount);
+		Resources[EDisplayClusterViewportResource::RenderTargets].AddZeroed(FrameTargetsAmount);
 	}
 
 	if (!bDisableInternalResources)
 	{
-		InputShaderResources.AddZeroed(FrameTargetsAmount);
+		Resources[EDisplayClusterViewportResource::InputShaderResources].AddZeroed(FrameTargetsAmount);
 
 		if (ShouldUseAdditionalTargetableResource())
 		{
-			AdditionalTargetableResources.AddZeroed(FrameTargetsAmount);
+			Resources[EDisplayClusterViewportResource::AdditionalTargetableResources].AddZeroed(FrameTargetsAmount);
 		}
 
 		// Setup Mips resources:
 		for (FDisplayClusterViewport_Context& ContextIt : Contexts)
 		{
-			ContextIt.NumMips = DisplayClusterViewportHelpers::GetMaxTextureNumMips(InFrameSettings, PostRenderSettings.GenerateMips.GetRequiredNumMips(ContextIt.ContextSize));
+			ContextIt.NumMips = FDisplayClusterViewportHelpers::GetMaxTextureNumMips(InFrameSettings, PostRenderSettings.GenerateMips.GetRequiredNumMips(ContextIt.ContextSize));
 			if (ContextIt.NumMips > 1)
 			{
-				MipsShaderResources.AddZeroed(1);
+				Resources[EDisplayClusterViewportResource::MipsShaderResources].AddZeroed(1);
 			}
 		}
 	}
diff --git a/Engine/Plugins/Runtime/nDisplay/Source/DisplayCluster/Private/Render/Viewport/DisplayClusterViewport.h b/Engine/Plugins/Runtime/nDisplay/Source/DisplayCluster/Private/Render/Viewport/DisplayClusterViewport.h
--- a/Engine/Plugins/Runtime/nDisplay/Source/DisplayCluster/Private/Render/Viewport/DisplayClusterViewport.h
+++ b/Engine/Plugins/Runtime/nDisplay/Source/DisplayCluster/Private/Render/Viewport/DisplayClusterViewport.h
@@ -13,14 +13,13 @@
 #include "Render/Viewport/DisplayClusterViewport_CustomPostProcessSettings.h"
 #include "Render/Viewport/DisplayClusterViewport_VisibilitySettings.h"
 #include "Render/Viewport/DisplayClusterViewport_OpenColorIO.h"
+#include "Render/Viewport/DisplayClusterViewportResources.h"
 
 #include "SceneViewExtension.h"
 #include "SceneViewExtensionContext.h"
 
 #include "Templates/SharedPointer.h"
 
-class FDisplayClusterViewportRenderTargetResource;
-class FDisplayClusterViewportTextureResource;
 class FDisplayClusterViewportManager;
 class FDisplayClusterViewportManagerProxy;
 class FDisplayClusterRenderTargetManager;
@@ -168,11 +167,11 @@ public:
 	FDisplayClusterViewportManager* GetViewportManagerImpl() const;
 	FDisplayClusterViewportManagerProxy* GetViewportManagerProxyImpl() const;
 
-#if WITH_EDITOR
 	FSceneView* ImplCalcScenePreview(class FSceneViewFamilyContext& InOutViewFamily, uint32 ContextNum);
 	bool    ImplPreview_CalculateStereoViewOffset(const uint32 InContextNum, FRotator& ViewRotation, const float WorldToMeters, FVector& ViewLocation);
 	FMatrix ImplPreview_GetStereoProjectionMatrix(const uint32 InContextNum);
 
+#if WITH_EDITOR
 	bool GetPreviewPixels(TSharedPtr<class FDisplayClusterViewportReadPixelsData, ESPMode::ThreadSafe>& OutPixelsData) const;
 #endif //WITH_EDITOR
 
@@ -249,18 +248,13 @@ private:
 	FIntPoint GetDesiredContextSize(const FIntPoint& InSize, const FDisplayClusterRenderFrameSettings& InFrameSettings) const;
 	float GetCustomBufferRatio(const FDisplayClusterRenderFrameSettings& InFrameSettings) const;
 
-#if WITH_EDITOR
-	// Support view states for preview
 private:
+	// Support view states for preview
 	FSceneViewStateInterface* GetViewState(uint32 ViewIndex);
 
 public:
 	void CleanupViewState();
 
-private:
-	TArray<TSharedPtr<FSceneViewStateReference, ESPMode::ThreadSafe>> ViewStates;
-#endif
-
 public:
 	/** nDisplay OpenColorIO object. */
 	TSharedPtr<FDisplayClusterViewport_OpenColorIO, ESPMode::ThreadSafe> OpenColorIO;
@@ -282,6 +276,9 @@ public:
 	// viewport OutputRemap feature
 	FDisplayClusterViewportRemap ViewportRemap;
 
+	// Unified repository of viewport resources
+	FDisplayClusterViewportResources Resources;
+
 protected:
 	friend FDisplayClusterViewportProxy;
 	friend FDisplayClusterViewportProxyData;
@@ -318,29 +315,15 @@ protected:
 	// Viewport contexts (left/center/right eyes)
 	TArray<FDisplayClusterViewport_Context> Contexts;
 
-	// View family render to this resources
-	TArray<FDisplayClusterViewportRenderTargetResource*> RenderTargets;
-	
-	// Projection policy output resources
-	TArray<FDisplayClusterViewportTextureResource*> OutputFrameTargetableResources;
-	TArray<FDisplayClusterViewportTextureResource*> AdditionalFrameTargetableResources;
-
-#if WITH_EDITOR
-	friend class UDisplayClusterPreviewComponent;
-
-	FTextureRHIRef OutputPreviewTargetableResource;
-#endif
-
-	// unique viewport resources
-	TArray<FDisplayClusterViewportTextureResource*> InputShaderResources;
-	TArray<FDisplayClusterViewportTextureResource*> AdditionalTargetableResources;
-	TArray<FDisplayClusterViewportTextureResource*> MipsShaderResources;
-
 	// viewport owners
 	TWeakPtr<FDisplayClusterViewportManager, ESPMode::ThreadSafe> ViewportManagerWeakPtr;
 	TWeakPtr<FDisplayClusterViewportManagerProxy, ESPMode::ThreadSafe> ViewportManagerProxyWeakPtr;
 
 private:
+	// View states (preview only)
+	TArray<TSharedPtr<FSceneViewStateReference, ESPMode::ThreadSafe>> ViewStates;
+
+	// Auxiliary variable for the log
 	bool bProjectionPolicyCalculateViewWarningOnce = false;
 
 	// Near clipping plane value (obtained from the GetDesiredView() functions).
diff --git a/Engine/Plugins/Runtime/nDisplay/Source/DisplayCluster/Private/Render/Viewport/DisplayClusterViewportHelpers.cpp b/Engine/Plugins/Runtime/nDisplay/Source/DisplayCluster/Private/Render/Viewport/DisplayClusterViewportHelpers.cpp
--- a/Engine/Plugins/Runtime/nDisplay/Source/DisplayCluster/Private/Render/Viewport/DisplayClusterViewportHelpers.cpp
+++ b/Engine/Plugins/Runtime/nDisplay/Source/DisplayCluster/Private/Render/Viewport/DisplayClusterViewportHelpers.cpp
@@ -23,17 +23,6 @@ static FAutoConsoleVariableRef CVarDisplayClusterOverrideMinTextureDimension(
 	ECVF_RenderThreadSafe
 );
 
-int32 GDisplayClusterDefaultPreviewPixelFormat = 1;
-static FAutoConsoleVariableRef CVarDisplayClusterDefaultPreviewPixelFormat(
-	TEXT("DC.Preview.DefaultPixelFormat"),
-	GDisplayClusterDefaultPreviewPixelFormat,
-	TEXT("Defines the default preview RTT pixel format.\n")
-	TEXT(" 0: 8bit fixed point RGBA\n")
-	TEXT(" 1: 16bit Float RGBA\n")
-	TEXT(" 2: 10bit fixed point RGB and 2bit Alpha\n"),
-	ECVF_RenderThreadSafe
-);
-
 int32 GDisplayClusterMaxNumMips = -1;
 static FAutoConsoleVariableRef CVarDisplayClusterMaxNumMips(
 	TEXT("DC.TextureMaxNumMips"),
@@ -45,26 +34,35 @@ static FAutoConsoleVariableRef CVarDisplayClusterMaxNumMips(
 	ECVF_RenderThreadSafe
 );
 
-namespace PreviewPixelFormat
+// Default pixel format for preview rendering
+int32 GDisplayClusterPreviewDefaultPixelFormat = 1;
+static FAutoConsoleVariableRef CVarDisplayClusterPreviewDefaultPixelFormat(
+	TEXT("nDisplay.preview.DefaultPixelFormat"),
+	GDisplayClusterPreviewDefaultPixelFormat,
+	TEXT("Defines the default preview RTT pixel format.\n")
+	TEXT(" 0: 8bit fixed point RGBA\n")
+	TEXT(" 1: 16bit Float RGBA\n")
+	TEXT(" 2: 10bit fixed point RGB and 2bit Alpha\n"),
+	ECVF_RenderThreadSafe
+);
+
+////////////////////////////////////////////////////////////////////////////////
+namespace UE::DisplayCluster::ViewportHelpers
 {
-	static ETextureRenderTargetFormat ImplGetRenderTargetFormatFromInt(const int32 InDefaultPreviewPixelFormat)
+	static inline ETextureRenderTargetFormat ImplGetRenderTargetFormatFromInt(const int32 InDefaultPreviewPixelFormat)
 	{
 		const int32 ValidIndex = FMath::Clamp(InDefaultPreviewPixelFormat, 0, 2);
 		static const ETextureRenderTargetFormat SPixelFormat[] = { RTF_RGBA8, RTF_RGBA16f, RTF_RGB10A2 };
 
 		return SPixelFormat[ValidIndex];
 	}
-
-	static EPixelFormat GetPreviewPixelFormat()
-	{
-		const ETextureRenderTargetFormat RenderTargetFormat = ImplGetRenderTargetFormatFromInt(GDisplayClusterDefaultPreviewPixelFormat);
-		
-		return GetPixelFormatFromRenderTargetFormat(RenderTargetFormat);
-	}
 };
+using namespace UE::DisplayCluster::ViewportHelpers;
 
 ////////////////////////////////////////////////////////////////////////////////
-int32 DisplayClusterViewportHelpers::GetMaxTextureNumMips(const FDisplayClusterRenderFrameSettings& InRenderFrameSettings, const int32 InNumMips)
+// FDisplayClusterViewportHelpers
+////////////////////////////////////////////////////////////////////////////////
+int32 FDisplayClusterViewportHelpers::GetMaxTextureNumMips(const FDisplayClusterRenderFrameSettings& InRenderFrameSettings, const int32 InNumMips)
 {
 	int32 NumMips = InNumMips;
 
@@ -72,7 +70,7 @@ int32 DisplayClusterViewportHelpers::GetMaxTextureNumMips(const FDisplayClusterR
 	switch (InRenderFrameSettings.RenderMode)
 	{
 	case EDisplayClusterRenderFrameMode::PreviewInScene:
-		if (GDisplayClusterDefaultPreviewPixelFormat != 0)
+		if (GDisplayClusterPreviewDefaultPixelFormat != 0)
 		{
 			//@todo: now UE support mips generation only for fixed point textures (8bit RGBA)
 			// Remove this hack latter
@@ -93,7 +91,7 @@ int32 DisplayClusterViewportHelpers::GetMaxTextureNumMips(const FDisplayClusterR
 	return NumMips;
 }
 
-int32 DisplayClusterViewportHelpers::GetMaxTextureDimension()
+int32 FDisplayClusterViewportHelpers::GetMaxTextureDimension()
 {
 	// The target always needs be within GMaxTextureDimensions, larger dimensions are not supported by the engine
 	static const int32 MaxTextureDimension = 1 << (GMaxTextureMipCount - 1);
@@ -111,7 +109,7 @@ int32 DisplayClusterViewportHelpers::GetMaxTextureDimension()
 	return MaxTextureDimension;
 }
 
-int32 DisplayClusterViewportHelpers::GetMinTextureDimension()
+int32 FDisplayClusterViewportHelpers::GetMinTextureDimension()
 {
 	static const int32 MinTextureDimension = 16;
 
@@ -128,24 +126,13 @@ int32 DisplayClusterViewportHelpers::GetMinTextureDimension()
 	return MinTextureDimension;
 }
 
-void DisplayClusterViewportHelpers::FreezeRenderingOfViewportTextureResources(TArray<FDisplayClusterViewportTextureResource*>& TextureResources)
-{
-	for (FDisplayClusterViewportTextureResource* TextureResource : TextureResources)
-	{
-		if (TextureResource != nullptr)
-		{
-			TextureResource->RaiseViewportResourceState(EDisplayClusterViewportResourceState::DisableReallocate);
-		}
-	}
-}
-
-bool DisplayClusterViewportHelpers::IsValidTextureSize(const FIntPoint& InSize)
+bool FDisplayClusterViewportHelpers::IsValidTextureSize(const FIntPoint& InSize)
 {
 	return InSize.GetMin() >= GetMinTextureDimension() && InSize.GetMax() <= GetMaxTextureDimension();
 }
 
 // Return size less than max texture dimension. Aspect ratio not changed
-FIntPoint DisplayClusterViewportHelpers::GetTextureSizeLessThanMax(const FIntPoint& InSize, const int32 InMaxTextureDimension)
+FIntPoint FDisplayClusterViewportHelpers::GetTextureSizeLessThanMax(const FIntPoint& InSize, const int32 InMaxTextureDimension)
 {
 	if (InSize.GetMax() > InMaxTextureDimension)
 	{
@@ -157,7 +144,7 @@ FIntPoint DisplayClusterViewportHelpers::GetTextureSizeLessThanMax(const FIntPoi
 	return InSize;
 }
 
-FIntPoint DisplayClusterViewportHelpers::ScaleTextureSize(const FIntPoint& InSize, float InMult)
+FIntPoint FDisplayClusterViewportHelpers::ScaleTextureSize(const FIntPoint& InSize, float InMult)
 {
 	const int32 ScaledX = FMath::CeilToInt(InSize.X * InMult);
 	const int32 ScaledY = FMath::CeilToInt(InSize.Y * InMult);
@@ -165,19 +152,19 @@ FIntPoint DisplayClusterViewportHelpers::ScaleTextureSize(const FIntPoint& InSiz
 	return FIntPoint(ScaledX, ScaledY);
 }
 
-float DisplayClusterViewportHelpers::GetValidSizeMultiplier(const FIntPoint& InSize, const float InSizeMult, const float InBaseSizeMult)
+float FDisplayClusterViewportHelpers::GetValidSizeMultiplier(const FIntPoint& InSize, const float InSizeMult, const float InBaseSizeMult)
 {
 	// find best possible size mult in range 1..InSizeMult
 	if (InSizeMult > 1.f)
 	{
 		const FIntPoint ScaledSize = ScaleTextureSize(InSize, FMath::Max(InSizeMult * InBaseSizeMult, 0.f));
-		if (DisplayClusterViewportHelpers::IsValidTextureSize(ScaledSize) == false)
+		if (FDisplayClusterViewportHelpers::IsValidTextureSize(ScaledSize) == false)
 		{
 			// Try change 'RenderTargetAdaptRatio' to min possible value
 			const float BaseMult = FMath::Max(InBaseSizeMult, 0.f);
 			const FIntPoint MinScaledSize = ScaleTextureSize(InSize, BaseMult);
 
-			if (DisplayClusterViewportHelpers::IsValidTextureSize(MinScaledSize) == false)
+			if (FDisplayClusterViewportHelpers::IsValidTextureSize(MinScaledSize) == false)
 			{
 				// BaseSizeMult to big. Disable size mult
 				return 1.f;
@@ -193,7 +180,7 @@ float DisplayClusterViewportHelpers::GetValidSizeMultiplier(const FIntPoint& InS
 #if UE_BUILD_DEBUG
 				// debug purpose
 				const FIntPoint FinalSize = ScaleTextureSize(InSize, FMath::Max(OutMult * InBaseSizeMult, 0.f));
-				check(DisplayClusterViewportHelpers::IsValidTextureSize(FinalSize));
+				check(FDisplayClusterViewportHelpers::IsValidTextureSize(FinalSize));
 #endif
 
 				return OutMult;
@@ -204,29 +191,14 @@ float DisplayClusterViewportHelpers::GetValidSizeMultiplier(const FIntPoint& InS
 	return InSizeMult;
 }
 
-#if WITH_EDITOR
-void DisplayClusterViewportHelpers::GetPreviewRenderTargetDesc_Editor(const FDisplayClusterRenderFrameSettings& InRenderFrameSettings, EPixelFormat& OutPixelFormat, float& OutDisplayGamma, bool& bOutSRGB)
+EPixelFormat FDisplayClusterViewportHelpers::GetPreviewDefaultPixelFormat()
 {
-	OutPixelFormat = PreviewPixelFormat::GetPreviewPixelFormat();
-
-	// Default gamma value
-	OutDisplayGamma = 2.2f;
-
-	//!
-	bOutSRGB = true;
+	const ETextureRenderTargetFormat RenderTargetFormat = ImplGetRenderTargetFormatFromInt(GDisplayClusterPreviewDefaultPixelFormat);
 
-	switch (InRenderFrameSettings.RenderMode)
-	{
-	case EDisplayClusterRenderFrameMode::PreviewInScene:
-		if (InRenderFrameSettings.bPreviewEnablePostProcess == false)
-		{
-			// Disable postprocess for preview. Use Gamma 1.f
-			OutDisplayGamma = 1.f;
-		}
-		break;
-	default:
-		break;
-	}
+	return GetPixelFormatFromRenderTargetFormat(RenderTargetFormat);
 }
-#endif
 
+EPixelFormat FDisplayClusterViewportHelpers::GetDefaultPixelFormat()
+{
+	return EPixelFormat::PF_FloatRGBA;
+}
diff --git a/Engine/Plugins/Runtime/nDisplay/Source/DisplayCluster/Private/Render/Viewport/DisplayClusterViewportHelpers.h b/Engine/Plugins/Runtime/nDisplay/Source/DisplayCluster/Private/Render/Viewport/DisplayClusterViewportHelpers.h
--- a/Engine/Plugins/Runtime/nDisplay/Source/DisplayCluster/Private/Render/Viewport/DisplayClusterViewportHelpers.h
+++ b/Engine/Plugins/Runtime/nDisplay/Source/DisplayCluster/Private/Render/Viewport/DisplayClusterViewportHelpers.h
@@ -4,32 +4,73 @@
 #include "CoreMinimal.h"
 #include "PixelFormat.h"
 
-class FDisplayClusterViewportTextureResource;
+class FDisplayClusterViewportResource;
 struct FDisplayClusterRenderFrameSettings;
 
-class DisplayClusterViewportHelpers
+/**
+ * Helper class for DC viewports
+ */
+class FDisplayClusterViewportHelpers
 {
 public:
-	// Control texture resources dimensions inside nDisplay
+	/** Get the maximum allowable texture size used for the nDisplay viewport. */
 	static int32 GetMinTextureDimension();
+
+	/** Get the minimum allowable texture size used for the nDisplay viewport. */
 	static int32 GetMaxTextureDimension();
 
-	// Texture size helpers
+	/** Return true, if size is valid.
+	* 
+	* @param InSize - the size to be checked.
+	* 
+	* @return true, if the size is within the minimum and maximum dimensions
+	*/
 	static bool IsValidTextureSize(const FIntPoint& InSize);
 
+	/** Scaling texture size with a multiplier
+	* 
+	* @param InSize - the input texture size
+	* @param InMult - texture size multiplier
+	* 
+	* @return - scaled texture size
+	*/
 	static FIntPoint ScaleTextureSize(const FIntPoint& InSize, float InMult);
 
-	// Return size less than max texture dimension. Aspect ratio not changed
+	/** The returned size is smaller than the maximum texture size. The aspect ratio does not change
+	* 
+	* @param InSize - input texture size
+	* @param InMaxTextureDimension - max alloable dimension
+	* 
+	* @return - the same as the input size, or reduced to the InMaxTextureDimension
+	 */
 	static FIntPoint GetTextureSizeLessThanMax(const FIntPoint& InSize, const int32 InMaxTextureDimension);
 
+	/** Find an acceptable multiplier for the texture size
+	* 
+	* @param InSize - in texture size
+	* @param InSizeMult - the target multiplier
+	* @param InBaseSizeMult - the base multiplier
+	* 
+	* @return - a multiplier that gives the valid size of the texture
+	*/
 	static float GetValidSizeMultiplier(const FIntPoint& InSize, const float InSizeMult, const float InBaseSizeMult);
 
-	// Support render freeze feature
-	static void FreezeRenderingOfViewportTextureResources(TArray<FDisplayClusterViewportTextureResource*>& TextureResources);
-
+	/** Getting the maximum allowable mips value for the rendering frame settings
+	* 
+	* @param InRenderFrameSettings  -the current render frame settings
+	* @param InNumMips - texture mips value from the settings
+	* 
+	* @return - the mips value to be used.
+	*/
 	static int32 GetMaxTextureNumMips(const FDisplayClusterRenderFrameSettings& InRenderFrameSettings, const int32 InNumMips);
 
-#if WITH_EDITOR
-	static void GetPreviewRenderTargetDesc_Editor(const FDisplayClusterRenderFrameSettings& InRenderFrameSettings, EPixelFormat& OutPixelFormat, float& OutDisplayGamma, bool& bOutSRGB);
-#endif
+	/**
+	 * Getting the default pixel format for preview rendering
+	 */
+	static EPixelFormat GetPreviewDefaultPixelFormat();
+
+	/**
+	 * Getting the default pixel format
+	 */
+	static EPixelFormat GetDefaultPixelFormat();
 };
diff --git a/Engine/Plugins/Runtime/nDisplay/Source/DisplayCluster/Private/Render/Viewport/DisplayClusterViewportManager.cpp b/Engine/Plugins/Runtime/nDisplay/Source/DisplayCluster/Private/Render/Viewport/DisplayClusterViewportManager.cpp
--- a/Engine/Plugins/Runtime/nDisplay/Source/DisplayCluster/Private/Render/Viewport/DisplayClusterViewportManager.cpp
+++ b/Engine/Plugins/Runtime/nDisplay/Source/DisplayCluster/Private/Render/Viewport/DisplayClusterViewportManager.cpp
@@ -510,7 +510,7 @@ void FDisplayClusterViewportManager::ImplUpdateClusterNodeViewports(const EDispl
 				continue;
 			}
 
-			if (Viewport->RenderSettings.IsViewportOverrided())
+			if (Viewport->RenderSettings.IsViewportOverridden())
 			{
 				OverriddenViewports.Add(Viewport);
 			}
@@ -669,8 +669,8 @@ bool FDisplayClusterViewportManager::BeginNewFrame(FViewport* InViewport, UWorld
 			for (FDisplayClusterViewport_Context& ContextIt : ViewportIt->Contexts)
 			{
 				// Get Context Display gamma
-				const FDisplayClusterViewportRenderTargetResource* ContextRTT = ViewportIt->RenderTargets.IsValidIndex(ContextIt.ContextNum) ? ViewportIt->RenderTargets[ContextIt.ContextNum] : nullptr;
-				const float ViewportDisplayGamma = ContextRTT ? ContextRTT->GetDisplayGamma() : DefaultDisplayGamma;
+				const TSharedPtr<FDisplayClusterViewportResource, ESPMode::ThreadSafe> ContextRTT = ViewportIt->Resources[EDisplayClusterViewportResource::RenderTargets].IsValidIndex(ContextIt.ContextNum) ? ViewportIt->Resources[EDisplayClusterViewportResource::RenderTargets][ContextIt.ContextNum] : nullptr;
+				const float ViewportDisplayGamma = ContextRTT.IsValid() ? ContextRTT->GetResourceSettings().GetDisplayGamma() : DefaultDisplayGamma;
 
 				ContextIt.RenderThreadData.EngineDisplayGamma = ViewportDisplayGamma;
 
diff --git a/Engine/Plugins/Runtime/nDisplay/Source/DisplayCluster/Private/Render/Viewport/DisplayClusterViewportManagerProxy.cpp b/Engine/Plugins/Runtime/nDisplay/Source/DisplayCluster/Private/Render/Viewport/DisplayClusterViewportManagerProxy.cpp
--- a/Engine/Plugins/Runtime/nDisplay/Source/DisplayCluster/Private/Render/Viewport/DisplayClusterViewportManagerProxy.cpp
+++ b/Engine/Plugins/Runtime/nDisplay/Source/DisplayCluster/Private/Render/Viewport/DisplayClusterViewportManagerProxy.cpp
@@ -26,17 +26,6 @@
 #include "Render/Viewport/DisplayClusterViewportStrings.h"
 #include "RHIContext.h"
 
-// Enable/disable warp&blend
-static TAutoConsoleVariable<int32> CVarWarpBlendEnabled(
-	TEXT("nDisplay.render.WarpBlendEnabled"),
-	1,
-	TEXT("Warp & Blend status\n")
-	TEXT("0 : disabled\n")
-	TEXT("1 : enabled\n")
-	,
-	ECVF_RenderThreadSafe
-);
-
 // Enable/Disable ClearTexture for Frame RTT
 static TAutoConsoleVariable<int32> CVarClearFrameRTTEnabled(
 	TEXT("nDisplay.render.ClearFrameRTTEnabled"),
@@ -48,27 +37,6 @@ static TAutoConsoleVariable<int32> CVarClearFrameRTTEnabled(
 	ECVF_RenderThreadSafe
 );
 
-///////////////////////////////////////////////////////////////////////////////////////
-namespace DisplayClusterViewportManagerProxyHelpers
-{
-	// Support warp blend logic
-	static inline bool ShouldApplyWarpBlend(IDisplayClusterViewportProxy* ViewportProxy)
-	{
-		if (ViewportProxy->GetPostRenderSettings_RenderThread().Replace.IsEnabled())
-		{
-			// When used override texture, disable warp blend
-			return false;
-		}
-
-		const TSharedPtr<IDisplayClusterProjectionPolicy, ESPMode::ThreadSafe>& PrjPolicy = ViewportProxy->GetProjectionPolicy_RenderThread();
-
-		// Ask current projection policy if it's warp&blend compatible
-		return PrjPolicy.IsValid() && PrjPolicy->IsWarpBlendSupported();
-	}
-};
-
-using namespace DisplayClusterViewportManagerProxyHelpers;
-
 ///////////////////////////////////////////////////////////////////////////////////////
 //          FDisplayClusterViewportManagerProxy
 ///////////////////////////////////////////////////////////////////////////////////////
@@ -122,24 +90,6 @@ void FDisplayClusterViewportManagerProxy::Initialize(FDisplayClusterViewportMana
 	ViewportManagerViewExtension = InViewportManager.ViewportManagerViewExtension;
 }
 
-void FDisplayClusterViewportManagerProxy::DeleteResource_RenderThread(FDisplayClusterViewportResource* InDeletedResourcePtr)
-{
-	if (InDeletedResourcePtr)
-	{
-		// Handle resource refs must be removed from all viewports
-		for (const TSharedPtr<FDisplayClusterViewportProxy, ESPMode::ThreadSafe>& ViewportProxyIt : ImplGetEntireClusterViewportProxies_RenderThread())
-		{
-			if (ViewportProxyIt.IsValid())
-			{
-				ViewportProxyIt->HandleResourceDelete_RenderThread(InDeletedResourcePtr);
-			}
-		}
-
-		InDeletedResourcePtr->ReleaseResource();
-		delete InDeletedResourcePtr;
-	}
-}
-
 void FDisplayClusterViewportManagerProxy::ImplUpdateClusterNodeViewportProxies_RenderThread()
 {
 	// Get viewports used to rendering of the current frame
@@ -312,10 +262,8 @@ void FDisplayClusterViewportManagerProxy::ImplRenderFrame_GameThread(FViewport*
 			ViewportManagerProxy->PostProcessManager->HandleBeginUpdateFrameResources_RenderThread(RHICmdList, ViewportManagerProxy);
 		}
 
-		const bool bWarpBlendEnabled = ViewportManagerProxy->RenderFrameSettings.bAllowWarpBlend && CVarWarpBlendEnabled.GetValueOnRenderThread() != 0;
-
 		// Update the frame resources: post-processing, warping, and finally resolving everything to the frame resource
-		ViewportManagerProxy->UpdateFrameResources_RenderThread(RHICmdList, bWarpBlendEnabled);
+		ViewportManagerProxy->UpdateFrameResources_RenderThread(RHICmdList);
 
 		if (ViewportManagerProxy->PostProcessManager.IsValid())
 		{
@@ -363,7 +311,7 @@ void FDisplayClusterViewportManagerProxy::UpdateDeferredResources_RenderThread(F
 
 	for (const TSharedPtr<FDisplayClusterViewportProxy, ESPMode::ThreadSafe>& ViewportProxy : ImplGetCurrentRenderFrameViewportProxies_RenderThread())
 	{
-		if (!ViewportProxy->RenderSettings.IsViewportOverrided())
+		if (!ViewportProxy->RenderSettings.IsViewportOverridden())
 		{
 			ViewportProxy->UpdateDeferredResources(RHICmdList);
 		}
@@ -420,7 +368,7 @@ enum class EWarpPass : uint8
 	COUNT
 };
 
-void FDisplayClusterViewportManagerProxy::UpdateFrameResources_RenderThread(FRHICommandListImmediate& RHICmdList, bool bWarpBlendEnabled) const
+void FDisplayClusterViewportManagerProxy::UpdateFrameResources_RenderThread(FRHICommandListImmediate& RHICmdList) const
 {
 	check(IsInRenderingThread());
 
@@ -457,7 +405,7 @@ void FDisplayClusterViewportManagerProxy::UpdateFrameResources_RenderThread(FRHI
 			// Iterate over visible viewports:
 			if (ViewportProxyIt.IsValid() && ViewportProxyIt->GetRenderSettings_RenderThread().bVisible)
 			{
-				if (bWarpBlendEnabled && ShouldApplyWarpBlend(ViewportProxyIt.Get()))
+				if (ViewportProxyIt->ShouldApplyWarpBlend_RenderThread())
 				{
 					const TSharedPtr<IDisplayClusterProjectionPolicy, ESPMode::ThreadSafe>& PrjPolicy = ViewportProxyIt->GetProjectionPolicy_RenderThread();
 					switch ((EWarpPass)WarpPass)
@@ -500,18 +448,8 @@ void FDisplayClusterViewportManagerProxy::UpdateFrameResources_RenderThread(FRHI
 		// Iterate over visible viewports:
 		if (ViewportProxyIt.IsValid() && ViewportProxyIt->GetRenderSettings_RenderThread().bVisible)
 		{
-			EDisplayClusterViewportResourceType ViewportSource = EDisplayClusterViewportResourceType::InputShaderResource;
-			if (bWarpBlendEnabled && ShouldApplyWarpBlend(ViewportProxyIt.Get()))
-			{
-				const TSharedPtr<IDisplayClusterProjectionPolicy, ESPMode::ThreadSafe>& PrjPolicy = ViewportProxyIt->GetProjectionPolicy_RenderThread();
-				if (PrjPolicy->ShouldUseAdditionalTargetableResource())
-				{
-					ViewportSource = EDisplayClusterViewportResourceType::AdditionalTargetableResource;
-				}
-			}
-
-			// resolve viewports to the frame target texture
-			ViewportProxyIt->ResolveResources_RenderThread(RHICmdList, ViewportSource, ViewportProxyIt->GetOutputResourceType_RenderThread());
+			// resolve warped viewport resource to the output texture
+			ViewportProxyIt->ResolveResources_RenderThread(RHICmdList, EDisplayClusterViewportResourceType::AfterWarpBlendTargetableResource, EDisplayClusterViewportResourceType::OutputTargetableResource);
 
 			// Apply post-warp (viewport remap, etc)
 			ViewportProxyIt->PostResolveViewport_RenderThread(RHICmdList);
@@ -547,7 +485,7 @@ void FDisplayClusterViewportManagerProxy::DoCrossGPUTransfers_RenderThread(FRHIC
 			continue;
 		}
 
-		if (!ViewportProxyIt->RenderSettings.bEnableCrossGPUTransfer || ViewportProxyIt->RenderTargets.IsEmpty())
+		if (!ViewportProxyIt->RenderSettings.bEnableCrossGPUTransfer || ViewportProxyIt->Resources[EDisplayClusterViewportResource::RenderTargets].IsEmpty())
 		{
 			// Skip a frozen viewport that has already been transferred between GPUs
 			// The first time freezing should do the transfer (RenderTargets must be assigned on the first pass)
@@ -570,7 +508,7 @@ void FDisplayClusterViewportManagerProxy::DoCrossGPUTransfers_RenderThread(FRHIC
 				continue;
 			}
 
-			if (!ViewportProxyIt->RenderTargets.IsValidIndex(ViewportContext.ContextNum))
+			if (!ViewportProxyIt->Resources[EDisplayClusterViewportResource::RenderTargets].IsValidIndex(ViewportContext.ContextNum))
 			{
 				// RTT does not exist for this context
 				continue;
@@ -585,22 +523,25 @@ void FDisplayClusterViewportManagerProxy::DoCrossGPUTransfers_RenderThread(FRHIC
 			}
 
 			// Perform an MGPU transfer for a specific viewport context:
-			if (FDisplayClusterViewportRenderTargetResource* RenderTarget = ViewportProxyIt->RenderTargets[ViewportContext.ContextNum])
+			TSharedPtr<FDisplayClusterViewportResource, ESPMode::ThreadSafe> ViewportRenderTargetResource = ViewportProxyIt->Resources[EDisplayClusterViewportResource::RenderTargets][ViewportContext.ContextNum];
 			{
-				if (FRHITexture2D* TextureRHI = RenderTarget->GetViewportResource2DRHI())
+				if (FRenderTarget* RenderTarget = ViewportRenderTargetResource.IsValid() ? ViewportRenderTargetResource->GetViewportResourceRenderTarget() : nullptr)
 				{
-					const FRHIGPUMask RenderTargetGPUMask = RenderTarget->GetGPUMask(RHICmdList);
-
-					for (uint32 DestGPUIndex : RenderTargetGPUMask)
+					if (FRHITexture2D* TextureRHI = ViewportRenderTargetResource->GetViewportResourceRHI())
 					{
-						if (DestGPUIndex != SrcGPUIndex)
+						const FRHIGPUMask RenderTargetGPUMask = RenderTarget->GetGPUMask(RHICmdList);
+
+						for (uint32 DestGPUIndex : RenderTargetGPUMask)
 						{
-							TransferResources.Add(FTransferResourceParams(
-								TextureRHI, TransferRect,
-								SrcGPUIndex, DestGPUIndex,
-								RenderFrameSettings.CrossGPUTransfer.bPullData,
-								RenderFrameSettings.CrossGPUTransfer.bLockSteps
-							));
+							if (DestGPUIndex != SrcGPUIndex)
+							{
+								TransferResources.Add(FTransferResourceParams(
+									TextureRHI, TransferRect,
+									SrcGPUIndex, DestGPUIndex,
+									RenderFrameSettings.CrossGPUTransfer.bPullData,
+									RenderFrameSettings.CrossGPUTransfer.bLockSteps
+								));
+							}
 						}
 					}
 				}
@@ -624,21 +565,21 @@ bool FDisplayClusterViewportManagerProxy::GetFrameTargets_RenderThread(TArray<FR
 	for (const TSharedPtr<FDisplayClusterViewportProxy, ESPMode::ThreadSafe>& ViewportProxyIt : ImplGetCurrentRenderFrameViewportProxies_RenderThread())
 	{
 		// Process only valid viewports with output frame resources
-		if (ViewportProxyIt.IsValid() && !ViewportProxyIt->OutputFrameTargetableResources.IsEmpty())
+		if (ViewportProxyIt.IsValid() && !ViewportProxyIt->Resources[EDisplayClusterViewportResource::OutputFrameTargetableResources].IsEmpty())
 		{
-			const TArray<FDisplayClusterViewportTextureResource*>& Frames = ViewportProxyIt->OutputFrameTargetableResources;
-			const TArray<FDisplayClusterViewportTextureResource*>& AdditionalFrames = ViewportProxyIt->AdditionalFrameTargetableResources;
+			const TArray<TSharedPtr<FDisplayClusterViewportResource, ESPMode::ThreadSafe>>& Frames = ViewportProxyIt->Resources[EDisplayClusterViewportResource::OutputFrameTargetableResources];
+			const TArray<TSharedPtr<FDisplayClusterViewportResource, ESPMode::ThreadSafe>>& AdditionalFrames = ViewportProxyIt->Resources[EDisplayClusterViewportResource::AdditionalFrameTargetableResources];
 
 			for (int32 FrameIt = 0; FrameIt < Frames.Num(); FrameIt++)
 			{
-				if (FRHITexture2D* FrameTexture = Frames[FrameIt] ? Frames[FrameIt]->GetViewportResource2DRHI() : nullptr)
+				if (FRHITexture2D* FrameTexture = Frames[FrameIt].IsValid() ? Frames[FrameIt]->GetViewportResourceRHI() : nullptr)
 				{
 					OutFrameResources.Add(FrameTexture);
-					OutTargetOffsets.Add(Frames[FrameIt]->BackbufferFrameOffset);
+					OutTargetOffsets.Add(Frames[FrameIt]->GetBackbufferFrameOffset());
 
 					if (OutAdditionalFrameResources && AdditionalFrames.IsValidIndex(FrameIt))
 					{
-						if (FRHITexture2D* AdditionalFrameTexture = AdditionalFrames[FrameIt] ? AdditionalFrames[FrameIt]->GetViewportResource2DRHI() : nullptr)
+						if (FRHITexture2D* AdditionalFrameTexture = AdditionalFrames[FrameIt].IsValid() ? AdditionalFrames[FrameIt]->GetViewportResourceRHI() : nullptr)
 						{
 							OutAdditionalFrameResources->Add(AdditionalFrameTexture);
 						}
diff --git a/Engine/Plugins/Runtime/nDisplay/Source/DisplayCluster/Private/Render/Viewport/DisplayClusterViewportManagerProxy.h b/Engine/Plugins/Runtime/nDisplay/Source/DisplayCluster/Private/Render/Viewport/DisplayClusterViewportManagerProxy.h
--- a/Engine/Plugins/Runtime/nDisplay/Source/DisplayCluster/Private/Render/Viewport/DisplayClusterViewportManagerProxy.h
+++ b/Engine/Plugins/Runtime/nDisplay/Source/DisplayCluster/Private/Render/Viewport/DisplayClusterViewportManagerProxy.h
@@ -71,7 +71,7 @@ public:
 	void UpdateDeferredResources_RenderThread(FRHICommandListImmediate& RHICmdList) const;
 
 	/** Apply WarpBlend and resolve to frame resources. */
-	void UpdateFrameResources_RenderThread(FRHICommandListImmediate& RHICmdList, bool bWarpBlendEnabled) const;
+	void UpdateFrameResources_RenderThread(FRHICommandListImmediate& RHICmdList) const;
 
 	/** Release all referenced objects and resources. */
 	void Release_RenderThread();
@@ -82,9 +82,6 @@ public:
 	/** Unregister exist viewport proxy. */
 	void DeleteViewport_RenderThread(const TSharedPtr<FDisplayClusterViewportProxy, ESPMode::ThreadSafe>& InViewportProxy);
 
-	/** Reset all references to InDeletedResourcePtr for all viewport proxies and delete the resource. */
-	void DeleteResource_RenderThread(FDisplayClusterViewportResource* InDeletedResourcePtr);
-
 	/** Initialize ViewportManagerProxy from ViewportManager. */
 	void Initialize(FDisplayClusterViewportManager& InViewportManager);
 
diff --git a/Engine/Plugins/Runtime/nDisplay/Source/DisplayCluster/Private/Render/Viewport/DisplayClusterViewportManagerViewExtension.cpp b/Engine/Plugins/Runtime/nDisplay/Source/DisplayCluster/Private/Render/Viewport/DisplayClusterViewportManagerViewExtension.cpp
--- a/Engine/Plugins/Runtime/nDisplay/Source/DisplayCluster/Private/Render/Viewport/DisplayClusterViewportManagerViewExtension.cpp
+++ b/Engine/Plugins/Runtime/nDisplay/Source/DisplayCluster/Private/Render/Viewport/DisplayClusterViewportManagerViewExtension.cpp
@@ -13,11 +13,12 @@
 #include "CommonRenderResources.h"
 #include "PostProcess/PostProcessMaterialInputs.h"
 
-namespace DisplayClusterViewportManagerViewExtensionHelpers
+///////////////////////////////////////////////////////////////////////////////////////
+namespace UE::DisplayCluster::ViewportManagerViewExtension
 {
 	static const FName RendererModuleName(TEXT("Renderer"));
 };
-using namespace DisplayClusterViewportManagerViewExtensionHelpers;
+using namespace UE::DisplayCluster::ViewportManagerViewExtension;
 
 ///////////////////////////////////////////////////////////////////////////////////////
 // FDisplayClusterViewportManagerViewExtension
diff --git a/Engine/Plugins/Runtime/nDisplay/Source/DisplayCluster/Private/Render/Viewport/DisplayClusterViewportManager_EditorPreview.cpp b/Engine/Plugins/Runtime/nDisplay/Source/DisplayCluster/Private/Render/Viewport/DisplayClusterViewportManager_EditorPreview.cpp
--- a/Engine/Plugins/Runtime/nDisplay/Source/DisplayCluster/Private/Render/Viewport/DisplayClusterViewportManager_EditorPreview.cpp
+++ b/Engine/Plugins/Runtime/nDisplay/Source/DisplayCluster/Private/Render/Viewport/DisplayClusterViewportManager_EditorPreview.cpp
@@ -3,8 +3,6 @@
 #include "Render/Viewport/DisplayClusterViewportManager.h"
 #include "Render/Viewport/DisplayClusterViewportManagerProxy.h"
 
-#if WITH_EDITOR
-
 #include "EngineModule.h"
 #include "CanvasTypes.h"
 #include "LegacyScreenPercentageDriver.h"
@@ -18,6 +16,7 @@
 #include "Render/Projection/IDisplayClusterProjectionPolicy.h"
 #include "Render/Viewport/DisplayClusterViewport.h"
 #include "Render/Viewport/DisplayClusterViewport_CustomPostProcessSettings.h"
+#include "Render/Viewport/Resource/DisplayClusterViewportResource.h"
 
 #include "Render/Viewport/RenderFrame/DisplayClusterRenderFrame.h"
 #include "Render/Viewport/RenderFrame/DisplayClusterRenderFrameSettings.h"
@@ -32,6 +31,7 @@
 ///////////////////////////////////////////////////////////////////////////////////////
 // FDisplayClusterViewportManager
 ///////////////////////////////////////////////////////////////////////////////////////
+#if WITH_EDITOR
 void FDisplayClusterViewportManager::ImplUpdatePreviewRTTResources()
 {
 	check(IsInGameThread());
@@ -81,10 +81,12 @@ void FDisplayClusterViewportManager::ImplUpdatePreviewRTTResources()
 		if (DesiredViewport.IsValid())
 		{
 			FTextureRHIRef& PreviewRTT = PreviewRenderTargetableTextures[ViewportIndex];
-			if (PreviewRTT.IsValid())
+			if (PreviewRTT.IsValid()
+				&& !DesiredViewport->Resources[EDisplayClusterViewportResource::OutputPreviewTargetableResources].IsEmpty()
+				&& DesiredViewport->Resources[EDisplayClusterViewportResource::OutputPreviewTargetableResources][0].IsValid())
 			{
 				// Use mapped preview viewport
-				DesiredViewport->OutputPreviewTargetableResource = PreviewRTT;
+				DesiredViewport->Resources[EDisplayClusterViewportResource::OutputPreviewTargetableResources][0]->SetExternalViewportResourceRHI(PreviewRTT);
 			}
 			else
 			{
@@ -94,6 +96,7 @@ void FDisplayClusterViewportManager::ImplUpdatePreviewRTTResources()
 		}
 	}
 }
+#endif
 
 void FDisplayClusterViewportManager::OnPreGarbageCollect()
 {
@@ -119,11 +122,14 @@ bool FDisplayClusterViewportManager::RenderInEditor(FDisplayClusterRenderFrame&
 	}
 
 	const ADisplayClusterRootActor* RootActor = GetRootActor();
+
+#if WITH_EDITOR
 	if (RootActor && !RootActor->IsEditorRenderEnabled())
 	{
 		bOutFrameRendered = true;
 		return true;
 	}
+#endif
 	
 	FSceneInterface* PreviewScene = CurrentWorld->Scene;
 	FEngineShowFlags EngineShowFlags = FEngineShowFlags(EShowFlagInitMode::ESFIM_Game);
@@ -169,9 +175,12 @@ bool FDisplayClusterViewportManager::RenderInEditor(FDisplayClusterRenderFrame&
 				ConfigureViewFamily(RenderTargetIt, ViewFamiliesIt, ViewFamily);
 
 				if (RenderTargetIt.CaptureMode == EDisplayClusterViewportCaptureMode::Default
+#if WITH_EDITOR
 					&& RootActor
 					&& !RootActor->bPreviewEnablePostProcess
-					&& RootActor->DoObserversNeedPostProcessRenderTarget())
+					&& RootActor->DoObserversNeedPostProcessRenderTarget()
+#endif
+				)
 				{
 					if (ViewFamily.EngineShowFlags.TemporalAA)
 					{
@@ -276,5 +285,3 @@ bool FDisplayClusterViewportManager::RenderInEditor(FDisplayClusterRenderFrame&
 
 	return true;
 }
-
-#endif
diff --git a/Engine/Plugins/Runtime/nDisplay/Source/DisplayCluster/Private/Render/Viewport/DisplayClusterViewportProxy.cpp b/Engine/Plugins/Runtime/nDisplay/Source/DisplayCluster/Private/Render/Viewport/DisplayClusterViewportProxy.cpp
--- a/Engine/Plugins/Runtime/nDisplay/Source/DisplayCluster/Private/Render/Viewport/DisplayClusterViewportProxy.cpp
+++ b/Engine/Plugins/Runtime/nDisplay/Source/DisplayCluster/Private/Render/Viewport/DisplayClusterViewportProxy.cpp
@@ -78,9 +78,27 @@ static FAutoConsoleVariableRef CVarDisplayClusterShadersICVFXFXAAChromakey(
 	ECVF_RenderThreadSafe
 );
 
+// Enable/disable warp&blend
+int32 GDisplayClusterRenderWarpBlendEnabled = 1;
+static FAutoConsoleVariableRef CVarDisplayClusterRenderWarpBlendEnabled(
+	TEXT("nDisplay.render.WarpBlendEnabled"),
+	GDisplayClusterRenderWarpBlendEnabled,
+	TEXT("Warp & Blend status\n")
+	TEXT("0 : disabled\n")
+	TEXT("1 : enabled\n"),
+	ECVF_RenderThreadSafe
+);
+
 ///////////////////////////////////////////////////////////////////////////////////////
-namespace DisplayClusterViewportProxyHelpers
+namespace UE::DisplayCluster::ViewportProxy
 {
+	static IDisplayClusterShaders& GetShadersAPI()
+	{
+		static IDisplayClusterShaders& ShadersAPISingleton = IDisplayClusterShaders::Get();
+
+		return ShadersAPISingleton;
+	}
+
 	// The viewport override has the maximum depth. This protects against a link cycle
 	static const int32 DisplayClusterViewportProxyResourcesOverrideRecursionDepthMax = 4;
 
@@ -114,90 +132,95 @@ namespace DisplayClusterViewportProxyHelpers
 		// No FXAA
 		return false;
 	}
+};
+using namespace UE::DisplayCluster::ViewportProxy;
 
-	/**
-	* Get viewport RHI resources
-	*
-	* @param InResources - Array with DC resources (for mono num=1, for stereo num=2)
-	* @param OutResources - Array with RHI resources
-	*
-	* @return true, if all resources are valid
-	*/
-	static bool GetViewportRHIResourcesImpl_RenderThread(const TArrayView<FDisplayClusterViewportResource*>& InResources, TArray<FRHITexture2D*>& OutResources)
+///////////////////////////////////////////////////////////////////////////////////////
+FDisplayClusterViewportProxy::FDisplayClusterViewportProxy(const FDisplayClusterViewport& RenderViewport)
+	: ViewportId(RenderViewport.ViewportId)
+	, ClusterNodeId(RenderViewport.ClusterNodeId)
+	, RenderSettings(RenderViewport.RenderSettings)
+	, ProjectionPolicy(RenderViewport.UninitializedProjectionPolicy)
+	, ViewportManagerProxyWeakRef(RenderViewport.GetViewportManagerProxyRefImpl())
+{
+	check(ProjectionPolicy.IsValid());
+}
+
+FDisplayClusterViewportProxy::~FDisplayClusterViewportProxy()
+{
+}
+
+EDisplayClusterViewportResourceType FDisplayClusterViewportProxy::GetResourceType_RenderThread(const EDisplayClusterViewportResourceType& InResourceType) const
+{
+	check(IsInRenderingThread());
+
+	switch (InResourceType)
 	{
-		OutResources.Reset();
+	case EDisplayClusterViewportResourceType::BeforeWarpBlendTargetableResource:
+		return EDisplayClusterViewportResourceType::InputShaderResource;
 
-		for (FDisplayClusterViewportResource* ViewportResourceIt : InResources)
+	case EDisplayClusterViewportResourceType::AfterWarpBlendTargetableResource:
+		return ShouldApplyWarpBlend_RenderThread() ? EDisplayClusterViewportResourceType::AdditionalTargetableResource : EDisplayClusterViewportResourceType::InputShaderResource;
+
+	case EDisplayClusterViewportResourceType::OutputTargetableResource:
+
+		/**
+		* Output textures for preview rendering.
+		*/
+		if (GetRenderMode() == EDisplayClusterRenderFrameMode::PreviewInScene)
 		{
-			if (FRHITexture2D* RHITexture2D = ViewportResourceIt ? ViewportResourceIt->GetViewportResource2DRHI() : nullptr)
-			{
-				// Collects only valid resources.
-				OutResources.Add(RHITexture2D);
-			}
+			// [warp] -> OutputPreviewTargetableResource
+			return EDisplayClusterViewportResourceType::OutputPreviewTargetableResource;
 		}
 
-		if (OutResources.Num() != InResources.Num())
+		/**
+		* Output 'Frame' textures for cluster rendering
+		*/
+		if (RemapMesh.IsValid())
 		{
-			// Some resources lost
-			OutResources.Reset();
+			const IDisplayClusterRender_MeshComponentProxy* MeshProxy = RemapMesh->GetMeshComponentProxy_RenderThread();
+			if (MeshProxy != nullptr && MeshProxy->IsEnabled_RenderThread())
+			{
+				// In this case render to additional frame targetable
+				// to minimize the rendering steps:
+				// [warp] -> AdditionalFrameTargetableResource -> [OutputRemap] -> OutputFrameTargetableResource
+				return EDisplayClusterViewportResourceType::AdditionalFrameTargetableResource;
+			}
 		}
 
-		// returns success if the number of output resources is equal to the input and is not empty
-		return !OutResources.IsEmpty();
-	}
+		// [warp] -> OutputFrameTargetableResource
+		return EDisplayClusterViewportResourceType::OutputFrameTargetableResource;
 
-	static bool GetViewportRHIResourcesImpl_RenderThread(const TArray<FDisplayClusterViewportTextureResource*>& InResources, TArray<FRHITexture2D*>& OutResources)
-	{
-		return GetViewportRHIResourcesImpl_RenderThread(TArrayView<FDisplayClusterViewportResource*>((FDisplayClusterViewportResource**)(InResources.GetData()), InResources.Num()), OutResources);
+
+	default:
+		// return the same value
+		break;
 	}
 
-	static bool GetViewportRHIResourcesImpl_RenderThread(const TArray<FDisplayClusterViewportRenderTargetResource*>& InResources, TArray<FRHITexture2D*>& OutResources)
+	return InResourceType;
+}
+
+bool FDisplayClusterViewportProxy::ShouldApplyWarpBlend_RenderThread() const
+{
+	if (GDisplayClusterRenderWarpBlendEnabled == 0)
 	{
-		return GetViewportRHIResourcesImpl_RenderThread(TArrayView<FDisplayClusterViewportResource*>((FDisplayClusterViewportResource**)(InResources.GetData()), InResources.Num()), OutResources);
+		return false;
 	}
 
-	// Reset reference to deleted resource
-	static void HandleResourceDeleteImpl_RenderThread(TArray<FDisplayClusterViewportTextureResource*>& InResources, const FDisplayClusterViewportResource* InDeletedResourcePtr)
+	if (GetPostRenderSettings_RenderThread().Replace.IsEnabled())
 	{
-		for (int32 Index = 0; Index < InResources.Num(); Index++)
-		{
-			if (InResources[Index] == InDeletedResourcePtr)
-			{
-				// remove references to a deleted resource
-				InResources[Index] = nullptr;
-			}
-		}
+		// When used override texture, disable warp blend
+		return false;
 	}
 
-	// Reset reference to deleted resource
-	static void HandleResourceDeleteImpl_RenderThread(TArray<FDisplayClusterViewportRenderTargetResource*>& InResources, const FDisplayClusterViewportResource* InDeletedResourcePtr)
+	const FDisplayClusterRenderFrameSettings* RenderFrameSettings = GetRenderFrameSettings_RenderThread();
+	if (!RenderFrameSettings || !RenderFrameSettings->bAllowWarpBlend)
 	{
-		for (int32 Index = 0; Index < InResources.Num(); Index++)
-		{
-			if (InResources[Index] == InDeletedResourcePtr)
-			{
-				// remove references to a deleted resource
-				InResources[Index] = nullptr;
-			}
-		}
+		return false;
 	}
-};
-using namespace DisplayClusterViewportProxyHelpers;
-
-///////////////////////////////////////////////////////////////////////////////////////
-FDisplayClusterViewportProxy::FDisplayClusterViewportProxy(const FDisplayClusterViewport& RenderViewport)
-	: ViewportId(RenderViewport.ViewportId)
-	, ClusterNodeId(RenderViewport.ClusterNodeId)
-	, RenderSettings(RenderViewport.RenderSettings)
-	, ProjectionPolicy(RenderViewport.UninitializedProjectionPolicy)
-	, ViewportManagerProxyWeakRef(RenderViewport.GetViewportManagerProxyRefImpl())
-	, ShadersAPI(IDisplayClusterShaders::Get())
-{
-	check(ProjectionPolicy.IsValid());
-}
 
-FDisplayClusterViewportProxy::~FDisplayClusterViewportProxy()
-{
+	// Ask current projection policy if it's warp&blend compatible
+	return ProjectionPolicy.IsValid() && ProjectionPolicy->IsWarpBlendSupported();
 }
 
 const FDisplayClusterRenderFrameSettings* FDisplayClusterViewportProxy::GetRenderFrameSettings_RenderThread() const
@@ -235,29 +258,16 @@ EDisplayClusterRenderFrameMode FDisplayClusterViewportProxy::GetRenderMode() con
 	return EDisplayClusterRenderFrameMode::Unknown;
 }
 
-void FDisplayClusterViewportProxy::HandleResourceDelete_RenderThread(FDisplayClusterViewportResource* InDeletedResourcePtr)
-{
-	// Reset all references to deleted resource
-	HandleResourceDeleteImpl_RenderThread(RenderTargets, InDeletedResourcePtr);
-	HandleResourceDeleteImpl_RenderThread(OutputFrameTargetableResources, InDeletedResourcePtr);
-	HandleResourceDeleteImpl_RenderThread(AdditionalFrameTargetableResources, InDeletedResourcePtr);
-
-	HandleResourceDeleteImpl_RenderThread(InputShaderResources, InDeletedResourcePtr);
-	HandleResourceDeleteImpl_RenderThread(AdditionalTargetableResources, InDeletedResourcePtr);
-	HandleResourceDeleteImpl_RenderThread(MipsShaderResources, InDeletedResourcePtr);
-}
-
-
-bool FDisplayClusterViewportProxy::ShouldOverrideViewportResource(const EDisplayClusterViewportResourceType InResourceType) const
+bool FDisplayClusterViewportProxy::ShouldOverrideViewportResource(const EDisplayClusterViewportResourceType InExtResourceType) const
 {
 	// Override resources from other viewport
-	if (RenderSettings.IsViewportOverrided())
+	if (RenderSettings.IsViewportOverridden())
 	{
-		switch (RenderSettings.ViewportOverrideMode)
+		switch (RenderSettings.GetViewportOverrideMode())
 		{
 		case EDisplayClusterViewportOverrideMode::All:
 		{
-			switch (InResourceType)
+			switch (GetResourceType_RenderThread(InExtResourceType))
 			{
 			case EDisplayClusterViewportResourceType::InternalRenderTargetResource:
 			case EDisplayClusterViewportResourceType::InputShaderResource:
@@ -273,7 +283,7 @@ bool FDisplayClusterViewportProxy::ShouldOverrideViewportResource(const EDisplay
 
 		case EDisplayClusterViewportOverrideMode::InernalRTT:
 		{
-			switch (InResourceType)
+			switch (GetResourceType_RenderThread(InExtResourceType))
 			{
 			case EDisplayClusterViewportResourceType::InternalRenderTargetResource:
 				return true;
@@ -293,20 +303,20 @@ bool FDisplayClusterViewportProxy::ShouldOverrideViewportResource(const EDisplay
 }
 
 //  Return viewport scene proxy resources by type
-bool FDisplayClusterViewportProxy::GetResources_RenderThread(const EDisplayClusterViewportResourceType InResourceType, TArray<FRHITexture2D*>& OutResources) const
+bool FDisplayClusterViewportProxy::GetResources_RenderThread(const EDisplayClusterViewportResourceType InExtResourceType, TArray<FRHITexture2D*>& OutResources) const
 {
-	return ImplGetResources_RenderThread(InResourceType, OutResources, false);
+	return ImplGetResources_RenderThread(InExtResourceType, OutResources, false);
 }
 
 const FDisplayClusterViewportProxy& FDisplayClusterViewportProxy::GetRenderingViewportProxy() const
 {
-	switch (RenderSettings.ViewportOverrideMode)
+	switch (RenderSettings.GetViewportOverrideMode())
 	{
 	case EDisplayClusterViewportOverrideMode::All:
 	case EDisplayClusterViewportOverrideMode::InernalRTT:
 		if (FDisplayClusterViewportManagerProxy* ViewportManagerProxy = GetViewportManagerProxyImpl_RenderThread())
 		{
-			if (FDisplayClusterViewportProxy const* OverrideViewportProxy = ViewportManagerProxy->ImplFindViewportProxy_RenderThread(RenderSettings.ViewportOverrideId))
+			if (FDisplayClusterViewportProxy const* OverrideViewportProxy = ViewportManagerProxy->ImplFindViewportProxy_RenderThread(RenderSettings.GetViewportOverrideId()))
 			{
 				return *OverrideViewportProxy;
 			}
@@ -332,19 +342,21 @@ bool FDisplayClusterViewportProxy::IsInputRenderTargetResourceExists() const
 		return true;
 	}
 
-	return !RenderTargets.IsEmpty();
+	return !Resources[EDisplayClusterViewportResource::RenderTargets].IsEmpty();
 }
 
-bool FDisplayClusterViewportProxy::ImplGetResources_RenderThread(const EDisplayClusterViewportResourceType InResourceType, TArray<FRHITexture2D*>& OutResources, const int32 InRecursionDepth) const
+bool FDisplayClusterViewportProxy::ImplGetResources_RenderThread(const EDisplayClusterViewportResourceType InExtResourceType, TArray<FRHITexture2D*>& OutResources, const int32 InRecursionDepth) const
 {
 	check(IsInRenderingThread());
 
+	const EDisplayClusterViewportResourceType InResourceType = GetResourceType_RenderThread(InExtResourceType);
+
 	// Override resources from other viewport
 	if (ShouldOverrideViewportResource(InResourceType))
 	{
 		if (InRecursionDepth < DisplayClusterViewportProxyResourcesOverrideRecursionDepthMax)
 		{
-			return GetRenderingViewportProxy().ImplGetResources_RenderThread(InResourceType, OutResources, InRecursionDepth + 1);
+			return GetRenderingViewportProxy().ImplGetResources_RenderThread(InExtResourceType, OutResources, InRecursionDepth + 1);
 		}
 
 		return false;
@@ -401,7 +413,7 @@ bool FDisplayClusterViewportProxy::ImplGetResources_RenderThread(const EDisplayC
 			// 3. Finally Use InternalRTT
 			if (!bResult)
 			{
-				bResult = GetViewportRHIResourcesImpl_RenderThread(RenderTargets, OutResources);
+				bResult = Resources.GetRHIResources_RenderThread(EDisplayClusterViewportResource::RenderTargets, OutResources);
 			}
 		}
 
@@ -414,38 +426,22 @@ bool FDisplayClusterViewportProxy::ImplGetResources_RenderThread(const EDisplayC
 	}
 
 	case EDisplayClusterViewportResourceType::InputShaderResource:
-		return GetViewportRHIResourcesImpl_RenderThread(InputShaderResources, OutResources);
+		return Resources.GetRHIResources_RenderThread(EDisplayClusterViewportResource::InputShaderResources, OutResources);
 
 	case EDisplayClusterViewportResourceType::AdditionalTargetableResource:
-		return GetViewportRHIResourcesImpl_RenderThread(AdditionalTargetableResources, OutResources);
+		return Resources.GetRHIResources_RenderThread(EDisplayClusterViewportResource::AdditionalTargetableResources, OutResources);
 
 	case EDisplayClusterViewportResourceType::MipsShaderResource:
-		return GetViewportRHIResourcesImpl_RenderThread(MipsShaderResources, OutResources);
+		return Resources.GetRHIResources_RenderThread(EDisplayClusterViewportResource::MipsShaderResources, OutResources);
 
 	case EDisplayClusterViewportResourceType::OutputFrameTargetableResource:
-		return GetViewportRHIResourcesImpl_RenderThread(OutputFrameTargetableResources, OutResources);
+		return Resources.GetRHIResources_RenderThread(EDisplayClusterViewportResource::OutputFrameTargetableResources, OutResources);
 
 	case EDisplayClusterViewportResourceType::AdditionalFrameTargetableResource:
-		return GetViewportRHIResourcesImpl_RenderThread(AdditionalFrameTargetableResources, OutResources);
+		return Resources.GetRHIResources_RenderThread(EDisplayClusterViewportResource::AdditionalFrameTargetableResources, OutResources);
 
-#if WITH_EDITOR
-		// Support preview:
 	case EDisplayClusterViewportResourceType::OutputPreviewTargetableResource:
-	{
-		// Get external resource:
-		if (OutputPreviewTargetableResource.IsValid())
-		{
-			FRHITexture* Texture = OutputPreviewTargetableResource;
-			FRHITexture2D* Texture2D = static_cast<FRHITexture2D*>(Texture);
-			if (Texture2D != nullptr)
-			{
-				OutResources.Add(Texture2D);
-				return true;
-			}
-		}
-		break;
-	}
-#endif
+		return Resources.GetRHIResources_RenderThread(EDisplayClusterViewportResource::OutputPreviewTargetableResources, OutResources);
 
 	default:
 		break;
@@ -495,23 +491,23 @@ void FDisplayClusterViewportProxy::ImplViewportRemap_RenderThread(FRHICommandLis
 		const IDisplayClusterRender_MeshComponentProxy* MeshProxy = RemapMesh->GetMeshComponentProxy_RenderThread();
 		if (MeshProxy!=nullptr && MeshProxy->IsEnabled_RenderThread())
 		{
-			if (AdditionalFrameTargetableResources.Num() != OutputFrameTargetableResources.Num())
+			if (Resources[EDisplayClusterViewportResource::AdditionalFrameTargetableResources].Num() != Resources[EDisplayClusterViewportResource::OutputFrameTargetableResources].Num())
 			{
 				// error
 				return;
 			}
 
-			for (int32 ContextIt = 0; ContextIt < AdditionalFrameTargetableResources.Num(); ContextIt++)
+			for (int32 ContextIt = 0; ContextIt < Resources[EDisplayClusterViewportResource::AdditionalFrameTargetableResources].Num(); ContextIt++)
 			{
-				FDisplayClusterViewportTextureResource* Src = AdditionalFrameTargetableResources[ContextIt];
-				FDisplayClusterViewportTextureResource* Dst = OutputFrameTargetableResources[ContextIt];
+				const TSharedPtr<FDisplayClusterViewportResource, ESPMode::ThreadSafe>& Src = Resources[EDisplayClusterViewportResource::AdditionalFrameTargetableResources][ContextIt];
+				const TSharedPtr<FDisplayClusterViewportResource, ESPMode::ThreadSafe>& Dst = Resources[EDisplayClusterViewportResource::OutputFrameTargetableResources][ContextIt];
 
-				FRHITexture2D* Input = Src ? Src->GetViewportResource2DRHI() : nullptr;
-				FRHITexture2D* Output = Dst ? Dst->GetViewportResource2DRHI() : nullptr;
+				FRHITexture2D* Input = Src.IsValid() ? Src->GetViewportResourceRHI() : nullptr;
+				FRHITexture2D* Output = Dst.IsValid() ? Dst->GetViewportResourceRHI() : nullptr;
 
 				if (Input && Output)
 				{
-					ShadersAPI.RenderPostprocess_OutputRemap(RHICmdList, Input, Output, *MeshProxy);
+					GetShadersAPI().RenderPostprocess_OutputRemap(RHICmdList, Input, Output, *MeshProxy);
 				}
 			}
 		}
@@ -527,7 +523,7 @@ void FDisplayClusterViewportProxy::ImplPreviewReadPixels_RenderThread(FRHIComman
 	}
 
 	// Now try to read until success:
-	if (bPreviewReadPixels && Contexts.Num() > 0 && OutputPreviewTargetableResource.IsValid())
+	if (bPreviewReadPixels && Contexts.Num() > 0 && !Resources[EDisplayClusterViewportResource::OutputPreviewTargetableResources].IsEmpty() && Resources[EDisplayClusterViewportResource::OutputPreviewTargetableResources][0].IsValid())
 	{
 		const FDisplayClusterRenderFrameSettings* RenderFrameSettings = GetRenderFrameSettings_RenderThread();
 		check(RenderFrameSettings && RenderFrameSettings->RenderMode == EDisplayClusterRenderFrameMode::PreviewInScene);
@@ -538,8 +534,7 @@ void FDisplayClusterViewportProxy::ImplPreviewReadPixels_RenderThread(FRHIComman
 		if (PreviewPixels.IsValid() == false)
 		{
 			// Read pixels from this texture
-			FRHITexture* Texture = OutputPreviewTargetableResource;
-			if (FRHITexture2D* Texture2D = static_cast<FRHITexture2D*>(Texture))
+			if (FRHITexture2D* Texture2D = Resources[EDisplayClusterViewportResource::OutputPreviewTargetableResources][0]->GetViewportResourceRHI())
 			{
 				
 				// Clear deferred read flag
@@ -548,7 +543,7 @@ void FDisplayClusterViewportProxy::ImplPreviewReadPixels_RenderThread(FRHIComman
 				TSharedPtr<FDisplayClusterViewportReadPixelsData, ESPMode::ThreadSafe> ReadData = MakeShared<FDisplayClusterViewportReadPixelsData, ESPMode::ThreadSafe>();
 
 				RHICmdList.ReadSurfaceData(
-					Texture,
+					Texture2D,
 					FIntRect(FIntPoint(0, 0), Texture2D->GetSizeXY()),
 					ReadData->Pixels,
 					FReadSurfaceDataFlags(RCM_UNorm, CubeFace_MAX)
@@ -583,50 +578,28 @@ bool FDisplayClusterViewportProxy::GetPreviewPixels_GameThread(TSharedPtr<FDispl
 }
 #endif
 
-EDisplayClusterViewportResourceType FDisplayClusterViewportProxy::GetOutputResourceType_RenderThread() const
-{
-	check(IsInRenderingThread());
-
-#if WITH_EDITOR
-	if(GetRenderMode() == EDisplayClusterRenderFrameMode::PreviewInScene)
-	{
-		return EDisplayClusterViewportResourceType::OutputPreviewTargetableResource;
-	}
-#endif
-
-	if (RemapMesh.IsValid())
-	{
-		const IDisplayClusterRender_MeshComponentProxy* MeshProxy = RemapMesh->GetMeshComponentProxy_RenderThread();
-		if (MeshProxy!=nullptr && MeshProxy->IsEnabled_RenderThread())
-		{
-			// In this case render to additional frame targetable
-			return EDisplayClusterViewportResourceType::AdditionalFrameTargetableResource;
-		}
-	}
-
-	return EDisplayClusterViewportResourceType::OutputFrameTargetableResource;
-}
-
-bool FDisplayClusterViewportProxy::GetResourcesWithRects_RenderThread(const EDisplayClusterViewportResourceType InResourceType, TArray<FRHITexture2D*>& OutResources, TArray<FIntRect>& OutResourceRects) const
+bool FDisplayClusterViewportProxy::GetResourcesWithRects_RenderThread(const EDisplayClusterViewportResourceType InExtResourceType, TArray<FRHITexture2D*>& OutResources, TArray<FIntRect>& OutResourceRects) const
 {
-	return ImplGetResourcesWithRects_RenderThread(InResourceType, OutResources, OutResourceRects, false);
+	return ImplGetResourcesWithRects_RenderThread(InExtResourceType, OutResources, OutResourceRects, false);
 }
 
-bool FDisplayClusterViewportProxy::ImplGetResourcesWithRects_RenderThread(const EDisplayClusterViewportResourceType InResourceType, TArray<FRHITexture2D*>& OutResources, TArray<FIntRect>& OutResourceRects, const int32 InRecursionDepth) const
+bool FDisplayClusterViewportProxy::ImplGetResourcesWithRects_RenderThread(const EDisplayClusterViewportResourceType InExtResourceType, TArray<FRHITexture2D*>& OutResources, TArray<FIntRect>& OutResourceRects, const int32 InRecursionDepth) const
 {
 	check(IsInRenderingThread());
 
 	// Override resources from other viewport
-	if(ShouldOverrideViewportResource(InResourceType))
+	if(ShouldOverrideViewportResource(InExtResourceType))
 	{
 		if (InRecursionDepth < DisplayClusterViewportProxyResourcesOverrideRecursionDepthMax)
 		{
-			return GetRenderingViewportProxy().ImplGetResourcesWithRects_RenderThread(InResourceType, OutResources, OutResourceRects, InRecursionDepth + 1);
+			return GetRenderingViewportProxy().ImplGetResourcesWithRects_RenderThread(InExtResourceType, OutResources, OutResourceRects, InRecursionDepth + 1);
 		}
 
 		return false;
 	}
 
+	const EDisplayClusterViewportResourceType InResourceType = GetResourceType_RenderThread(InExtResourceType);
+
 	if (!GetResources_RenderThread(InResourceType, OutResources))
 	{
 		return false;
@@ -678,7 +651,7 @@ void FDisplayClusterViewportProxy::UpdateDeferredResources(FRHICommandListImmedi
 		return;
 	}
 
-	if (RenderSettings.IsViewportOverrided() && RenderSettings.ViewportOverrideMode == EDisplayClusterViewportOverrideMode::All)
+	if (RenderSettings.GetViewportOverrideMode() == EDisplayClusterViewportOverrideMode::All)
 	{
 		// Disable deferred update for clone viewports
 		return;
@@ -744,7 +717,7 @@ void FDisplayClusterViewportProxy::UpdateDeferredResources(FRHICommandListImmedi
 			// Render postprocess blur:
 			for (int32 ContextNum = 0; ContextNum < InShaderResources.Num(); ContextNum++)
 			{
-				ShadersAPI.RenderPostprocess_Blur(RHICmdList, InShaderResources[ContextNum], OutTargetableResources[ContextNum], PostRenderSettings.PostprocessBlur);
+				GetShadersAPI().RenderPostprocess_Blur(RHICmdList, InShaderResources[ContextNum], OutTargetableResources[ContextNum], PostRenderSettings.PostprocessBlur);
 			}
 
 			// Copy result back to input
@@ -764,7 +737,7 @@ void FDisplayClusterViewportProxy::UpdateDeferredResources(FRHICommandListImmedi
 			// Generate mips
 			for (FRHITexture2D*& ResourceIt : InOutMipsResources)
 			{
-				ShadersAPI.GenerateMips(RHICmdList, ResourceIt, PostRenderSettings.GenerateMips);
+				GetShadersAPI().GenerateMips(RHICmdList, ResourceIt, PostRenderSettings.GenerateMips);
 			}
 		}
 	}
@@ -871,10 +844,12 @@ void ImplResolveResource(FRHICommandListImmediate& RHICmdList, FRHITexture2D* In
 	}
 }
 
-FIntRect FDisplayClusterViewportProxy::GetFinalContextRect(const EDisplayClusterViewportResourceType InputResourceType, const FIntRect& InRect) const
+FIntRect FDisplayClusterViewportProxy::GetFinalContextRect(const EDisplayClusterViewportResourceType InExtResourceType, const FIntRect& InRect) const
 {
+	const EDisplayClusterViewportResourceType InResourceType = GetResourceType_RenderThread(InExtResourceType);
+
 	// When resolving without warp, apply overscan
-	switch (InputResourceType)
+	switch (InResourceType)
 	{
 	case EDisplayClusterViewportResourceType::InternalRenderTargetResource:
 		if (OverscanSettings.bIsEnabled && CVarDisplayClusterRenderOverscanResolve.GetValueOnRenderThread() != 0)
@@ -891,53 +866,53 @@ FIntRect FDisplayClusterViewportProxy::GetFinalContextRect(const EDisplayCluster
 }
 
 // Resolve resource contexts
-bool FDisplayClusterViewportProxy::ResolveResources_RenderThread(FRHICommandListImmediate& RHICmdList, const EDisplayClusterViewportResourceType InputResourceType, const EDisplayClusterViewportResourceType OutputResourceType, const int32 InContextNum) const
+bool FDisplayClusterViewportProxy::ResolveResources_RenderThread(FRHICommandListImmediate& RHICmdList, const EDisplayClusterViewportResourceType InExtResourceType, const EDisplayClusterViewportResourceType OutExtResourceType, const int32 InContextNum) const
 {
-	return ImplResolveResources_RenderThread(RHICmdList, this, InputResourceType, OutputResourceType, InContextNum);
+	return ImplResolveResources_RenderThread(RHICmdList, this, InExtResourceType, OutExtResourceType, InContextNum);
 }
 
-bool FDisplayClusterViewportProxy::ImplResolveResources_RenderThread(FRHICommandListImmediate& RHICmdList, FDisplayClusterViewportProxy const* SourceProxy, const EDisplayClusterViewportResourceType InputResourceType, const EDisplayClusterViewportResourceType OutputResourceType, const int32 InContextNum) const
+bool FDisplayClusterViewportProxy::ImplResolveResources_RenderThread(FRHICommandListImmediate& RHICmdList, FDisplayClusterViewportProxy const* SourceProxy, const EDisplayClusterViewportResourceType InExtResourceType, const EDisplayClusterViewportResourceType OutExtResourceType, const int32 InContextNum) const
 {
 	check(IsInRenderingThread());
+	check(SourceProxy);
 
-	if (InputResourceType == EDisplayClusterViewportResourceType::MipsShaderResource) {
+	const EDisplayClusterViewportResourceType InResourceType = SourceProxy->GetResourceType_RenderThread(InExtResourceType);
+	const EDisplayClusterViewportResourceType OutResourceType = GetResourceType_RenderThread(OutExtResourceType);
+
+	if (InResourceType == EDisplayClusterViewportResourceType::MipsShaderResource) {
 		// RenderTargetMips not allowved for resolve op
 		return false;
 	}
 
-	bool bOutputIsMipsResource = OutputResourceType == EDisplayClusterViewportResourceType::MipsShaderResource;
+	bool bOutputIsMipsResource = OutResourceType == EDisplayClusterViewportResourceType::MipsShaderResource;
 
-#if WITH_EDITOR
 	// This resolve pattern always called once for preview. This flag force to invert alpha (at this point alpha from engine is inverted)
-	const bool bOutputIsPreviewResource = InputResourceType == EDisplayClusterViewportResourceType::InputShaderResource && OutputResourceType == EDisplayClusterViewportResourceType::OutputPreviewTargetableResource;
-#else
-	const bool bOutputIsPreviewResource = false;
-#endif
+	const bool bOutputIsPreviewResource = InResourceType == EDisplayClusterViewportResourceType::InputShaderResource && OutResourceType == EDisplayClusterViewportResourceType::OutputPreviewTargetableResource;
 
 	TArray<FRHITexture2D*> SrcResources, DestResources;
 	TArray<FIntRect> SrcResourcesRect, DestResourcesRect;
-	if (SourceProxy->GetResourcesWithRects_RenderThread(InputResourceType, SrcResources, SrcResourcesRect) && GetResourcesWithRects_RenderThread(OutputResourceType, DestResources, DestResourcesRect))
+	if (SourceProxy->GetResourcesWithRects_RenderThread(InExtResourceType, SrcResources, SrcResourcesRect) && GetResourcesWithRects_RenderThread(OutExtResourceType, DestResources, DestResourcesRect))
 	{
 		const int32 SrcResourcesAmount = FMath::Min(SrcResources.Num(), DestResources.Num());
 		for (int32 SrcResourceContextIndex = 0; SrcResourceContextIndex < SrcResourcesAmount; SrcResourceContextIndex++)
 		{
 			// Get context num to copy
 			const int32 SrcContextNum = (InContextNum == INDEX_NONE) ? SrcResourceContextIndex : InContextNum;
-			const FIntRect SrcRect = GetFinalContextRect(InputResourceType, SrcResourcesRect[SrcContextNum]);
+			const FIntRect SrcRect = SourceProxy->GetFinalContextRect(InExtResourceType, SrcResourcesRect[SrcContextNum]);
 
 			if ((SrcContextNum + 1) == SrcResourcesAmount)
 			{
 				// last input mono -> stereo outputs
 				for (int32 DestResourceContextIndex = SrcContextNum; DestResourceContextIndex < DestResources.Num(); DestResourceContextIndex++)
 				{
-					const FIntRect DestRect = GetFinalContextRect(OutputResourceType, DestResourcesRect[DestResourceContextIndex]);
+					const FIntRect DestRect = GetFinalContextRect(OutResourceType, DestResourcesRect[DestResourceContextIndex]);
 					ImplResolveResource(RHICmdList, SrcResources[SrcContextNum], SrcRect, DestResources[DestResourceContextIndex], DestRect, bOutputIsMipsResource, bOutputIsPreviewResource);
 				}
 				break;
 			}
 			else
 			{
-				const FIntRect DestRect = GetFinalContextRect(OutputResourceType, DestResourcesRect[SrcContextNum]);
+				const FIntRect DestRect = GetFinalContextRect(OutResourceType, DestResourcesRect[SrcContextNum]);
 				ImplResolveResource(RHICmdList, SrcResources[SrcContextNum], SrcRect, DestResources[SrcContextNum], DestRect, bOutputIsMipsResource, bOutputIsPreviewResource);
 			}
 
@@ -954,11 +929,11 @@ bool FDisplayClusterViewportProxy::ImplResolveResources_RenderThread(FRHICommand
 	return false;
 }
 
-bool FDisplayClusterViewportProxy::CopyResource_RenderThread(FRDGBuilder& GraphBuilder, const EDisplayClusterTextureCopyMode InCopyMode, const int32 InContextNum, const EDisplayClusterViewportResourceType InSrcResourceType, const EDisplayClusterViewportResourceType InDestResourceType)
+bool FDisplayClusterViewportProxy::CopyResource_RenderThread(FRDGBuilder& GraphBuilder, const EDisplayClusterTextureCopyMode InCopyMode, const int32 InContextNum, const EDisplayClusterViewportResourceType InExtSrcResourceType, const EDisplayClusterViewportResourceType InExtDestResourceType)
 {
 	TArray<FRHITexture2D*> SrcResources;
 	TArray<FIntRect> SrcResourceRects;
-	if (GetResourcesWithRects_RenderThread(InSrcResourceType, SrcResources, SrcResourceRects))
+	if (GetResourcesWithRects_RenderThread(InExtSrcResourceType, SrcResources, SrcResourceRects))
 	{
 		check(SrcResources.IsValidIndex(InContextNum));
 		check(SrcResourceRects.IsValidIndex(InContextNum));
@@ -968,7 +943,7 @@ bool FDisplayClusterViewportProxy::CopyResource_RenderThread(FRDGBuilder& GraphB
 		FRDGTextureRef SrcTextureRef = RegisterExternalTexture(GraphBuilder, SrcResources[InContextNum]->GetTexture2D(), TEXT("DCViewportProxyCopySrcResource"));
 		GraphBuilder.SetTextureAccessFinal(SrcTextureRef, ERHIAccess::SRVGraphics);
 
-		return CopyResource_RenderThread(GraphBuilder, InCopyMode, InContextNum, SrcTextureRef, SrcRect, InDestResourceType);
+		return CopyResource_RenderThread(GraphBuilder, InCopyMode, InContextNum, SrcTextureRef, SrcRect, InExtDestResourceType);
 	}
 
 	return false;
@@ -979,11 +954,11 @@ RDG_TEXTURE_ACCESS(Input, ERHIAccess::CopySrc)
 RDG_TEXTURE_ACCESS(Output, ERHIAccess::CopyDest)
 END_SHADER_PARAMETER_STRUCT()
 
-bool FDisplayClusterViewportProxy::CopyResource_RenderThread(FRDGBuilder& GraphBuilder, const EDisplayClusterTextureCopyMode InCopyMode, const int32 InContextNum, FRDGTextureRef InSrcTextureRef, const FIntRect& InSrcRect, const EDisplayClusterViewportResourceType InDestResourceType)
+bool FDisplayClusterViewportProxy::CopyResource_RenderThread(FRDGBuilder& GraphBuilder, const EDisplayClusterTextureCopyMode InCopyMode, const int32 InContextNum, FRDGTextureRef InSrcTextureRef, const FIntRect& InSrcRect, const EDisplayClusterViewportResourceType InExtDestResourceType)
 {
 	TArray<FRHITexture2D*> DestResources;
 	TArray<FIntRect> DestResourceRects;
-	if (HasBeenProduced(InSrcTextureRef) && GetResourcesWithRects_RenderThread(InDestResourceType, DestResources, DestResourceRects))
+	if (HasBeenProduced(InSrcTextureRef) && GetResourcesWithRects_RenderThread(InExtDestResourceType, DestResources, DestResourceRects))
 	{
 		check(DestResources.IsValidIndex(InContextNum));
 		check(DestResourceRects.IsValidIndex(InContextNum));
@@ -1011,11 +986,11 @@ bool FDisplayClusterViewportProxy::CopyResource_RenderThread(FRDGBuilder& GraphB
 	return false;
 }
 
-bool FDisplayClusterViewportProxy::CopyResource_RenderThread(FRDGBuilder& GraphBuilder, const EDisplayClusterTextureCopyMode InCopyMode, const int32 InContextNum, const EDisplayClusterViewportResourceType InSrcResourceType, FRDGTextureRef InDestTextureRef, const FIntRect& InDestRect)
+bool FDisplayClusterViewportProxy::CopyResource_RenderThread(FRDGBuilder& GraphBuilder, const EDisplayClusterTextureCopyMode InCopyMode, const int32 InContextNum, const EDisplayClusterViewportResourceType InExtSrcResourceType, FRDGTextureRef InDestTextureRef, const FIntRect& InDestRect)
 {
 	TArray<FRHITexture2D*> SrcResources;
 	TArray<FIntRect> SrcResourceRects;
-	if (HasBeenProduced(InDestTextureRef) && GetResourcesWithRects_RenderThread(InSrcResourceType, SrcResources, SrcResourceRects))
+	if (HasBeenProduced(InDestTextureRef) && GetResourcesWithRects_RenderThread(InExtSrcResourceType, SrcResources, SrcResourceRects))
 	{
 		check(SrcResources.IsValidIndex(InContextNum));
 		check(SrcResourceRects.IsValidIndex(InContextNum));
@@ -1195,9 +1170,9 @@ void FDisplayClusterViewportProxy::OnPostRenderViewFamily_RenderThread(FRDGBuild
 		{
 			// Restore alpha channed
 			EFXAAQuality FXAAQuality = EFXAAQuality::Q0;
-			if (GetFXAAQuality(RenderSettings.CaptureMode, FXAAQuality) && InputShaderResources.IsValidIndex(InContextNum))
+			if (GetFXAAQuality(RenderSettings.CaptureMode, FXAAQuality) && Resources[EDisplayClusterViewportResource::InputShaderResources].IsValidIndex(InContextNum))
 			{
-				if (FRHITexture2D* InputTextureRHI = InputShaderResources[InContextNum] ? InputShaderResources[InContextNum]->GetViewportResource2DRHI() : nullptr)
+				if (FRHITexture2D* InputTextureRHI = Resources[EDisplayClusterViewportResource::InputShaderResources][InContextNum] ? Resources[EDisplayClusterViewportResource::InputShaderResources][InContextNum]->GetViewportResourceRHI() : nullptr)
 				{
 					// Apply FXAA for RGB only
 					// Note: Add AA for alpha channel
diff --git a/Engine/Plugins/Runtime/nDisplay/Source/DisplayCluster/Private/Render/Viewport/DisplayClusterViewportProxy.h b/Engine/Plugins/Runtime/nDisplay/Source/DisplayCluster/Private/Render/Viewport/DisplayClusterViewportProxy.h
--- a/Engine/Plugins/Runtime/nDisplay/Source/DisplayCluster/Private/Render/Viewport/DisplayClusterViewportProxy.h
+++ b/Engine/Plugins/Runtime/nDisplay/Source/DisplayCluster/Private/Render/Viewport/DisplayClusterViewportProxy.h
@@ -10,10 +10,9 @@
 #include "Templates/SharedPointer.h"
 
 #include "Render/Viewport/IDisplayClusterViewportProxy.h"
+#include "Render/Viewport/DisplayClusterViewportResources.h"
 #include "Render/Viewport/Containers/DisplayClusterViewport_OverscanSettings.h"
 
-class FDisplayClusterViewportRenderTargetResource;
-class FDisplayClusterViewportTextureResource;
 class FDisplayClusterViewportProxyData;
 class IDisplayClusterViewportManagerProxy;
 class FDisplayClusterViewportManagerProxy;
@@ -143,7 +142,6 @@ public:
 	// Resolve resource contexts
 	virtual bool ResolveResources_RenderThread(FRHICommandListImmediate& RHICmdList, const EDisplayClusterViewportResourceType InputResourceType, const EDisplayClusterViewportResourceType OutputResourceType, const int32 InContextNum = INDEX_NONE) const override;
 
-	virtual EDisplayClusterViewportResourceType GetOutputResourceType_RenderThread() const override;
 
 	virtual const class IDisplayClusterViewportManagerProxy* GetViewportManagerProxy_RenderThread() const override;
 
@@ -153,8 +151,22 @@ public:
 	// ~IDisplayClusterViewportProxy
 	///////////////////////////////
 
-	/** Release internal resource refs. */
-	void HandleResourceDelete_RenderThread(class FDisplayClusterViewportResource* InDeletedResourcePtr);
+	/** Get valid resource type
+	 * 
+	 * @param InResourceType - the requested resource type from the entire namespace
+	 * 
+	 * @return - the type of resource actually used, depending on the current configuration of the rendering pipeline
+	 */
+	EDisplayClusterViewportResourceType GetResourceType_RenderThread(const EDisplayClusterViewportResourceType& InResourceType) const;
+
+	/** Get actual region for viewport context
+	 * 
+	 * @param InResourceType - the requested resource type from the entire namespace
+	 * @param InRect - viewport context region
+	 * 
+	 * @return - InRect with applied overscan for RTT
+	 */
+	FIntRect GetFinalContextRect(const EDisplayClusterViewportResourceType InResourceType, const FIntRect& InRect) const;
 
 	/** Resolve viewport RTT: render OCIO, PP, generate MIPS, etc.
 	 *
@@ -262,7 +274,9 @@ public:
 		return false;
 	}
 
-	FIntRect GetFinalContextRect(const EDisplayClusterViewportResourceType InputResourceType, const FIntRect& InRect) const;
+public:
+	/* Returns true if the warp can be applied to this viewport. */
+	bool ShouldApplyWarpBlend_RenderThread() const;
 
 private:
 	bool ImplGetResourcesWithRects_RenderThread(const EDisplayClusterViewportResourceType InResourceType, TArray<FRHITexture2D*>& OutResources, TArray<FIntRect>& OutResourceRects, const int32 InRecursionDepth) const;
@@ -357,29 +371,17 @@ protected:
 	// Viewport contexts (left/center/right eyes)
 	mutable TArray<FDisplayClusterViewport_Context> Contexts;
 
-	// View family render to this resources
-	TArray<FDisplayClusterViewportRenderTargetResource*> RenderTargets;
+	// Unified repository of viewport resources
+	FDisplayClusterViewportResources Resources;
 
-	// Projection policy output resources
-	TArray<FDisplayClusterViewportTextureResource*> OutputFrameTargetableResources;
-	TArray<FDisplayClusterViewportTextureResource*> AdditionalFrameTargetableResources;
-
-#if WITH_EDITOR
-	FTextureRHIRef OutputPreviewTargetableResource;
+	// Used ViewStates
 	TArray<TSharedPtr<FSceneViewStateReference, ESPMode::ThreadSafe>> ViewStates;
 	
+#if WITH_EDITOR
 	mutable bool bPreviewReadPixels = false;
 	mutable FCriticalSection PreviewPixelsCSGuard;
 	mutable TSharedPtr<FDisplayClusterViewportReadPixelsData, ESPMode::ThreadSafe> PreviewPixels;
 #endif
 
-	// unique viewport resources
-	TArray<FDisplayClusterViewportTextureResource*> InputShaderResources;
-	TArray<FDisplayClusterViewportTextureResource*> AdditionalTargetableResources;
-	TArray<FDisplayClusterViewportTextureResource*> MipsShaderResources;
-
 	TWeakPtr<FDisplayClusterViewportManagerProxy, ESPMode::ThreadSafe> ViewportManagerProxyWeakRef;
-
-	IDisplayClusterShaders& ShadersAPI;
 };
-
diff --git a/Engine/Plugins/Runtime/nDisplay/Source/DisplayCluster/Private/Render/Viewport/DisplayClusterViewportResources.cpp b/Engine/Plugins/Runtime/nDisplay/Source/DisplayCluster/Private/Render/Viewport/DisplayClusterViewportResources.cpp
new file mode 100644
--- /dev/null
+++ b/Engine/Plugins/Runtime/nDisplay/Source/DisplayCluster/Private/Render/Viewport/DisplayClusterViewportResources.cpp
@@ -0,0 +1,46 @@
+// Copyright Epic Games, Inc. All Rights Reserved.
+
+#include "Render/Viewport/DisplayClusterViewportResources.h"
+#include "Render/Viewport/Resource/DisplayClusterViewportResource.h"
+
+/////////////////////////////////////////////////////////////////////
+// FDisplayClusterViewportResources
+/////////////////////////////////////////////////////////////////////
+const TArray<TSharedPtr<FDisplayClusterViewportResource, ESPMode::ThreadSafe>> FDisplayClusterViewportResources::EmptyResources;
+
+void FDisplayClusterViewportResources::FreezeRendering(const EDisplayClusterViewportResource InResourceType)
+{
+	for (const TSharedPtr<FDisplayClusterViewportResource, ESPMode::ThreadSafe>& TextureResource : operator[](InResourceType))
+	{
+		if (TextureResource.IsValid())
+		{
+			EnumAddFlags(TextureResource->GetResourceState(), EDisplayClusterViewportResourceState::DisableReallocate);
+		}
+	}
+}
+
+bool FDisplayClusterViewportResources::GetRHIResources_RenderThread(const EDisplayClusterViewportResource InResourceType, TArray<FRHITexture2D*>& OutResources) const
+{
+	OutResources.Reset();
+
+	if (const TArray<TSharedPtr<FDisplayClusterViewportResource, ESPMode::ThreadSafe>>* ExistResources = ViewportResources.Find(InResourceType))
+	{
+		for (const TSharedPtr<FDisplayClusterViewportResource, ESPMode::ThreadSafe>& ViewportResourceIt : *ExistResources)
+		{
+			if (FRHITexture2D* RHITexture2D = ViewportResourceIt.IsValid() ? ViewportResourceIt->GetViewportResourceRHI() : nullptr)
+			{
+				// Collects only valid resources.
+				OutResources.Add(RHITexture2D);
+			}
+		}
+
+		if (OutResources.Num() != ExistResources->Num())
+		{
+			// Some resources lost
+			OutResources.Reset();
+		}
+	}
+
+	// returns success if the number of output resources is equal to the input and is not empty
+	return !OutResources.IsEmpty();
+}
diff --git a/Engine/Plugins/Runtime/nDisplay/Source/DisplayCluster/Private/Render/Viewport/DisplayClusterViewportResources.h b/Engine/Plugins/Runtime/nDisplay/Source/DisplayCluster/Private/Render/Viewport/DisplayClusterViewportResources.h
new file mode 100644
--- /dev/null
+++ b/Engine/Plugins/Runtime/nDisplay/Source/DisplayCluster/Private/Render/Viewport/DisplayClusterViewportResources.h
@@ -0,0 +1,128 @@
+// Copyright Epic Games, Inc. All Rights Reserved.
+
+#pragma once
+
+#include "CoreMinimal.h"
+
+#include "RHI.h"
+#include "RHIResources.h"
+#include "RenderResource.h"
+
+class FDisplayClusterViewportResource;
+
+/**
+ * Internal types of DC viewport resources
+ * These types are used to point to actual instances of the resources in the viewport.
+ */
+enum class EDisplayClusterViewportResource : uint8
+{
+	// View family render to this resources
+	// This is the resource that can be pointed to using the enum 'EDisplayClusterViewportResourceType::InternalRenderTargetResource' value.
+	RenderTargets,
+
+	// unique viewport resources
+	// This is the resource that can be pointed to using the enum 'EDisplayClusterViewportResourceType::InputShaderResource' value.
+	InputShaderResources,
+
+	// unique viewport resources
+	// This is the resource that can be pointed to using the enum 'EDisplayClusterViewportResourceType::AdditionalTargetableResource' value.
+	AdditionalTargetableResources,
+
+	// unique viewport resources
+	// This is the resource that can be pointed to using the enum 'EDisplayClusterViewportResourceType::MipsShaderResource' value.
+	MipsShaderResources,
+
+	// Projection policy output resources
+	// This is the resource that can be pointed to using the enum 'EDisplayClusterViewportResourceType::OutputPreviewTargetableResource' value.
+	OutputPreviewTargetableResources,
+
+	// Projection policy output resources
+	// This is the resource that can be pointed to using the enum 'EDisplayClusterViewportResourceType::OutputFrameTargetableResource' value.
+	OutputFrameTargetableResources,
+
+	// Projection policy output resources
+	// This is the resource that can be pointed to using the enum 'EDisplayClusterViewportResourceType::AdditionalFrameTargetableResource' value.
+	AdditionalFrameTargetableResources,
+};
+
+/** Helper class for unified viewport resource repository. */
+class FDisplayClusterViewportResources
+{
+public:
+	FDisplayClusterViewportResources() = default;
+	~FDisplayClusterViewportResources() = default;
+
+public:
+	/** Get viewport resources for all contexts by type. */
+	TArray<TSharedPtr<FDisplayClusterViewportResource, ESPMode::ThreadSafe>>& operator[](const EDisplayClusterViewportResource InResourceType)
+	{
+		if (TArray<TSharedPtr<FDisplayClusterViewportResource, ESPMode::ThreadSafe>>* ExistResources = ViewportResources.Find(InResourceType))
+		{
+			return *ExistResources;
+		}
+
+		return ViewportResources.Emplace(InResourceType);
+	}
+
+	/** Get viewport const resources for all contexts by type. */
+	const TArray<TSharedPtr<FDisplayClusterViewportResource, ESPMode::ThreadSafe>>& operator[](const EDisplayClusterViewportResource InResourceType) const
+	{
+		if (const TArray<TSharedPtr<FDisplayClusterViewportResource, ESPMode::ThreadSafe>>* ExistResources = ViewportResources.Find(InResourceType))
+		{
+			return *ExistResources;
+		}
+
+		// This value is used for const-functions when resources are not defined
+		return EmptyResources;
+	}
+
+	/** Release refs to all viewport resources. */
+	void Release(bool bReleaseAllResources = true)
+	{
+		if (bReleaseAllResources)
+		{
+			ViewportResources.Empty();
+		}
+		else
+		{
+			// Release only part of the resources, leaving resources that can be used by other viewports (viewport override feature)
+			// Discard resources that are not used in frame composition
+			ImplRelease(EDisplayClusterViewportResource::OutputPreviewTargetableResources);
+			ImplRelease(EDisplayClusterViewportResource::RenderTargets);
+			ImplRelease(EDisplayClusterViewportResource::OutputFrameTargetableResources);
+			ImplRelease(EDisplayClusterViewportResource::AdditionalFrameTargetableResources);
+		}
+	}
+
+	/** Raise DisableReallocate flag for all viewport resources of specified type.
+	* Support render freeze feature
+	*/
+	void FreezeRendering(const EDisplayClusterViewportResource InResourceType);
+
+	/**
+	* Get viewport RHI resources
+	*
+	* @param InResourceType - The type of resources you want to retrieve.
+	* @param OutResources - Array with RHI resources
+	*
+	* @return true, if all resources are valid
+	*/
+	bool GetRHIResources_RenderThread(const EDisplayClusterViewportResource InResourceType, TArray<FRHITexture2D*>& OutResources) const;
+
+protected:
+	/** Release specified type of resources. */
+	inline void ImplRelease(const EDisplayClusterViewportResource InResourceType)
+	{
+		if (TArray<TSharedPtr<FDisplayClusterViewportResource, ESPMode::ThreadSafe>>* ExistResources = ViewportResources.Find(InResourceType))
+		{
+			ExistResources->Empty();
+		}
+	}
+
+private:
+	// Unified repository of viewport resources
+	TMap<EDisplayClusterViewportResource, TArray<TSharedPtr<FDisplayClusterViewportResource, ESPMode::ThreadSafe>>> ViewportResources;
+
+	// This value is used for const-functions when resources are not defined
+	const static TArray<TSharedPtr<FDisplayClusterViewportResource, ESPMode::ThreadSafe>> EmptyResources;
+};
diff --git a/Engine/Plugins/Runtime/nDisplay/Source/DisplayCluster/Private/Render/Viewport/RenderFrame/DisplayClusterRenderFrameManager.cpp b/Engine/Plugins/Runtime/nDisplay/Source/DisplayCluster/Private/Render/Viewport/RenderFrame/DisplayClusterRenderFrameManager.cpp
--- a/Engine/Plugins/Runtime/nDisplay/Source/DisplayCluster/Private/Render/Viewport/RenderFrame/DisplayClusterRenderFrameManager.cpp
+++ b/Engine/Plugins/Runtime/nDisplay/Source/DisplayCluster/Private/Render/Viewport/RenderFrame/DisplayClusterRenderFrameManager.cpp
@@ -142,7 +142,7 @@ bool FDisplayClusterRenderFrameManager::BuildSimpleFrame(FViewport* InViewport,
 				FDisplayClusterRenderFrameTarget FrameRenderTarget;
 				{
 					// Simple frame use unique RTT  for each viewport, so disable RTT when viewport rendering disabled
-					FrameRenderTarget.bShouldUseRenderTarget = ViewportIt->RenderTargets.Num() > 0;
+					FrameRenderTarget.bShouldUseRenderTarget = ViewportIt->Resources[EDisplayClusterViewportResource::RenderTargets].Num() > 0;
 
 					FrameRenderTarget.ViewFamilies.Add(FrameViewFamily);
 
diff --git a/Engine/Plugins/Runtime/nDisplay/Source/DisplayCluster/Private/Render/Viewport/RenderTarget/DisplayClusterRenderTargetManager.cpp b/Engine/Plugins/Runtime/nDisplay/Source/DisplayCluster/Private/Render/Viewport/RenderTarget/DisplayClusterRenderTargetManager.cpp
--- a/Engine/Plugins/Runtime/nDisplay/Source/DisplayCluster/Private/Render/Viewport/RenderTarget/DisplayClusterRenderTargetManager.cpp
+++ b/Engine/Plugins/Runtime/nDisplay/Source/DisplayCluster/Private/Render/Viewport/RenderTarget/DisplayClusterRenderTargetManager.cpp
@@ -12,7 +12,7 @@
 
 #include "Misc/DisplayClusterLog.h"
 
-
+////////////////////////////////////////////////////////////////////////////////////////////////////////////
 int32 GDisplayClusterSceneColorFormat = 0;
 static FAutoConsoleVariableRef CVarDisplayClusterSceneColorFormat(
 	TEXT("nDisplay.render.SceneColorFormat"),
@@ -26,7 +26,7 @@ static FAutoConsoleVariableRef CVarDisplayClusterSceneColorFormat(
 );
 
 ////////////////////////////////////////////////////////////////////////////////////////////////////////////
-namespace DisplayClusterRenderTargetManager
+namespace UE::DisplayCluster::RenderTargetManager
 {
 	static EPixelFormat ImplGetCustomFormat(EDisplayClusterViewportCaptureMode CaptureMode)
 	{
@@ -75,17 +75,16 @@ namespace DisplayClusterRenderTargetManager
 		return OutPixelFormat;
 	}
 
-	static void ImplViewportTextureResourceLogging(const TSharedPtr<FDisplayClusterViewport, ESPMode::ThreadSafe>& InViewport, const uint32 ContextNum, const FString& ResourceId, const FDisplayClusterViewportTextureResource* InTextureResource)
+	static void ImplViewportTextureResourceLogging(const TSharedPtr<FDisplayClusterViewport, ESPMode::ThreadSafe>& InViewport, const uint32 ContextNum, const FString& ResourceId, const TSharedPtr<FDisplayClusterViewportResource, ESPMode::ThreadSafe>& InTextureResource)
 	{
-		if (InViewport.IsValid() && InTextureResource != nullptr && InTextureResource->GetViewportResourceState(EDisplayClusterViewportResourceState::Initialized) == false)
+		if (InViewport.IsValid() && InTextureResource != nullptr && EnumHasAnyFlags(InTextureResource->GetResourceState(), EDisplayClusterViewportResourceState::Initialized) == false)
 		{
 			// Log: New resource created
-			UE_LOG(LogDisplayClusterViewport, Verbose, TEXT("Created new %s resource (%dx%d) for viewport '%s':%d"), *ResourceId, InTextureResource->GetSizeX(), InTextureResource->GetSizeY(), *InViewport->GetId(), ContextNum);
+			UE_LOG(LogDisplayClusterViewport, Verbose, TEXT("Created new %s resource (%dx%d) for viewport '%s':%d"), *ResourceId, InTextureResource->GetResourceSettings().GetSizeXY().X, InTextureResource->GetResourceSettings().GetSizeXY().Y, *InViewport->GetId(), ContextNum);
 		}
 	}
 };
-
-using namespace DisplayClusterRenderTargetManager;
+using namespace UE::DisplayCluster::RenderTargetManager;
 
 ////////////////////////////////////////////////////////////////////////////////////////////////////////////
 /// FDisplayClusterRenderTargetManager
@@ -117,11 +116,12 @@ bool FDisplayClusterRenderTargetManager::AllocateRenderFrameResources(FViewport*
 			if (FrameRenderTargetIt.bShouldUseRenderTarget)
 			{
 				// reallocate
-				FDisplayClusterViewportRenderTargetResource* NewResource = ResourcesPool->AllocateRenderTargetResource(FrameRenderTargetIt.RenderTargetSize, ImplGetCustomFormat(FrameRenderTargetIt.CaptureMode));
-				if (NewResource != nullptr)
+				TSharedPtr<FDisplayClusterViewportResource, ESPMode::ThreadSafe> NewResource = ResourcesPool->AllocateResource(FrameRenderTargetIt.RenderTargetSize, ImplGetCustomFormat(FrameRenderTargetIt.CaptureMode), EDisplayClusterViewportResourceSettingsFlags::RenderTarget);
+				if (NewResource.IsValid())
 				{
 					// Set RenderFrame resource
-					FrameRenderTargetIt.RenderTargetPtr = NewResource;
+					FrameRenderTargetIt.RenderTargetPtr = NewResource->GetViewportResourceRenderTarget();
+					check(FrameRenderTargetIt.RenderTargetPtr);
 
 					// Assign for all views in render target families
 					for (FDisplayClusterRenderFrameTargetViewFamily& ViewFamily : FrameRenderTargetIt.ViewFamilies)
@@ -131,15 +131,15 @@ bool FDisplayClusterRenderTargetManager::AllocateRenderFrameResources(FViewport*
 							if (FDisplayClusterViewport* ViewportPtr = static_cast<FDisplayClusterViewport*>(ViewIt.Viewport.Get()))
 							{
 								// Array already resized in function FDisplayClusterViewport::UpdateFrameContexts() with RenderTargets.AddZeroed(ViewportContextAmount);
-								check(ViewportPtr->RenderTargets.IsValidIndex(ViewIt.ContextNum));
+								check(ViewportPtr->Resources[EDisplayClusterViewportResource::RenderTargets].IsValidIndex(ViewIt.ContextNum));
 								check(!ViewportPtr->Contexts[ViewIt.ContextNum].bDisableRender);
 
-								ViewportPtr->RenderTargets[ViewIt.ContextNum] = NewResource;
+								ViewportPtr->Resources[EDisplayClusterViewportResource::RenderTargets][ViewIt.ContextNum] = NewResource;
 
-								if (NewResource->GetViewportResourceState(EDisplayClusterViewportResourceState::Initialized) == false)
+								if (EnumHasAnyFlags(NewResource->GetResourceState(), EDisplayClusterViewportResourceState::Initialized) == false)
 								{
 									// Log: New resource created
-									UE_LOG(LogDisplayClusterViewport, Verbose, TEXT("Created new ViewportRenderTarget resource %08X (%dx%d) for viewport '%s'"), NewResource, NewResource->GetSizeX(), NewResource->GetSizeY(), *ViewportPtr->GetId());
+									UE_LOG(LogDisplayClusterViewport, Verbose, TEXT("Created new ViewportRenderTarget resource %08X (%dx%d) for viewport '%s'"), NewResource.Get(), NewResource->GetResourceSettings().GetSizeX(), NewResource->GetResourceSettings().GetSizeY(), *ViewportPtr->GetId());
 								}
 							}
 						}
@@ -161,23 +161,30 @@ bool FDisplayClusterRenderTargetManager::AllocateRenderFrameResources(FViewport*
 					const FIntPoint& ContextSize = ContextIt.ContextSize;
 					const uint32& ContextNum = ContextIt.ContextNum;
 
-					if (ViewportIt->InputShaderResources.Num() > (int32)ContextNum)
+					// Allocate per-viewport preview output textures:
+					if (ViewportIt->Resources[EDisplayClusterViewportResource::OutputPreviewTargetableResources].Num() > (int32)ContextNum)
+					{
+						ViewportIt->Resources[EDisplayClusterViewportResource::OutputPreviewTargetableResources][ContextNum] = ResourcesPool->AllocateResource(ContextSize, CustomFormat, EDisplayClusterViewportResourceSettingsFlags::PreviewTargetableTexture);
+						ImplViewportTextureResourceLogging(ViewportIt, ContextNum, TEXT("InputShader"), ViewportIt->Resources[EDisplayClusterViewportResource::OutputPreviewTargetableResources][ContextNum]);
+					}
+
+					if (ViewportIt->Resources[EDisplayClusterViewportResource::InputShaderResources].Num() > (int32)ContextNum)
 					{
-						ViewportIt->InputShaderResources[ContextNum] = ResourcesPool->AllocateTextureResource(ContextSize, false, CustomFormat);
-						ImplViewportTextureResourceLogging(ViewportIt, ContextNum, TEXT("InputShader"), ViewportIt->InputShaderResources[ContextNum]);
+						ViewportIt->Resources[EDisplayClusterViewportResource::InputShaderResources][ContextNum] = ResourcesPool->AllocateResource(ContextSize, CustomFormat, EDisplayClusterViewportResourceSettingsFlags::ResolveTargetableTexture);
+						ImplViewportTextureResourceLogging(ViewportIt, ContextNum, TEXT("InputShader"), ViewportIt->Resources[EDisplayClusterViewportResource::InputShaderResources][ContextNum]);
 					}
 
 					// Allocate custom resources:
-					if (ViewportIt->AdditionalTargetableResources.Num() > (int32)ContextNum)
+					if (ViewportIt->Resources[EDisplayClusterViewportResource::AdditionalTargetableResources].Num() > (int32)ContextNum)
 					{
-						ViewportIt->AdditionalTargetableResources[ContextNum] = ResourcesPool->AllocateTextureResource(ContextSize, true, CustomFormat);
-						ImplViewportTextureResourceLogging(ViewportIt, ContextNum, TEXT("AdditionalRTT"), ViewportIt->AdditionalTargetableResources[ContextNum]);
+						ViewportIt->Resources[EDisplayClusterViewportResource::AdditionalTargetableResources][ContextNum] = ResourcesPool->AllocateResource(ContextSize, CustomFormat, EDisplayClusterViewportResourceSettingsFlags::RenderTargetableTexture);
+						ImplViewportTextureResourceLogging(ViewportIt, ContextNum, TEXT("AdditionalRTT"), ViewportIt->Resources[EDisplayClusterViewportResource::AdditionalTargetableResources][ContextNum]);
 					}
 
-					if (ViewportIt->MipsShaderResources.Num() > (int32)ContextNum)
+					if (ViewportIt->Resources[EDisplayClusterViewportResource::MipsShaderResources].Num() > (int32)ContextNum)
 					{
-						ViewportIt->MipsShaderResources[ContextNum] = ResourcesPool->AllocateTextureResource(ContextSize, false, CustomFormat, ContextIt.NumMips);
-						ImplViewportTextureResourceLogging(ViewportIt, ContextNum, TEXT("Mips"), ViewportIt->MipsShaderResources[ContextNum]);
+						ViewportIt->Resources[EDisplayClusterViewportResource::MipsShaderResources][ContextNum] = ResourcesPool->AllocateResource(ContextSize, CustomFormat, EDisplayClusterViewportResourceSettingsFlags::ResolveTargetableTexture, ContextIt.NumMips);
+						ImplViewportTextureResourceLogging(ViewportIt, ContextNum, TEXT("Mips"), ViewportIt->Resources[EDisplayClusterViewportResource::MipsShaderResources][ContextNum]);
 					}
 				}
 			}
@@ -192,7 +199,7 @@ bool FDisplayClusterRenderTargetManager::AllocateRenderFrameResources(FViewport*
 			bResult = false;
 		}
 
-		ResourcesPool->FinishReallocateResources();
+		ResourcesPool->EndReallocateResources();
 	}
 
 	return bResult;
@@ -237,35 +244,35 @@ bool FDisplayClusterRenderTargetManager::AllocateFrameTargets(const FDisplayClus
 	}
 
 	// Reallocate frame target resources
-	TArray<FDisplayClusterViewportTextureResource*> NewFrameTargetResources;
-	TArray<FDisplayClusterViewportTextureResource*> NewAdditionalFrameTargetableResources;
+	TArray<TSharedPtr<FDisplayClusterViewportResource, ESPMode::ThreadSafe>> NewFrameTargetResources;
+	TArray<TSharedPtr<FDisplayClusterViewportResource, ESPMode::ThreadSafe>> NewAdditionalFrameTargetableResources;
 
 	for (uint32 FrameTargetsIt = 0; FrameTargetsIt < FrameTargetsAmount; FrameTargetsIt++)
 	{
-		FDisplayClusterViewportTextureResource* NewResource = ResourcesPool->AllocateTextureResource(InOutRenderFrame.FrameRect.Size(), true, PF_Unknown);
-		if (NewResource != nullptr)
+		TSharedPtr<FDisplayClusterViewportResource, ESPMode::ThreadSafe> NewResource = ResourcesPool->AllocateResource(InOutRenderFrame.FrameRect.Size(), PF_Unknown, EDisplayClusterViewportResourceSettingsFlags::RenderTargetableTexture);
+		if (NewResource.IsValid())
 		{
-			if (NewResource->GetViewportResourceState(EDisplayClusterViewportResourceState::Initialized) == false)
+			if (EnumHasAnyFlags(NewResource->GetResourceState(), EDisplayClusterViewportResourceState::Initialized) == false)
 			{
 				// Log: New resource created
-				UE_LOG(LogDisplayClusterViewport, Verbose, TEXT("Created new RenderFrame resource (%dx%d)"), NewResource->GetSizeX(), NewResource->GetSizeY());
+				UE_LOG(LogDisplayClusterViewport, Verbose, TEXT("Created new RenderFrame resource (%dx%d)"), NewResource->GetResourceSettings().GetSizeX(), NewResource->GetResourceSettings().GetSizeY());
 			}
 
 			// calc and assign backbuffer offset (side_by_side, top_bottom)
-			NewResource->BackbufferFrameOffset = InOutRenderFrame.FrameRect.Min + TargetLocation;
+			NewResource->SetBackbufferFrameOffset(InOutRenderFrame.FrameRect.Min + TargetLocation);
 			TargetLocation += TargetOffset;
 			NewFrameTargetResources.Add(NewResource);
 		}
 
 		if (InRenderFrameSettings.bShouldUseAdditionalFrameTargetableResource)
 		{
-			FDisplayClusterViewportTextureResource* NewAdditionalResource = ResourcesPool->AllocateTextureResource(InOutRenderFrame.FrameRect.Size(), true, PF_Unknown);
-			if (NewAdditionalResource != nullptr)
+			TSharedPtr<FDisplayClusterViewportResource, ESPMode::ThreadSafe> NewAdditionalResource = ResourcesPool->AllocateResource(InOutRenderFrame.FrameRect.Size(), PF_Unknown, EDisplayClusterViewportResourceSettingsFlags::RenderTargetableTexture);
+			if (NewAdditionalResource.IsValid())
 			{
-				if (NewAdditionalResource->GetViewportResourceState(EDisplayClusterViewportResourceState::Initialized) == false)
+				if (EnumHasAnyFlags(NewAdditionalResource->GetResourceState(), EDisplayClusterViewportResourceState::Initialized) == false)
 				{
 					// Log: New resource created
-					UE_LOG(LogDisplayClusterViewport, Verbose, TEXT("Created new RenderFrame2 resource (%dx%d)"), NewAdditionalResource->GetSizeX(), NewAdditionalResource->GetSizeY());
+					UE_LOG(LogDisplayClusterViewport, Verbose, TEXT("Created new RenderFrame2 resource (%dx%d)"), NewAdditionalResource->GetResourceSettings().GetSizeX(), NewAdditionalResource->GetResourceSettings().GetSizeY());
 				}
 
 				NewAdditionalFrameTargetableResources.Add(NewAdditionalResource);
@@ -283,8 +290,8 @@ bool FDisplayClusterRenderTargetManager::AllocateFrameTargets(const FDisplayClus
 				FDisplayClusterViewport* ViewportPtr = static_cast<FDisplayClusterViewport*>(ViewIt.Viewport.Get());
 				if (ViewportPtr && ViewportPtr->RenderSettings.bVisible)
 				{
-					ViewportPtr->OutputFrameTargetableResources = NewFrameTargetResources;
-					ViewportPtr->AdditionalFrameTargetableResources = NewAdditionalFrameTargetableResources;
+					ViewportPtr->Resources[EDisplayClusterViewportResource::OutputFrameTargetableResources] = NewFrameTargetResources;
+					ViewportPtr->Resources[EDisplayClusterViewportResource::AdditionalFrameTargetableResources] = NewAdditionalFrameTargetableResources;
 
 					// Adjust viewports frame rects. This offset saved in 'BackbufferFrameOffset'
 					if (ViewportPtr->Contexts.IsValidIndex(ViewIt.ContextNum))
diff --git a/Engine/Plugins/Runtime/nDisplay/Source/DisplayCluster/Private/Render/Viewport/RenderTarget/DisplayClusterRenderTargetResource.cpp b/Engine/Plugins/Runtime/nDisplay/Source/DisplayCluster/Private/Render/Viewport/RenderTarget/DisplayClusterRenderTargetResource.cpp
--- a/Engine/Plugins/Runtime/nDisplay/Source/DisplayCluster/Private/Render/Viewport/RenderTarget/DisplayClusterRenderTargetResource.cpp
+++ b/Engine/Plugins/Runtime/nDisplay/Source/DisplayCluster/Private/Render/Viewport/RenderTarget/DisplayClusterRenderTargetResource.cpp
@@ -6,85 +6,22 @@
 #include "RenderResource.h"
 #include "UnrealClient.h"
 
-//------------------------------------------------------------------------------------------------
-// FDisplayClusterViewportResource
-//------------------------------------------------------------------------------------------------
-FDisplayClusterViewportResource::FDisplayClusterViewportResource(const FDisplayClusterViewportResourceSettings& InResourceSettings)
-	: ViewportResourceSettings(InResourceSettings)
+////////////////////////////////////////////////////////////////////////
+// FDisplayClusterViewportTextureResource
+////////////////////////////////////////////////////////////////////////
+FDisplayClusterViewportTextureResource::FDisplayClusterViewportTextureResource(const FDisplayClusterViewportResourceSettings& InResourceSettings)
+	: FDisplayClusterViewportResource(InResourceSettings)
 {
-	bSRGB = ViewportResourceSettings.bShouldUseSRGB;
+	bSRGB = EnumHasAnyFlags(GetResourceSettings().GetResourceFlags(), EDisplayClusterViewportResourceSettingsFlags::ShouldUseSRGB);
 	bGreyScaleFormat = false;
 	bIgnoreGammaConversions = true;
 }
 
-void FDisplayClusterViewportResource::ImplInitDynamicRHI_RenderTargetResource2D(FTexture2DRHIRef& OutRenderTargetTextureRHI, FTexture2DRHIRef& OutTextureRHI)
-{
-	ETextureCreateFlags CreateFlags = TexCreate_Dynamic;
-
-	// -- we will be manually copying this cross GPU, tell render graph not to
-	CreateFlags |= TexCreate_MultiGPUGraphIgnore;
-
-	// reflect srgb from settings
-	if (bSRGB)
-	{
-		CreateFlags |= TexCreate_SRGB;
-	}
-
-	uint32 NumMips = 1;
-	if (ViewportResourceSettings.NumMips > 1)
-	{
-		// Create nummips texture!
-		CreateFlags |= TexCreate_GenerateMipCapable | TexCreate_UAV;
-		NumMips = ViewportResourceSettings.NumMips;
-	}
-
-	const FRHITextureCreateDesc Desc =
-		FRHITextureCreateDesc::Create2D(TEXT("DisplayClusterViewportRenderTargetResource"))
-		.SetExtent(GetSizeX(), GetSizeY())
-		.SetFormat(ViewportResourceSettings.Format)
-		.SetNumMips(NumMips)
-		.SetFlags(CreateFlags | ETextureCreateFlags::RenderTargetable | ETextureCreateFlags::ShaderResource)
-		.SetInitialState(ERHIAccess::SRVMask)
-		.SetClearValue(FClearValueBinding::Black);
-
-	OutRenderTargetTextureRHI = OutTextureRHI = RHICreateTexture(Desc);
-}
-
-void FDisplayClusterViewportResource::ImplInitDynamicRHI_TextureResource2D(FTexture2DRHIRef& OutTextureRHI)
-{
-	FRHITextureCreateDesc Desc =
-		FRHITextureCreateDesc::Create2D(TEXT("DisplayClusterViewportTextureResource"), GetSizeX(), GetSizeY(), ViewportResourceSettings.Format)
-		.SetClearValue(FClearValueBinding::Black)
-		.SetFlags(ETextureCreateFlags::Dynamic | ETextureCreateFlags::ShaderResource)
-		// -- we will be manually copying this cross GPU, tell render graph not to
-		.AddFlags(ETextureCreateFlags::MultiGPUGraphIgnore)
-		.SetInitialState(ERHIAccess::SRVMask);
-
-	// reflect srgb from settings
-	if (bSRGB)
-	{
-		Desc.AddFlags(ETextureCreateFlags::SRGB);
-	}
-	if (ViewportResourceSettings.bIsRenderTargetable)
-	{
-		Desc.AddFlags(ETextureCreateFlags::RenderTargetable);
-	}
-	else
-	{
-		Desc.AddFlags(ETextureCreateFlags::ResolveTargetable);
-	}
-
-	OutTextureRHI = RHICreateTexture(Desc);
-}
-
-//------------------------------------------------------------------------------------------------
-// FDisplayClusterViewportTextureResource
-//------------------------------------------------------------------------------------------------
 void FDisplayClusterViewportTextureResource::InitRHI(FRHICommandListBase&)
 {
 	FTexture2DRHIRef NewTextureRHI;
 
-	if (GetResourceSettingsConstRef().NumMips > 1)
+	if (GetResourceSettings().GetNumMips() > 1)
 	{
 		FTexture2DRHIRef DummyTextureRHI;
 		ImplInitDynamicRHI_RenderTargetResource2D(NewTextureRHI, DummyTextureRHI);
@@ -97,9 +34,9 @@ void FDisplayClusterViewportTextureResource::InitRHI(FRHICommandListBase&)
 	TextureRHI = (FTextureRHIRef&)NewTextureRHI;
 }
 
-//------------------------------------------------------------------------------------------------
+////////////////////////////////////////////////////////////////////////
 // FDisplayClusterViewportRenderTargetResource
-//------------------------------------------------------------------------------------------------
+////////////////////////////////////////////////////////////////////////
 void FDisplayClusterViewportRenderTargetResource::InitRHI(FRHICommandListBase&)
 {
 	// Create RTT and shader resources
diff --git a/Engine/Plugins/Runtime/nDisplay/Source/DisplayCluster/Private/Render/Viewport/RenderTarget/DisplayClusterRenderTargetResource.h b/Engine/Plugins/Runtime/nDisplay/Source/DisplayCluster/Private/Render/Viewport/RenderTarget/DisplayClusterRenderTargetResource.h
--- a/Engine/Plugins/Runtime/nDisplay/Source/DisplayCluster/Private/Render/Viewport/RenderTarget/DisplayClusterRenderTargetResource.h
+++ b/Engine/Plugins/Runtime/nDisplay/Source/DisplayCluster/Private/Render/Viewport/RenderTarget/DisplayClusterRenderTargetResource.h
@@ -1,202 +1,119 @@
 // Copyright Epic Games, Inc. All Rights Reserved.
 
 #pragma once
+#include "Render/Viewport/Resource/DisplayClusterViewportResource.h"
 
-#include "CoreMinimal.h"
-
-#include "PixelFormat.h"
 #include "RHI.h"
 #include "RHIResources.h"
 #include "RenderResource.h"
 #include "UnrealClient.h"
-#include "Containers/StaticBitArray.h"
-
-///////////////////////////////////////////////////////////////////
-struct FDisplayClusterViewportResourceSettings
-{
-	FDisplayClusterViewportResourceSettings()
-		: Size(ForceInit)
-		, Format(EPixelFormat(0))
-		, bShouldUseSRGB(false)
-		, DisplayGamma(1)
-		, bIsRenderTargetable(false)
-		, NumMips(1)
-	{ }
-
-	FDisplayClusterViewportResourceSettings(const FString& InClusterNodeId, const EPixelFormat InFormat, const bool InbShouldUseSRGB, const float InDisplayGamma)
-		: Size(ForceInit)
-		, Format(InFormat)
-		, bShouldUseSRGB(InbShouldUseSRGB)
-		, DisplayGamma(InDisplayGamma)
-		, bIsRenderTargetable(false)
-		, NumMips(1)
-		, ClusterNodeId(InClusterNodeId)
-	{ }
-
-	FDisplayClusterViewportResourceSettings(const FDisplayClusterViewportResourceSettings& InBaseSettings, const FIntPoint& InSize, const EPixelFormat InFormat, const bool InbIsRenderTargetable = false, const int32 InNumMips = 1)
-		: Size(InSize)
-		, Format(InFormat == PF_Unknown ? InBaseSettings.Format : InFormat)
-		, bShouldUseSRGB(InBaseSettings.bShouldUseSRGB)
-		, DisplayGamma(InBaseSettings.DisplayGamma)
-		, bIsRenderTargetable(InbIsRenderTargetable)
-		, NumMips(InNumMips)
-		, ClusterNodeId(InBaseSettings.ClusterNodeId)
-	{ }
-
-	FORCEINLINE bool IsResourceSettingsEqual(const FDisplayClusterViewportResourceSettings& In) const
-	{
-		return (In.Size == Size)
-			&& (In.Format == Format)
-			&& (In.bShouldUseSRGB == bShouldUseSRGB)
-			&& (In.DisplayGamma == DisplayGamma)
-			&& (In.bIsRenderTargetable == bIsRenderTargetable)
-			&& (In.NumMips == NumMips);
-	}
-
-	FORCEINLINE bool IsClusterNodeNameEqual(const FDisplayClusterViewportResourceSettings& In) const
-	{
-		return ClusterNodeId == In.ClusterNodeId;
-	}
-
-	const FIntPoint    Size;
-	const EPixelFormat Format;
-	const bool         bShouldUseSRGB;
-
-	// Render target params
-	const float DisplayGamma;
-
-	// Texture target params
-	const bool bIsRenderTargetable;
-	const int32 NumMips;
-
-private:
-	// Cluster per-node render
-	const FString ClusterNodeId;
-};
-
-enum class EDisplayClusterViewportResourceState : uint8
-{
-	Initialized = 0,
-	Unused,
-	Deleted,
 
-	// Disable reallocation of this resource for the current frame
-	DisableReallocate,
-
-	COUNT
-};
-///////////////////////////////////////////////////////////////////
-class FDisplayClusterViewportResource
-	: public FTexture
+/**
+ * Viewport texture resource
+ */
+class FDisplayClusterViewportTextureResource
+	: public FDisplayClusterViewportResource
+	, public FTexture
 {
 public:
-	FDisplayClusterViewportResource(const FDisplayClusterViewportResourceSettings& InResourceSettings);
-	virtual ~FDisplayClusterViewportResource()
+	FDisplayClusterViewportTextureResource(const FDisplayClusterViewportResourceSettings& InResourceSettings);
+
+	virtual ~FDisplayClusterViewportTextureResource()
 	{ }
 
 public:
-	FORCEINLINE const FDisplayClusterViewportResourceSettings& GetResourceSettingsConstRef() const
+	//~ Begin FDisplayClusterViewportResource
+	virtual FRHITexture2D* GetViewportResourceRHI() const override
 	{
-		return ViewportResourceSettings;
+		return TextureRHI.IsValid() ? TextureRHI->GetTexture2D() : nullptr;
 	}
 
-	void RaiseViewportResourceState(const EDisplayClusterViewportResourceState InState)
+	virtual FIntPoint GetBackbufferFrameOffset() const override
 	{
-		ResourceState[(uint8)InState] = true;
+		return BackbufferFrameOffset;
 	}
 
-	void ClearViewportResourceState(const EDisplayClusterViewportResourceState InState)
+	virtual void SetBackbufferFrameOffset(const FIntPoint& InBackbufferFrameOffset) override
 	{
-		ResourceState[(uint8)InState] = false;
+		BackbufferFrameOffset = InBackbufferFrameOffset;
 	}
 
-	bool GetViewportResourceState(const EDisplayClusterViewportResourceState InState) const
+	virtual void InitializeViewportResource_RenderThread(FRHICommandListBase& RHICmdList) override
 	{
-		return ResourceState[(uint8)InState];
+		// Call FTexture::InitResource()
+		InitResource(RHICmdList);
 	}
 
-	virtual FRHITexture2D* GetViewportResource2DRHI() const
+	virtual void ReleaseViewportResource_RenderThread(FRHICommandListBase& RHICmdList) override
 	{
-		return TextureRHI.IsValid() ? TextureRHI->GetTexture2D() : nullptr;
+		// Call FTexture::ReleaseResource()
+		ReleaseResource();
 	}
+	//~~ End FDisplayClusterViewportResource
 
-	/** Returns the width of the texture in pixels. */
+	//~ Begin FTexture
 	virtual uint32 GetSizeX() const override
 	{
-		return ViewportResourceSettings.Size.X;
+		return GetResourceSettings().GetSizeXY().X;
 	}
 
-	/** Returns the height of the texture in pixels. */
 	virtual uint32 GetSizeY() const override
 	{
-		return ViewportResourceSettings.Size.Y;
+		return GetResourceSettings().GetSizeXY().Y;
 	}
 
-protected:
-	void ImplInitDynamicRHI_RenderTargetResource2D(FTexture2DRHIRef& OutRenderTargetTextureRHI, FTexture2DRHIRef& OutTextureRHI);
-	void ImplInitDynamicRHI_TextureResource2D(FTexture2DRHIRef& OutTextureRHI);
-
-private:
-	const FDisplayClusterViewportResourceSettings ViewportResourceSettings;
-
-	// resource state for game thread
-	TStaticBitArray<(uint8)EDisplayClusterViewportResourceState::COUNT> ResourceState;
-};
-
-///////////////////////////////////////////////////////////////////
-class FDisplayClusterViewportTextureResource
-	: public FDisplayClusterViewportResource
-{
-public:
-	FDisplayClusterViewportTextureResource(const FDisplayClusterViewportResourceSettings& InResourceSettings)
-		: FDisplayClusterViewportResource(InResourceSettings)
-		, BackbufferFrameOffset(EForceInit::ForceInitToZero)
-	{ }
-
-	virtual ~FDisplayClusterViewportTextureResource()
-	{ }
-
-public:
 	virtual void InitRHI(FRHICommandListBase& RHICmdList) override;
 
 	virtual FString GetFriendlyName() const override
 	{
 		return TEXT("DisplayClusterViewportTextureResource");
 	}
+	//~~ End FTexture
 
 public:
 	// OutputFrameTargetableResources frame offset on backbuffer (special for 'side_by_side' and 'top_bottom' DCRenderDevices)
 	FIntPoint BackbufferFrameOffset;
 };
 
-///////////////////////////////////////////////////////////////////
+/**
+ * Viewport RTT resource
+ */
 class FDisplayClusterViewportRenderTargetResource
-	: public FDisplayClusterViewportResource
+	: public FDisplayClusterViewportTextureResource
 	, public FRenderTarget
 {
 public:
 	FDisplayClusterViewportRenderTargetResource(const FDisplayClusterViewportResourceSettings& InResourceSettings)
-		: FDisplayClusterViewportResource(InResourceSettings)
+		: FDisplayClusterViewportTextureResource(InResourceSettings)
 	{ }
 
 	virtual ~FDisplayClusterViewportRenderTargetResource()
 	{ }
 
 public:
+	//~ Begin FDisplayClusterViewportResource
+	virtual FRenderTarget* GetViewportResourceRenderTarget()
+	{
+		return this;
+	}
+	//~~ End FDisplayClusterViewportResource
+
+	//~ Begin FRenderTarget
 	virtual void InitRHI(FRHICommandListBase& RHICmdList) override;
 
 	virtual FIntPoint GetSizeXY() const override
 	{
-		return GetResourceSettingsConstRef().Size;
+		return GetResourceSettings().GetSizeXY();
 	}
 
 	virtual float GetDisplayGamma() const 
 	{
-		return GetResourceSettingsConstRef().DisplayGamma;
+		return GetResourceSettings().GetDisplayGamma();
 	}
 
 	virtual FString GetFriendlyName() const override 
 	{
 		return TEXT("DisplayClusterViewportRenderTargetResource");
 	}
+	//~~ End FRenderTarget
 };
diff --git a/Engine/Plugins/Runtime/nDisplay/Source/DisplayCluster/Private/Render/Viewport/RenderTarget/DisplayClusterRenderTargetResourcesPool.cpp b/Engine/Plugins/Runtime/nDisplay/Source/DisplayCluster/Private/Render/Viewport/RenderTarget/DisplayClusterRenderTargetResourcesPool.cpp
--- a/Engine/Plugins/Runtime/nDisplay/Source/DisplayCluster/Private/Render/Viewport/RenderTarget/DisplayClusterRenderTargetResourcesPool.cpp
+++ b/Engine/Plugins/Runtime/nDisplay/Source/DisplayCluster/Private/Render/Viewport/RenderTarget/DisplayClusterRenderTargetResourcesPool.cpp
@@ -1,6 +1,10 @@
 // Copyright Epic Games, Inc. All Rights Reserved.
 
 #include "DisplayClusterRenderTargetResourcesPool.h"
+#include "DisplayClusterRenderTargetResource.h"
+
+#include "Render/Viewport/Resource/DisplayClusterViewportPreviewResource.h"
+
 #include "Render/Viewport/RenderFrame/DisplayClusterRenderFrameSettings.h"
 #include "Render/Viewport/DisplayClusterViewportHelpers.h"
 #include "Render/Viewport/DisplayClusterViewportManagerProxy.h"
@@ -9,72 +13,6 @@
 #include "Engine/RendererSettings.h"
 #include "RenderingThread.h"
 
-////////////////////////////////////////////////////////////////////////////////////////////////////////////
-namespace DisplayClusterRenderTargetResourcesPool
-{
-	static void ImplInitializeViewportResourcesRHI(const TArrayView<FDisplayClusterViewportResource*>& InResources)
-	{
-		TArray<FDisplayClusterViewportResource*> ResourcesForRenderThread(InResources);
-
-		ENQUEUE_RENDER_COMMAND(DisplayCluster_InitializeViewportResourcesRHI)(
-			[NewResources = std::move(ResourcesForRenderThread)](FRHICommandListImmediate& RHICmdList)
-		{
-			for (FDisplayClusterViewportResource* ResourceIt : NewResources)
-			{
-				ResourceIt->InitResource(RHICmdList);
-			}
-		});
-	}
-
-	static void ImplReleaseViewportResourcesRHI(FDisplayClusterViewportManagerProxy* InViewportManagerProxy, const TArrayView<FDisplayClusterViewportResource*>& InResources)
-	{
-		if (InViewportManagerProxy != nullptr && !InResources.IsEmpty())
-		{
-		TArray<FDisplayClusterViewportResource*> ResourcesForRenderThread(InResources);
-
-		ENQUEUE_RENDER_COMMAND(DisplayCluster_ReleaseViewportResourcesRHI)(
-				[InViewportManagerProxy = InViewportManagerProxy->AsShared(), ReleasedResources = std::move(ResourcesForRenderThread)](FRHICommandListImmediate& RHICmdList)
-			{
-				for (FDisplayClusterViewportResource* DeletedResourcePtrIt : ReleasedResources)
-				{
-					InViewportManagerProxy->DeleteResource_RenderThread(DeletedResourcePtrIt);
-				}
-			});
-		}
-	}
-
-	static FDisplayClusterViewportResourceSettings* ImplCreateViewportResourceSettings(const FDisplayClusterRenderFrameSettings& InRenderFrameSettings, FViewport* InViewport)
-	{
-		FDisplayClusterViewportResourceSettings Result;
-
-		if (InViewport != nullptr)
-		{
-			FRHITexture2D* ViewportTexture = InViewport->GetRenderTargetTexture();
-			if (ViewportTexture != nullptr)
-			{
-				EPixelFormat Format = ViewportTexture->GetFormat();
-				bool bShouldUseSRGB = EnumHasAnyFlags(ViewportTexture->GetFlags(), TexCreate_SRGB);
-				float Gamma = InViewport->GetDisplayGamma();
-
-				return new FDisplayClusterViewportResourceSettings(InRenderFrameSettings.ClusterNodeId, Format, bShouldUseSRGB, Gamma);
-			}
-		}
-
-#if WITH_EDITOR
-		EPixelFormat PreviewPixelFormat;
-		float PreviewDisplayGamma = 1.f;
-		bool bPreviewSRGB = true;
-		DisplayClusterViewportHelpers::GetPreviewRenderTargetDesc_Editor(InRenderFrameSettings, PreviewPixelFormat, PreviewDisplayGamma, bPreviewSRGB);
-
-		return new FDisplayClusterViewportResourceSettings(InRenderFrameSettings.ClusterNodeId, PreviewPixelFormat, bPreviewSRGB, PreviewDisplayGamma);
-#else
-		return nullptr;
-#endif
-	}
-};
-
-using namespace DisplayClusterRenderTargetResourcesPool;
-
 ////////////////////////////////////////////////////////////////////////////////////////////////////////////
 /// FDisplayClusterRenderTargetResourcesPool
 ////////////////////////////////////////////////////////////////////////////////////////////////////////////
@@ -90,177 +28,199 @@ FDisplayClusterRenderTargetResourcesPool::~FDisplayClusterRenderTargetResourcesP
 void FDisplayClusterRenderTargetResourcesPool::Release()
 {
 	// Release all resources
-	ImplReleaseResources<FDisplayClusterViewportRenderTargetResource>(RenderTargetResources);
-	ImplReleaseResources<FDisplayClusterViewportTextureResource>(TextureResources);
+	ImplUpdateResources(ViewportResources, EResourceUpdateMode::Release);
 }
 
-template <typename TViewportResourceType>
-void FDisplayClusterRenderTargetResourcesPool::ImplBeginReallocateResources(TArray<TViewportResourceType*>& InOutViewportResources)
+void FDisplayClusterRenderTargetResourcesPool::ImplUpdateResources(TArray<TSharedPtr<FDisplayClusterViewportResource, ESPMode::ThreadSafe>>& InOutViewportResources, const EResourceUpdateMode InUpdateMode)
 {
-	check(ResourceSettings != nullptr);
+	if (InOutViewportResources.IsEmpty())
+	{
+		return;
+	}
 
-	// Mark cluster node resources as unused
-	for (TViewportResourceType* ResourceIt : InOutViewportResources)
+	// Update 'Initialized' resource state
+	for (const TSharedPtr<FDisplayClusterViewportResource, ESPMode::ThreadSafe>& ResourceIt : InOutViewportResources)
 	{
-		if (ResourceIt != nullptr && ResourceIt->GetViewportResourceState(EDisplayClusterViewportResourceState::DisableReallocate) == false)
+		if (ResourceIt.IsValid())
 		{
-			if (const bool bIsCurrentClusterNodeResource = ResourceIt->GetResourceSettingsConstRef().IsClusterNodeNameEqual(*ResourceSettings))
+			switch (InUpdateMode)
 			{
-				ResourceIt->RaiseViewportResourceState(EDisplayClusterViewportResourceState::Unused);
+			case EResourceUpdateMode::Release:
+				EnumRemoveFlags(ResourceIt->GetResourceState(), EDisplayClusterViewportResourceState::Initialized);
+				break;
+
+			default:
+			case EResourceUpdateMode::Initialize:
+				EnumAddFlags(ResourceIt->GetResourceState(), EDisplayClusterViewportResourceState::Initialized);
+				break;
 			}
-			else
+		}
+	}
+
+	// Update all viewports resources on rendering thread
+	ENQUEUE_RENDER_COMMAND(DisplayCluster_UpdateViewportResources)
+	([ViewportResources = InOutViewportResources, InUpdateMode](FRHICommandListImmediate& RHICmdList)
+	{
+		for (const TSharedPtr<FDisplayClusterViewportResource, ESPMode::ThreadSafe>& ResourceIt : ViewportResources)
+		{
+			if (ResourceIt.IsValid())
 			{
-				ResourceIt->RaiseViewportResourceState(EDisplayClusterViewportResourceState::DisableReallocate);
+				switch (InUpdateMode)
+				{
+				case EResourceUpdateMode::Release:
+					ResourceIt->ReleaseViewportResource_RenderThread(RHICmdList);
+					break;
+
+				default:
+				case EResourceUpdateMode::Initialize:
+					ResourceIt->InitializeViewportResource_RenderThread(RHICmdList);
+					break;
+				}
 			}
 		}
-	}
+	});
+
+	InOutViewportResources.Empty();
 }
 
-template <typename TViewportResourceType>
-void FDisplayClusterRenderTargetResourcesPool::ImplFinishReallocateResources(TArray<TViewportResourceType*>& InOutViewportResources)
+bool FDisplayClusterRenderTargetResourcesPool::BeginReallocateResources(const FDisplayClusterRenderFrameSettings& InRenderFrameSettings, FViewport* InViewport)
 {
-	// Collect new and unused resources 
-	TArray<TViewportResourceType*> NewResources;
-	TArray<TViewportResourceType*> UnusedResources;
+	check(ResourceSettings == nullptr);
 
-	for (TViewportResourceType* ResourceIt : InOutViewportResources)
+	// Initialize settings for new render frame
+	ResourceSettings = new FDisplayClusterViewportResourceSettings(InRenderFrameSettings, InViewport);
+
+	// Begin reallocate resources for current cluster node
+	// Mark cluster node resources as unused
+	for (TSharedPtr<FDisplayClusterViewportResource, ESPMode::ThreadSafe>& ViewportResourceIt : ViewportResources)
 	{
-		if (ResourceIt != nullptr)
+		if (ViewportResourceIt.IsValid() && EnumHasAnyFlags(ViewportResourceIt->GetResourceState(), EDisplayClusterViewportResourceState::DisableReallocate) == false)
 		{
-			check(ResourceIt->GetViewportResourceState(EDisplayClusterViewportResourceState::Deleted) == false);
-
-			if (ResourceIt->GetViewportResourceState(EDisplayClusterViewportResourceState::Initialized) == false)
+			if (const bool bIsCurrentClusterNodeResource = ViewportResourceIt->GetResourceSettings().IsClusterNodeNameEqual(*ResourceSettings))
 			{
-				// Collect new resources
-				NewResources.Add(ResourceIt);
-				ResourceIt->RaiseViewportResourceState(EDisplayClusterViewportResourceState::Initialized);
+				EnumAddFlags(ViewportResourceIt->GetResourceState(), EDisplayClusterViewportResourceState::Unused);
 			}
-			else if (ResourceIt->GetViewportResourceState(EDisplayClusterViewportResourceState::Unused) == true)
+			else
 			{
-				// Collect unused resources
-				UnusedResources.Add(ResourceIt);
-				ResourceIt->RaiseViewportResourceState(EDisplayClusterViewportResourceState::Deleted);
+				EnumAddFlags(ViewportResourceIt->GetResourceState(), EDisplayClusterViewportResourceState::DisableReallocate);
 			}
-
-			// Clear 'DisableReallocate' at frame end
-			ResourceIt->ClearViewportResourceState(EDisplayClusterViewportResourceState::DisableReallocate);
 		}
 	}
 
-	// Init RHI for new resources
-	if (NewResources.Num() > 0)
-	{
-		// Send to rendering thread
-		ImplInitializeViewportResourcesRHI(TArrayView<FDisplayClusterViewportResource*>((FDisplayClusterViewportResource**)(NewResources.GetData()), NewResources.Num()));
-	}
+	return true;
+}
 
-	// Remove unused resources
-	if (UnusedResources.Num() > 0)
+void FDisplayClusterRenderTargetResourcesPool::EndReallocateResources()
+{
+	// Finish reallocate resources for current cluster node
+	if (ResourceSettings != nullptr)
 	{
+		// Collect unused resources for removal.
+		TArray<TSharedPtr<FDisplayClusterViewportResource, ESPMode::ThreadSafe>> UnusedResources;
+
+		// Collect new resources
+		TArray<TSharedPtr<FDisplayClusterViewportResource, ESPMode::ThreadSafe>> NewResources;
+
+		for (TSharedPtr<FDisplayClusterViewportResource, ESPMode::ThreadSafe>& ResourceIt : ViewportResources)
+		{
+			if (ResourceIt.IsValid())
+			{
+				if (EnumHasAnyFlags(ResourceIt->GetResourceState(), EDisplayClusterViewportResourceState::Unused))
+				{
+					// Collect unused resources
+					UnusedResources.Add(ResourceIt);
+				}
+				else if (!EnumHasAnyFlags(ResourceIt->GetResourceState(), EDisplayClusterViewportResourceState::Initialized))
+				{
+					// Collect new resources
+					NewResources.Add(ResourceIt);
+				}
+
+				// Clear 'DisableReallocate' at frame end
+				EnumRemoveFlags(ResourceIt->GetResourceState(), EDisplayClusterViewportResourceState::DisableReallocate);
+			}
+		}
+
 		// Remove GameThread resource references
-		for (TViewportResourceType* ResourceIt : UnusedResources)
+		for (const TSharedPtr<FDisplayClusterViewportResource, ESPMode::ThreadSafe>& ResourceIt : UnusedResources)
 		{
 			int ResourceIndex = INDEX_NONE;
-			while (InOutViewportResources.Find(ResourceIt, ResourceIndex))
+			while (ViewportResources.Find(ResourceIt, ResourceIndex))
 			{
-				InOutViewportResources.RemoveAt(ResourceIndex);
+				ViewportResources.RemoveAt(ResourceIndex);
 			}
 		}
 
-		// Send to released resources to rendering thread
-		ImplReleaseViewportResourcesRHI(GetViewportManagerProxy(), TArrayView<FDisplayClusterViewportResource*>((FDisplayClusterViewportResource**)(UnusedResources.GetData()), UnusedResources.Num()));
+		// Initialize new resources on rendering thread:
+		ImplUpdateResources(NewResources, EResourceUpdateMode::Initialize);
+
+		// Remove unused resources
+		ImplUpdateResources(UnusedResources, EResourceUpdateMode::Release);
 	}
+
+	delete ResourceSettings;
+	ResourceSettings = nullptr;
 }
 
-template <typename TViewportResourceType>
-void FDisplayClusterRenderTargetResourcesPool::ImplReleaseResources(TArray<TViewportResourceType*>& InOutViewportResources)
+TSharedPtr<FDisplayClusterViewportResource, ESPMode::ThreadSafe> FDisplayClusterRenderTargetResourcesPool::AllocateResource(const FIntPoint& InSize, EPixelFormat CustomPixelFormat, const EDisplayClusterViewportResourceSettingsFlags InResourceFlags, int32 InNumMips)
 {
-	// Send all resources to rendering thread
-	ImplReleaseViewportResourcesRHI(GetViewportManagerProxy(), TArray<FDisplayClusterViewportResource*>((FDisplayClusterViewportResource**)(InOutViewportResources.GetData()), InOutViewportResources.Num()));
+	check(ResourceSettings != nullptr);
 
-	// Reset refs on game thread too
-	InOutViewportResources.Empty();
-}
+	FDisplayClusterViewportResourceSettings InSettings(*ResourceSettings, InSize, CustomPixelFormat, InResourceFlags, InNumMips);
 
-template <typename TViewportResourceType>
-TViewportResourceType* FDisplayClusterRenderTargetResourcesPool::ImplAllocateResource(TArray<TViewportResourceType*>& InOutViewportResources, const FDisplayClusterViewportResourceSettings& InSettings)
-{
-	if (!DisplayClusterViewportHelpers::IsValidTextureSize(InSettings.Size))
+	if (!FDisplayClusterViewportHelpers::IsValidTextureSize(InSettings.GetSizeXY()))
 	{
 		return nullptr;
 	}
 
 	// Unused resources marked for current cluster node
-	auto ExistAndUnusedResourceIndex = InOutViewportResources.IndexOfByPredicate([InSettings](const TViewportResourceType* ResourceIt)
-	{
-		if (ResourceIt != nullptr && ResourceIt->GetViewportResourceState(EDisplayClusterViewportResourceState::Unused) == true)
+	auto ExistAndUnusedResourceIndex = ViewportResources.IndexOfByPredicate([InSettings](const TSharedPtr<FDisplayClusterViewportResource, ESPMode::ThreadSafe>& ResourceIt)
 		{
-			return ResourceIt->GetResourceSettingsConstRef().IsResourceSettingsEqual(InSettings);
-		}
+			if (ResourceIt.IsValid() && EnumHasAnyFlags(ResourceIt->GetResourceState(), EDisplayClusterViewportResourceState::Unused) == true)
+			{
+				return ResourceIt->GetResourceSettings().IsResourceSettingsEqual(InSettings);
+			}
 
-		return false;
-	});
+			return false;
+		});
 
 	if (ExistAndUnusedResourceIndex != INDEX_NONE)
 	{
 		// Use exist resource again
-		TViewportResourceType* ExistResource = InOutViewportResources[ExistAndUnusedResourceIndex];
+		TSharedPtr<FDisplayClusterViewportResource, ESPMode::ThreadSafe> ExistResource = ViewportResources[ExistAndUnusedResourceIndex];
 
-		// Clear unused state for re-used resource
-		ExistResource->ClearViewportResourceState(EDisplayClusterViewportResourceState::Unused);
+		if (ExistResource.IsValid())
+		{
+			// Clear unused state for re-used resource
+			EnumRemoveFlags(ExistResource->GetResourceState(), EDisplayClusterViewportResourceState::Unused);
 
-		return ExistResource;
+			return ExistResource;
+		}
 	}
 
 	// Create new resource:
-	TViewportResourceType* NewResource = new TViewportResourceType(InSettings);
-	InOutViewportResources.Add(NewResource);
+	TSharedPtr<FDisplayClusterViewportResource, ESPMode::ThreadSafe> NewResource;
 
-	return NewResource;
-}
-
-bool FDisplayClusterRenderTargetResourcesPool::BeginReallocateResources(const FDisplayClusterRenderFrameSettings& InRenderFrameSettings, FViewport* InViewport)
-{
-	check(ResourceSettings == nullptr);
-
-	// Initialize settings for new render frame
-	ResourceSettings = ImplCreateViewportResourceSettings(InRenderFrameSettings, InViewport);
-
-	if(ResourceSettings != nullptr)
+	if (EnumHasAnyFlags(InSettings.GetResourceFlags(), EDisplayClusterViewportResourceSettingsFlags::RenderTarget))
 	{
-		// Begin reallocate resources for current cluster node
-		ImplBeginReallocateResources<FDisplayClusterViewportRenderTargetResource>(RenderTargetResources);
-		ImplBeginReallocateResources<FDisplayClusterViewportTextureResource>(TextureResources);
-
-		return true;
+		// Render Target
+		NewResource = MakeShared<FDisplayClusterViewportRenderTargetResource>(InSettings);
 	}
-
-	return false;
-}
-
-void FDisplayClusterRenderTargetResourcesPool::FinishReallocateResources()
-{
-	if (ResourceSettings != nullptr)
+	else if (EnumHasAnyFlags(InSettings.GetResourceFlags(), EDisplayClusterViewportResourceSettingsFlags::PreviewTargetableTexture))
 	{
-		// Finish reallocate resources for current cluster node
-		ImplFinishReallocateResources<FDisplayClusterViewportRenderTargetResource>(RenderTargetResources);
-		ImplFinishReallocateResources<FDisplayClusterViewportTextureResource>(TextureResources);
-
-		delete ResourceSettings;
-		ResourceSettings = nullptr;
+		// Preview texture
+		NewResource = MakeShared<FDisplayClusterViewportPreviewResource>(InSettings);
+	}
+	else if (EnumHasAnyFlags(InSettings.GetResourceFlags(), EDisplayClusterViewportResourceSettingsFlags::RenderTargetableTexture | EDisplayClusterViewportResourceSettingsFlags::ResolveTargetableTexture))
+	{
+		// Texture
+		NewResource = MakeShared<FDisplayClusterViewportTextureResource>(InSettings);
 	}
-}
-
-FDisplayClusterViewportRenderTargetResource* FDisplayClusterRenderTargetResourcesPool::AllocateRenderTargetResource(const FIntPoint& InSize, enum EPixelFormat CustomPixelFormat)
-{
-	check(ResourceSettings != nullptr);
-
-	return ImplAllocateResource<FDisplayClusterViewportRenderTargetResource>(RenderTargetResources, FDisplayClusterViewportResourceSettings(*ResourceSettings, InSize, CustomPixelFormat));
-}
 
-FDisplayClusterViewportTextureResource* FDisplayClusterRenderTargetResourcesPool::AllocateTextureResource(const FIntPoint& InSize, bool bIsRenderTargetable, enum EPixelFormat CustomPixelFormat, int32 InNumMips)
-{
-	check(ResourceSettings != nullptr);
+	if (NewResource.IsValid())
+	{
+		ViewportResources.Add(NewResource);
+	}
 
-	return ImplAllocateResource<FDisplayClusterViewportTextureResource>(TextureResources, FDisplayClusterViewportResourceSettings(*ResourceSettings, InSize, CustomPixelFormat, bIsRenderTargetable, InNumMips));
+	return NewResource;
 }
diff --git a/Engine/Plugins/Runtime/nDisplay/Source/DisplayCluster/Private/Render/Viewport/RenderTarget/DisplayClusterRenderTargetResourcesPool.h b/Engine/Plugins/Runtime/nDisplay/Source/DisplayCluster/Private/Render/Viewport/RenderTarget/DisplayClusterRenderTargetResourcesPool.h
--- a/Engine/Plugins/Runtime/nDisplay/Source/DisplayCluster/Private/Render/Viewport/RenderTarget/DisplayClusterRenderTargetResourcesPool.h
+++ b/Engine/Plugins/Runtime/nDisplay/Source/DisplayCluster/Private/Render/Viewport/RenderTarget/DisplayClusterRenderTargetResourcesPool.h
@@ -3,11 +3,14 @@
 #pragma once
 
 #include "CoreMinimal.h"
-#include "DisplayClusterRenderTargetResource.h"
+#include "Render/Viewport/Resource/DisplayClusterViewportResource.h"
 
 class FDisplayClusterViewportManagerProxy;
 struct FDisplayClusterRenderFrameSettings;
 
+/**
+ * DC viewport resources pool
+ */
 class FDisplayClusterRenderTargetResourcesPool
 {
 public:
@@ -17,38 +20,59 @@ public:
 	void Release();
 
 public:
+	/** Begin reallocate resources for the specified cluster node.
+	* 
+	* @param InRenderFrameSettings - default resource settings
+	* @param InViewport - the window that uses these resources
+	* 
+	* @return true, if success
+	*/
 	bool BeginReallocateResources(const FDisplayClusterRenderFrameSettings& InRenderFrameSettings, class FViewport* InViewport);
-	void FinishReallocateResources();
+	
+	/** Allocate a new resource or reuse exists
+	* 
+	* @param InSize            - resource dimensions
+	* @param CustomPixelFormat - resource pixel format
+	* @param InResourceFlags   - These flags define the behavior of the resource.
+	* @param NumMips           - number of mips in the texture
+	* 
+	* @return ref to the resource instance
+	*/
+	TSharedPtr<FDisplayClusterViewportResource, ESPMode::ThreadSafe> AllocateResource(const FIntPoint& InSize, EPixelFormat CustomPixelFormat, const EDisplayClusterViewportResourceSettingsFlags InResourceFlags, int32 NumMips = 1);
 
-	FDisplayClusterViewportRenderTargetResource* AllocateRenderTargetResource(const FIntPoint& InSize, EPixelFormat CustomPixelFormat);
-	FDisplayClusterViewportTextureResource*      AllocateTextureResource(const FIntPoint& InSize, bool bIsRenderTargetable, EPixelFormat CustomPixelFormat, int32 NumMips = 1);
+	/** End reallocate resources for the specified cluster node. */
+	void EndReallocateResources();
 
-protected:
+private:
+	/**
+	 * Resource update mode
+	 */
+	enum class EResourceUpdateMode : uint8
+	{
+		Initialize = 0,
+		Release
+	};
+
+	/** Update resources from array
+	* 
+	* @param InOutViewportResources - resources to update
+	* @param InUpdateMode           - initialize or release
+	*/
+	void ImplUpdateResources(TArray<TSharedPtr<FDisplayClusterViewportResource, ESPMode::ThreadSafe>>& InOutViewportResources, const EResourceUpdateMode InUpdateMode);
+
+	/** Retrieving a ViewportManager instance. */
 	inline FDisplayClusterViewportManagerProxy* GetViewportManagerProxy() const
 	{
 		return ViewportManagerProxyWeakPtr.IsValid() ? ViewportManagerProxyWeakPtr.Pin().Get() : nullptr;
 	}
 
-private:
-	template <typename TViewportResourceType>
-	void ImplBeginReallocateResources(TArray<TViewportResourceType*>& InOutViewportResources);
-
-	template <typename TViewportResourceType>
-	void ImplFinishReallocateResources(TArray<TViewportResourceType*>& InOutViewportResources);
-
-	template <typename TViewportResourceType>
-	TViewportResourceType* ImplAllocateResource(TArray<TViewportResourceType*>& InOutViewportResources, const FDisplayClusterViewportResourceSettings& InSettings);
-
-	template <typename TViewportResourceType>
-	void ImplReleaseResources(TArray<TViewportResourceType*>& InOutViewportResources);
-
 private:
 	// Current render resource settings
 	FDisplayClusterViewportResourceSettings* ResourceSettings = nullptr;
 
 	// Viewport render resources
-	TArray<FDisplayClusterViewportRenderTargetResource*> RenderTargetResources;
-	TArray<FDisplayClusterViewportTextureResource*>      TextureResources;
+	TArray<TSharedPtr<FDisplayClusterViewportResource, ESPMode::ThreadSafe>> ViewportResources;
 
+	// Weak ref to the viewport manager
 	TWeakPtr<FDisplayClusterViewportManagerProxy, ESPMode::ThreadSafe> ViewportManagerProxyWeakPtr;
 };
diff --git a/Engine/Plugins/Runtime/nDisplay/Source/DisplayCluster/Private/Render/Viewport/Resource/DisplayClusterViewportPreviewResource.cpp b/Engine/Plugins/Runtime/nDisplay/Source/DisplayCluster/Private/Render/Viewport/Resource/DisplayClusterViewportPreviewResource.cpp
new file mode 100644
--- /dev/null
+++ b/Engine/Plugins/Runtime/nDisplay/Source/DisplayCluster/Private/Render/Viewport/Resource/DisplayClusterViewportPreviewResource.cpp
@@ -0,0 +1,14 @@
+// Copyright Epic Games, Inc. All Rights Reserved.
+
+#include "Render/Viewport/Resource/DisplayClusterViewportPreviewResource.h"
+
+#include "RHI.h"
+#include "RenderResource.h"
+#include "UnrealClient.h"
+
+////////////////////////////////////////////////////////////////////////
+// FDisplayClusterViewportPreviewResource
+////////////////////////////////////////////////////////////////////////
+FDisplayClusterViewportPreviewResource::FDisplayClusterViewportPreviewResource(const FDisplayClusterViewportResourceSettings& InResourceSettings)
+	: FDisplayClusterViewportResource(InResourceSettings)
+{ }
diff --git a/Engine/Plugins/Runtime/nDisplay/Source/DisplayCluster/Private/Render/Viewport/Resource/DisplayClusterViewportPreviewResource.h b/Engine/Plugins/Runtime/nDisplay/Source/DisplayCluster/Private/Render/Viewport/Resource/DisplayClusterViewportPreviewResource.h
new file mode 100644
--- /dev/null
+++ b/Engine/Plugins/Runtime/nDisplay/Source/DisplayCluster/Private/Render/Viewport/Resource/DisplayClusterViewportPreviewResource.h
@@ -0,0 +1,42 @@
+// Copyright Epic Games, Inc. All Rights Reserved.
+
+#pragma once
+#include "Render/Viewport/Resource/DisplayClusterViewportResource.h"
+
+#include "RHI.h"
+#include "RHIResources.h"
+#include "RenderResource.h"
+
+/**
+ * Viewport preview texture resource
+ */
+class FDisplayClusterViewportPreviewResource
+	: public FDisplayClusterViewportResource
+{
+public:
+	FDisplayClusterViewportPreviewResource(const FDisplayClusterViewportResourceSettings& InResourceSettings);
+	virtual ~FDisplayClusterViewportPreviewResource() = default;
+
+public:
+	//~ Begin FDisplayClusterViewportResource
+	virtual FRHITexture2D* GetViewportResourceRHI() const override
+	{
+		return OutputPreviewTargetableResource.IsValid() ? OutputPreviewTargetableResource->GetTexture2D() : nullptr;
+	}
+
+	virtual void SetExternalViewportResourceRHI(FTextureRHIRef& InExternalViewportResourceRHI) override
+	{
+		OutputPreviewTargetableResource = InExternalViewportResourceRHI;
+	}
+
+	virtual void ReleaseViewportResource_RenderThread(FRHICommandListBase& RHICmdList) override
+	{
+		OutputPreviewTargetableResource.SafeRelease();
+	}
+	//~~ End FDisplayClusterViewportResource
+
+public:
+	// here we implement the preview resources the old way,
+	// but when we remove the DCPreviewComponent, this class will be redesigned using UTexture.
+	FTextureRHIRef OutputPreviewTargetableResource;
+};
diff --git a/Engine/Plugins/Runtime/nDisplay/Source/DisplayCluster/Private/Render/Viewport/Resource/DisplayClusterViewportResource.cpp b/Engine/Plugins/Runtime/nDisplay/Source/DisplayCluster/Private/Render/Viewport/Resource/DisplayClusterViewportResource.cpp
new file mode 100644
--- /dev/null
+++ b/Engine/Plugins/Runtime/nDisplay/Source/DisplayCluster/Private/Render/Viewport/Resource/DisplayClusterViewportResource.cpp
@@ -0,0 +1,78 @@
+// Copyright Epic Games, Inc. All Rights Reserved.
+
+#include "Render/Viewport/Resource/DisplayClusterViewportResource.h"
+
+#include "RHI.h"
+#include "RenderResource.h"
+#include "UnrealClient.h"
+
+////////////////////////////////////////////////////////////////////////
+// FDisplayClusterViewportResource
+////////////////////////////////////////////////////////////////////////
+FDisplayClusterViewportResource::~FDisplayClusterViewportResource()
+{
+	// resource must be released in the rendering thread before the destructor is called
+	check(!EnumHasAnyFlags(ResourceState, EDisplayClusterViewportResourceState::Initialized));
+}
+
+void FDisplayClusterViewportResource::ImplInitDynamicRHI_RenderTargetResource2D(FTexture2DRHIRef& OutRenderTargetTextureRHI, FTexture2DRHIRef& OutTextureRHI)
+{
+	ETextureCreateFlags CreateFlags = TexCreate_Dynamic;
+
+	// -- we will be manually copying this cross GPU, tell render graph not to
+	CreateFlags |= TexCreate_MultiGPUGraphIgnore;
+
+	// reflect srgb from settings
+	if (EnumHasAnyFlags(ResourceSettings.GetResourceFlags(), EDisplayClusterViewportResourceSettingsFlags::ShouldUseSRGB))
+	{
+		CreateFlags |= TexCreate_SRGB;
+	}
+
+	uint32 NumMips = 1;
+	if (ResourceSettings.GetNumMips() > 1)
+	{
+		// Create nummips texture!
+		CreateFlags |= TexCreate_GenerateMipCapable | TexCreate_UAV;
+		NumMips = ResourceSettings.GetNumMips();
+	}
+
+	const FRHITextureCreateDesc Desc =
+		FRHITextureCreateDesc::Create2D(TEXT("DisplayClusterViewportRenderTargetResource"))
+		.SetExtent(ResourceSettings.GetSizeXY().X, ResourceSettings.GetSizeXY().Y)
+		.SetFormat(ResourceSettings.GetFormat())
+		.SetNumMips(NumMips)
+		.SetFlags(CreateFlags | ETextureCreateFlags::RenderTargetable | ETextureCreateFlags::ShaderResource)
+		.SetInitialState(ERHIAccess::SRVMask)
+		.SetClearValue(FClearValueBinding::Black);
+
+	OutRenderTargetTextureRHI = OutTextureRHI = RHICreateTexture(Desc);
+}
+
+void FDisplayClusterViewportResource::ImplInitDynamicRHI_TextureResource2D(FTexture2DRHIRef& OutTextureRHI)
+{
+	FRHITextureCreateDesc Desc =
+		FRHITextureCreateDesc::Create2D(TEXT("DisplayClusterViewportTextureResource"), ResourceSettings.GetSizeXY().X, ResourceSettings.GetSizeXY().Y, ResourceSettings.GetFormat())
+		.SetClearValue(FClearValueBinding::Black)
+		.SetFlags(ETextureCreateFlags::Dynamic | ETextureCreateFlags::ShaderResource)
+		// -- we will be manually copying this cross GPU, tell render graph not to
+		.AddFlags(ETextureCreateFlags::MultiGPUGraphIgnore)
+		.SetInitialState(ERHIAccess::SRVMask);
+
+	// reflect srgb from settings
+	if (EnumHasAnyFlags(ResourceSettings.GetResourceFlags(), EDisplayClusterViewportResourceSettingsFlags::ShouldUseSRGB))
+	{
+		Desc.AddFlags(ETextureCreateFlags::SRGB);
+	}
+
+	// reflect RenderTargetable flags
+	if (EnumHasAnyFlags(ResourceSettings.GetResourceFlags(), EDisplayClusterViewportResourceSettingsFlags::RenderTargetableTexture))
+	{
+		Desc.AddFlags(ETextureCreateFlags::RenderTargetable);
+	}
+	else
+	{
+		Desc.AddFlags(ETextureCreateFlags::ResolveTargetable);
+	}
+
+	OutTextureRHI = RHICreateTexture(Desc);
+}
diff --git a/Engine/Plugins/Runtime/nDisplay/Source/DisplayCluster/Private/Render/Viewport/Resource/DisplayClusterViewportResource.h b/Engine/Plugins/Runtime/nDisplay/Source/DisplayCluster/Private/Render/Viewport/Resource/DisplayClusterViewportResource.h
new file mode 100644
--- /dev/null
+++ b/Engine/Plugins/Runtime/nDisplay/Source/DisplayCluster/Private/Render/Viewport/Resource/DisplayClusterViewportResource.h
@@ -0,0 +1,113 @@
+// Copyright Epic Games, Inc. All Rights Reserved.
+
+#pragma once
+#include "DisplayClusterViewportResourceSettings.h"
+
+#include "RHI.h"
+#include "RHIResources.h"
+#include "RenderResource.h"
+
+#include "Templates/SharedPointer.h"
+
+class FRenderTarget;
+
+/**
+ * Runtime state of the viewport resource
+ */
+enum class EDisplayClusterViewportResourceState : uint8
+{
+	// empty state, without flags
+	None = 0,
+
+	// This resource was initialized in the rendering thread
+	Initialized = 1 << 0,
+
+	// This resource is not used for the current frame
+	Unused = 1 << 1,
+
+	// Disable reallocation of this resource for the current frame
+	DisableReallocate = 1 << 3,
+};
+ENUM_CLASS_FLAGS(EDisplayClusterViewportResourceState);
+
+/**
+ * Viewport resource texture base class
+ */
+class FDisplayClusterViewportResource
+	: public TSharedFromThis<FDisplayClusterViewportResource, ESPMode::ThreadSafe>
+{
+public:
+	FDisplayClusterViewportResource(const FDisplayClusterViewportResourceSettings& InResourceSettings)
+		: ResourceSettings(InResourceSettings)
+	{ }
+
+	virtual ~FDisplayClusterViewportResource();
+
+public:
+	/** Get the current state of the resource. */
+	inline EDisplayClusterViewportResourceState& GetResourceState()
+	{
+		return ResourceState;
+	}
+
+	/** Get the current const state of the resource. */
+	inline const EDisplayClusterViewportResourceState& GetResourceState() const
+	{
+		return ResourceState;
+	}
+
+	/** Get the current settings of the resource. */
+	inline const FDisplayClusterViewportResourceSettings& GetResourceSettings() const
+	{
+		return ResourceSettings;
+	}
+
+public:
+	/** Assign external RHI resource. */
+	virtual void SetExternalViewportResourceRHI(FTextureRHIRef& InExternalViewportResourceRHI)
+	{ }
+
+	/** Get current RHI resource. */
+	virtual FRHITexture2D* GetViewportResourceRHI() const
+	{
+		return nullptr;
+	}
+
+	/** Get RenderTarget resource. */
+	virtual FRenderTarget* GetViewportResourceRenderTarget()
+	{
+		return nullptr;
+	}
+
+	/** Getting the viewport rect offset in the backbuffer. */
+	virtual FIntPoint GetBackbufferFrameOffset() const
+	{
+		return FIntPoint::ZeroValue;
+	}
+
+	/** Settings the viewport rect offset in the backbuffer. */
+	virtual void SetBackbufferFrameOffset(const FIntPoint& InBackbufferFrameOffset)
+	{ }
+
+	/** RHI initialization for this resource in the rendering thread. */
+	virtual void InitializeViewportResource_RenderThread(FRHICommandListBase& RHICmdList)
+	{ }
+
+	/** RHI release for this resource in the rendering thread*/
+	virtual void ReleaseViewportResource_RenderThread(FRHICommandListBase& RHICmdList)
+	{ }
+
+protected:
+	/** Helper function: create RTT texture. */
+	void ImplInitDynamicRHI_RenderTargetResource2D(FTexture2DRHIRef& OutRenderTargetTextureRHI, FTexture2DRHIRef& OutTextureRHI);
+
+	/** Helper function: create 2D texture. */
+	void ImplInitDynamicRHI_TextureResource2D(FTexture2DRHIRef& OutTextureRHI);
+
+private:
+	// Settings used by this resource
+	const FDisplayClusterViewportResourceSettings ResourceSettings;
+
+	// Resource states
+	EDisplayClusterViewportResourceState ResourceState = EDisplayClusterViewportResourceState::None;
+};
diff --git a/Engine/Plugins/Runtime/nDisplay/Source/DisplayCluster/Private/Render/Viewport/Resource/DisplayClusterViewportResourceSettings.cpp b/Engine/Plugins/Runtime/nDisplay/Source/DisplayCluster/Private/Render/Viewport/Resource/DisplayClusterViewportResourceSettings.cpp
new file mode 100644
--- /dev/null
+++ b/Engine/Plugins/Runtime/nDisplay/Source/DisplayCluster/Private/Render/Viewport/Resource/DisplayClusterViewportResourceSettings.cpp
@@ -0,0 +1,65 @@
+// Copyright Epic Games, Inc. All Rights Reserved.
+
+#include "Render/Viewport/Resource/DisplayClusterViewportResourceSettings.h"
+#include "Render/Viewport/RenderFrame/DisplayClusterRenderFrameSettings.h"
+#include "Render/Viewport/DisplayClusterViewportHelpers.h"
+
+#include "RHI.h"
+#include "RenderResource.h"
+#include "UnrealClient.h"
+
+/////////////////////////////////////////////////////////////////////
+// FDisplayClusterViewportResource
+/////////////////////////////////////////////////////////////////////
+FDisplayClusterViewportResourceSettings::FDisplayClusterViewportResourceSettings(const FDisplayClusterRenderFrameSettings& InRenderFrameSettings, FViewport* InViewport)
+	: ClusterNodeId(InRenderFrameSettings.ClusterNodeId)
+{
+	if (FRHITexture2D* ViewportTexture = InViewport ? InViewport->GetRenderTargetTexture() : nullptr)
+	{
+		Format = ViewportTexture->GetFormat();
+
+		if (EnumHasAnyFlags(ViewportTexture->GetFlags(), TexCreate_SRGB))
+		{
+			EnumAddFlags(ResourceFlags, EDisplayClusterViewportResourceSettingsFlags::ShouldUseSRGB);
+		}
+
+		DisplayGamma = InViewport->GetDisplayGamma();
+	}
+	else
+	{
+		// Use default settings:
+		Format = FDisplayClusterViewportHelpers::GetDefaultPixelFormat();
+		DisplayGamma = 2.2f;
+
+		// Always use srgb for preview rendering
+		EnumAddFlags(ResourceFlags, EDisplayClusterViewportResourceSettingsFlags::ShouldUseSRGB);
+
+		switch (InRenderFrameSettings.RenderMode)
+		{
+		case EDisplayClusterRenderFrameMode::PreviewInScene:
+			Format = FDisplayClusterViewportHelpers::GetPreviewDefaultPixelFormat();
+
+			// Preview display gamma
+			if (InRenderFrameSettings.bPreviewEnablePostProcess == false)
+			{
+				// Disable postprocess for preview. Use Gamma 1.f
+				DisplayGamma = 1.f;
+			}
+			break;
+
+		default:
+			break;
+		}
+	}
+}
+
+FDisplayClusterViewportResourceSettings::FDisplayClusterViewportResourceSettings(const FDisplayClusterViewportResourceSettings& InBaseSettings, const FIntPoint& InSize, const EPixelFormat InFormat, const EDisplayClusterViewportResourceSettingsFlags InResourceFlags, const int32 InNumMips)
+	: ClusterNodeId(InBaseSettings.ClusterNodeId)
+	, Size(InSize)
+	, Format((InFormat == PF_Unknown) ? InBaseSettings.Format : InFormat)
+	, DisplayGamma(InBaseSettings.DisplayGamma)
+	, NumMips(InNumMips)
+	, ResourceFlags(InBaseSettings.ResourceFlags)
+{
+	EnumAddFlags(ResourceFlags, InResourceFlags);
+}
diff --git a/Engine/Plugins/Runtime/nDisplay/Source/DisplayCluster/Private/Render/Viewport/Resource/DisplayClusterViewportResourceSettings.h b/Engine/Plugins/Runtime/nDisplay/Source/DisplayCluster/Private/Render/Viewport/Resource/DisplayClusterViewportResourceSettings.h
new file mode 100644
--- /dev/null
+++ b/Engine/Plugins/Runtime/nDisplay/Source/DisplayCluster/Private/Render/Viewport/Resource/DisplayClusterViewportResourceSettings.h
@@ -0,0 +1,144 @@
+// Copyright Epic Games, Inc. All Rights Reserved.
+
+#pragma once
+#include "CoreMinimal.h"
+#include "PixelFormat.h"
+
+struct FDisplayClusterRenderFrameSettings;
+class FViewport;
+
+/**
+ * Runtime state of the viewport resource
+ * Depending on the flags used, different child classes will be used.
+ */
+enum class EDisplayClusterViewportResourceSettingsFlags : uint8
+{
+	// Default value, no flags
+	None = 0,
+
+	// The resource will use the 'ETextureCreateFlags::SRGB' flag.
+	ShouldUseSRGB = 1 << 0,
+
+	// This is the RenderTarget resource used in the viewfamily (not the usual texture).
+	RenderTarget = 1 << 1,
+
+	// This is the texture.
+	// The render resource will be created with the 'ETextureCreateFlags::RenderTargetable' flag.
+	RenderTargetableTexture = 1 << 3,
+
+	// This is the texture.
+	// The render resource will be created with the 'ETextureCreateFlags::ResolveTargetable' flag.
+	ResolveTargetableTexture = 1 << 4,
+
+	// This is the texture that is used to render the DCRA preview
+	// (contains the UTexture that is used for UMaterialInstanceDynamic).
+	PreviewTargetableTexture = 1 << 5,
+};
+ENUM_CLASS_FLAGS(EDisplayClusterViewportResourceSettingsFlags);
+
+/**
+ * nDisplay viewport render resource settings
+ */
+struct FDisplayClusterViewportResourceSettings
+{
+	/** Create default resource settings
+	* 
+	* @param InRenderFrameSettings - render frame settings used (preview or normal rendering, etc.)
+	* @param InViewport - the window viewport interface (used display gamma, format, etc.)
+	*/
+	FDisplayClusterViewportResourceSettings(const FDisplayClusterRenderFrameSettings& InRenderFrameSettings, class FViewport* InViewport);
+
+	/** Create a custom resource setting based on another one.
+	* 
+	* @param InBaseSettings  -  basic settings, if no values are defined, the values from them will be used.
+	* @param InSize          - resource size
+	* @param InFormat        - resource format
+	* @param InResourceFlags - these flags determine what type of resource will be created.
+	* @param InNumMips       - the number of mips that will be created. This value is only supported by a certain type of resource.
+	*/
+	FDisplayClusterViewportResourceSettings(const FDisplayClusterViewportResourceSettings& InBaseSettings, const FIntPoint& InSize, const EPixelFormat InFormat, const EDisplayClusterViewportResourceSettingsFlags InResourceFlags = EDisplayClusterViewportResourceSettingsFlags::None, const int32 InNumMips = 1);
+
+public:
+	/** Returns true if the input resource settings match the local values. */
+	inline bool IsResourceSettingsEqual(const FDisplayClusterViewportResourceSettings& In) const
+	{
+		// This is a temporary rule, because the preview texture uses a link to an external resource
+		if (EnumHasAnyFlags(ResourceFlags | In.ResourceFlags, EDisplayClusterViewportResourceSettingsFlags::PreviewTargetableTexture))
+		{
+			return false;
+		}
+
+		return (In.Size == Size)
+			&& (In.Format == Format)
+			&& (In.ResourceFlags == ResourceFlags)
+			&& (In.DisplayGamma == DisplayGamma)
+			&& (In.NumMips == NumMips);
+	}
+
+	/** Returns true if the resource settings are used in the same cluster node. */
+	inline bool IsClusterNodeNameEqual(const FDisplayClusterViewportResourceSettings& In) const
+	{
+		return ClusterNodeId == In.ClusterNodeId;
+	}
+
+	/** Get the X and Y dimensions of the resource. */
+	inline const FIntPoint& GetSizeXY() const
+	{
+		return Size;
+	}
+
+	/** Get the width of the resource. */
+	inline const int32 GetSizeX() const
+	{
+		return Size.X;
+	}
+
+	/** Get the height of the resource. */
+	inline const int32 GetSizeY() const
+	{
+		return Size.Y;
+	}
+
+	/** Get the pixel format of the resource. */
+	inline EPixelFormat GetFormat() const
+	{
+		return Format;
+	}
+
+	/** Get the display gamma used by this resource */
+	inline float GetDisplayGamma() const
+	{
+		return DisplayGamma;
+	}
+
+	/** Get the number of resource mips. */
+	inline int32 GetNumMips() const
+	{
+		return NumMips;
+	}
+
+	/** Retrieving resource flags. */
+	inline EDisplayClusterViewportResourceSettingsFlags GetResourceFlags() const
+	{
+		return ResourceFlags;
+	}
+
+private:
+	// This is a resource belonging to the cluster node
+	const FString ClusterNodeId;
+
+	// Resource width and height
+	FIntPoint    Size;
+
+	// The pixel format of the resource.
+	EPixelFormat Format = EPixelFormat::PF_Unknown;
+
+	// The display gamma used by this resource
+	float DisplayGamma = 2.2f;
+
+	// Number of resource mips.
+	int32 NumMips = 0;
+
+	// Resource flags.
+	EDisplayClusterViewportResourceSettingsFlags ResourceFlags = EDisplayClusterViewportResourceSettingsFlags::None;
+};
diff --git a/Engine/Plugins/Runtime/nDisplay/Source/DisplayCluster/Public/Render/Viewport/Containers/DisplayClusterViewport_Enums.h b/Engine/Plugins/Runtime/nDisplay/Source/DisplayCluster/Public/Render/Viewport/Containers/DisplayClusterViewport_Enums.h
--- a/Engine/Plugins/Runtime/nDisplay/Source/DisplayCluster/Public/Render/Viewport/Containers/DisplayClusterViewport_Enums.h
+++ b/Engine/Plugins/Runtime/nDisplay/Source/DisplayCluster/Public/Render/Viewport/Containers/DisplayClusterViewport_Enums.h
@@ -4,47 +4,90 @@
 
 #include "CoreMinimal.h"
 
+/**
+ * Types of DC viewport resources
+ */
 enum class EDisplayClusterViewportResourceType : uint8
 {
-	// All viewport contexts rendered to render target (using atlasing to map multiple views)	
-	// Then called single cross gpu transfer for multiple render targets regions
-	// After corss gpu transfers viewport regions resolved to unique viewports contexts resources 'InputShaderResource'
-	// (Context region saved FDisplayClusterViewport_Context::RenderTargetRect)
-	// * this resource used in next function for postprocess and mips generation:
-	//
-	//    FDisplayClusterViewportProxy::UpdateDeferredResources()
-	//     {
-	//       Resolve(InternalRenderTargetResource -> InputShaderResource)
-	//       Postprocess(InputShaderResource -> AdditionalTargetableResource -> InputShaderResource)
-	//.......GenerateMips(InputShaderResource -> MipsShaderResource)
-	//     }
-	//
-	InternalRenderTargetResource = 0, /** Just for Internal use */
+	// Undefined resource type
+	Unknown = 0,
+
+	// This RTT is used to render the scene for the viewport eye context to the specified area on this texture.
+	// The mGPU rendering with cross-GPU transfer can be used for this texture.
+	InternalRenderTargetResource, /** Just for Internal use */
+
+
+	/**
+	 * Internal textures of the DC rendering pipeline.
+	 */
 
 	// Unique Viewport contexts shader resource. No regions used at this point
-	// Use this resources as source for projection policy and resolve to frame targets
+	// The image source for this resource can be: 'InternalRenderTargetResource', 'OverrideTexture', 'OverrideViewport', etc.
+	// This is the entry point to the DC rendering pipeline
 	InputShaderResource,
 
 	// Special resource with mips texture, generated from 'InputShaderResource' after postprocess
 	MipsShaderResource,
 
 	// Additional targetable resource, used by logic (external warpblend, blur, etc)
+	// This resource is used as an additional RTT for 'InputShaderResource'.
+	// It contains different contents depending on the point in time: 'PostprocessRTT','ViewportOutputRemap','AfterWarpBlend' etc.
 	AdditionalTargetableResource,
 
-	// Final frame targetable resources used externally as frame output
-	// Projection policy render output to this resources into viewport region
-	// (Context frame region saved FDisplayClusterViewport_Context::FrameTargetRect)
-	OutputFrameTargetableResource,
 
-	// Some frame postprocess require this additional full frame resource
-	AdditionalFrameTargetableResource,
+	/**
+	 * Textures for warp and blend
+	 */
+	 
+	 // The viewport texture before warpblend (InputShaderResource)
+	// This resource is used as the source image for the projection policy.
+	BeforeWarpBlendTargetableResource,
+
+	// The viewport texture after warpblend (AdditionalTargetableResource or InputShaderResource).
+	// This resource is used as the output image in the projection policy.
+	AfterWarpBlendTargetableResource,
+
+
+	/**
+	 * Output textures
+	 */
+
+	// This is the entry point into the output texture rendering.('OutputFrameTargetableResource' or 'OutputPreviewTargetableResource')
+	// The rendering time of this image in the DC rendering pipeline is right after warp&blend.
+	OutputTargetableResource,
 
-	// Finally OutputFrameTargetableResource resolved from OutputFrameTargetableResource to the Backbuffer from external logic
-#if WITH_EDITOR
+
+	/**
+	 * Output textures for preview rendering.
+	 */
+
+	// If DCRA uses a preview, this texture will be used instead of 'Frame' textures
 	OutputPreviewTargetableResource,
-#endif
+
+
+	/**
+	 * Output 'Frame' textures for cluster rendering
+	 * The 'Frame' resources are used to compose all the viewports into a single texture.
+	 * There is a separate 'frame' texture for each eye context.
+	 * And at the end of the frame these resources are copied to the backbuffer.
+	 * 
+	 * Projection policy render output to this resources into viewport region
+	 * (Context frame region saved FDisplayClusterViewport_Context::FrameTargetRect)
+	 */
+
+	// This texture contains the results of the DC rendering, and is copied directly to the backbuffer texture at the end of the frame.
+	// Each eye is in a separate texture.
+	OutputFrameTargetableResource,
+
+	// This resource is used as an additional RTT for 'OutputFrameTargetableResource'.
+	// It contains different contents depending on the point in time: 'FramePostprocessRTT','OutputRemap', etc.
+	AdditionalFrameTargetableResource,
 };
 
+/**
+ * Viewport capture mode
+ * This mode affects many viewport rendering settings.
+ */
 enum class EDisplayClusterViewportCaptureMode : uint8
 {
 	// Use current scene format, no alpha
@@ -60,6 +103,10 @@ enum class EDisplayClusterViewportCaptureMode : uint8
 	MoviePipeline,
 };
 
+/**
+ * Viewport can be overridden by another one.
+ * This mode determines how many resources will be overridden.
+ */
 enum class EDisplayClusterViewportOverrideMode : uint8
 {
 	// Do not override this viewport from the other one (Render viewport; create all resources)
diff --git a/Engine/Plugins/Runtime/nDisplay/Source/DisplayCluster/Public/Render/Viewport/Containers/DisplayClusterViewport_RenderSettings.h b/Engine/Plugins/Runtime/nDisplay/Source/DisplayCluster/Public/Render/Viewport/Containers/DisplayClusterViewport_RenderSettings.h
--- a/Engine/Plugins/Runtime/nDisplay/Source/DisplayCluster/Public/Render/Viewport/Containers/DisplayClusterViewport_RenderSettings.h
+++ b/Engine/Plugins/Runtime/nDisplay/Source/DisplayCluster/Public/Render/Viewport/Containers/DisplayClusterViewport_RenderSettings.h
@@ -88,12 +88,8 @@ public:
 	// Special capture modes (chromakey, lightcard) change RTT format and render flags
 	EDisplayClusterViewportCaptureMode CaptureMode = EDisplayClusterViewportCaptureMode::Default;
 
-	// Override resources from this viewport
-	EDisplayClusterViewportOverrideMode ViewportOverrideMode = EDisplayClusterViewportOverrideMode::None;
-	FString ViewportOverrideId;
-
 public:
-	// Reset runtime values from prev frame
+	/** Resets the viewport settings. This function is called every frame at the beginning of the frame. */
 	inline void BeginUpdateSettings()
 	{
 		bVisible = true;
@@ -108,54 +104,84 @@ public:
 		ViewportOverrideId.Empty();
 	}
 
+	/** Finishes setting the viewport in the game thread. Called once per frame at the end. */
 	inline void FinishUpdateSettings()
 	{
 		bPreviewReadPixels = false;
 	}
 
-	inline const FString& GetParentViewportId() const
+	/** Returns true if the viewport is assigned to a parent viewport. */
+	inline bool IsViewportHasParent() const
 	{
-		return ParentViewportId;
+		return !ParentViewportId.IsEmpty();
 	}
 
-	inline bool IsViewportHasParent() const
+	/** Get the name of the parent viewport. */
+	inline const FString& GetParentViewportId() const
 	{
-		return !ParentViewportId.IsEmpty();
+		return ParentViewportId;
 	}
 
-	inline void SetViewportOverride(const FString& InViewportOverrideId, const EDisplayClusterViewportOverrideMode InViewportOverrideMode = EDisplayClusterViewportOverrideMode::All)
+	/** Assign parent viewport to this.
+	* The main idea is copiing some render settings and math from the parent viewport.
+	* This is used for a 'link' projection policy to render LC and CK from the same frustum assigned as the parent (outer for LC, incamera for CK).
+	* Also, child viewports are only updated when the parent viewport is updated (sorted in ViewportManager/Proxy).
+	* 
+	* @param InParentViewportId - parent viewport name
+	* @param InParentSettings   - parent viewport rendering settings.
+	*/
+	inline void AssignParentViewport(const FString& InParentViewportId, const FDisplayClusterViewport_RenderSettings& InParentSettings)
 	{
-		ViewportOverrideMode = InViewportOverrideMode;
-		ViewportOverrideId = InViewportOverrideId;
+		ParentViewportId = InParentViewportId;
+
+		// Inherit values from parent viewport:
+		CameraId = InParentSettings.CameraId;
+		Rect = InParentSettings.Rect;
+
+		bForceMono = InParentSettings.bForceMono;
+
+		GPUIndex = (GPUIndex < 0) ? InParentSettings.GPUIndex : GPUIndex;
+		StereoGPUIndex = (StereoGPUIndex < 0) ? InParentSettings.StereoGPUIndex : StereoGPUIndex;
+
+		RenderFamilyGroup = (RenderFamilyGroup < 0) ? InParentSettings.RenderFamilyGroup : RenderFamilyGroup;
 	}
 
-	inline bool IsViewportOverrided() const
+	/** The viewport can be overridden from another viewport. This function returns true if it is. */
+	inline bool IsViewportOverridden() const
 	{
 		return ViewportOverrideMode != EDisplayClusterViewportOverrideMode::None && !ViewportOverrideId.IsEmpty();
 	}
 
-	// Call this after UpdateSettings()
-	inline void AssignParentViewport(const FString& InParentViewportId, const FDisplayClusterViewport_RenderSettings& InParentSettings, bool Inherit = true)
+	/** Getting the override mode that is currently in use. */
+	inline EDisplayClusterViewportOverrideMode GetViewportOverrideMode() const
 	{
-		ParentViewportId = InParentViewportId;
-
-		// Inherit values from parent viewport:
-		if (Inherit)
-		{
-			CameraId = InParentSettings.CameraId;
-			Rect = InParentSettings.Rect;
-
-			bForceMono = InParentSettings.bForceMono;
+		return IsViewportOverridden() ? ViewportOverrideMode : EDisplayClusterViewportOverrideMode::None;
+	}
 
-			GPUIndex = (GPUIndex < 0) ? InParentSettings.GPUIndex : GPUIndex;
-			StereoGPUIndex = (StereoGPUIndex < 0) ? InParentSettings.StereoGPUIndex : StereoGPUIndex;
+	/** Get the name of the viewport used as the image source. */
+	inline const FString& GetViewportOverrideId() const
+	{
+		return ViewportOverrideId;
+	}
 
-			RenderFamilyGroup = (RenderFamilyGroup < 0) ? InParentSettings.RenderFamilyGroup : RenderFamilyGroup;
-		}
+	/** Set an override for viewport images from another viewport.
+	* 
+	* @param InViewportOverrideId   - The name of the viewport used as the image source.
+	* @param InViewportOverrideMode - Override mode, which defines its rules.
+	*/
+	inline void SetViewportOverride(const FString& InViewportOverrideId, const EDisplayClusterViewportOverrideMode InViewportOverrideMode = EDisplayClusterViewportOverrideMode::All)
+	{
+		ViewportOverrideMode = InViewportOverrideMode;
+		ViewportOverrideId = InViewportOverrideId;
 	}
 
 protected:
 	// Parent viewport name
 	FString ParentViewportId;
-};
 
+	// Override resources from another viewport. The name of the viewport used as the image source.
+	FString ViewportOverrideId;
+
+	// Override mode, which defines its rules.
+	EDisplayClusterViewportOverrideMode ViewportOverrideMode = EDisplayClusterViewportOverrideMode::None;
+};
diff --git a/Engine/Plugins/Runtime/nDisplay/Source/DisplayCluster/Public/Render/Viewport/IDisplayClusterViewportManager.h b/Engine/Plugins/Runtime/nDisplay/Source/DisplayCluster/Public/Render/Viewport/IDisplayClusterViewportManager.h
--- a/Engine/Plugins/Runtime/nDisplay/Source/DisplayCluster/Public/Render/Viewport/IDisplayClusterViewportManager.h
+++ b/Engine/Plugins/Runtime/nDisplay/Source/DisplayCluster/Public/Render/Viewport/IDisplayClusterViewportManager.h
@@ -137,7 +137,6 @@ public:
 	/** Add internal DCVM objects to the reference collector. */
 	virtual void AddReferencedObjects(FReferenceCollector& Collector) = 0;
 
-#if WITH_EDITOR
 	/**
 	* Render in editor (preview)
 	* [Game thread func]
@@ -152,7 +151,6 @@ public:
 	* @return - true, if render success
 	*/
 	virtual bool RenderInEditor(FDisplayClusterRenderFrame& InRenderFrame, FViewport* InViewport, const uint32 InFirstViewportNum, const int32 InViewportsAmount, int32& OutViewportsAmount, bool& bOutFrameRendered) = 0;
-#endif
 
 	///////////////////////////////////////////////////////////////////////////////////////////////////////////////////////////////////////////////////////////////////////////
 	/**
diff --git a/Engine/Plugins/Runtime/nDisplay/Source/DisplayCluster/Public/Render/Viewport/IDisplayClusterViewportProxy.h b/Engine/Plugins/Runtime/nDisplay/Source/DisplayCluster/Public/Render/Viewport/IDisplayClusterViewportProxy.h
--- a/Engine/Plugins/Runtime/nDisplay/Source/DisplayCluster/Public/Render/Viewport/IDisplayClusterViewportProxy.h
+++ b/Engine/Plugins/Runtime/nDisplay/Source/DisplayCluster/Public/Render/Viewport/IDisplayClusterViewportProxy.h
@@ -76,7 +76,11 @@ public:
 	virtual bool ResolveResources_RenderThread(FRHICommandListImmediate& RHICmdList, const EDisplayClusterViewportResourceType InputResourceType, const EDisplayClusterViewportResourceType OutputResourceType, const int32 InContextNum = INDEX_NONE) const = 0;
 
 	/** Return output resource type (support preview, remap, etc). */
-	virtual EDisplayClusterViewportResourceType   GetOutputResourceType_RenderThread() const = 0;
+	UE_DEPRECATED(5.3, "This function has beend deprecate. Please use the new enumeration values.")
+	virtual EDisplayClusterViewportResourceType GetOutputResourceType_RenderThread() const
+	{
+		return EDisplayClusterViewportResourceType::OutputTargetableResource;
+	}
 
 	/** Returns ptr to ViewportManagerProxy (the owner of this viewport proxy) if it still exists. */
 	virtual const class IDisplayClusterViewportManagerProxy* GetViewportManagerProxy_RenderThread() const = 0;
diff --git a/Engine/Plugins/Runtime/nDisplay/Source/DisplayCluster/Public/Render/Warp/IDisplayClusterWarpPolicy.h b/Engine/Plugins/Runtime/nDisplay/Source/DisplayCluster/Public/Render/Warp/IDisplayClusterWarpPolicy.h
--- a/Engine/Plugins/Runtime/nDisplay/Source/DisplayCluster/Public/Render/Warp/IDisplayClusterWarpPolicy.h
+++ b/Engine/Plugins/Runtime/nDisplay/Source/DisplayCluster/Public/Render/Warp/IDisplayClusterWarpPolicy.h
@@ -82,7 +82,7 @@ public:
 	 * @param InViewport - a owner viewport
 	 * @param ContextNum - viewport eye context index
 	 * 
-	 * @return - true if frustum overrided.
+	 * @return - true if frustum is overridden.
 	 */
 	virtual bool OverrideCalcFrustum(IDisplayClusterViewport* InViewport, const uint32 ContextNum)
 	{
diff --git a/Engine/Plugins/Runtime/nDisplay/Source/DisplayClusterConfiguration/Private/DisplayClusterConfigurationTypes_Viewport.cpp b/Engine/Plugins/Runtime/nDisplay/Source/DisplayClusterConfiguration/Private/DisplayClusterConfigurationTypes_Viewport.cpp
--- a/Engine/Plugins/Runtime/nDisplay/Source/DisplayClusterConfiguration/Private/DisplayClusterConfigurationTypes_Viewport.cpp
+++ b/Engine/Plugins/Runtime/nDisplay/Source/DisplayClusterConfiguration/Private/DisplayClusterConfigurationTypes_Viewport.cpp
@@ -17,7 +17,7 @@ EDisplayClusterShaderParametersICVFX_LightCardRenderMode FDisplayClusterConfigur
 
 	if (LightcardRenderMode != EDisplayClusterConfigurationICVFX_OverrideLightcardRenderMode::Default)
 	{
-		// Use overrided values from viewport:
+		// Use overridden values from the viewport:
 		switch (LightcardRenderMode)
 		{
 		case EDisplayClusterConfigurationICVFX_OverrideLightcardRenderMode::Over:
diff --git a/Engine/Plugins/Runtime/nDisplay/Source/DisplayClusterMoviePipeline/Private/DisplayClusterMoviePipelineViewportPass.cpp b/Engine/Plugins/Runtime/nDisplay/Source/DisplayClusterMoviePipeline/Private/DisplayClusterMoviePipelineViewportPass.cpp
--- a/Engine/Plugins/Runtime/nDisplay/Source/DisplayClusterMoviePipeline/Private/DisplayClusterMoviePipelineViewportPass.cpp
+++ b/Engine/Plugins/Runtime/nDisplay/Source/DisplayClusterMoviePipeline/Private/DisplayClusterMoviePipelineViewportPass.cpp
@@ -18,7 +18,7 @@
 #include "RHIUtilities.h"
 
 /////////////////////////////////////////////////////////////////////////////////////////////////////////////////////////////////
-namespace DisplayClusterMoviePipelineHelpers
+namespace UE::DisplayCluster::MoviePipelineViewportPass
 {
 	static void DisplayClusterWarpBlendImpl_RenderThread(FRHICommandListImmediate& RHICmdList, const FTextureRHIRef& RenderTargetRHI, IDisplayClusterViewportManagerProxy* InViewportManagerProxy, const FString& InViewportId, const FIntPoint& OffsetMin, const FIntPoint& OffsetMax)
 	{
@@ -80,7 +80,7 @@ namespace DisplayClusterMoviePipelineHelpers
 		}
 	}
 };
-using namespace DisplayClusterMoviePipelineHelpers;
+using namespace UE::DisplayCluster::MoviePipelineViewportPass;
 
 /////////////////////////////////////////////////////////////////////////////////////////////////////////////////////////////////
 // UDisplayClusterMoviePipelineViewportPassBase
diff --git a/Engine/Plugins/Runtime/nDisplay/Source/DisplayClusterProjection/Private/Policy/Domeprojection/Windows/DX11/DisplayClusterProjectionDomeprojectionViewAdapterDX11.cpp b/Engine/Plugins/Runtime/nDisplay/Source/DisplayClusterProjection/Private/Policy/Domeprojection/Windows/DX11/DisplayClusterProjectionDomeprojectionViewAdapterDX11.cpp
--- a/Engine/Plugins/Runtime/nDisplay/Source/DisplayClusterProjection/Private/Policy/Domeprojection/Windows/DX11/DisplayClusterProjectionDomeprojectionViewAdapterDX11.cpp
+++ b/Engine/Plugins/Runtime/nDisplay/Source/DisplayClusterProjection/Private/Policy/Domeprojection/Windows/DX11/DisplayClusterProjectionDomeprojectionViewAdapterDX11.cpp
@@ -197,7 +197,7 @@ bool FDisplayClusterProjectionDomeprojectionViewAdapterDX11::ApplyWarpBlend_Rend
 
 	// Get output resources with rects
 	// warp result is now inside AdditionalRTT.  Later, from the DC ViewportManagerProxy it will be resolved to FrameRTT
-	if (!InViewportProxy->GetResources_RenderThread(EDisplayClusterViewportResourceType::AdditionalTargetableResource, OutputTextures))
+	if (!InViewportProxy->GetResources_RenderThread(EDisplayClusterViewportResourceType::AfterWarpBlendTargetableResource, OutputTextures))
 	{
 		return false;
 	}
diff --git a/Engine/Plugins/Runtime/nDisplay/Source/DisplayClusterProjection/Private/Policy/EasyBlend/DisplayClusterProjectionEasyBlendPolicyBase.cpp b/Engine/Plugins/Runtime/nDisplay/Source/DisplayClusterProjection/Private/Policy/EasyBlend/DisplayClusterProjectionEasyBlendPolicyBase.cpp
--- a/Engine/Plugins/Runtime/nDisplay/Source/DisplayClusterProjection/Private/Policy/EasyBlend/DisplayClusterProjectionEasyBlendPolicyBase.cpp
+++ b/Engine/Plugins/Runtime/nDisplay/Source/DisplayClusterProjection/Private/Policy/EasyBlend/DisplayClusterProjectionEasyBlendPolicyBase.cpp
@@ -118,7 +118,7 @@ void FDisplayClusterProjectionEasyBlendPolicyBase::ApplyWarpBlend_RenderThread(F
 		if (!ViewAdapter->ApplyWarpBlend_RenderThread(RHICmdList, InViewportProxy))
 		{
 			// Warpbled failed, just copy unwarped rtt for frame output
-			InViewportProxy->ResolveResources_RenderThread(RHICmdList, EDisplayClusterViewportResourceType::InputShaderResource, InViewportProxy->GetOutputResourceType_RenderThread());
+			InViewportProxy->ResolveResources_RenderThread(RHICmdList, EDisplayClusterViewportResourceType::InputShaderResource, EDisplayClusterViewportResourceType::OutputTargetableResource);
 		}
 	}
 }
diff --git a/Engine/Plugins/Runtime/nDisplay/Source/DisplayClusterProjection/Private/Policy/EasyBlend/Windows/DX11/DisplayClusterProjectionEasyBlendViewAdapterDX11.cpp b/Engine/Plugins/Runtime/nDisplay/Source/DisplayClusterProjection/Private/Policy/EasyBlend/Windows/DX11/DisplayClusterProjectionEasyBlendViewAdapterDX11.cpp
--- a/Engine/Plugins/Runtime/nDisplay/Source/DisplayClusterProjection/Private/Policy/EasyBlend/Windows/DX11/DisplayClusterProjectionEasyBlendViewAdapterDX11.cpp
+++ b/Engine/Plugins/Runtime/nDisplay/Source/DisplayClusterProjection/Private/Policy/EasyBlend/Windows/DX11/DisplayClusterProjectionEasyBlendViewAdapterDX11.cpp
@@ -235,7 +235,7 @@ bool FDisplayClusterProjectionEasyBlendViewAdapterDX11::ApplyWarpBlend_RenderThr
 
 	// Get output resources with rects
 	// warp result is now inside AdditionalRTT.  Later, from the DC ViewportManagerProxy it will be resolved to FrameRTT
-	if (!InViewportProxy->GetResources_RenderThread(EDisplayClusterViewportResourceType::AdditionalTargetableResource, OutputTextures))
+	if (!InViewportProxy->GetResources_RenderThread(EDisplayClusterViewportResourceType::AfterWarpBlendTargetableResource, OutputTextures))
 	{
 		return false;
 	}
diff --git a/Engine/Plugins/Runtime/nDisplay/Source/DisplayClusterProjection/Private/Policy/MPCDI/DisplayClusterProjectionMPCDIPolicy.cpp b/Engine/Plugins/Runtime/nDisplay/Source/DisplayClusterProjection/Private/Policy/MPCDI/DisplayClusterProjectionMPCDIPolicy.cpp
--- a/Engine/Plugins/Runtime/nDisplay/Source/DisplayClusterProjection/Private/Policy/MPCDI/DisplayClusterProjectionMPCDIPolicy.cpp
+++ b/Engine/Plugins/Runtime/nDisplay/Source/DisplayClusterProjection/Private/Policy/MPCDI/DisplayClusterProjectionMPCDIPolicy.cpp
@@ -307,7 +307,7 @@ void FDisplayClusterProjectionMPCDIPolicy::ApplyWarpBlend_RenderThread(FRHIComma
 
 	// Get output resources with rects
 	// warp result is now inside AdditionalRTT.  Later, from the DC ViewportManagerProxy it will be resolved to FrameRTT 
-	if (!InViewportProxy->GetResourcesWithRects_RenderThread(EDisplayClusterViewportResourceType::AdditionalTargetableResource, OutputTextures, OutputRects))
+	if (!InViewportProxy->GetResourcesWithRects_RenderThread(EDisplayClusterViewportResourceType::AfterWarpBlendTargetableResource, OutputTextures, OutputRects))
 	{
 		return;
 	}
diff --git a/Engine/Plugins/Runtime/nDisplay/Source/DisplayClusterProjection/Private/Policy/VIOSO/Windows/DisplayClusterProjectionVIOSOPolicy.cpp b/Engine/Plugins/Runtime/nDisplay/Source/DisplayClusterProjection/Private/Policy/VIOSO/Windows/DisplayClusterProjectionVIOSOPolicy.cpp
--- a/Engine/Plugins/Runtime/nDisplay/Source/DisplayClusterProjection/Private/Policy/VIOSO/Windows/DisplayClusterProjectionVIOSOPolicy.cpp
+++ b/Engine/Plugins/Runtime/nDisplay/Source/DisplayClusterProjection/Private/Policy/VIOSO/Windows/DisplayClusterProjectionVIOSOPolicy.cpp
@@ -275,7 +275,7 @@ void FDisplayClusterProjectionVIOSOPolicy::ApplyWarpBlend_RenderThread(FRHIComma
 	if (!ImplApplyWarpBlend_RenderThread(RHICmdList, InViewportProxy))
 	{
 		// warp failed, just resolve texture to frame
-		InViewportProxy->ResolveResources_RenderThread(RHICmdList, EDisplayClusterViewportResourceType::InputShaderResource, InViewportProxy->GetOutputResourceType_RenderThread());
+		InViewportProxy->ResolveResources_RenderThread(RHICmdList, EDisplayClusterViewportResourceType::InputShaderResource, EDisplayClusterViewportResourceType::OutputTargetableResource);
 	}
 }
 
@@ -300,7 +300,7 @@ bool FDisplayClusterProjectionVIOSOPolicy::ImplApplyWarpBlend_RenderThread(FRHIC
 
 	// Get output resources with rects
 	// warp result is now inside AdditionalRTT.  Later, from the DC ViewportManagerProxy it will be resolved to FrameRTT 
-	if (!InViewportProxy->GetResources_RenderThread(EDisplayClusterViewportResourceType::AdditionalTargetableResource, OutputTextures))
+	if (!InViewportProxy->GetResources_RenderThread(EDisplayClusterViewportResourceType::AfterWarpBlendTargetableResource, OutputTextures))
 	{
 		return false;
 	}
