nDisplay: Fixed  lightcards disappear when alpha is allowed through tonemapper (proj set.) and OCIO is enabled on them.
- removed the auto option due to the fact that the 'propagate alpha' rendering option is experimental and does not always work correctly.
- the default option is now FXAA, which always works

diff --git a/Engine/Plugins/Runtime/nDisplay/Source/DisplayCluster/Private/Render/Viewport/Configuration/DisplayClusterViewportConfiguration.cpp b/Engine/Plugins/Runtime/nDisplay/Source/DisplayCluster/Private/Render/Viewport/Configuration/DisplayClusterViewportConfiguration.cpp
--- a/Engine/Plugins/Runtime/nDisplay/Source/DisplayCluster/Private/Render/Viewport/Configuration/DisplayClusterViewportConfiguration.cpp
+++ b/Engine/Plugins/Runtime/nDisplay/Source/DisplayCluster/Private/Render/Viewport/Configuration/DisplayClusterViewportConfiguration.cpp
@@ -59,15 +59,10 @@ static FAutoConsoleVariableRef CVarDisplayClusterCrossGPUTransferPullData(
 */
 enum class ECVarDisplayClusterAlphaChannelCaptureMode : uint8
 {
-	/** [Auto]
-	 * Use one of the [ThroughTonemapper] or [FXAA] modes available for the current project settings.
-	 */
-	Auto = 0,
-
 	/** [Disabled]
 	 * Disable alpha channel saving.
 	 */
-	 Disabled,
+	Disabled,
 
 	/** [ThroughTonemapper]
 	 * When rendering with the PropagateAlpha experimental mode turned on, the alpha channel is forwarded to post-processes.
@@ -102,17 +97,16 @@ enum class ECVarDisplayClusterAlphaChannelCaptureMode : uint8
 /**
  *Choose method to preserve alpha channel
  */
-int32 GDisplayClusterAlphaChannelCaptureMode = 0;
+int32 GDisplayClusterAlphaChannelCaptureMode = (uint8)ECVarDisplayClusterAlphaChannelCaptureMode::FXAA;
 static FAutoConsoleVariableRef CVarDisplayClusterAlphaChannelCaptureMode(
 	TEXT("nDisplay.render.AlphaChannelCaptureMode"),
 	GDisplayClusterAlphaChannelCaptureMode,
-	TEXT("Alpha channel capture mode\n")
-	TEXT("0 - Auto (default) ThroughTonemapper or FXAA depend from prj settings\n")
-	TEXT("1 - Disabled\n")
-	TEXT("2 - ThroughTonemapper\n")
-	TEXT("3 - FXAA\n")
-	TEXT("4 - Copy [experimental]\n")
-	TEXT("5 - CopyAA [experimental]\n"),
+	TEXT("Alpha channel capture mode (FXAA - default)\n")
+	TEXT("0 - Disabled\n")
+	TEXT("1 - ThroughTonemapper\n")
+	TEXT("2 - FXAA\n")
+	TEXT("3 - Copy [experimental]\n")
+	TEXT("4 - CopyAA [experimental]\n"),
 	ECVF_RenderThreadSafe
 );
 
@@ -136,10 +130,6 @@ EDisplayClusterRenderFrameAlphaChannelCaptureMode FDisplayClusterViewportConfigu
 
 	switch (AlphaChannelCaptureMode)
 	{
-	case ECVarDisplayClusterAlphaChannelCaptureMode::Auto:
-		// Fit best possible AA
-		return bAllowThroughTonemapper ? EDisplayClusterRenderFrameAlphaChannelCaptureMode::ThroughTonemapper : EDisplayClusterRenderFrameAlphaChannelCaptureMode::FXAA;
-
 	case ECVarDisplayClusterAlphaChannelCaptureMode::ThroughTonemapper:
 		// Disable alpha capture if PropagateAlpha not valid
 		return bAllowThroughTonemapper ? EDisplayClusterRenderFrameAlphaChannelCaptureMode::ThroughTonemapper : EDisplayClusterRenderFrameAlphaChannelCaptureMode::None;
