nDisplay: The MPCDI warpblend interface has been redesigned.

- The IDisplayClusterWarpBlend interface has been redesigned and new functions have been added:
	* A new module DisplayClusterWarp has been added (the code with the warp API has been moved from DisplayClusterSharders to the new module).
	* Reworking warp math and code cleanup.
	* The internal data off the frustum calculation pipeline has been added to the public API.
	* Added MPCDI attributes loading (Loading from projection policy parameters can be done later).

- Performance optimization for MPCDI (the time to reload mpcdi resources is 0. Useful for large mpcdi files)
	* Added MPCDI file loader cache (The mpcdi loader instance remains alive until the wait time expires)
	* Added MPCDI Textures cache (multiple instances of the same DCRA reuse the same texture instance.)

- Added MPCDI 2d profile implementation:
	* Added math to support mpcdi 2d
	* Added preview geometry generation for this profile.
	* Added the ability to use an external static mesh as the geometry to calculate frustum instead of MPCDI attributes.
	* Added new property to MPCDI projection policy (ext mesh for mpcdi 2d) and to the DisplayClusterConfigurator

- Added changes to the 'mesh' and 'mpcdi' projection policy: adaptation to API changes and code cleanup

diff --git a/Engine/Plugins/Runtime/nDisplay/Source/DisplayCluster/Private/Render/Viewport/DisplayClusterViewportManager.cpp b/Engine/Plugins/Runtime/nDisplay/Source/DisplayCluster/Private/Render/Viewport/DisplayClusterViewportManager.cpp
--- a/Engine/Plugins/Runtime/nDisplay/Source/DisplayCluster/Private/Render/Viewport/DisplayClusterViewportManager.cpp
+++ b/Engine/Plugins/Runtime/nDisplay/Source/DisplayCluster/Private/Render/Viewport/DisplayClusterViewportManager.cpp
@@ -24,7 +24,7 @@
 #include "Render/Viewport/Configuration/DisplayClusterViewportConfigurationInstanceData.h"
 #include "Render/Viewport/DisplayClusterViewportStrings.h"
 
-#include "WarpBlend/IDisplayClusterWarpBlend.h"
+#include "IDisplayClusterWarpBlend.h"
 
 #include "SceneViewExtension.h"
 
diff --git a/Engine/Plugins/Runtime/nDisplay/Source/DisplayClusterConfigurator/Private/Views/Details/Policies/DisplayClusterConfiguratorPolicyDetailCustomization.cpp b/Engine/Plugins/Runtime/nDisplay/Source/DisplayClusterConfigurator/Private/Views/Details/Policies/DisplayClusterConfiguratorPolicyDetailCustomization.cpp
--- a/Engine/Plugins/Runtime/nDisplay/Source/DisplayClusterConfigurator/Private/Views/Details/Policies/DisplayClusterConfiguratorPolicyDetailCustomization.cpp
+++ b/Engine/Plugins/Runtime/nDisplay/Source/DisplayClusterConfigurator/Private/Views/Details/Policies/DisplayClusterConfiguratorPolicyDetailCustomization.cpp
@@ -754,6 +754,12 @@ void FDisplayClusterConfiguratorProjectionCustomization::CreateMPCDIPolicy(UDisp
 			DisplayClusterProjectionStrings::cfg::mpcdi::Region,
 			Blueprint,
 			ConfigurationViewports));
+
+		CustomPolicyParameters.Add(MakeShared<FPolicyParameterInfoText>(
+			"Screen",
+			DisplayClusterProjectionStrings::cfg::mpcdi::Component,
+			Blueprint,
+			ConfigurationViewports));
 	}
 	else if (Setting == TypePFM)
 	{
diff --git a/Engine/Plugins/Runtime/nDisplay/Source/DisplayClusterProjection/DisplayClusterProjection.Build.cs b/Engine/Plugins/Runtime/nDisplay/Source/DisplayClusterProjection/DisplayClusterProjection.Build.cs
--- a/Engine/Plugins/Runtime/nDisplay/Source/DisplayClusterProjection/DisplayClusterProjection.Build.cs
+++ b/Engine/Plugins/Runtime/nDisplay/Source/DisplayClusterProjection/DisplayClusterProjection.Build.cs
@@ -10,7 +10,8 @@ public class DisplayClusterProjection : ModuleRules
 		PublicIncludePathModuleNames.AddRange(
 			new string[] {
 				"DisplayClusterConfiguration",
-				"DisplayClusterShaders"
+				"DisplayClusterShaders",
+				"DisplayClusterWarp"
 			});
 
 		PublicDependencyModuleNames.AddRange(
diff --git a/Engine/Plugins/Runtime/nDisplay/Source/DisplayClusterProjection/Private/Policy/MPCDI/DisplayClusterProjectionMPCDIPolicy.cpp b/Engine/Plugins/Runtime/nDisplay/Source/DisplayClusterProjection/Private/Policy/MPCDI/DisplayClusterProjectionMPCDIPolicy.cpp
--- a/Engine/Plugins/Runtime/nDisplay/Source/DisplayClusterProjection/Private/Policy/MPCDI/DisplayClusterProjectionMPCDIPolicy.cpp
+++ b/Engine/Plugins/Runtime/nDisplay/Source/DisplayClusterProjection/Private/Policy/MPCDI/DisplayClusterProjectionMPCDIPolicy.cpp
@@ -8,6 +8,9 @@
 
 #include "IDisplayCluster.h"
 #include "IDisplayClusterCallbacks.h"
+#include "IDisplayClusterShaders.h"
+#include "IDisplayClusterWarp.h"
+
 #include "Config/IDisplayClusterConfigManager.h"
 #include "Game/IDisplayClusterGameManager.h"
 
@@ -20,20 +23,36 @@
 #include "Render/Viewport/IDisplayClusterViewportManagerProxy.h"
 #include "Render/Viewport/IDisplayClusterViewport.h"
 #include "Render/Viewport/IDisplayClusterViewportProxy.h"
-
-#include "IDisplayClusterShaders.h"
-#include "WarpBlend/IDisplayClusterWarpBlend.h"
-#include "WarpBlend/IDisplayClusterWarpBlendManager.h"
-
-#include "Render/Warp/IDisplayClusterWarpPolicy.h"
 #include "Render/Viewport/Containers/DisplayClusterViewport_RenderSettingsICVFX.h"
+#include "Render/Warp/IDisplayClusterWarpPolicy.h"
 
 #include "ShaderParameters/DisplayClusterShaderParameters_WarpBlend.h"
+#include "Blueprints/DisplayClusterWarpGeometry.h"
 
-#include "Blueprints/MPCDIGeometryData.h"
 #include "DisplayClusterRootActor.h"
 
 
+namespace UE::DisplayClusterProjection::MPCDIPolicy
+{
+	static IDisplayClusterWarp& GetWarpAPI()
+	{
+		static IDisplayClusterWarp& DisplayClusterWarpAPISingleton = IDisplayClusterWarp::Get();
+
+		return DisplayClusterWarpAPISingleton;
+	}
+
+	static inline IDisplayClusterShaders& GetShadersAPI()
+	{
+		IDisplayClusterShaders& ShadersAPISingleton = IDisplayClusterShaders::Get();
+
+		return ShadersAPISingleton;
+	}
+};
+using namespace UE::DisplayClusterProjection::MPCDIPolicy;
+
+//---------------------------------------------------------------------------------------------
+// FDisplayClusterProjectionMPCDIPolicy
+//---------------------------------------------------------------------------------------------
 FDisplayClusterProjectionMPCDIPolicy::FDisplayClusterProjectionMPCDIPolicy(const FString& ProjectionPolicyId, const FDisplayClusterConfigurationProjection* InConfigurationProjectionPolicy)
 	: FDisplayClusterProjectionPolicyBase(ProjectionPolicyId, InConfigurationProjectionPolicy)
 {
@@ -44,27 +63,25 @@ FDisplayClusterProjectionMPCDIPolicy::~FDisplayClusterProjectionMPCDIPolicy()
 	ImplRelease();
 }
 
-void FDisplayClusterProjectionMPCDIPolicy::UpdateProxyData(IDisplayClusterViewport* InViewport)
+void FDisplayClusterProjectionMPCDIPolicy::SetWarpPolicy(IDisplayClusterWarpPolicy* InWarpPolicy)
 {
-	check(InViewport);
+	check(IsInGameThread());
 
-	const TSharedPtr<IDisplayClusterProjectionPolicy, ESPMode::ThreadSafe> ProjectionPolicyPtr = InViewport->GetProjectionPolicy();
-	const TSharedPtr<IDisplayClusterWarpBlend, ESPMode::ThreadSafe> WarpBlendInterfacePtr = WarpBlendInterface;
+	WarpPolicyInterface = InWarpPolicy ? InWarpPolicy->ToSharedPtr() : nullptr;
+}
 
-	ENQUEUE_RENDER_COMMAND(DisplayClusterProjectionMPCDIPolicy_UpdateProxyData)(
-		[ProjectionPolicyPtr, WarpBlendInterfacePtr, Contexts = WarpBlendContexts](FRHICommandListImmediate& RHICmdList)
-	{
-		IDisplayClusterProjectionPolicy* ProjectionPolicy = ProjectionPolicyPtr.Get();
-		if (ProjectionPolicy)
-		{
-			FDisplayClusterProjectionMPCDIPolicy* MPCDIPolicy = static_cast<FDisplayClusterProjectionMPCDIPolicy*>(ProjectionPolicy);
-			if (MPCDIPolicy)
+IDisplayClusterWarpPolicy* FDisplayClusterProjectionMPCDIPolicy::GetWarpPolicy() const
 			{
-				MPCDIPolicy->WarpBlendInterface_Proxy = WarpBlendInterfacePtr;
-				MPCDIPolicy->WarpBlendContexts_Proxy = Contexts;
-			}
+	check(IsInGameThread());
+
+	return WarpPolicyInterface.Get();
 		}
-	});
+
+IDisplayClusterWarpPolicy* FDisplayClusterProjectionMPCDIPolicy::GetWarpPolicy_RenderThread() const
+{
+	check(IsInRenderingThread());
+
+	return WarpPolicyInterface_Proxy.Get();
 }
 
 bool FDisplayClusterProjectionMPCDIPolicy::ShouldSupportICVFX(IDisplayClusterViewport* InViewport) const
@@ -72,7 +89,7 @@ bool FDisplayClusterProjectionMPCDIPolicy::ShouldSupportICVFX(IDisplayClusterVie
 	check(IsInGameThread());
 	check(InViewport);
 
-	if (!WarpBlendInterface.IsValid() || !WarpBlendInterface->ShouldSupportICVFX(InViewport))
+	if (!WarpBlendInterface.IsValid() || !WarpBlendInterface->ShouldSupportICVFX())
 	{
 		// WarpBlend does not support the ICVFX pipeline (mpcdi 2D,3D,SL profiles)
 		return false;
@@ -80,32 +97,34 @@ bool FDisplayClusterProjectionMPCDIPolicy::ShouldSupportICVFX(IDisplayClusterVie
 
 	if (WarpPolicyInterface.IsValid() && !WarpPolicyInterface->ShouldSupportICVFX(InViewport))
 	{
-		// The warp policy is used for this projection and does not support the ICVFX pipeline
+		// Warp policy does not support the ICVFX pipeline
 		return false;
 	}
 
 	return true;
 }
 
-void FDisplayClusterProjectionMPCDIPolicy::SetWarpPolicy(IDisplayClusterWarpPolicy* InWarpPolicy)
+void FDisplayClusterProjectionMPCDIPolicy::UpdateProxyData(IDisplayClusterViewport* InViewport)
 {
-	check(IsInGameThread());
-
-	WarpPolicyInterface = InWarpPolicy ? InWarpPolicy->ToSharedPtr() : nullptr;
-}
+	check(InViewport);
 
-IDisplayClusterWarpPolicy* FDisplayClusterProjectionMPCDIPolicy::GetWarpPolicy() const
+	ENQUEUE_RENDER_COMMAND(DisplayClusterProjectionMPCDIPolicy_UpdateProxyData)(
+		[ProjectionPolicy = InViewport->GetProjectionPolicy(), WarpBlendInterfacePtr = WarpBlendInterface, WarPolicyInterfacePtr = WarpPolicyInterface, Contexts = WarpBlendContexts](FRHICommandListImmediate& RHICmdList)
 {
-	check(IsInGameThread());
-
-	return WarpPolicyInterface.Get();
-}
-
-IDisplayClusterWarpPolicy* FDisplayClusterProjectionMPCDIPolicy::GetWarpPolicy_RenderThread() const
+		if (ProjectionPolicy.IsValid())
 {
-	check(IsInRenderingThread());
+			FDisplayClusterProjectionMPCDIPolicy* MPCDIPolicy = static_cast<FDisplayClusterProjectionMPCDIPolicy*>(ProjectionPolicy.Get());
+			if (MPCDIPolicy)
+			{
+				MPCDIPolicy->WarpBlendInterface_Proxy = WarpBlendInterfacePtr;
+				MPCDIPolicy->WarpPolicyInterface_Proxy = WarPolicyInterfacePtr;
+				MPCDIPolicy->WarpBlendContexts_Proxy = Contexts;
+			}
+		}
+	});
 
-	return WarpPolicyInterface_Proxy.Get();
+	// reset warp policy after game thread is done
+	WarpPolicyInterface.Reset();
 }
 
 //////////////////////////////////////////////////////////////////////////////////////////////
@@ -183,7 +202,7 @@ bool FDisplayClusterProjectionMPCDIPolicy::GetWarpBlendInterface(TSharedPtr<IDis
 	return false;
 }
 
-bool FDisplayClusterProjectionMPCDIPolicy::GetWarpBlendInterface_RenderThread(TSharedPtr<class IDisplayClusterWarpBlend, ESPMode::ThreadSafe>& OutWarpBlendInterfaceProxy) const
+bool FDisplayClusterProjectionMPCDIPolicy::GetWarpBlendInterface_RenderThread(TSharedPtr<IDisplayClusterWarpBlend, ESPMode::ThreadSafe>& OutWarpBlendInterfaceProxy) const
 {
 	check(IsInRenderingThread());
 
@@ -224,33 +243,37 @@ bool FDisplayClusterProjectionMPCDIPolicy::CalculateView(IDisplayClusterViewport
 
 	// Get view location in local space
 	const USceneComponent* const OriginComp = GetOriginComp();
-	const FTransform& World2LocalTransform = (OriginComp ? OriginComp->GetComponentTransform() : FTransform::Identity);
+
+	// Initialize frustum
+	TSharedPtr<FDisplayClusterWarpEye, ESPMode::ThreadSafe> WarpEye = MakeShared<FDisplayClusterWarpEye, ESPMode::ThreadSafe>(InViewport->ToSharedPtr(), InContextNum);
+
+	WarpEye->World2LocalTransform = (OriginComp ? OriginComp->GetComponentTransform() : FTransform::Identity);
 
 	// Get our base camera location and view offset in local space (MPCDI space)
-	const FVector LocalOrigin    = World2LocalTransform.InverseTransformPosition(InOutViewLocation - ViewOffset);
-	const FVector LocalEyeOrigin = World2LocalTransform.InverseTransformPosition(InOutViewLocation);
+	WarpEye->ViewPoint.Location  = WarpEye->World2LocalTransform.InverseTransformPosition(InOutViewLocation - ViewOffset);
+	WarpEye->ViewPoint.EyeOffset = WarpEye->World2LocalTransform.InverseTransformPosition(InOutViewLocation) - WarpEye->ViewPoint.Location;
+	WarpEye->ViewPoint.Rotation  = WarpEye->World2LocalTransform.InverseTransformRotation(InOutViewRotation.Quaternion()).Rotator();
 
-	// Initialize frustum
-	FDisplayClusterWarpEye Eye;
-	Eye.OriginLocation  = LocalOrigin;
-	Eye.OriginEyeOffset = LocalEyeOrigin - LocalOrigin;
-	Eye.WorldScale = WorldScale;
-	Eye.ZNear = NCP;
-	Eye.ZFar  = FCP;
+	WarpEye->WorldScale = WorldScale;
+
+	// Use current warp policy
+	WarpEye->WarpPolicy = WarpPolicyInterface;
+
+	// discard current context
+	WarpBlendContexts[InContextNum].bIsValid = false;
 
 	// Compute frustum
-	if (!WarpBlendInterface->CalcFrustumContext(InViewport, InContextNum, Eye, WarpBlendContexts[InContextNum]))
+	if (!WarpBlendInterface->CalcFrustumContext(WarpEye))
 	{
 		return false;
 	}
 
-	// Get rotation in warp space
-	const FRotator MpcdiRotation = WarpBlendContexts[InContextNum].OutCameraRotation;
-	const FVector  MpcdiOrigin  = WarpBlendContexts[InContextNum].OutCameraOrigin;
+	// Readback warp context
+	WarpBlendContexts[InContextNum] = WarpBlendInterface->GetWarpData(InContextNum).WarpContext;
 
-	// Transform rotation to world space
-	InOutViewRotation = World2LocalTransform.TransformRotation(MpcdiRotation.Quaternion()).Rotator();
-	InOutViewLocation = World2LocalTransform.TransformPosition(MpcdiOrigin);
+	// Transform viewpoint back to world space
+	InOutViewRotation = WarpEye->World2LocalTransform.TransformRotation(WarpBlendContexts[InContextNum].Rotation.Quaternion()).Rotator();
+	InOutViewLocation = WarpEye->World2LocalTransform.TransformPosition(WarpBlendContexts[InContextNum].Location);
 
 	WarpBlendContexts[InContextNum].bIsValid = true;
 
@@ -312,8 +335,6 @@ void FDisplayClusterProjectionMPCDIPolicy::ApplyWarpBlend_RenderThread(FRHIComma
 		return;
 	}
 
-	IDisplayClusterShaders& ShadersAPI = IDisplayClusterShaders::Get();
-
 	const FDisplayClusterViewport_RenderSettingsICVFX& SettingsICVFX = InViewportProxy->GetRenderSettingsICVFX_RenderThread();
 	if (EnumHasAllFlags(SettingsICVFX.RuntimeFlags, EDisplayClusterViewportRuntimeICVFXFlags::Target))
 	{
@@ -385,7 +406,7 @@ void FDisplayClusterProjectionMPCDIPolicy::ApplyWarpBlend_RenderThread(FRHIComma
 			IDisplayCluster::Get().GetCallbacks().OnDisplayClusterPreProcessIcvfx_RenderThread().Broadcast(RHICmdList, InViewportProxy, WarpBlendParameters, ShaderICVFX);
 
 			// Start ICVFX pipeline
-			if (!ShadersAPI.RenderWarpBlend_ICVFX(RHICmdList, WarpBlendParameters, ShaderICVFX))
+			if (!GetShadersAPI().RenderWarpBlend_ICVFX(RHICmdList, WarpBlendParameters, ShaderICVFX))
 			{
 				if (!IsEditorOperationMode_RenderThread(InViewportProxy))
 				{
@@ -415,7 +436,7 @@ void FDisplayClusterProjectionMPCDIPolicy::ApplyWarpBlend_RenderThread(FRHIComma
 		
 		WarpBlendParameters.bRenderAlphaChannel = InViewportProxy->GetRenderSettings_RenderThread().bWarpBlendRenderAlphaChannel;
 
-		if (!ShadersAPI.RenderWarpBlend_MPCDI(RHICmdList, WarpBlendParameters))
+		if (!GetShadersAPI().RenderWarpBlend_MPCDI(RHICmdList, WarpBlendParameters))
 		{
 			if (!IsEditorOperationMode_RenderThread(InViewportProxy))
 			{
@@ -452,6 +473,17 @@ UMeshComponent* FDisplayClusterProjectionMPCDIPolicy::GetOrCreatePreviewMeshComp
 		return nullptr;
 	}
 
+	// MPCDI can use a static mesh component instead of a warpmap (2D)
+	if (WarpBlendInterface.IsValid())
+	{
+		if (UMeshComponent* StaticMeshComponentRef = WarpBlendInterface->GetStaticMeshComponent())
+		{
+			bOutIsRootActorComponent = true;
+
+			return StaticMeshComponentRef;
+		}
+	}
+
 	// used created mesh component
 	bOutIsRootActorComponent = false;
 
@@ -474,8 +506,8 @@ UMeshComponent* FDisplayClusterProjectionMPCDIPolicy::GetOrCreatePreviewMeshComp
 	const uint32 PreviewGeometryDimLimit = 128;
 
 	// Create new WarpMesh component
-	FMPCDIGeometryExportData MeshData;
-	if (WarpBlendInterface->ExportWarpMapGeometry(&MeshData, PreviewGeometryDimLimit))
+	FDisplayClusterWarpGeometryOBJ MeshData;
+	if (WarpBlendInterface->ExportWarpMapGeometry(MeshData, PreviewGeometryDimLimit))
 	{
 		const FString CompName = FString::Printf(TEXT("MPCDI_%s_impl"), *GetId());
 
@@ -508,11 +540,9 @@ bool FDisplayClusterProjectionMPCDIPolicy::CreateWarpBlendFromConfig(IDisplayClu
 
 	bool bResult = false;
 
-	FConfigParser CfgData;
-	if (CfgData.ImplLoadConfig(InViewport, GetParameters()))
+	FDisplayClusterProjectionMPCDIPolicy_ConfigParser CfgData(InViewport, GetParameters());
+	if (CfgData.IsValid())
 	{
-		IDisplayClusterShaders& ShadersAPI = IDisplayClusterShaders::Get();
-
 		// Support custom origin node
 		InitializeOriginComponent(InViewport, CfgData.OriginType);
 
@@ -538,12 +568,27 @@ bool FDisplayClusterProjectionMPCDIPolicy::CreateWarpBlendFromConfig(IDisplayClu
 				return false;
 			}
 
-			FDisplayClusterWarpBlendConstruct::FLoadMPCDIFile CreateParameters;
+			if (CfgData.ScreenComponent == nullptr)
+			{
+				FDisplayClusterWarpInitializer_MPCDIFile CreateParameters;
 			CreateParameters.MPCDIFileName = CfgData.MPCDIFileName;
 			CreateParameters.BufferId = CfgData.BufferId;
 			CreateParameters.RegionId = CfgData.RegionId;
 
-			if (!ShadersAPI.GetWarpBlendManager().Create(CreateParameters, WarpBlendInterface))
+				WarpBlendInterface = GetWarpAPI().Create(CreateParameters);
+			}
+			else
+			{
+				FDisplayClusterWarpInitializer_MPCDIFile_Profile2DScreen CreateParameters;
+			CreateParameters.MPCDIFileName = CfgData.MPCDIFileName;
+			CreateParameters.BufferId = CfgData.BufferId;
+			CreateParameters.RegionId = CfgData.RegionId;
+				CreateParameters.StaticMeshComponent = CfgData.ScreenComponent;
+
+				WarpBlendInterface = GetWarpAPI().Create(CreateParameters);
+			}
+
+			if (!WarpBlendInterface.IsValid())
 			{
 				UE_LOG(LogDisplayClusterProjectionMPCDI, Warning, TEXT("Couldn't load MPCDI file: %s"), *CfgData.MPCDIFileName);
 				return false;
@@ -553,8 +598,9 @@ bool FDisplayClusterProjectionMPCDIPolicy::CreateWarpBlendFromConfig(IDisplayClu
 			return true;
 		}
 
-		FDisplayClusterWarpBlendConstruct::FLoadPFMFile CreateParameters;
-		CreateParameters.ProfileType = CfgData.MPCDIType;
+		FDisplayClusterWarpInitializer_PFMFile CreateParameters;
+
+		CreateParameters.MPCDIAttributes = CfgData.MPCDIAttributes;
 		CreateParameters.PFMFileName = CfgData.PFMFile;
 
 		CreateParameters.PFMScale = CfgData.PFMFileScale;
@@ -565,7 +611,9 @@ bool FDisplayClusterProjectionMPCDIPolicy::CreateWarpBlendFromConfig(IDisplayClu
 
 		CreateParameters.BetaMapFileName = CfgData.BetaFile;
 
-		if (!ShadersAPI.GetWarpBlendManager().Create(CreateParameters, WarpBlendInterface))
+		WarpBlendInterface = GetWarpAPI().Create(CreateParameters);
+
+		if (!WarpBlendInterface.IsValid())
 		{
 			UE_LOG(LogDisplayClusterProjectionMPCDI, Error, TEXT("Failed to load PFM from file: %s"), *CfgData.PFMFile);
 			return false;
diff --git a/Engine/Plugins/Runtime/nDisplay/Source/DisplayClusterProjection/Private/Policy/MPCDI/DisplayClusterProjectionMPCDIPolicy.h b/Engine/Plugins/Runtime/nDisplay/Source/DisplayClusterProjection/Private/Policy/MPCDI/DisplayClusterProjectionMPCDIPolicy.h
--- a/Engine/Plugins/Runtime/nDisplay/Source/DisplayClusterProjection/Private/Policy/MPCDI/DisplayClusterProjectionMPCDIPolicy.h
+++ b/Engine/Plugins/Runtime/nDisplay/Source/DisplayClusterProjection/Private/Policy/MPCDI/DisplayClusterProjectionMPCDIPolicy.h
@@ -3,10 +3,8 @@
 #pragma once
 
 #include "Policy/DisplayClusterProjectionPolicyBase.h"
-
-#include "WarpBlend/IDisplayClusterWarpBlend.h"
-#include "WarpBlend/DisplayClusterWarpContext.h"
-
+#include "Containers/DisplayClusterWarpContext.h"
+#include "IDisplayClusterWarpBlend.h"
 
 /**
  * MPCDI projection policy
@@ -15,23 +13,11 @@
 class FDisplayClusterProjectionMPCDIPolicy
 	: public FDisplayClusterProjectionPolicyBase
 {
-public:
-	enum class EWarpType : uint8
-	{
-		mpcdi = 0,
-		mesh
-	};
-
 public:
 	FDisplayClusterProjectionMPCDIPolicy(const FString& ProjectionPolicyId, const FDisplayClusterConfigurationProjection* InConfigurationProjectionPolicy);
 	virtual ~FDisplayClusterProjectionMPCDIPolicy();
 
 public:
-	virtual EWarpType GetWarpType() const
-	{
-		return EWarpType::mpcdi;
-	}
-
 	// This policy can support ICVFX rendering
 	virtual bool ShouldSupportICVFX(IDisplayClusterViewport* InViewport) const override;
 
@@ -76,12 +62,16 @@ protected:
 	void ImplRelease();
 
 protected:
+	// GameThread: WarpBlend and WarpPolicy interfaces
 	TSharedPtr<IDisplayClusterWarpBlend, ESPMode::ThreadSafe> WarpBlendInterface;
 	TSharedPtr<IDisplayClusterWarpPolicy, ESPMode::ThreadSafe> WarpPolicyInterface;
-	TArray<FDisplayClusterWarpContext> WarpBlendContexts;
 
+	// RenderingThread: WarpBlend and WarpPolicy interfaces
 	TSharedPtr<IDisplayClusterWarpBlend, ESPMode::ThreadSafe> WarpBlendInterface_Proxy;
 	TSharedPtr<IDisplayClusterWarpPolicy, ESPMode::ThreadSafe> WarpPolicyInterface_Proxy;
+
+	// Context for both game and rendering threads
+	TArray<FDisplayClusterWarpContext> WarpBlendContexts;
 	TArray<FDisplayClusterWarpContext> WarpBlendContexts_Proxy;
 
 private:
diff --git a/Engine/Plugins/Runtime/nDisplay/Source/DisplayClusterProjection/Private/Policy/MPCDI/DisplayClusterProjectionMPCDIPolicy_Config.cpp b/Engine/Plugins/Runtime/nDisplay/Source/DisplayClusterProjection/Private/Policy/MPCDI/DisplayClusterProjectionMPCDIPolicy_Config.cpp
new file mode 100644
--- /dev/null
+++ b/Engine/Plugins/Runtime/nDisplay/Source/DisplayClusterProjection/Private/Policy/MPCDI/DisplayClusterProjectionMPCDIPolicy_Config.cpp
@@ -0,0 +1,235 @@
+// Copyright Epic Games, Inc. All Rights Reserved.
+
+#include "Policy/MPCDI/DisplayClusterProjectionMPCDIPolicy_Config.h"
+#include "Policy/DisplayClusterProjectionPolicyBase.h"
+
+#include "Render/Viewport/IDisplayClusterViewport.h"
+#include "Containers/DisplayClusterWarpContainers.h"
+#include "Components/DisplayClusterScreenComponent.h"
+
+#include "DisplayClusterProjectionLog.h"
+#include "DisplayClusterProjectionStrings.h"
+
+#include "Misc/DisplayClusterHelpers.h"
+#include "Misc/FileHelper.h"
+
+#include "DisplayClusterRootActor.h"
+
+//----------------------------------------------------------------------
+// FDisplayClusterProjectionMPCDIPolicy_ConfigParser
+//----------------------------------------------------------------------
+FDisplayClusterProjectionMPCDIPolicy_ConfigParser::FDisplayClusterProjectionMPCDIPolicy_ConfigParser(IDisplayClusterViewport* InViewport, const TMap<FString, FString>& InConfigParameters)
+	: ConfigParameters(InConfigParameters)
+	, Viewport(InViewport ? InViewport->ToSharedPtr() : nullptr)
+{
+	bValid = Viewport.IsValid() && ReadConfig();
+}
+
+bool FDisplayClusterProjectionMPCDIPolicy_ConfigParser::ReadConfig()
+{
+	FString MPCDITypeKey;
+	if (DisplayClusterHelpers::map::template ExtractValue(ConfigParameters, DisplayClusterProjectionStrings::cfg::mpcdi::MPCDITypeKey, MPCDITypeKey))
+	{
+		UE_LOG(LogDisplayClusterProjectionMPCDI, Verbose, TEXT("Found Argument '%s'='%s'"), DisplayClusterProjectionStrings::cfg::mpcdi::MPCDITypeKey, *MPCDITypeKey);
+	}
+
+	if (MPCDITypeKey.IsEmpty())
+	{
+		if (!FDisplayClusterProjectionPolicyBase::IsEditorOperationMode(Viewport.Get()))
+		{
+			UE_LOG(LogDisplayClusterProjectionMPCDI, Error, TEXT("Undefined mpcdi type key '%s'='%s'"), DisplayClusterProjectionStrings::cfg::mpcdi::MPCDITypeKey, *MPCDITypeKey);
+		}
+		return false;
+	}
+
+	if (MPCDITypeKey.Compare(DisplayClusterProjectionStrings::cfg::mpcdi::TypeMPCDI) == 0)
+	{
+		return ImplGetMPCDIConfig() && ImplGetBaseConfig();
+	}
+
+	if (MPCDITypeKey.Compare(DisplayClusterProjectionStrings::cfg::mpcdi::TypePFM) == 0)
+	{
+		return ImplGetPFMConfig() && ImplGetMPCDIAttributes() && ImplGetBaseConfig();
+	}
+
+	UE_LOG(LogDisplayClusterProjectionMPCDI, Error, TEXT("Unknown mpcdi type key '%s'='%s'"), DisplayClusterProjectionStrings::cfg::mpcdi::MPCDITypeKey, *MPCDITypeKey);
+	return false;
+}
+
+bool FDisplayClusterProjectionMPCDIPolicy_ConfigParser::ImplGetMPCDIConfig()
+{
+	// Filename
+	FString LocalMPCDIFileName;
+	if (DisplayClusterHelpers::map::template ExtractValue(ConfigParameters, DisplayClusterProjectionStrings::cfg::mpcdi::File, LocalMPCDIFileName))
+	{
+		UE_LOG(LogDisplayClusterProjectionMPCDI, Verbose, TEXT("Found mpcdi file name for %s:%s - %s"), *BufferId, *RegionId, *LocalMPCDIFileName);
+		MPCDIFileName = LocalMPCDIFileName;
+	}
+
+	if (MPCDIFileName.IsEmpty())
+	{
+		return false;
+	}
+
+	// Buffer
+	if (!DisplayClusterHelpers::map::template ExtractValue(ConfigParameters, DisplayClusterProjectionStrings::cfg::mpcdi::Buffer, BufferId))
+	{
+		if (!FDisplayClusterProjectionPolicyBase::IsEditorOperationMode(Viewport.Get()))
+		{
+			UE_LOG(LogDisplayClusterProjectionMPCDI, Error, TEXT("Argument '%s' not found in the config file"), DisplayClusterProjectionStrings::cfg::mpcdi::Buffer);
+		}
+
+		return false;
+	}
+
+	if (BufferId.IsEmpty())
+	{
+		return false;
+	}
+
+	// Region
+	if (!DisplayClusterHelpers::map::template ExtractValue(ConfigParameters, DisplayClusterProjectionStrings::cfg::mpcdi::Region, RegionId))
+	{
+		if (!FDisplayClusterProjectionPolicyBase::IsEditorOperationMode(Viewport.Get()))
+		{
+			UE_LOG(LogDisplayClusterProjectionMPCDI, Error, TEXT("Argument '%s' not found in the config file"), DisplayClusterProjectionStrings::cfg::mpcdi::Region);
+		}
+
+		return false;
+	}
+
+	if (RegionId.IsEmpty())
+	{
+		return false;
+	}
+
+	// Screen component
+	FString ScreenComponentName;
+	if (DisplayClusterHelpers::map::template ExtractValue(ConfigParameters, DisplayClusterProjectionStrings::cfg::mpcdi::Component, ScreenComponentName))
+	{
+		if (!ScreenComponentName.IsEmpty())
+		{
+			if (ADisplayClusterRootActor* RootActorPtr = Viewport->GetRootActor())
+			{
+				if (UDisplayClusterScreenComponent* ScreenComp = RootActorPtr->GetComponentByName<UDisplayClusterScreenComponent>(ScreenComponentName))
+				{
+					ScreenComponent = ScreenComp;
+				}
+			}
+		}
+	}
+
+	return true;
+}
+
+bool FDisplayClusterProjectionMPCDIPolicy_ConfigParser::ImplGetPFMConfig()
+{
+	// PFM file
+	FString LocalPFMFile;
+	if (DisplayClusterHelpers::map::template ExtractValue(ConfigParameters, DisplayClusterProjectionStrings::cfg::mpcdi::FilePFM, LocalPFMFile))
+	{
+		UE_LOG(LogDisplayClusterProjectionMPCDI, Verbose, TEXT("Found Argument '%s'='%s'"), DisplayClusterProjectionStrings::cfg::mpcdi::FilePFM, *LocalPFMFile);
+		PFMFile = LocalPFMFile;
+	}
+
+	if (PFMFile.IsEmpty())
+	{
+		return false;
+	}
+
+	// MPCDIType (optional)
+	FString MPCDITypeStr;
+	if (!DisplayClusterHelpers::map::template ExtractValue(ConfigParameters, DisplayClusterProjectionStrings::cfg::mpcdi::MPCDIType, MPCDITypeStr))
+	{
+		MPCDIAttributes.ProfileType = EDisplayClusterWarpProfileType::warp_A3D;
+	}
+	else
+	{
+		MPCDIAttributes.ProfileType = EDisplayClusterWarpProfileType::Invalid;
+
+		static const TArray<FString> Profiles({ "2d","3d","a3d","sl" });
+		for (int32 ProfileIndex = 0; ProfileIndex < Profiles.Num(); ++ProfileIndex)
+		{
+			if (!MPCDITypeStr.Compare(Profiles[ProfileIndex], ESearchCase::IgnoreCase))
+			{
+				MPCDIAttributes.ProfileType = (EDisplayClusterWarpProfileType)ProfileIndex;
+				break;
+			}
+		}
+
+		if (MPCDIAttributes.ProfileType == EDisplayClusterWarpProfileType::Invalid)
+		{
+			UE_LOG(LogDisplayClusterProjectionMPCDI, Error, TEXT("Argument '%s' has unknown value '%s'"), DisplayClusterProjectionStrings::cfg::mpcdi::MPCDIType, *MPCDITypeStr);
+			return false;
+		}
+	}
+
+	// Default is UE scale, cm
+	PFMFileScale = 1;
+	if (DisplayClusterHelpers::map::template ExtractValueFromString(ConfigParameters, DisplayClusterProjectionStrings::cfg::mpcdi::WorldScale, PFMFileScale))
+	{
+		UE_LOG(LogDisplayClusterProjectionMPCDI, Verbose, TEXT("Found WorldScale value - %.f"), PFMFileScale);
+	}
+
+	bIsUnrealGameSpace = false;
+	if (DisplayClusterHelpers::map::template ExtractValueFromString(ConfigParameters, DisplayClusterProjectionStrings::cfg::mpcdi::UseUnrealAxis, bIsUnrealGameSpace))
+	{
+		UE_LOG(LogDisplayClusterProjectionMPCDI, Verbose, TEXT("Found bIsUnrealGameSpace value - %s"), bIsUnrealGameSpace ? TEXT("true") : TEXT("false"));
+	}
+
+	// AlphaFile file (optional)
+	FString LocalAlphaFile;
+	if (DisplayClusterHelpers::map::template ExtractValueFromString(ConfigParameters, DisplayClusterProjectionStrings::cfg::mpcdi::FileAlpha, LocalAlphaFile))
+	{
+		UE_LOG(LogDisplayClusterProjectionMPCDI, Verbose, TEXT("Found external AlphaMap file - %s"), *LocalAlphaFile);
+		AlphaFile = LocalAlphaFile;
+	}
+
+	AlphaGamma = 1;
+	if (DisplayClusterHelpers::map::template ExtractValueFromString(ConfigParameters, DisplayClusterProjectionStrings::cfg::mpcdi::AlphaGamma, AlphaGamma))
+	{
+		UE_LOG(LogDisplayClusterProjectionMPCDI, Verbose, TEXT("Found AlphaGamma value - %.f"), AlphaGamma);
+	}
+
+	// BetaFile file (optional)
+	FString LocalBetaFile;
+	if (DisplayClusterHelpers::map::template ExtractValueFromString(ConfigParameters, DisplayClusterProjectionStrings::cfg::mpcdi::FileBeta, LocalBetaFile))
+	{
+		UE_LOG(LogDisplayClusterProjectionMPCDI, Verbose, TEXT("Found external BetaMap file - %s"), *LocalBetaFile);
+		BetaFile = LocalBetaFile;
+	}
+
+	return true;
+}
+
+bool FDisplayClusterProjectionMPCDIPolicy_ConfigParser::ImplGetBaseConfig()
+{
+	// Origin node (optional)
+	if (DisplayClusterHelpers::map::template ExtractValue(ConfigParameters, DisplayClusterProjectionStrings::cfg::mpcdi::Origin, OriginType))
+	{
+		UE_LOG(LogDisplayClusterProjectionMPCDI, Verbose, TEXT("Found origin node - %s"), *OriginType);
+	}
+	else
+	{
+		if (!FDisplayClusterProjectionPolicyBase::IsEditorOperationMode(Viewport.Get()))
+		{
+			UE_LOG(LogDisplayClusterProjectionMPCDI, Log, TEXT("No origin node found. VR root will be used as default."));
+		}
+	}
+
+	bEnablePreview = false;
+	if (DisplayClusterHelpers::map::template ExtractValueFromString(ConfigParameters, DisplayClusterProjectionStrings::cfg::mpcdi::EnablePreview, bEnablePreview))
+	{
+		UE_LOG(LogDisplayClusterProjectionMPCDI, Verbose, TEXT("Found EnablePreview value - %s"), bEnablePreview ? TEXT("true") : TEXT("false"));
+	}
+
+	return true;
+}
+
+bool FDisplayClusterProjectionMPCDIPolicy_ConfigParser::ImplGetMPCDIAttributes()
+{
+	// Todo: Read parameters with names from the DisplayClusterProjectionStrings::cfg::mpcdi::Attributes namespace
+	// into a local variable named MPCDIAttributes;
+
+	return true;
+}
diff --git a/Engine/Plugins/Runtime/nDisplay/Source/DisplayClusterProjection/Private/Policy/MPCDI/DisplayClusterProjectionMPCDIPolicy_Config.h b/Engine/Plugins/Runtime/nDisplay/Source/DisplayClusterProjection/Private/Policy/MPCDI/DisplayClusterProjectionMPCDIPolicy_Config.h
--- a/Engine/Plugins/Runtime/nDisplay/Source/DisplayClusterProjection/Private/Policy/MPCDI/DisplayClusterProjectionMPCDIPolicy_Config.h
+++ b/Engine/Plugins/Runtime/nDisplay/Source/DisplayClusterProjection/Private/Policy/MPCDI/DisplayClusterProjectionMPCDIPolicy_Config.h
@@ -3,224 +3,92 @@
 #pragma once
 
 #include "CoreMinimal.h"
+#include "Containers/DisplayClusterWarpEnums.h"
+#include "Containers/DisplayClusterWarpContainers.h"
 
-#include "DisplayClusterProjectionLog.h"
-#include "DisplayClusterProjectionStrings.h"
+class IDisplayClusterViewport;
+class UStaticMeshComponent;
 
-#include "Misc/DisplayClusterHelpers.h"
-#include "Misc/FileHelper.h"
+/**
+ * MPCDI projection policy configuration parser
+ */
+struct FDisplayClusterProjectionMPCDIPolicy_ConfigParser
+{
+	// Name of the MPCDI file that is used as the data source
+	FString MPCDIFileName; // Single mpcdi file name
 
-#include "WarpBlend/DisplayClusterWarpEnums.h"
+	// Name of the buffer inside the MPCDI file that will be used as the data source
+	FString BufferId;
 
-#include "Policy/DisplayClusterProjectionPolicyBase.h"
+	// Name of the region inside the MPCDI file that will be used as the data source
+	FString RegionId;
 
-struct FConfigParser
-{
-	FString  MPCDIFileName; // Single mpcdi file name
+	// The origin point name
+	FString OriginType;
+
+	// Todo: fill this data from InConfigParameters
+	FDisplayClusterWarpMPCDIAttributes MPCDIAttributes;
 
-	FString  BufferId;
-	FString  RegionId;
+	// Screen component with source geometry
+	UStaticMeshComponent* ScreenComponent = nullptr;
 
-	FString  OriginType;
+	// Name of the PFM file that is used as the data source
+	FString PFMFile;
 
-	// Support external pfm (warp)  and png(blend) files
-	EDisplayClusterWarpProfileType  MPCDIType;
+	// Scale unit of the PFM file geometry
+	float PFMFileScale = 1;
 
-	FString  PFMFile;
-	float    PFMFileScale;
-	bool     bIsUnrealGameSpace = false;
+	// Are the axes of the PFM geometry the same as in UE
+	bool bIsUnrealGameSpace = false;
 
-	FString  AlphaFile;
-	float AlphaGamma;
+	// Filename of the image to be used as an AlphaMap
+	FString AlphaFile;
+	// Gamma used for the AlphaMap image
+	float AlphaGamma = 1;
 
-	FString  BetaFile;
+	// Filename of the image to be used as an BetaMap
+	FString BetaFile;
 
-	bool     bEnablePreview = false;
+	// The preview geometry must be imported from the PFM\MPCDI file and displayed.
+	bool bEnablePreview = false;
 
-	inline bool ImplLoadConfig(IDisplayClusterViewport* InViewport, const TMap<FString, FString>& InConfigParameters)
+public:
+	/** Read MPCDI policy configuration from projection policy parameters
+	* 
+	* @param InViewport         - the projection policy owner viewport
+	* @param InConfigParameters - the projection policy parameters
+	*/
+	FDisplayClusterProjectionMPCDIPolicy_ConfigParser(IDisplayClusterViewport* InViewport, const TMap<FString, FString>& InConfigParameters);
+
+	/** When this configuration is valid, it returns true. */
+	inline bool IsValid() const
 	{
-		FString MPCDITypeKey;
-		if (DisplayClusterHelpers::map::template ExtractValue(InConfigParameters, DisplayClusterProjectionStrings::cfg::mpcdi::MPCDITypeKey, MPCDITypeKey))
-		{
-			UE_LOG(LogDisplayClusterProjectionMPCDI, Verbose, TEXT("Found Argument '%s'='%s'"), DisplayClusterProjectionStrings::cfg::mpcdi::MPCDITypeKey, *MPCDITypeKey);
-		}
-
-		if (MPCDITypeKey.IsEmpty())
-		{
-			if (!FDisplayClusterProjectionPolicyBase::IsEditorOperationMode(InViewport))
-			{
-				UE_LOG(LogDisplayClusterProjectionMPCDI, Error, TEXT("Undefined mpcdi type key '%s'='%s'"), DisplayClusterProjectionStrings::cfg::mpcdi::MPCDITypeKey, *MPCDITypeKey);
-			}
-			return false;
-		}
-
-		if (MPCDITypeKey.Compare(DisplayClusterProjectionStrings::cfg::mpcdi::TypeMPCDI) == 0)
-		{
-			return ImplLoadMPCDIConfig(InViewport, InConfigParameters) && ImplLoadBase(InViewport, InConfigParameters);
-		}
-
-		if (MPCDITypeKey.Compare(DisplayClusterProjectionStrings::cfg::mpcdi::TypePFM) == 0)
-		{
-			return ImplLoadPFMConfig(InViewport, InConfigParameters) && ImplLoadBase(InViewport, InConfigParameters);
-		}
-
-		UE_LOG(LogDisplayClusterProjectionMPCDI, Error, TEXT("Unknown mpcdi type key '%s'='%s'"), DisplayClusterProjectionStrings::cfg::mpcdi::MPCDITypeKey, *MPCDITypeKey);
-		return false;
+		return bValid;
 	}
 
 private:
-	inline bool ImplLoadMPCDIConfig(IDisplayClusterViewport* InViewport, const TMap<FString, FString>& InConfigParameters)
-	{
-		// Filename
-		FString LocalMPCDIFileName;
-		if (DisplayClusterHelpers::map::template ExtractValue(InConfigParameters, DisplayClusterProjectionStrings::cfg::mpcdi::File, LocalMPCDIFileName))
-		{
-			UE_LOG(LogDisplayClusterProjectionMPCDI, Verbose, TEXT("Found mpcdi file name for %s:%s - %s"), *BufferId, *RegionId, *LocalMPCDIFileName);
-			MPCDIFileName = LocalMPCDIFileName;
-		}
-
-		if (MPCDIFileName.IsEmpty())
-		{
-			return false;
-		}
-
-		// Buffer
-		if (!DisplayClusterHelpers::map::template ExtractValue(InConfigParameters, DisplayClusterProjectionStrings::cfg::mpcdi::Buffer, BufferId))
-		{
-			if (!FDisplayClusterProjectionPolicyBase::IsEditorOperationMode(InViewport))
-			{
-				UE_LOG(LogDisplayClusterProjectionMPCDI, Error, TEXT("Argument '%s' not found in the config file"), DisplayClusterProjectionStrings::cfg::mpcdi::Buffer);
-			}
-
-			return false;
-		}
-
-		if (BufferId.IsEmpty())
-		{
-			return false;
-		}
-
-		// Region
-		if (!DisplayClusterHelpers::map::template ExtractValue(InConfigParameters, DisplayClusterProjectionStrings::cfg::mpcdi::Region, RegionId))
-		{
-			if (!FDisplayClusterProjectionPolicyBase::IsEditorOperationMode(InViewport))
-			{
-				UE_LOG(LogDisplayClusterProjectionMPCDI, Error, TEXT("Argument '%s' not found in the config file"), DisplayClusterProjectionStrings::cfg::mpcdi::Region);
-			}
-
-			return false;
-		}
-
-		if (RegionId.IsEmpty())
-		{
-			return false;
-		}
-
-		return true;
-	}
+	/** Implementation of reading the configuration from the parameters. */
+	bool ReadConfig();
 
-	inline bool ImplLoadPFMConfig(IDisplayClusterViewport* InViewport, const TMap<FString, FString>& InConfigParameters)
-	{
-		// PFM file
-		FString LocalPFMFile;
-		if (DisplayClusterHelpers::map::template ExtractValue(InConfigParameters, DisplayClusterProjectionStrings::cfg::mpcdi::FilePFM, LocalPFMFile))
-		{
-			UE_LOG(LogDisplayClusterProjectionMPCDI, Verbose, TEXT("Found Argument '%s'='%s'"), DisplayClusterProjectionStrings::cfg::mpcdi::FilePFM, *LocalPFMFile);
-			PFMFile = LocalPFMFile;
-		}
-
-		if (PFMFile.IsEmpty())
-		{
-			return false;
-		}
-
-		// MPCDIType (optional)
-		FString MPCDITypeStr;
-		if (!DisplayClusterHelpers::map::template ExtractValue(InConfigParameters, DisplayClusterProjectionStrings::cfg::mpcdi::MPCDIType, MPCDITypeStr))
-		{
-			MPCDIType = EDisplayClusterWarpProfileType::warp_A3D;
-		}
-		else
-		{
-			MPCDIType = EDisplayClusterWarpProfileType::Invalid;
-
-			static const TArray<FString> Profiles({ "2d","3d","a3d","sl" });
-			for (int32 ProfileIndex = 0; ProfileIndex < Profiles.Num(); ++ProfileIndex)
-			{
-				if (!MPCDITypeStr.Compare(Profiles[ProfileIndex], ESearchCase::IgnoreCase))
-				{
-					MPCDIType = (EDisplayClusterWarpProfileType)ProfileIndex;
-					break;
-				}
-			}
-
-			if (MPCDIType == EDisplayClusterWarpProfileType::Invalid)
-			{
-				UE_LOG(LogDisplayClusterProjectionMPCDI, Error, TEXT("Argument '%s' has unknown value '%s'"), DisplayClusterProjectionStrings::cfg::mpcdi::MPCDIType, *MPCDITypeStr);
-				return false;
-			}
-		}
-
-		// Default is UE scale, cm
-		PFMFileScale = 1;
-		if (DisplayClusterHelpers::map::template ExtractValueFromString(InConfigParameters, DisplayClusterProjectionStrings::cfg::mpcdi::WorldScale, PFMFileScale))
-		{
-			UE_LOG(LogDisplayClusterProjectionMPCDI, Verbose, TEXT("Found WorldScale value - %.f"), PFMFileScale);
-		}
-
-		bIsUnrealGameSpace = false;
-		if (DisplayClusterHelpers::map::template ExtractValueFromString(InConfigParameters, DisplayClusterProjectionStrings::cfg::mpcdi::UseUnrealAxis, bIsUnrealGameSpace))
-		{
-			UE_LOG(LogDisplayClusterProjectionMPCDI, Verbose, TEXT("Found bIsUnrealGameSpace value - %s"), bIsUnrealGameSpace ? TEXT("true") : TEXT("false"));
-		}
-
-		// AlphaFile file (optional)
-		FString LocalAlphaFile;
-		if (DisplayClusterHelpers::map::template ExtractValueFromString(InConfigParameters, DisplayClusterProjectionStrings::cfg::mpcdi::FileAlpha, LocalAlphaFile))
-		{
-			UE_LOG(LogDisplayClusterProjectionMPCDI, Verbose, TEXT("Found external AlphaMap file - %s"), *LocalAlphaFile);
-			AlphaFile = LocalAlphaFile;
-		}
-
-		AlphaGamma = 1;
-		if (DisplayClusterHelpers::map::template ExtractValueFromString(InConfigParameters, DisplayClusterProjectionStrings::cfg::mpcdi::AlphaGamma, AlphaGamma))
-		{
-			UE_LOG(LogDisplayClusterProjectionMPCDI, Verbose, TEXT("Found AlphaGamma value - %.f"), AlphaGamma);
-		}
-
-		// BetaFile file (optional)
-		FString LocalBetaFile;
-		if (DisplayClusterHelpers::map::template ExtractValueFromString(InConfigParameters, DisplayClusterProjectionStrings::cfg::mpcdi::FileBeta, LocalBetaFile))
-		{
-			UE_LOG(LogDisplayClusterProjectionMPCDI, Verbose, TEXT("Found external BetaMap file - %s"), *LocalBetaFile);
-			BetaFile = LocalBetaFile;
-		}
-
-		return true;
-	}
+	/** Implementation of reading basic parameters. */
+	bool ImplGetBaseConfig();
 
-	inline bool ImplLoadBase(IDisplayClusterViewport* InViewport, const TMap<FString, FString>& InConfigParameters)
-	{
-		// Origin node (optional)
-		if (DisplayClusterHelpers::map::template ExtractValue(InConfigParameters, DisplayClusterProjectionStrings::cfg::mpcdi::Origin, OriginType))
-		{
-			UE_LOG(LogDisplayClusterProjectionMPCDI, Verbose, TEXT("Found origin node - %s"), *OriginType);
-		}
-		else
-		{
-			if (!FDisplayClusterProjectionPolicyBase::IsEditorOperationMode(InViewport))
-			{
-				UE_LOG(LogDisplayClusterProjectionMPCDI, Log, TEXT("No origin node found. VR root will be used as default."));
-			}
-		}
-
-		bEnablePreview = false;
-		if (DisplayClusterHelpers::map::template ExtractValueFromString(InConfigParameters, DisplayClusterProjectionStrings::cfg::mpcdi::EnablePreview, bEnablePreview))
-		{
-			UE_LOG(LogDisplayClusterProjectionMPCDI, Verbose, TEXT("Found EnablePreview value - %s"), bEnablePreview ? TEXT("true") : TEXT("false"));
-		}
-
-		return true;
-	}
-};
+	/** Implementation of parameter reading for the case when we initialize from an MPCDI file. */
+	bool ImplGetMPCDIConfig();
+
+	/** Implementation of parameter reading for the case when we initialize from an PFM file. */
+	bool ImplGetPFMConfig();
+
+	/** Implementation of MPCDI attribute reading for the case when we initialize from a PFM file. */
+	bool ImplGetMPCDIAttributes();
 
+private:
+	// The projection policy parameters used to read this configuration
+	const TMap<FString, FString>& ConfigParameters;
+
+	// The projection policy owner viewport
+	TSharedPtr<IDisplayClusterViewport, ESPMode::ThreadSafe> Viewport;
+
+	// Is this configuration valid.
+	bool bValid = false;
+};
diff --git a/Engine/Plugins/Runtime/nDisplay/Source/DisplayClusterProjection/Private/Policy/Mesh/DisplayClusterProjectionMeshPolicy.cpp b/Engine/Plugins/Runtime/nDisplay/Source/DisplayClusterProjection/Private/Policy/Mesh/DisplayClusterProjectionMeshPolicy.cpp
--- a/Engine/Plugins/Runtime/nDisplay/Source/DisplayClusterProjection/Private/Policy/Mesh/DisplayClusterProjectionMeshPolicy.cpp
+++ b/Engine/Plugins/Runtime/nDisplay/Source/DisplayClusterProjection/Private/Policy/Mesh/DisplayClusterProjectionMeshPolicy.cpp
@@ -11,13 +11,36 @@
 #include "Render/Viewport/IDisplayClusterViewport.h"
 
 #include "IDisplayClusterShaders.h"
-#include "WarpBlend/IDisplayClusterWarpBlend.h"
-#include "WarpBlend/IDisplayClusterWarpBlendManager.h"
+#include "IDisplayClusterWarp.h"
 
 #include "Components/StaticMeshComponent.h"
 #include "ProceduralMeshComponent.h"
 
+/**
+ * WarpBlend configuration
+ */
+struct FDisplayClusterProjectionMeshPolicyConfiguration
+{
+	// StaticMesh component with source geometry
+	UStaticMeshComponent* StaticMeshComponent = nullptr;
+
+	// StaticMesh geometry LOD
+	int32 StaticMeshComponentLODIndex = 0;
+
+	// ProceduralMesh component with source geometry
+	UProceduralMeshComponent* ProceduralMeshComponent = nullptr;
+
+	// ProceduralMesh section index
+	int32 ProceduralMeshComponentSectionIndex = 0;
+
+	// Customize source geometry UV channels
+	int32 BaseUVIndex = INDEX_NONE;
+	int32 ChromakeyUVIndex = INDEX_NONE;
+};
 
+//------------------------------------------------------------------------------------------
+// FDisplayClusterProjectionMeshPolicy
+//------------------------------------------------------------------------------------------
 FDisplayClusterProjectionMeshPolicy::FDisplayClusterProjectionMeshPolicy(const FString& ProjectionPolicyId, const FDisplayClusterConfigurationProjection* InConfigurationProjectionPolicy)
 	: FDisplayClusterProjectionMPCDIPolicy(ProjectionPolicyId, InConfigurationProjectionPolicy)
 {
@@ -27,10 +50,10 @@ bool FDisplayClusterProjectionMeshPolicy::CreateWarpMeshInterface(IDisplayCluste
 {
 	check(IsInGameThread());
 
-	if (WarpBlendInterface.IsValid() == false)
+	if (!WarpBlendInterface.IsValid())
 	{
 		// Read configuration from config
-		FWarpMeshConfiguration WarpCfg;
+		FDisplayClusterProjectionMeshPolicyConfiguration WarpCfg;
 		if (!GetWarpMeshConfiguration(InViewport, WarpCfg))
 		{
 			return false;
@@ -39,12 +62,11 @@ bool FDisplayClusterProjectionMeshPolicy::CreateWarpMeshInterface(IDisplayCluste
 		// The mesh always uses DCRootActor as its origin because all geometry is in UE space.
 		InitializeOriginComponent(InViewport, TEXT(""));
 
-		IDisplayClusterShaders& ShadersAPI = IDisplayClusterShaders::Get();
-		bool bResult = false;
+		static IDisplayClusterWarp& DisplayClusterWarpAPI = IDisplayClusterWarp::Get();
 
 		if (WarpCfg.StaticMeshComponent != nullptr)
 		{
-			FDisplayClusterWarpBlendConstruct::FAssignWarpStaticMesh CreateParameters;
+			FDisplayClusterWarpInitializer_StaticMesh CreateParameters;
 
 			CreateParameters.OriginComponent = GetOriginComp();
 
@@ -54,11 +76,11 @@ bool FDisplayClusterProjectionMeshPolicy::CreateWarpMeshInterface(IDisplayCluste
 			CreateParameters.BaseUVIndex      = WarpCfg.BaseUVIndex;
 			CreateParameters.ChromakeyUVIndex = WarpCfg.ChromakeyUVIndex;
 
-			bResult = ShadersAPI.GetWarpBlendManager().Create(CreateParameters, WarpBlendInterface);
+			WarpBlendInterface = DisplayClusterWarpAPI.Create(CreateParameters);
 		}
 		else
 		{
-			FDisplayClusterWarpBlendConstruct::FAssignWarpProceduralMesh CreateParameters;
+			FDisplayClusterWarpInitializer_ProceduralMesh CreateParameters;
 
 			CreateParameters.OriginComponent = GetOriginComp();
 
@@ -68,10 +90,10 @@ bool FDisplayClusterProjectionMeshPolicy::CreateWarpMeshInterface(IDisplayCluste
 			CreateParameters.BaseUVIndex      = WarpCfg.BaseUVIndex;
 			CreateParameters.ChromakeyUVIndex = WarpCfg.ChromakeyUVIndex;
 
-			bResult = ShadersAPI.GetWarpBlendManager().Create(CreateParameters, WarpBlendInterface);
+			WarpBlendInterface = DisplayClusterWarpAPI.Create(CreateParameters);
 		}
 
-		if (!bResult)
+		if (!WarpBlendInterface.IsValid())
 		{
 			if (!IsEditorOperationMode(InViewport))
 			{
@@ -116,7 +138,7 @@ bool FDisplayClusterProjectionMeshPolicy::HandleStartScene(IDisplayClusterViewpo
 	return true;
 }
 
-bool FDisplayClusterProjectionMeshPolicy::GetWarpMeshConfiguration(IDisplayClusterViewport* InViewport, FWarpMeshConfiguration& OutWarpCfg)
+bool FDisplayClusterProjectionMeshPolicy::GetWarpMeshConfiguration(IDisplayClusterViewport* InViewport, FDisplayClusterProjectionMeshPolicyConfiguration& OutWarpCfg)
 {
 	check(InViewport);
 
@@ -233,20 +255,11 @@ bool FDisplayClusterProjectionMeshPolicy::GetWarpMeshConfiguration(IDisplayClust
 #if WITH_EDITOR
 UMeshComponent* FDisplayClusterProjectionMeshPolicy::GetOrCreatePreviewMeshComponent(IDisplayClusterViewport* InViewport, bool& bOutIsRootActorComponent)
 {
-	bOutIsRootActorComponent = true;
-
-	FWarpMeshConfiguration WarpCfg;
-	if (GetWarpMeshConfiguration(InViewport, WarpCfg))
-	{
-		if(WarpCfg.StaticMeshComponent != nullptr)
+	if (WarpBlendInterface.IsValid())
 		{
-			return WarpCfg.StaticMeshComponent;
-		}
+		bOutIsRootActorComponent = true;
 
-		if (WarpCfg.ProceduralMeshComponent != nullptr)
-		{
-			return WarpCfg.ProceduralMeshComponent;
-		}
+		return WarpBlendInterface->GetStaticMeshComponent();
 	}
 
 	return nullptr;
diff --git a/Engine/Plugins/Runtime/nDisplay/Source/DisplayClusterProjection/Private/Policy/Mesh/DisplayClusterProjectionMeshPolicy.h b/Engine/Plugins/Runtime/nDisplay/Source/DisplayClusterProjection/Private/Policy/Mesh/DisplayClusterProjectionMeshPolicy.h
--- a/Engine/Plugins/Runtime/nDisplay/Source/DisplayClusterProjection/Private/Policy/Mesh/DisplayClusterProjectionMeshPolicy.h
+++ b/Engine/Plugins/Runtime/nDisplay/Source/DisplayClusterProjection/Private/Policy/Mesh/DisplayClusterProjectionMeshPolicy.h
@@ -34,32 +34,17 @@ public:
 	virtual UMeshComponent* GetOrCreatePreviewMeshComponent(IDisplayClusterViewport* InViewport, bool& bOutIsRootActorComponent) override;
 #endif
 
-public:	
-	virtual EWarpType GetWarpType() const override
-	{
-		return EWarpType::mesh;
-	}
-
+public:
 	/** Parse the config data for a mesh id and try to retrieve it from the root actor. */
 	bool CreateWarpMeshInterface(IDisplayClusterViewport* InViewport);
 
 private:
-	struct FWarpMeshConfiguration
-	{
-		// StaticMesh component with source geometry
-		UStaticMeshComponent*     StaticMeshComponent = nullptr;
-		// StaticMesh geometry LOD
-		int32 StaticMeshComponentLODIndex = 0;
-
-		// ProceduralMesh component with source geometry
-		UProceduralMeshComponent* ProceduralMeshComponent = nullptr;
-		// ProceduralMesh section index
-		int32 ProceduralMeshComponentSectionIndex = 0;
-
-		// Customize source geometry UV channels
-		int32 BaseUVIndex = INDEX_NONE;
-		int32 ChromakeyUVIndex = INDEX_NONE;
-	};
-
-	bool GetWarpMeshConfiguration(IDisplayClusterViewport* InViewport, FWarpMeshConfiguration& OutWarpCfg);
+	/** Read mesh policy configuration from the projection policy parameters
+	* 
+	* @param InViewport - projection policy owner viewport
+	* @param OutWarpCfg - (out) warpblend configuration
+	* 
+	* @return - true when successful.
+	*/
+	bool GetWarpMeshConfiguration(IDisplayClusterViewport* InViewport, struct FDisplayClusterProjectionMeshPolicyConfiguration& OutWarpCfg);
 };
diff --git a/Engine/Plugins/Runtime/nDisplay/Source/DisplayClusterProjection/Public/DisplayClusterProjectionStrings.h b/Engine/Plugins/Runtime/nDisplay/Source/DisplayClusterProjection/Public/DisplayClusterProjectionStrings.h
--- a/Engine/Plugins/Runtime/nDisplay/Source/DisplayClusterProjection/Public/DisplayClusterProjectionStrings.h
+++ b/Engine/Plugins/Runtime/nDisplay/Source/DisplayClusterProjection/Public/DisplayClusterProjectionStrings.h
@@ -4,7 +4,9 @@
 
 #include "CoreMinimal.h"
 
-
+/**
+* Textual contants used by projection policies.
+*/
 namespace DisplayClusterProjectionStrings
 {
 	namespace cfg
@@ -55,6 +57,7 @@ namespace DisplayClusterProjectionStrings
 			static constexpr const TCHAR* AngleR       = TEXT("r");
 			static constexpr const TCHAR* AngleT       = TEXT("t");
 			static constexpr const TCHAR* AngleB       = TEXT("b");
+
 			namespace FrustumType
 			{
 				static constexpr const TCHAR* Matrix = TEXT("Matrix");
@@ -107,6 +110,46 @@ namespace DisplayClusterProjectionStrings
 			static constexpr const TCHAR* TypePFM = TEXT("Explicit PFM");
 
 			static constexpr const TCHAR* EnablePreview = TEXT("EnablePreview");
+
+			static constexpr const TCHAR* Component = TEXT("screen_component");
+
+			namespace Attributes
+			{
+				namespace Buffer
+				{
+					static constexpr const TCHAR* Resolution = TEXT("BufferResolution");
+				}
+
+				namespace Region
+				{
+					static constexpr const TCHAR* Resolution = TEXT("RegionResolution");
+					static constexpr const TCHAR* Pos = TEXT("RegionPos");
+					static constexpr const TCHAR* Size = TEXT("RegionSize");
+				}
+
+				namespace Frustum
+				{
+					static constexpr const TCHAR* Pitch = TEXT("FrustumPitch");
+					static constexpr const TCHAR* Yaw   = TEXT("FrustumYaw");
+					static constexpr const TCHAR* Roll  = TEXT("FrustumRoll");
+					
+					namespace Angle
+					{
+						static constexpr const TCHAR* Left   = TEXT("FrustumAngleLeft");
+						static constexpr const TCHAR* Right  = TEXT("FrustumAngleRight");
+						static constexpr const TCHAR* Top    = TEXT("FrustumAngleTop");
+						static constexpr const TCHAR* Bottom = TEXT("FrustumAngleBottom");
+					}
+				}
+
+				namespace CoordinateFrame
+				{
+					static constexpr const TCHAR* Pos   = TEXT("CoordinateFramePos");
+					static constexpr const TCHAR* Yaw   = TEXT("CoordinateFrameYaw");
+					static constexpr const TCHAR* Pitch = TEXT("CoordinateFramePitch");
+					static constexpr const TCHAR* Roll  = TEXT("CoordinateFrameRoll");
+				}
+			}
 		}
 
 		namespace domeprojection
diff --git a/Engine/Plugins/Runtime/nDisplay/Source/DisplayClusterShaders/DisplayClusterShaders.Build.cs b/Engine/Plugins/Runtime/nDisplay/Source/DisplayClusterShaders/DisplayClusterShaders.Build.cs
--- a/Engine/Plugins/Runtime/nDisplay/Source/DisplayClusterShaders/DisplayClusterShaders.Build.cs
+++ b/Engine/Plugins/Runtime/nDisplay/Source/DisplayClusterShaders/DisplayClusterShaders.Build.cs
@@ -7,9 +7,6 @@ public class DisplayClusterShaders : ModuleRules
 {
 	public DisplayClusterShaders(ReadOnlyTargetRules ROTargetRules) : base(ROTargetRules)
 	{
-		// This define required for mpcdi library
-		PublicDefinitions.Add("MPCDI_STATIC");
-
 		PublicDependencyModuleNames.AddRange(
 			new string[] {
 				"Core",
@@ -21,55 +18,15 @@ public class DisplayClusterShaders : ModuleRules
 		PrivateDependencyModuleNames.AddRange(
 			new string[] {
 				"DisplayCluster",
+				"DisplayClusterWarp",
 				"RenderCore",
 				"RHI",
-				"Renderer",
-				"ProceduralMeshComponent",
+				"Renderer"
 			});
 
 		if (Target.bBuildEditor == true)
 		{
 			PrivateDependencyModuleNames.Add("UnrealEd");
 		}
-		
-		// NOTE: General rule is not to access the private folder of another module
-		PrivateIncludePaths.AddRange(new string[] {
-		});
-
-		AddEngineThirdPartyPrivateStaticDependencies(Target, "UElibPNG");
-		AddEngineThirdPartyPrivateStaticDependencies(Target, "zlib");
-
-		AddThirdPartyDependencies(ROTargetRules);
-	}
-
-	public void AddThirdPartyDependencies(ReadOnlyTargetRules ROTargetRules)
-	{
-		string ThirdPartyPath = Path.GetFullPath(Path.Combine(ModuleDirectory, "../../ThirdParty/"));
-
-		string PathLib = string.Empty;
-		string PathInc = string.Empty;
-
-		// MPCDI
-		PathLib = Path.Combine(ThirdPartyPath, "MPCDI/Lib");
-		PathInc = Path.Combine(ThirdPartyPath, "MPCDI/Include");
-
-		// Libs
-		if (Target.Platform == UnrealTargetPlatform.Win64)
-		{
-			PublicAdditionalLibraries.Add(Path.Combine(PathLib, "mpcdi.lib"));
-			PublicAdditionalLibraries.Add(Path.Combine(PathLib, "tinyxml2.lib"));
-		}
-		else if (Target.Platform == UnrealTargetPlatform.Linux && Target.Architecture == UnrealArch.X64)
-		{
-			PublicAdditionalLibraries.Add(Path.Combine(PathLib, "libmpcdi.a"));
-		}
-
-		// Include paths
-		PrivateIncludePaths.Add(PathInc);
-		PrivateIncludePaths.Add(Path.Combine(PathInc, "Base"));
-		PrivateIncludePaths.Add(Path.Combine(PathInc, "Container"));
-		PrivateIncludePaths.Add(Path.Combine(PathInc, "Creators"));
-		PrivateIncludePaths.Add(Path.Combine(PathInc, "IO"));
-		PrivateIncludePaths.Add(Path.Combine(PathInc, "Utils"));
 	}
 }
diff --git a/Engine/Plugins/Runtime/nDisplay/Source/DisplayClusterShaders/Private/Blueprints/MPCDIGeometryData.cpp b/Engine/Plugins/Runtime/nDisplay/Source/DisplayClusterShaders/Private/Blueprints/MPCDIGeometryData.cpp
deleted file mode 100644
diff --git a/Engine/Plugins/Runtime/nDisplay/Source/DisplayClusterShaders/Private/DisplayClusterShadersModule.h b/Engine/Plugins/Runtime/nDisplay/Source/DisplayClusterShaders/Private/DisplayClusterShadersModule.h
--- a/Engine/Plugins/Runtime/nDisplay/Source/DisplayClusterShaders/Private/DisplayClusterShadersModule.h
+++ b/Engine/Plugins/Runtime/nDisplay/Source/DisplayClusterShaders/Private/DisplayClusterShadersModule.h
@@ -4,8 +4,6 @@
 
 #include "IDisplayClusterShaders.h"
 
-#include "WarpBlend/DisplayClusterWarpBlendManager.h"
-
 class FDisplayClusterShadersModule
 	: public IDisplayClusterShaders
 {
@@ -24,12 +22,4 @@ public:
 	virtual bool RenderPostprocess_OutputRemap(FRHICommandListImmediate& RHICmdList, FRHITexture2D* InSourceTexture, FRHITexture2D* InRenderTargetableDestTexture, const IDisplayClusterRender_MeshComponentProxy& MeshProxy) const override;
 	virtual bool RenderPostprocess_Blur(FRHICommandListImmediate& RHICmdList, FRHITexture2D* InSourceTexture, FRHITexture2D* InRenderTargetableDestTexture, const FDisplayClusterShaderParameters_PostprocessBlur& InSettings) const override;
 	virtual bool GenerateMips(FRHICommandListImmediate& RHICmdList, FRHITexture2D* InOutMipsTexture, const FDisplayClusterShaderParameters_GenerateMips& InSettings) const override;
-
-	virtual const IDisplayClusterWarpBlendManager& GetWarpBlendManager() const override
-	{
-		return WarpBlendManager;
-	}
-
-private:
-	FDisplayClusterWarpBlendManager WarpBlendManager;
 };
diff --git a/Engine/Plugins/Runtime/nDisplay/Source/DisplayClusterShaders/Private/Render/DisplayClusterRenderTexture.cpp b/Engine/Plugins/Runtime/nDisplay/Source/DisplayClusterShaders/Private/Render/DisplayClusterRenderTexture.cpp
deleted file mode 100644
diff --git a/Engine/Plugins/Runtime/nDisplay/Source/DisplayClusterShaders/Private/Render/DisplayClusterRenderTexture.h b/Engine/Plugins/Runtime/nDisplay/Source/DisplayClusterShaders/Private/Render/DisplayClusterRenderTexture.h
deleted file mode 100644
diff --git a/Engine/Plugins/Runtime/nDisplay/Source/DisplayClusterShaders/Private/Shaders/DisplayClusterShadersWarpblend_ICVFX.cpp b/Engine/Plugins/Runtime/nDisplay/Source/DisplayClusterShaders/Private/Shaders/DisplayClusterShadersWarpblend_ICVFX.cpp
--- a/Engine/Plugins/Runtime/nDisplay/Source/DisplayClusterShaders/Private/Shaders/DisplayClusterShadersWarpblend_ICVFX.cpp
+++ b/Engine/Plugins/Runtime/nDisplay/Source/DisplayClusterShaders/Private/Shaders/DisplayClusterShadersWarpblend_ICVFX.cpp
@@ -20,7 +20,7 @@
 #include "ShaderParameters/DisplayClusterShaderParameters_WarpBlend.h"
 #include "ShaderParameters/DisplayClusterShaderParameters_ICVFX.h"
 
-#include "WarpBlend/IDisplayClusterWarpBlend.h"
+#include "IDisplayClusterWarpBlend.h"
 
 #define ICVFX_ShaderFileName "/Plugin/nDisplay/Private/MPCDIOverlayShaders.usf"
 
diff --git a/Engine/Plugins/Runtime/nDisplay/Source/DisplayClusterShaders/Private/Shaders/DisplayClusterShadersWarpblend_MPCDI.cpp b/Engine/Plugins/Runtime/nDisplay/Source/DisplayClusterShaders/Private/Shaders/DisplayClusterShadersWarpblend_MPCDI.cpp
--- a/Engine/Plugins/Runtime/nDisplay/Source/DisplayClusterShaders/Private/Shaders/DisplayClusterShadersWarpblend_MPCDI.cpp
+++ b/Engine/Plugins/Runtime/nDisplay/Source/DisplayClusterShaders/Private/Shaders/DisplayClusterShadersWarpblend_MPCDI.cpp
@@ -19,7 +19,7 @@
 #include "ShaderParameterStruct.h"
 
 #include "Render/Containers/IDisplayClusterRender_MeshComponentProxy.h"
-#include "WarpBlend/IDisplayClusterWarpBlend.h"
+#include "IDisplayClusterWarpBlend.h"
 
 #include "ShaderParameters/DisplayClusterShaderParameters_WarpBlend.h"
 
diff --git a/Engine/Plugins/Runtime/nDisplay/Source/DisplayClusterShaders/Private/WarpBlend/DisplayClusterWarpBlend.cpp b/Engine/Plugins/Runtime/nDisplay/Source/DisplayClusterShaders/Private/WarpBlend/DisplayClusterWarpBlend.cpp
deleted file mode 100644
diff --git a/Engine/Plugins/Runtime/nDisplay/Source/DisplayClusterShaders/Private/WarpBlend/DisplayClusterWarpBlend.h b/Engine/Plugins/Runtime/nDisplay/Source/DisplayClusterShaders/Private/WarpBlend/DisplayClusterWarpBlend.h
deleted file mode 100644
diff --git a/Engine/Plugins/Runtime/nDisplay/Source/DisplayClusterShaders/Private/WarpBlend/DisplayClusterWarpBlendManager.cpp b/Engine/Plugins/Runtime/nDisplay/Source/DisplayClusterShaders/Private/WarpBlend/DisplayClusterWarpBlendManager.cpp
deleted file mode 100644
diff --git a/Engine/Plugins/Runtime/nDisplay/Source/DisplayClusterShaders/Private/WarpBlend/DisplayClusterWarpBlendManager.h b/Engine/Plugins/Runtime/nDisplay/Source/DisplayClusterShaders/Private/WarpBlend/DisplayClusterWarpBlendManager.h
deleted file mode 100644
diff --git a/Engine/Plugins/Runtime/nDisplay/Source/DisplayClusterShaders/Private/WarpBlend/DisplayClusterWarpBlend_FrustumCache.cpp b/Engine/Plugins/Runtime/nDisplay/Source/DisplayClusterShaders/Private/WarpBlend/DisplayClusterWarpBlend_FrustumCache.cpp
deleted file mode 100644
diff --git a/Engine/Plugins/Runtime/nDisplay/Source/DisplayClusterShaders/Private/WarpBlend/DisplayClusterWarpBlend_FrustumCache.h b/Engine/Plugins/Runtime/nDisplay/Source/DisplayClusterShaders/Private/WarpBlend/DisplayClusterWarpBlend_FrustumCache.h
deleted file mode 100644
diff --git a/Engine/Plugins/Runtime/nDisplay/Source/DisplayClusterShaders/Private/WarpBlend/DisplayClusterWarpBlend_GeometryContext.cpp b/Engine/Plugins/Runtime/nDisplay/Source/DisplayClusterShaders/Private/WarpBlend/DisplayClusterWarpBlend_GeometryContext.cpp
deleted file mode 100644
diff --git a/Engine/Plugins/Runtime/nDisplay/Source/DisplayClusterShaders/Private/WarpBlend/DisplayClusterWarpBlend_GeometryContext.h b/Engine/Plugins/Runtime/nDisplay/Source/DisplayClusterShaders/Private/WarpBlend/DisplayClusterWarpBlend_GeometryContext.h
deleted file mode 100644
diff --git a/Engine/Plugins/Runtime/nDisplay/Source/DisplayClusterShaders/Private/WarpBlend/DisplayClusterWarpBlend_GeometryProxy.cpp b/Engine/Plugins/Runtime/nDisplay/Source/DisplayClusterShaders/Private/WarpBlend/DisplayClusterWarpBlend_GeometryProxy.cpp
deleted file mode 100644
diff --git a/Engine/Plugins/Runtime/nDisplay/Source/DisplayClusterShaders/Private/WarpBlend/DisplayClusterWarpBlend_GeometryProxy.h b/Engine/Plugins/Runtime/nDisplay/Source/DisplayClusterShaders/Private/WarpBlend/DisplayClusterWarpBlend_GeometryProxy.h
deleted file mode 100644
diff --git a/Engine/Plugins/Runtime/nDisplay/Source/DisplayClusterShaders/Private/WarpBlend/Exporter/DisplayClusterWarpBlendExporter_WarpMap.cpp b/Engine/Plugins/Runtime/nDisplay/Source/DisplayClusterShaders/Private/WarpBlend/Exporter/DisplayClusterWarpBlendExporter_WarpMap.cpp
deleted file mode 100644
diff --git a/Engine/Plugins/Runtime/nDisplay/Source/DisplayClusterShaders/Private/WarpBlend/Exporter/DisplayClusterWarpBlendExporter_WarpMap.h b/Engine/Plugins/Runtime/nDisplay/Source/DisplayClusterShaders/Private/WarpBlend/Exporter/DisplayClusterWarpBlendExporter_WarpMap.h
deleted file mode 100644
diff --git a/Engine/Plugins/Runtime/nDisplay/Source/DisplayClusterShaders/Private/WarpBlend/Loader/DisplayClusterWarpBlendLoader_MPCDI.cpp b/Engine/Plugins/Runtime/nDisplay/Source/DisplayClusterShaders/Private/WarpBlend/Loader/DisplayClusterWarpBlendLoader_MPCDI.cpp
deleted file mode 100644
diff --git a/Engine/Plugins/Runtime/nDisplay/Source/DisplayClusterShaders/Private/WarpBlend/Loader/DisplayClusterWarpBlendLoader_MPCDI.h b/Engine/Plugins/Runtime/nDisplay/Source/DisplayClusterShaders/Private/WarpBlend/Loader/DisplayClusterWarpBlendLoader_MPCDI.h
deleted file mode 100644
diff --git a/Engine/Plugins/Runtime/nDisplay/Source/DisplayClusterShaders/Private/WarpBlend/Loader/DisplayClusterWarpBlendLoader_MeshComponent.cpp b/Engine/Plugins/Runtime/nDisplay/Source/DisplayClusterShaders/Private/WarpBlend/Loader/DisplayClusterWarpBlendLoader_MeshComponent.cpp
deleted file mode 100644
diff --git a/Engine/Plugins/Runtime/nDisplay/Source/DisplayClusterShaders/Private/WarpBlend/Loader/DisplayClusterWarpBlendLoader_MeshComponent.h b/Engine/Plugins/Runtime/nDisplay/Source/DisplayClusterShaders/Private/WarpBlend/Loader/DisplayClusterWarpBlendLoader_MeshComponent.h
deleted file mode 100644
diff --git a/Engine/Plugins/Runtime/nDisplay/Source/DisplayClusterShaders/Private/WarpBlend/Loader/DisplayClusterWarpBlendLoader_ProceduralMeshComponent.cpp b/Engine/Plugins/Runtime/nDisplay/Source/DisplayClusterShaders/Private/WarpBlend/Loader/DisplayClusterWarpBlendLoader_ProceduralMeshComponent.cpp
deleted file mode 100644
diff --git a/Engine/Plugins/Runtime/nDisplay/Source/DisplayClusterShaders/Private/WarpBlend/Loader/DisplayClusterWarpBlendLoader_ProceduralMeshComponent.h b/Engine/Plugins/Runtime/nDisplay/Source/DisplayClusterShaders/Private/WarpBlend/Loader/DisplayClusterWarpBlendLoader_ProceduralMeshComponent.h
deleted file mode 100644
diff --git a/Engine/Plugins/Runtime/nDisplay/Source/DisplayClusterShaders/Private/WarpBlend/Loader/DisplayClusterWarpBlendLoader_Texture.cpp b/Engine/Plugins/Runtime/nDisplay/Source/DisplayClusterShaders/Private/WarpBlend/Loader/DisplayClusterWarpBlendLoader_Texture.cpp
deleted file mode 100644
diff --git a/Engine/Plugins/Runtime/nDisplay/Source/DisplayClusterShaders/Private/WarpBlend/Loader/DisplayClusterWarpBlendLoader_Texture.h b/Engine/Plugins/Runtime/nDisplay/Source/DisplayClusterShaders/Private/WarpBlend/Loader/DisplayClusterWarpBlendLoader_Texture.h
deleted file mode 100644
diff --git a/Engine/Plugins/Runtime/nDisplay/Source/DisplayClusterShaders/Private/WarpBlend/Loader/DisplayClusterWarpBlendLoader_WarpMap.cpp b/Engine/Plugins/Runtime/nDisplay/Source/DisplayClusterShaders/Private/WarpBlend/Loader/DisplayClusterWarpBlendLoader_WarpMap.cpp
deleted file mode 100644
diff --git a/Engine/Plugins/Runtime/nDisplay/Source/DisplayClusterShaders/Private/WarpBlend/Loader/DisplayClusterWarpBlendLoader_WarpMap.h b/Engine/Plugins/Runtime/nDisplay/Source/DisplayClusterShaders/Private/WarpBlend/Loader/DisplayClusterWarpBlendLoader_WarpMap.h
deleted file mode 100644
diff --git a/Engine/Plugins/Runtime/nDisplay/Source/DisplayClusterShaders/Private/WarpBlend/Math/DisplayClusterWarpBlendMath_Frustum.cpp b/Engine/Plugins/Runtime/nDisplay/Source/DisplayClusterShaders/Private/WarpBlend/Math/DisplayClusterWarpBlendMath_Frustum.cpp
deleted file mode 100644
diff --git a/Engine/Plugins/Runtime/nDisplay/Source/DisplayClusterShaders/Private/WarpBlend/Math/DisplayClusterWarpBlendMath_Frustum.h b/Engine/Plugins/Runtime/nDisplay/Source/DisplayClusterShaders/Private/WarpBlend/Math/DisplayClusterWarpBlendMath_Frustum.h
deleted file mode 100644
diff --git a/Engine/Plugins/Runtime/nDisplay/Source/DisplayClusterShaders/Private/WarpBlend/Math/DisplayClusterWarpBlendMath_WarpMap.h b/Engine/Plugins/Runtime/nDisplay/Source/DisplayClusterShaders/Private/WarpBlend/Math/DisplayClusterWarpBlendMath_WarpMap.h
deleted file mode 100644
diff --git a/Engine/Plugins/Runtime/nDisplay/Source/DisplayClusterShaders/Private/WarpBlend/Math/DisplayClusterWarpBlendMath_WarpMesh.h b/Engine/Plugins/Runtime/nDisplay/Source/DisplayClusterShaders/Private/WarpBlend/Math/DisplayClusterWarpBlendMath_WarpMesh.h
deleted file mode 100644
diff --git a/Engine/Plugins/Runtime/nDisplay/Source/DisplayClusterShaders/Private/WarpBlend/Math/DisplayClusterWarpBlendMath_WarpProceduralMesh.h b/Engine/Plugins/Runtime/nDisplay/Source/DisplayClusterShaders/Private/WarpBlend/Math/DisplayClusterWarpBlendMath_WarpProceduralMesh.h
deleted file mode 100644
diff --git a/Engine/Plugins/Runtime/nDisplay/Source/DisplayClusterShaders/Public/IDisplayClusterShaders.h b/Engine/Plugins/Runtime/nDisplay/Source/DisplayClusterShaders/Public/IDisplayClusterShaders.h
--- a/Engine/Plugins/Runtime/nDisplay/Source/DisplayClusterShaders/Public/IDisplayClusterShaders.h
+++ b/Engine/Plugins/Runtime/nDisplay/Source/DisplayClusterShaders/Public/IDisplayClusterShaders.h
@@ -9,19 +9,18 @@
 #include "RHICommandList.h"
 #include "RHIResources.h"
 
-
 class IDisplayClusterRender_MeshComponentProxy;
 class FDisplayClusterShaderParameters_ICVFX;
 class FDisplayClusterShaderParameters_Override;
+
 struct FDisplayClusterShaderParameters_GenerateMips;
 struct FDisplayClusterShaderParameters_PostprocessBlur;
 struct FDisplayClusterShaderParameters_WarpBlend;
 struct FDisplayClusterShaderParameters_UVLightCards;
-class IDisplayClusterWarpBlendManager;
+
 class FSceneInterface;
 class FRenderTarget;
 
-
 class IDisplayClusterShaders : public IModuleInterface
 {
 public:
@@ -53,13 +52,6 @@ public:
 	}
 
 public:
-	/**
-	* Warp&blend manager interface
-	*
-	* @return interface
-	*/
-	virtual const IDisplayClusterWarpBlendManager& GetWarpBlendManager() const = 0;
-
 	/**
 	* Render warp&blend
 	*
diff --git a/Engine/Plugins/Runtime/nDisplay/Source/DisplayClusterShaders/Public/ShaderParameters/DisplayClusterShaderParameters_WarpBlend.h b/Engine/Plugins/Runtime/nDisplay/Source/DisplayClusterShaders/Public/ShaderParameters/DisplayClusterShaderParameters_WarpBlend.h
--- a/Engine/Plugins/Runtime/nDisplay/Source/DisplayClusterShaders/Public/ShaderParameters/DisplayClusterShaderParameters_WarpBlend.h
+++ b/Engine/Plugins/Runtime/nDisplay/Source/DisplayClusterShaders/Public/ShaderParameters/DisplayClusterShaderParameters_WarpBlend.h
@@ -7,8 +7,7 @@
 #include "RHI.h"
 #include "RHIResources.h"
 
-#include "WarpBlend/DisplayClusterWarpContext.h"
-#include "WarpBlend/IDisplayClusterWarpBlend.h"
+#include "Containers/DisplayClusterWarpContext.h"
 
 struct FDisplayClusterShaderParameters_WarpBlend
 {
@@ -29,7 +28,7 @@ struct FDisplayClusterShaderParameters_WarpBlend
 	FResourceWithRect Dest;
 
 	// Warp interface
-	TSharedPtr<IDisplayClusterWarpBlend, ESPMode::ThreadSafe> WarpInterface;
+	TSharedPtr<class IDisplayClusterWarpBlend, ESPMode::ThreadSafe> WarpInterface;
 
 	// Context data
 	FDisplayClusterWarpContext Context;
diff --git a/Engine/Plugins/Runtime/nDisplay/Source/DisplayClusterShaders/Public/WarpBlend/DisplayClusterWarpContext.h b/Engine/Plugins/Runtime/nDisplay/Source/DisplayClusterShaders/Public/WarpBlend/DisplayClusterWarpContext.h
deleted file mode 100644
diff --git a/Engine/Plugins/Runtime/nDisplay/Source/DisplayClusterShaders/Public/WarpBlend/DisplayClusterWarpEnums.h b/Engine/Plugins/Runtime/nDisplay/Source/DisplayClusterShaders/Public/WarpBlend/DisplayClusterWarpEnums.h
deleted file mode 100644
diff --git a/Engine/Plugins/Runtime/nDisplay/Source/DisplayClusterShaders/Public/WarpBlend/IDisplayClusterWarpBlend.h b/Engine/Plugins/Runtime/nDisplay/Source/DisplayClusterShaders/Public/WarpBlend/IDisplayClusterWarpBlend.h
deleted file mode 100644
diff --git a/Engine/Plugins/Runtime/nDisplay/Source/DisplayClusterShaders/Public/WarpBlend/IDisplayClusterWarpBlendManager.h b/Engine/Plugins/Runtime/nDisplay/Source/DisplayClusterShaders/Public/WarpBlend/IDisplayClusterWarpBlendManager.h
deleted file mode 100644
diff --git a/Engine/Plugins/Runtime/nDisplay/Source/DisplayClusterWarp/DisplayClusterWarp.Build.cs b/Engine/Plugins/Runtime/nDisplay/Source/DisplayClusterWarp/DisplayClusterWarp.Build.cs
new file mode 100644
--- /dev/null
+++ b/Engine/Plugins/Runtime/nDisplay/Source/DisplayClusterWarp/DisplayClusterWarp.Build.cs
@@ -0,0 +1,77 @@
+// Copyright Epic Games, Inc. All Rights Reserved.
+
+using UnrealBuildTool;
+using System.IO;
+
+public class DisplayClusterWarp : ModuleRules
+{
+	public DisplayClusterWarp(ReadOnlyTargetRules ROTargetRules) : base(ROTargetRules)
+	{
+		// This define required for mpcdi library
+		PublicDefinitions.Add("MPCDI_STATIC");
+
+		PublicDependencyModuleNames.AddRange(
+			new string[] {
+				"CinematicCamera",
+				"Core",
+				"CoreUObject",
+				"DisplayCluster",
+				"Engine",
+				"Projects"
+			});
+
+		PrivateDependencyModuleNames.AddRange(
+			new string[] {
+				"DisplayCluster",
+				"RenderCore",
+				"RHI",
+				"Renderer",
+				"ProceduralMeshComponent",
+			});
+
+		if (Target.bBuildEditor == true)
+		{
+			PrivateDependencyModuleNames.Add("UnrealEd");
+		}
+		
+		// NOTE: General rule is not to access the private folder of another module
+		PrivateIncludePaths.AddRange(new string[] {
+		});
+
+		AddEngineThirdPartyPrivateStaticDependencies(Target, "UElibPNG");
+		AddEngineThirdPartyPrivateStaticDependencies(Target, "zlib");
+
+		AddThirdPartyDependencies(ROTargetRules);
+	}
+
+	public void AddThirdPartyDependencies(ReadOnlyTargetRules ROTargetRules)
+	{
+		string ThirdPartyPath = Path.GetFullPath(Path.Combine(ModuleDirectory, "../../ThirdParty/"));
+
+		string PathLib = string.Empty;
+		string PathInc = string.Empty;
+
+		// MPCDI
+		PathLib = Path.Combine(ThirdPartyPath, "MPCDI/Lib");
+		PathInc = Path.Combine(ThirdPartyPath, "MPCDI/Include");
+
+		// Libs
+		if (Target.Platform == UnrealTargetPlatform.Win64)
+		{
+			PublicAdditionalLibraries.Add(Path.Combine(PathLib, "mpcdi.lib"));
+			PublicAdditionalLibraries.Add(Path.Combine(PathLib, "tinyxml2.lib"));
+		}
+		else if (Target.Platform == UnrealTargetPlatform.Linux && Target.Architecture == UnrealArch.X64)
+		{
+			PublicAdditionalLibraries.Add(Path.Combine(PathLib, "libmpcdi.a"));
+		}
+
+		// Include paths
+		PrivateIncludePaths.Add(PathInc);
+		PrivateIncludePaths.Add(Path.Combine(PathInc, "Base"));
+		PrivateIncludePaths.Add(Path.Combine(PathInc, "Container"));
+		PrivateIncludePaths.Add(Path.Combine(PathInc, "Creators"));
+		PrivateIncludePaths.Add(Path.Combine(PathInc, "IO"));
+		PrivateIncludePaths.Add(Path.Combine(PathInc, "Utils"));
+	}
+}
diff --git a/Engine/Plugins/Runtime/nDisplay/Source/DisplayClusterWarp/Private/Blueprints/DisplayClusterWarpGeometry.cpp b/Engine/Plugins/Runtime/nDisplay/Source/DisplayClusterWarp/Private/Blueprints/DisplayClusterWarpGeometry.cpp
new file mode 100644
--- /dev/null
+++ b/Engine/Plugins/Runtime/nDisplay/Source/DisplayClusterWarp/Private/Blueprints/DisplayClusterWarpGeometry.cpp
@@ -0,0 +1,23 @@
+// Copyright Epic Games, Inc. All Rights Reserved.
+
+#include "Blueprints/DisplayClusterWarpGeometry.h"
+
+//------------------------------------------------------------------------------
+// FDisplayClusterWarpGeometryOBJ
+//------------------------------------------------------------------------------
+void FDisplayClusterWarpGeometryOBJ::PostAddFace(int32 f0, int32 f1, int32 f2)
+{
+	Triangles.Add(f0);
+	Triangles.Add(f1);
+	Triangles.Add(f2);
+
+	//Update normal
+	const FVector FaceDir1 = Vertices[f1] - Vertices[f0];
+	const FVector FaceDir2 = Vertices[f2] - Vertices[f0];
+
+	const FVector FaceNornal = FVector::CrossProduct(FaceDir1, FaceDir2);
+
+	Normal[f0] = FaceNornal;
+	Normal[f1] = FaceNornal;
+	Normal[f2] = FaceNornal;
+}
diff --git a/Engine/Plugins/Runtime/nDisplay/Source/DisplayClusterWarp/Private/DisplayClusterWarpLog.cpp b/Engine/Plugins/Runtime/nDisplay/Source/DisplayClusterWarp/Private/DisplayClusterWarpLog.cpp
new file mode 100644
--- /dev/null
+++ b/Engine/Plugins/Runtime/nDisplay/Source/DisplayClusterWarp/Private/DisplayClusterWarpLog.cpp
@@ -0,0 +1,8 @@
+// Copyright Epic Games, Inc. All Rights Reserved.
+
+#include "DisplayClusterWarpLog.h"
+
+
+DEFINE_LOG_CATEGORY(LogDisplayClusterWarp);
+DEFINE_LOG_CATEGORY(LogDisplayClusterWarpBlend);
+DEFINE_LOG_CATEGORY(LogDisplayClusterWarpInFrustumFitPolicy);
diff --git a/Engine/Plugins/Runtime/nDisplay/Source/DisplayClusterWarp/Private/DisplayClusterWarpLog.h b/Engine/Plugins/Runtime/nDisplay/Source/DisplayClusterWarp/Private/DisplayClusterWarpLog.h
new file mode 100644
--- /dev/null
+++ b/Engine/Plugins/Runtime/nDisplay/Source/DisplayClusterWarp/Private/DisplayClusterWarpLog.h
@@ -0,0 +1,10 @@
+// Copyright Epic Games, Inc. All Rights Reserved.
+
+#pragma once
+
+#include "CoreMinimal.h"
+
+// Plugin-wide log categories
+DECLARE_LOG_CATEGORY_EXTERN(LogDisplayClusterWarp,            Log, All);
+DECLARE_LOG_CATEGORY_EXTERN(LogDisplayClusterWarpBlend, Log, All);
+DECLARE_LOG_CATEGORY_EXTERN(LogDisplayClusterWarpInFrustumFitPolicy, Log, All);
diff --git a/Engine/Plugins/Runtime/nDisplay/Source/DisplayClusterWarp/Private/DisplayClusterWarpModule.cpp b/Engine/Plugins/Runtime/nDisplay/Source/DisplayClusterWarp/Private/DisplayClusterWarpModule.cpp
new file mode 100644
--- /dev/null
+++ b/Engine/Plugins/Runtime/nDisplay/Source/DisplayClusterWarp/Private/DisplayClusterWarpModule.cpp
@@ -0,0 +1,34 @@
+// Copyright Epic Games, Inc. All Rights Reserved.
+
+#include "DisplayClusterWarpModule.h"
+
+#include "DisplayClusterWarpLog.h"
+#include "DisplayClusterWarpStrings.h"
+
+#include "IDisplayCluster.h"
+#include "Render/IDisplayClusterRenderManager.h"
+
+FDisplayClusterWarpModule::FDisplayClusterWarpModule()
+{
+	UE_LOG(LogDisplayClusterWarp, Verbose, TEXT("Warp module has been instantiated"));
+}
+
+FDisplayClusterWarpModule::~FDisplayClusterWarpModule()
+{
+	UE_LOG(LogDisplayClusterWarp, Verbose, TEXT("Warp module has been destroyed"));
+}
+
+//////////////////////////////////////////////////////////////////////////////////////////////
+// IModuleInterface
+//////////////////////////////////////////////////////////////////////////////////////////////
+void FDisplayClusterWarpModule::StartupModule()
+{
+	UE_LOG(LogDisplayClusterWarp, Verbose, TEXT("Warp module has been started"));
+}
+
+void FDisplayClusterWarpModule::ShutdownModule()
+{
+	UE_LOG(LogDisplayClusterWarp, Verbose, TEXT("Warp module has been shutdown"));
+}
+
+IMPLEMENT_MODULE(FDisplayClusterWarpModule, DisplayClusterWarp);
diff --git a/Engine/Plugins/Runtime/nDisplay/Source/DisplayClusterWarp/Private/DisplayClusterWarpModule.h b/Engine/Plugins/Runtime/nDisplay/Source/DisplayClusterWarp/Private/DisplayClusterWarpModule.h
new file mode 100644
--- /dev/null
+++ b/Engine/Plugins/Runtime/nDisplay/Source/DisplayClusterWarp/Private/DisplayClusterWarpModule.h
@@ -0,0 +1,57 @@
+// Copyright Epic Games, Inc. All Rights Reserved.
+
+#pragma once
+
+#include "IDisplayClusterWarp.h"
+#include "Render/Warp/IDisplayClusterWarpPolicyFactory.h"
+#include "WarpBlend/Loader/DisplayClusterWarpBlendLoader_MPCDI.h"
+#include "WarpBlend/Loader/DisplayClusterWarpBlendLoader_MeshComponent.h"
+
+/**
+ * Implement DisplayClusterWarp module
+ */
+class FDisplayClusterWarpModule
+	: public IDisplayClusterWarp
+{
+public:
+	FDisplayClusterWarpModule();
+	virtual ~FDisplayClusterWarpModule();
+
+public:
+	//////////////////////////////////////////////////////////////////////////////////////////////
+	// IModuleInterface
+	//////////////////////////////////////////////////////////////////////////////////////////////
+	virtual void StartupModule() override;
+	virtual void ShutdownModule() override;
+
+public:
+	virtual TSharedPtr<IDisplayClusterWarpBlend, ESPMode::ThreadSafe> Create(const FDisplayClusterWarpInitializer_MPCDIFile& InConstructParameters) const override
+	{
+		return FDisplayClusterWarpBlendLoader_MPCDI::Create(InConstructParameters);
+	}
+
+	virtual TSharedPtr<IDisplayClusterWarpBlend, ESPMode::ThreadSafe> Create(const FDisplayClusterWarpInitializer_MPCDIFile_Profile2DScreen& InConstructParameters) const override
+	{
+		return FDisplayClusterWarpBlendLoader_MPCDI::Create(InConstructParameters);
+	}
+
+	virtual TSharedPtr<IDisplayClusterWarpBlend, ESPMode::ThreadSafe> Create(const FDisplayClusterWarpInitializer_PFMFile& InConstructParameters) const override
+	{
+		return FDisplayClusterWarpBlendLoader_MPCDI::Create(InConstructParameters);
+	}
+
+	virtual TSharedPtr<IDisplayClusterWarpBlend, ESPMode::ThreadSafe> Create(const FDisplayClusterWarpInitializer_StaticMesh& InConstructParameters) const override
+	{
+		return FDisplayClusterWarpBlendLoader_MeshComponent::Create(InConstructParameters);
+	}
+
+	virtual TSharedPtr<IDisplayClusterWarpBlend, ESPMode::ThreadSafe> Create(const FDisplayClusterWarpInitializer_ProceduralMesh& InConstructParameters) const override
+	{
+		return FDisplayClusterWarpBlendLoader_MeshComponent::Create(InConstructParameters);
+	}
+
+	virtual bool ReadMPCDFileStructure(const FString& InMPCDIFileName, TMap<FString, TMap<FString, FDisplayClusterWarpMPCDIAttributes>>& OutMPCDIFileStructure) const override
+	{
+		return FDisplayClusterWarpBlendLoader_MPCDI::ReadMPCDFileStructure(InMPCDIFileName, OutMPCDIFileStructure);
+	}
+};
diff --git a/Engine/Plugins/Runtime/nDisplay/Source/DisplayClusterWarp/Private/WarpBlend/DisplayClusterWarpBlend.cpp b/Engine/Plugins/Runtime/nDisplay/Source/DisplayClusterWarp/Private/WarpBlend/DisplayClusterWarpBlend.cpp
new file mode 100644
--- /dev/null
+++ b/Engine/Plugins/Runtime/nDisplay/Source/DisplayClusterWarp/Private/WarpBlend/DisplayClusterWarpBlend.cpp
@@ -0,0 +1,159 @@
+// Copyright Epic Games, Inc. All Rights Reserved.
+
+#include "WarpBlend/DisplayClusterWarpBlend.h"
+#include "WarpBlend/Math/DisplayClusterWarpBlendMath_Frustum.h"
+
+#include "Render/Containers/IDisplayClusterRender_MeshComponent.h"
+
+#include "HAL/IConsoleManager.h"
+#include "ProceduralMeshComponent.h"
+#include "Components/StaticMeshComponent.h"
+
+// Setup frustum projection cache
+static TAutoConsoleVariable<int32> CVarMPCDIFrustumCacheDepth(
+	TEXT("nDisplay.render.mpcdi.cache_depth"),
+	0,// Default disabled
+	TEXT("Frustum values cache (depth, num).\n")
+	TEXT("By default cache is disabled. For better performance (EDisplayClusterWarpBlendFrustumType::FULL) set value to 512).\n")
+	TEXT(" 0: Disabled\n")
+	TEXT(" N: Cache size, integer\n"),
+	ECVF_RenderThreadSafe
+);
+
+static TAutoConsoleVariable<float> CVarMPCDIFrustumCachePrecision(
+	TEXT("nDisplay.render.mpcdi.cache_precision"),
+	0.1f, // 1mm
+	TEXT("Frustum cache values comparison precision (float, unit is sm).\n"),
+	ECVF_RenderThreadSafe
+);
+
+//--------------------------------------------------------------------------------------------
+// FDisplayClusterWarpBlend
+//--------------------------------------------------------------------------------------------
+FDisplayClusterWarpBlend::FDisplayClusterWarpBlend()
+{
+	WarpData.AddDefaulted(2);
+}
+
+FDisplayClusterWarpBlend::~FDisplayClusterWarpBlend()
+{
+	// Release resources
+	GeometryContext.GeometryProxy.ReleaseResources();
+}
+
+bool FDisplayClusterWarpBlend::MarkWarpGeometryComponentDirty(const FName& InComponentName)
+{
+	return GeometryContext.GeometryProxy.MarkWarpFrustumGeometryComponentDirty(InComponentName);
+}
+
+bool FDisplayClusterWarpBlend::ShouldSupportICVFX() const
+{
+	if (GetWarpProfileType() == EDisplayClusterWarpProfileType::warp_A3D)
+	{
+		return true;
+	}
+
+	return false;
+}
+
+bool FDisplayClusterWarpBlend::UpdateGeometryContext(const float InWorldScale)
+{
+	return GeometryContext.UpdateGeometryContext(InWorldScale);
+}
+
+const FDisplayClusterWarpGeometryContext& FDisplayClusterWarpBlend::GetGeometryContext() const
+{
+	return GeometryContext.Context;
+}
+
+bool FDisplayClusterWarpBlend::CalcFrustumContext(const TSharedPtr<FDisplayClusterWarpEye, ESPMode::ThreadSafe>& InWarpEye)
+{
+	if(!InWarpEye.IsValid())
+	{
+		return false;
+	}
+
+	// Update current warp data
+	check(WarpData.IsValidIndex(InWarpEye->ContextNum));
+	FDisplayClusterWarpData& CurrentWarpData = WarpData[InWarpEye->ContextNum];
+	CurrentWarpData.WarpEye = InWarpEye;
+
+	FDisplayClusterWarpBlendMath_Frustum Frustum(CurrentWarpData, GeometryContext);
+
+	// Override settings from warp policy
+	BeginCalcFrustum(InWarpEye);
+
+	if (InWarpEye->bUpdateGeometryContext)
+	{
+		if (!GeometryContext.UpdateGeometryContext(InWarpEye->WorldScale))
+		{
+			// wrong geometry
+			return false;
+		}
+	}
+
+	if (!Frustum.CalcFrustum())
+	{
+		return false;
+	}
+
+	return true;
+}
+
+void FDisplayClusterWarpBlend::BeginCalcFrustum(const TSharedPtr<FDisplayClusterWarpEye, ESPMode::ThreadSafe>& InWarpEye)
+{
+	check(InWarpEye.IsValid());
+
+	if (InWarpEye->WarpPolicy.IsValid())
+	{
+		if (IDisplayClusterViewport* Viewport = InWarpEye->GetViewport())
+		{
+			InWarpEye->WarpPolicy->BeginCalcFrustum(Viewport, InWarpEye->ContextNum);
+		}
+	}
+}
+
+FDisplayClusterWarpData& FDisplayClusterWarpBlend::GetWarpData(const uint32 ContextNum)
+{
+	check(WarpData.IsValidIndex(ContextNum));
+
+	return WarpData[ContextNum];
+
+}
+
+const FDisplayClusterWarpData& FDisplayClusterWarpBlend::GetWarpData(const uint32 ContextNum) const
+{
+	check(WarpData.IsValidIndex(ContextNum));
+
+	return WarpData[ContextNum];
+
+}
+
+UMeshComponent* FDisplayClusterWarpBlend::GetStaticMeshComponent() const
+{
+	const TSharedPtr<IDisplayClusterRender_MeshComponent, ESPMode::ThreadSafe>& MeshComponent = GeometryContext.GeometryProxy.MeshComponent;
+	if (MeshComponent.IsValid())
+	{
+		switch (MeshComponent->GetGeometrySource())
+		{
+		case EDisplayClusterRender_MeshComponentGeometrySource::StaticMeshComponentRef:
+			return MeshComponent->GetStaticMeshComponent();
+
+		case EDisplayClusterRender_MeshComponentGeometrySource::ProceduralMeshComponentRef:
+			return MeshComponent->GetProceduralMeshComponent();
+
+		default:
+			break;
+		}
+	}
+
+	return nullptr;
+}
+
+bool FDisplayClusterWarpBlend::ExportWarpMapGeometry(FDisplayClusterWarpGeometryOBJ& OutMeshData, uint32 InMaxDimension) const
+{
+#if WITH_EDITOR
+	return FDisplayClusterWarpBlendExporter_WarpMap::ExportWarpMap(GeometryContext, OutMeshData, InMaxDimension);
+#endif
+	return false;
+}
diff --git a/Engine/Plugins/Runtime/nDisplay/Source/DisplayClusterWarp/Private/WarpBlend/DisplayClusterWarpBlend.h b/Engine/Plugins/Runtime/nDisplay/Source/DisplayClusterWarp/Private/WarpBlend/DisplayClusterWarpBlend.h
new file mode 100644
--- /dev/null
+++ b/Engine/Plugins/Runtime/nDisplay/Source/DisplayClusterWarp/Private/WarpBlend/DisplayClusterWarpBlend.h
@@ -0,0 +1,97 @@
+// Copyright Epic Games, Inc. All Rights Reserved.
+
+#pragma once
+#include "IDisplayClusterWarpBlend.h"
+#include "WarpBlend/DisplayClusterWarpBlend_GeometryContext.h"
+#include "WarpBlend/Exporter/DisplayClusterWarpBlendExporter_WarpMap.h"
+#include "Render/Containers/IDisplayClusterRender_Texture.h"
+#include "Templates/SharedPointer.h"
+
+/**
+ * WarpBlend interface implementation for MPCDI and mesh projection policies
+ */
+class FDisplayClusterWarpBlend
+	: public IDisplayClusterWarpBlend
+	, public TSharedFromThis<FDisplayClusterWarpBlend, ESPMode::ThreadSafe>
+{
+public:
+	FDisplayClusterWarpBlend();
+
+	//~ IDisplayClusterWarpBlend
+	virtual ~FDisplayClusterWarpBlend() override;
+
+	virtual TSharedPtr<IDisplayClusterWarpBlend, ESPMode::ThreadSafe> ToSharedPtr() override
+	{
+		return AsShared();
+	}
+
+	virtual TSharedPtr<const IDisplayClusterWarpBlend, ESPMode::ThreadSafe> ToSharedPtr() const override
+	{
+		return AsShared();
+	}
+
+	virtual bool MarkWarpGeometryComponentDirty(const FName& InComponentName) override;
+
+	virtual bool UpdateGeometryContext(const float InWorldScale) override;
+	virtual const FDisplayClusterWarpGeometryContext& GetGeometryContext() const override;
+
+	virtual bool CalcFrustumContext(const TSharedPtr<FDisplayClusterWarpEye, ESPMode::ThreadSafe>& InWarpEye) override;
+
+	/** Get texture by type. */
+	virtual class FRHITexture* GetTexture(EDisplayClusterWarpBlendTextureType TextureType) const override
+	{
+		const TSharedPtr<IDisplayClusterRender_Texture, ESPMode::ThreadSafe> pTexture = GeometryContext.GeometryProxy.ImplGetTexture(TextureType);
+		return pTexture.IsValid() ? pTexture->GetRHITexture() : nullptr;
+	}
+
+	virtual float GetAlphaMapEmbeddedGamma() const override
+	{ 
+		return GeometryContext.GeometryProxy.AlphaMapEmbeddedGamma; 
+	}
+
+	virtual const IDisplayClusterRender_MeshComponentProxy* GetWarpMeshProxy_RenderThread() const override
+	{
+		return GeometryContext.GeometryProxy.GetWarpMeshProxy_RenderThread();
+	}
+
+	// Return current warp profile type
+	virtual EDisplayClusterWarpProfileType  GetWarpProfileType() const override
+	{
+		return GeometryContext.GeometryProxy.MPCDIAttributes.ProfileType;
+	}
+
+	virtual EDisplayClusterWarpGeometryType GetWarpGeometryType() const override
+	{
+		return GeometryContext.GeometryProxy.GeometryType;
+	}
+
+	virtual EDisplayClusterWarpFrustumGeometryType GetWarpFrustumGeometryType() const override
+	{
+		return GeometryContext.GeometryProxy.FrustumGeometryType;
+	}
+
+	virtual const FDisplayClusterWarpMPCDIAttributes& GetMPCDIAttributes() const override
+	{
+		return GeometryContext.GeometryProxy.MPCDIAttributes;
+	}
+
+	virtual bool ExportWarpMapGeometry(FDisplayClusterWarpGeometryOBJ& OutMeshData, uint32 InMaxDimension = 0) const override final;
+
+	virtual bool ShouldSupportICVFX() const override;
+
+	virtual FDisplayClusterWarpData& GetWarpData(const uint32 ContextNum) override;
+	virtual const FDisplayClusterWarpData& GetWarpData(const uint32 ContextNum) const override;
+
+	virtual UMeshComponent* GetStaticMeshComponent() const override;
+
+	//~!IDisplayClusterWarpBlend
+
+private:
+	void BeginCalcFrustum(const TSharedPtr<FDisplayClusterWarpEye, ESPMode::ThreadSafe>& InWarpEye);
+
+public:
+	FDisplayClusterWarpBlend_GeometryContext GeometryContext;
+
+private:
+	TArray<FDisplayClusterWarpData> WarpData;
+};
diff --git a/Engine/Plugins/Runtime/nDisplay/Source/DisplayClusterWarp/Private/WarpBlend/DisplayClusterWarpBlend_GeometryContext.cpp b/Engine/Plugins/Runtime/nDisplay/Source/DisplayClusterWarp/Private/WarpBlend/DisplayClusterWarpBlend_GeometryContext.cpp
new file mode 100644
--- /dev/null
+++ b/Engine/Plugins/Runtime/nDisplay/Source/DisplayClusterWarp/Private/WarpBlend/DisplayClusterWarpBlend_GeometryContext.cpp
@@ -0,0 +1,72 @@
+// Copyright Epic Games, Inc. All Rights Reserved.
+
+#include "DisplayClusterWarpBlend_GeometryContext.h"
+#include "WarpBlend/DisplayClusterWarpBlend_GeometryProxy.h"
+
+//------------------------------------------------------------------------------
+// FDisplayClusterWarpBlend_GeometryContext
+//------------------------------------------------------------------------------
+bool FDisplayClusterWarpBlend_GeometryContext::UpdateGeometryContext(const double InWorldScale)
+{
+	if (!GeometryProxy.UpdateFrustumGeometry())
+	{
+		return false;
+	}
+
+	// Update runtime data:
+	Context.GeometryToOrigin = GeometryProxy.GeometryCache.GeometryToOrigin.ToMatrixWithScale();
+
+	Context.SurfaceViewNormal = Context.GeometryToOrigin.TransformVector(GeometryProxy.GeometryCache.SurfaceViewNormal);
+	Context.SurfaceViewPlane = Context.GeometryToOrigin.TransformVector(GeometryProxy.GeometryCache.SurfaceViewPlane);
+
+	Context.AABBox.Max = Context.GeometryToOrigin.TransformPosition(GeometryProxy.GeometryCache.AABBox.Max) * InWorldScale;
+	Context.AABBox.Min = Context.GeometryToOrigin.TransformPosition(GeometryProxy.GeometryCache.AABBox.Min) * InWorldScale;
+
+	return true;
+}
+
+FMatrix FDisplayClusterWarpBlend_GeometryContext::GetTextureMatrix() const
+{
+	switch (GetWarpProfileType())
+	{
+	case EDisplayClusterWarpProfileType::warp_A3D:
+		// Fetching from a 3D source
+		return FMatrix(
+			FPlane(0.5f, 0, 0, 0),
+			FPlane(0, -0.5f, 0, 0),
+			FPlane(0, 0, 1, 0),
+			FPlane(0.5f, 0.5f, 0, 1));
+
+	case EDisplayClusterWarpProfileType::warp_2D:
+		// Fetching from a 2D source
+		return FMatrix(
+			FPlane(1, 0, 0, 0),
+			FPlane(0, -1, 0, 0),
+			FPlane(0, 0, 1, 0),
+			FPlane(0, 1, 0, 1));
+		break;
+
+	default:
+		break;
+	}
+
+	return FMatrix::Identity;
+}
+
+FMatrix FDisplayClusterWarpBlend_GeometryContext::GetRegionMatrix() const
+{
+	if (Context.bUseRegionMatrix)
+	{
+		const FVector2D& Size = GeometryProxy.MPCDIAttributes.Region.Size;
+		const FVector2D& Pos  = GeometryProxy.MPCDIAttributes.Region.Pos;
+
+		// Build Region matrix
+		return FMatrix(
+			FPlane(Size.X, 0,      0, 0),
+			FPlane(0,      Size.Y, 0, 0),
+			FPlane(0,      0,      1, 0),
+			FPlane(Pos.X,  Pos.Y,  0, 1));
+	}
+
+	return FMatrix::Identity;
+}
diff --git a/Engine/Plugins/Runtime/nDisplay/Source/DisplayClusterWarp/Private/WarpBlend/DisplayClusterWarpBlend_GeometryContext.h b/Engine/Plugins/Runtime/nDisplay/Source/DisplayClusterWarp/Private/WarpBlend/DisplayClusterWarpBlend_GeometryContext.h
new file mode 100644
--- /dev/null
+++ b/Engine/Plugins/Runtime/nDisplay/Source/DisplayClusterWarp/Private/WarpBlend/DisplayClusterWarpBlend_GeometryContext.h
@@ -0,0 +1,49 @@
+// Copyright Epic Games, Inc. All Rights Reserved.
+
+#pragma once
+#include "Containers/DisplayClusterWarpEye.h"
+#include "WarpBlend/DisplayClusterWarpBlend_GeometryProxy.h"
+#include "Containers/DisplayClusterWarpContext.h"
+
+/**
+ * Math context for Warp
+ */
+class FDisplayClusterWarpBlend_GeometryContext
+{
+public:
+	FDisplayClusterWarpBlend_GeometryContext()
+	{ }
+
+	~FDisplayClusterWarpBlend_GeometryContext()
+	{ }
+
+public:
+	/** Update internal geometry context (AABB, matrices, etc)
+	*/
+	bool UpdateGeometryContext(const double InWorldScale);
+
+	/** Get matrix for WarpMap texture. */
+	FMatrix GetTextureMatrix() const;
+
+	/** Obtaining a matrix for the area on the input texture (mpcdi 2d profile). */
+	FMatrix GetRegionMatrix() const;
+
+	/** Returns the profile type used. */
+	inline EDisplayClusterWarpProfileType  GetWarpProfileType() const
+	{
+		return GeometryProxy.MPCDIAttributes.ProfileType;
+	}
+
+	/** Assignment of the profile type to be used. */
+	inline void SetWarpProfileType(const EDisplayClusterWarpProfileType InWarpProfileType)
+	{
+		GeometryProxy.MPCDIAttributes.ProfileType = InWarpProfileType;
+	}
+
+public:
+	// Geometry container
+	FDisplayClusterWarpBlend_GeometryProxy GeometryProxy;
+
+	// Geometry context
+	FDisplayClusterWarpGeometryContext Context;
+};
diff --git a/Engine/Plugins/Runtime/nDisplay/Source/DisplayClusterWarp/Private/WarpBlend/DisplayClusterWarpBlend_GeometryProxy.cpp b/Engine/Plugins/Runtime/nDisplay/Source/DisplayClusterWarp/Private/WarpBlend/DisplayClusterWarpBlend_GeometryProxy.cpp
new file mode 100644
--- /dev/null
+++ b/Engine/Plugins/Runtime/nDisplay/Source/DisplayClusterWarp/Private/WarpBlend/DisplayClusterWarpBlend_GeometryProxy.cpp
@@ -0,0 +1,368 @@
+// Copyright Epic Games, Inc. All Rights Reserved.
+
+#include "WarpBlend/DisplayClusterWarpBlend_GeometryProxy.h"
+
+#include "Render/Containers/IDisplayClusterRender_MeshComponent.h"
+
+#include "WarpBlend/Math/DisplayClusterWarpBlendMath_WarpMesh.h"
+#include "WarpBlend/Math/DisplayClusterWarpBlendMath_WarpProceduralMesh.h"
+#include "WarpBlend/Math/DisplayClusterWarpBlendMath_WarpMap.h"
+
+#include "WarpBlend/Exporter/DisplayClusterWarpBlendExporter_WarpMap.h"
+
+/////////////////////////////////////////////////////////////////////////////////
+/// FDisplayClusterWarpBlend_GeometryProxy
+/////////////////////////////////////////////////////////////////////////////////
+bool FDisplayClusterWarpBlend_GeometryProxy::UpdateFrustumGeometry()
+{
+	switch (FrustumGeometryType)
+	{
+	case EDisplayClusterWarpFrustumGeometryType::WarpMesh:
+		return ImplUpdateFrustumGeometry_WarpMesh();
+
+	case EDisplayClusterWarpFrustumGeometryType::WarpProceduralMesh:
+		return ImplUpdateFrustumGeometry_WarpProceduralMesh();
+
+	case EDisplayClusterWarpFrustumGeometryType::WarpMap:
+		return ImplUpdateFrustumGeometry_WarpMap();
+
+	case EDisplayClusterWarpFrustumGeometryType::MPCDIAttributes:
+		return ImplUpdateFrustumGeometry_MPCDIAttributes();
+
+	default:
+		break;
+	}
+
+	bIsGeometryValid = false;
+	return false;
+}
+
+const IDisplayClusterRender_MeshComponentProxy* FDisplayClusterWarpBlend_GeometryProxy::GetWarpMeshProxy_RenderThread() const
+{
+	check(IsInRenderingThread());
+
+	switch (GeometryType)
+	{
+	case EDisplayClusterWarpGeometryType::WarpMesh:
+	case EDisplayClusterWarpGeometryType::WarpProceduralMesh:
+		return MeshComponent.IsValid() ? MeshComponent->GetMeshComponentProxy_RenderThread() : nullptr;
+
+	default:
+		break;
+	}
+
+	return nullptr;
+}
+
+bool FDisplayClusterWarpBlend_GeometryProxy::MarkWarpFrustumGeometryComponentDirty(const FName& InComponentName)
+{
+	switch (FrustumGeometryType)
+	{
+	case EDisplayClusterWarpFrustumGeometryType::WarpMesh:
+	case EDisplayClusterWarpFrustumGeometryType::WarpProceduralMesh:
+		if (MeshComponent.IsValid())
+		{
+			if (InComponentName == NAME_None || MeshComponent->EqualsMeshComponentName(InComponentName))
+			{
+				MeshComponent->MarkMeshComponentRefGeometryDirty();
+				return true;
+			}
+		}
+		break;
+	default:
+		break;
+	}
+
+	return false;
+}
+
+bool FDisplayClusterWarpBlend_GeometryProxy::ImplUpdateFrustumGeometry_MPCDIAttributes()
+{
+	bIsGeometryValid = false;
+
+	// Update caches
+	if ((bIsGeometryValid = bIsGeometryCacheValid) == false)
+	{
+		bIsGeometryValid = ImplUpdateFrustumGeometryCache_MPCDIAttributes();
+	}
+
+	GeometryCache.GeometryToOrigin = FTransform::Identity;
+
+	return bIsGeometryValid;
+}
+
+bool FDisplayClusterWarpBlend_GeometryProxy::ImplUpdateFrustumGeometry_WarpMap()
+{
+	bIsGeometryValid = false;
+
+	if (WarpMapTexture.IsValid() && WarpMapTexture->IsEnabled())
+	{
+		// Update caches
+		if ((bIsGeometryValid = bIsGeometryCacheValid) == false)
+		{
+			bIsGeometryValid = ImplUpdateFrustumGeometryCache_WarpMap();
+		}
+
+		GeometryCache.GeometryToOrigin = FTransform::Identity;
+	}
+
+	return bIsGeometryValid;
+}
+
+bool FDisplayClusterWarpBlend_GeometryProxy::ImplUpdateFrustumGeometry_WarpMesh()
+{
+	bIsGeometryValid = false;
+
+	if (!MeshComponent.IsValid())
+	{
+		return false;
+	}
+
+	UStaticMeshComponent* StaticMeshComponent = MeshComponent->GetStaticMeshComponent();
+	USceneComponent*      OriginComponent     = MeshComponent->GetOriginComponent();
+
+	const FStaticMeshLODResources* StaticMeshLODResources = (StaticMeshComponent!=nullptr) ? MeshComponent->GetStaticMeshComponentLODResources(StaticMeshComponentLODIndex) : nullptr;
+	if (StaticMeshLODResources == nullptr)
+	{
+		// mesh deleted?
+		MeshComponent->ReleaseProxyGeometry();
+		bIsMeshComponentLost = true;
+		return false;
+	};
+
+	// If StaticMesh geometry changed, update mpcdi math and RHI resources
+	if (MeshComponent->IsMeshComponentRefGeometryDirty() || bIsMeshComponentLost)
+	{
+		MeshComponent->AssignStaticMeshComponentRefs(StaticMeshComponent, WarpMeshUVs, OriginComponent, StaticMeshComponentLODIndex);
+		bIsMeshComponentLost = false;
+	}
+	
+	// Update caches
+	if ((bIsGeometryValid = bIsGeometryCacheValid) == false)
+	{
+		bIsGeometryValid = ImplUpdateFrustumGeometryCache_WarpMesh();
+	}
+
+	if (OriginComponent)
+	{
+		FMatrix MeshToWorldMatrix = StaticMeshComponent->GetComponentTransform().ToMatrixWithScale();
+		FMatrix WorldToOriginMatrix = OriginComponent->GetComponentTransform().ToInverseMatrixWithScale();
+
+		GeometryCache.GeometryToOrigin.SetFromMatrix(MeshToWorldMatrix * WorldToOriginMatrix);
+	}
+	else
+	{
+		GeometryCache.GeometryToOrigin = StaticMeshComponent->GetRelativeTransform();
+	}
+
+	return bIsGeometryValid;
+}
+
+bool FDisplayClusterWarpBlend_GeometryProxy::ImplUpdateFrustumGeometry_WarpProceduralMesh()
+{
+	bIsGeometryValid = false;
+
+	if (!MeshComponent.IsValid())
+	{
+		return false;
+	}
+
+	UProceduralMeshComponent* ProceduralMeshComponent = MeshComponent->GetProceduralMeshComponent();
+	USceneComponent*          OriginComponent         = MeshComponent->GetOriginComponent();
+
+	const FProcMeshSection* ProcMeshSection = (ProceduralMeshComponent != nullptr) ? MeshComponent->GetProceduralMeshComponentSection(ProceduralMeshComponentSectionIndex) : nullptr;
+	if (ProcMeshSection == nullptr)
+	{
+		// mesh deleted, lost or section not defined
+		MeshComponent->ReleaseProxyGeometry();
+		bIsMeshComponentLost = true;
+		return false;
+	};
+
+	// If ProceduralMesh geometry changed, update mpcdi math and RHI resources
+	if (MeshComponent->IsMeshComponentRefGeometryDirty() || bIsMeshComponentLost)
+	{
+		MeshComponent->AssignProceduralMeshComponentRefs(ProceduralMeshComponent, WarpMeshUVs, OriginComponent, ProceduralMeshComponentSectionIndex);
+		bIsMeshComponentLost = false;
+	}
+
+	// Update caches
+	if ((bIsGeometryValid = bIsGeometryCacheValid) == false)
+	{
+		bIsGeometryValid = ImplUpdateFrustumGeometryCache_WarpProceduralMesh();
+	}
+
+	// These matrices were copied from LocalPlayer.cpp.
+	// They change the coordinate system from the Unreal "Game" coordinate system to the Unreal "Render" coordinate system
+	static const FMatrix Game2Render(
+		FPlane(0, 0, 1, 0),
+		FPlane(1, 0, 0, 0),
+		FPlane(0, 1, 0, 0),
+		FPlane(0, 0, 0, 1));
+
+	static const FMatrix Render2Game = Game2Render.Inverse();
+
+	// Procedural mesh coord is render, so convert to game at first
+	FMatrix GeometryToOriginMatrix;// = Render2Game;
+
+	if (OriginComponent)
+	{
+		FMatrix MeshToWorldMatrix = ProceduralMeshComponent->GetComponentTransform().ToMatrixWithScale();
+		FMatrix WorldToOriginMatrix = OriginComponent->GetComponentTransform().ToInverseMatrixWithScale();
+
+		GeometryToOriginMatrix = MeshToWorldMatrix * WorldToOriginMatrix;
+	}
+	else
+	{
+		GeometryToOriginMatrix = ProceduralMeshComponent->GetRelativeTransform().ToMatrixWithScale();
+	}
+
+	GeometryCache.GeometryToOrigin.SetFromMatrix(GeometryToOriginMatrix);
+
+	return bIsGeometryValid;
+}
+
+bool FDisplayClusterWarpBlend_GeometryProxy::ImplUpdateFrustumGeometryCache_WarpMesh()
+{
+	bIsGeometryValid = false;
+
+	if (MeshComponent.IsValid())
+	{
+		const FStaticMeshLODResources* StaticMeshLODResources = MeshComponent->GetStaticMeshComponentLODResources(StaticMeshComponentLODIndex);
+		if (StaticMeshLODResources != nullptr)
+		{
+			FDisplayClusterWarpBlendMath_WarpMesh MeshHelper(*StaticMeshLODResources);
+
+			GeometryCache.AABBox = MeshHelper.CalcAABBox();
+			MeshHelper.CalcSurfaceVectors(GeometryCache.SurfaceViewNormal, GeometryCache.SurfaceViewPlane);
+
+			bIsGeometryCacheValid = true;
+			bIsGeometryValid = true;
+		}
+	}
+
+	return bIsGeometryValid;
+}
+
+bool FDisplayClusterWarpBlend_GeometryProxy::ImplUpdateFrustumGeometryCache_WarpProceduralMesh()
+{
+	bIsGeometryValid = false;
+
+	if (MeshComponent.IsValid())
+	{
+		const FProcMeshSection* ProcMeshSection = MeshComponent->GetProceduralMeshComponentSection(ProceduralMeshComponentSectionIndex);
+		if (ProcMeshSection != nullptr)
+		{
+			FDisplayClusterWarpBlendMath_WarpProceduralMesh ProceduralMeshHelper(*ProcMeshSection);
+
+			GeometryCache.AABBox = ProceduralMeshHelper.CalcAABBox();
+			ProceduralMeshHelper.CalcSurfaceVectors(GeometryCache.SurfaceViewNormal, GeometryCache.SurfaceViewPlane);
+
+			bIsGeometryCacheValid = true;
+			bIsGeometryValid = true;
+
+			return true;
+		}
+	}
+
+	return false;
+}
+
+bool FDisplayClusterWarpBlend_GeometryProxy::ImplUpdateFrustumGeometryCache_WarpMap()
+{
+	// WarpMap must be initialized before, outside
+	bIsGeometryValid = false;
+	
+	if (WarpMapTexture.IsValid() && WarpMapTexture->IsEnabled())
+	{
+		FDisplayClusterWarpBlendMath_WarpMap DataHelper(*(WarpMapTexture.Get()));
+
+		GeometryCache.AABBox = DataHelper.GetAABBox();
+		GeometryCache.SurfaceViewNormal = DataHelper.GetSurfaceViewNormal();
+		GeometryCache.SurfaceViewPlane  = DataHelper.GetSurfaceViewPlane();
+
+		bIsGeometryCacheValid = true;
+		bIsGeometryValid = true;
+
+		return true;
+	}
+
+	return false;
+}
+
+bool FDisplayClusterWarpBlend_GeometryProxy::ImplUpdateFrustumGeometryCache_MPCDIAttributes()
+{
+	// WarpMap must be initialized before, outside
+	bIsGeometryValid = false;
+
+	switch (MPCDIAttributes.ProfileType)
+	{
+	case EDisplayClusterWarpProfileType::warp_2D:
+		GeometryCache.SurfaceViewNormal = FVector(1, 0, 0);
+		GeometryCache.SurfaceViewPlane = FVector(1, 0, 0);
+
+		// Calc AABB for 2D profile geometry:
+		{
+			TArray<FVector> ScreenPoints;
+			FDisplayClusterWarpBlendExporter_WarpMap::Get2DProfileGeometry(MPCDIAttributes, ScreenPoints);
+
+			FDisplayClusterWarpAABB WarpAABB;
+			WarpAABB.UpdateAABB(ScreenPoints);
+
+			GeometryCache.AABBox = WarpAABB;
+		}
+		break;
+
+	default:
+		return false;
+	}
+
+	bIsGeometryCacheValid = true;
+	bIsGeometryValid = true;
+
+	return true;
+}
+
+bool FDisplayClusterWarpBlend_GeometryProxy::UpdateFrustumGeometryLOD(const FIntPoint& InSizeLOD)
+{
+	check(InSizeLOD.X > 0 && InSizeLOD.Y > 0);
+
+	GeometryCache.IndexLOD.Empty();
+
+	if (WarpMapTexture.IsValid() && WarpMapTexture->IsEnabled())
+	{
+		switch (FrustumGeometryType)
+		{
+		case EDisplayClusterWarpFrustumGeometryType::WarpMap:
+		{
+			// Generate valid points for texturebox method:
+			FDisplayClusterWarpBlendMath_WarpMap DataHelper(*(WarpMapTexture.Get()));
+			DataHelper.BuildIndexLOD(InSizeLOD.X, InSizeLOD.Y, GeometryCache.IndexLOD);
+
+			return true;
+		}
+		default:
+			break;
+		}
+	}
+
+	return false;
+}
+
+const FStaticMeshLODResources* FDisplayClusterWarpBlend_GeometryProxy::GetStaticMeshComponentLODResources() const
+{
+	return MeshComponent.IsValid() ? MeshComponent->GetStaticMeshComponentLODResources(StaticMeshComponentLODIndex) : nullptr;
+}
+
+const FProcMeshSection* FDisplayClusterWarpBlend_GeometryProxy::GetProceduralMeshComponentSection() const
+{
+	return MeshComponent.IsValid() ? MeshComponent->GetProceduralMeshComponentSection(ProceduralMeshComponentSectionIndex) : nullptr;
+}
+
+void FDisplayClusterWarpBlend_GeometryProxy::ReleaseResources()
+{
+	WarpMapTexture.Reset();
+	AlphaMapTexture.Reset();
+	BetaMapTexture.Reset();
+	MeshComponent.Reset();
+}
diff --git a/Engine/Plugins/Runtime/nDisplay/Source/DisplayClusterWarp/Private/WarpBlend/DisplayClusterWarpBlend_GeometryProxy.h b/Engine/Plugins/Runtime/nDisplay/Source/DisplayClusterWarp/Private/WarpBlend/DisplayClusterWarpBlend_GeometryProxy.h
new file mode 100644
--- /dev/null
+++ b/Engine/Plugins/Runtime/nDisplay/Source/DisplayClusterWarp/Private/WarpBlend/DisplayClusterWarpBlend_GeometryProxy.h
@@ -0,0 +1,123 @@
+// Copyright Epic Games, Inc. All Rights Reserved.
+
+#pragma once
+#include "Containers/DisplayClusterWarpContainers.h"
+#include "Render/Containers/IDisplayClusterRender_MeshComponent.h"
+#include "Render/Containers/IDisplayClusterRender_Texture.h"
+
+class IDisplayClusterRender_MeshComponentProxy;
+struct FStaticMeshLODResources;
+struct FProcMeshSection;
+
+/**
+ * Geometry container for warp
+ */
+class FDisplayClusterWarpBlend_GeometryProxy
+{
+public:
+	FDisplayClusterWarpBlend_GeometryProxy()
+	{ }
+
+	~FDisplayClusterWarpBlend_GeometryProxy()
+	{
+		ReleaseResources();
+	}
+
+public:
+	/** Get texture resource by type. */
+	const TSharedPtr<IDisplayClusterRender_Texture, ESPMode::ThreadSafe> ImplGetTexture(EDisplayClusterWarpBlendTextureType TextureType) const
+	{
+		switch (TextureType)
+		{
+		case EDisplayClusterWarpBlendTextureType::WarpMap:
+			return WarpMapTexture;
+
+		case EDisplayClusterWarpBlendTextureType::AlphaMap:
+			return AlphaMapTexture;
+
+		case EDisplayClusterWarpBlendTextureType::BetaMap:
+			return BetaMapTexture;
+
+		default:
+			break;
+		}
+		return nullptr;
+	}
+
+	bool MarkWarpFrustumGeometryComponentDirty(const FName& InComponentName);
+
+	bool UpdateFrustumGeometry();
+	bool UpdateFrustumGeometryLOD(const FIntPoint& InSizeLOD);
+
+	const IDisplayClusterRender_MeshComponentProxy* GetWarpMeshProxy_RenderThread() const;
+
+	const FStaticMeshLODResources* GetStaticMeshComponentLODResources() const;
+	const FProcMeshSection* GetProceduralMeshComponentSection() const;
+
+private:
+	bool ImplUpdateFrustumGeometry_WarpMesh();
+	bool ImplUpdateFrustumGeometry_WarpProceduralMesh();
+	bool ImplUpdateFrustumGeometry_WarpMap();
+	bool ImplUpdateFrustumGeometry_MPCDIAttributes();
+
+public:
+	bool bIsGeometryCacheValid = false;
+	bool bIsGeometryValid = false;
+
+	// Geometry source for render
+	EDisplayClusterWarpGeometryType GeometryType = EDisplayClusterWarpGeometryType::Invalid;
+
+	// Geometry source for frustum
+	EDisplayClusterWarpFrustumGeometryType FrustumGeometryType = EDisplayClusterWarpFrustumGeometryType::Invalid;
+	
+	void ReleaseResources();
+
+public:
+	// MPCDI attributes
+	FDisplayClusterWarpMPCDIAttributes MPCDIAttributes;
+
+	// Render resources:
+	TSharedPtr<IDisplayClusterRender_Texture, ESPMode::ThreadSafe> WarpMapTexture;
+	TSharedPtr<IDisplayClusterRender_Texture, ESPMode::ThreadSafe> AlphaMapTexture;
+	TSharedPtr<IDisplayClusterRender_Texture, ESPMode::ThreadSafe> BetaMapTexture;
+
+	TSharedPtr<IDisplayClusterRender_MeshComponent, ESPMode::ThreadSafe> MeshComponent;
+
+	FDisplayClusterMeshUVs WarpMeshUVs;
+
+	// for huge warp mesh geometry, change this value, and use LOD geometry in frustum math
+	int32 StaticMeshComponentLODIndex = 0;
+	int32 ProceduralMeshComponentSectionIndex = 0;
+
+	float AlphaMapEmbeddedGamma = 1.f;
+
+private:
+	bool bIsMeshComponentLost = false;
+
+private:
+	bool ImplUpdateFrustumGeometryCache_WarpMesh();
+	bool ImplUpdateFrustumGeometryCache_WarpProceduralMesh();
+	bool ImplUpdateFrustumGeometryCache_WarpMap();
+	bool ImplUpdateFrustumGeometryCache_MPCDIAttributes();
+
+	inline EDisplayClusterWarpProfileType  GetWarpProfileType() const
+	{
+		return MPCDIAttributes.ProfileType;
+	}
+
+public:
+	// Cached values for geometry (updated only if geometry changed)
+	class FDisplayClusterWarpBlend_GeometryCache
+	{
+	public:
+		FTransform GeometryToOrigin;
+
+		FBox       AABBox;
+		FVector    SurfaceViewNormal;
+		FVector    SurfaceViewPlane;
+
+		TArray<int32> IndexLOD;
+	};
+
+	FDisplayClusterWarpBlend_GeometryCache GeometryCache;
+};
diff --git a/Engine/Plugins/Runtime/nDisplay/Source/DisplayClusterWarp/Private/WarpBlend/Exporter/DisplayClusterWarpBlendExporter_WarpMap.cpp b/Engine/Plugins/Runtime/nDisplay/Source/DisplayClusterWarp/Private/WarpBlend/Exporter/DisplayClusterWarpBlendExporter_WarpMap.cpp
new file mode 100644
--- /dev/null
+++ b/Engine/Plugins/Runtime/nDisplay/Source/DisplayClusterWarp/Private/WarpBlend/Exporter/DisplayClusterWarpBlendExporter_WarpMap.cpp
@@ -0,0 +1,195 @@
+// Copyright Epic Games, Inc. All Rights Reserved.
+
+#include "DisplayClusterWarpBlendExporter_WarpMap.h"
+
+#include "WarpBlend/DisplayClusterWarpBlend_GeometryContext.h"
+#include "Blueprints/DisplayClusterWarpGeometry.h"
+
+namespace UE::DisplayClusterWarp::PFM
+{
+	// Export geometry from texture (a3 profile)
+	static bool ExportProfileA3(const TSharedPtr<IDisplayClusterRender_Texture, ESPMode::ThreadSafe>& InWarpMap, struct FDisplayClusterWarpGeometryOBJ& Dst, uint32 InMaxDimension)
+	{
+		if (!InWarpMap.IsValid() || !InWarpMap->IsEnabled())
+		{
+			return false;
+		}
+
+		const uint32 Width = InWarpMap->GetWidth();
+		const uint32 Height = InWarpMap->GetHeight();
+
+		const FVector4f* WarpData = (FVector4f*)InWarpMap->GetData();
+
+		uint32 DownScaleFactor = 1;
+
+		if (InMaxDimension > 0)
+		{
+			uint32 MeshDimension = FMath::Max(Width, Height);
+			if (MeshDimension > InMaxDimension)
+			{
+				DownScaleFactor = FMath::RoundToInt(float(MeshDimension) / InMaxDimension);
+			}
+		}
+
+
+		TMap<int32, int32> VIndexMap;
+		int32 VIndex = 0;
+
+		const uint32 MaxHeight = Height / DownScaleFactor;
+		const uint32 MaxWidth = Width / DownScaleFactor;
+
+		{
+			//Pts + Normals + UV
+			const float ScaleU = 1.0f / float(MaxWidth);
+			const float ScaleV = 1.0f / float(MaxHeight);
+
+			for (uint32 j = 0; j < MaxHeight; ++j)
+			{
+				const uint32 MeshY = (j == (MaxHeight - 1)) ? Height : (j * DownScaleFactor);
+				for (uint32 i = 0; i < MaxWidth; ++i)
+				{
+					const uint32 MeshX = (i == (MaxWidth - 1)) ? Width : (i * DownScaleFactor);
+
+					const int32 SrcIdx = MeshX + (MeshY)*Width;
+					const FVector4f& v = WarpData[SrcIdx];
+					if (v.W > 0)
+					{
+						Dst.Vertices.Add(FVector(v.X, v.Y, v.Z));
+						VIndexMap.Add(SrcIdx, VIndex++);
+
+						Dst.UV.Add(FVector2D(
+							float(i) * ScaleU,
+							float(j) * ScaleV
+						));
+
+						Dst.Normal.Add(FVector(0, 0, 0)); // Fill on face pass
+					}
+				}
+			}
+		}
+
+		{
+			//faces
+			for (uint32 j = 0; j < MaxHeight - 1; ++j)
+			{
+				const uint32 MeshY = (j == (MaxHeight - 1)) ? Height : (j * DownScaleFactor);
+				const uint32 NextMeshY = ((j + 1) == (MaxHeight - 1)) ? Height : ((j + 1) * DownScaleFactor);
+
+				for (uint32 i = 0; i < MaxWidth - 1; ++i)
+				{
+					const uint32 MeshX = (i == (MaxWidth - 1)) ? Width : (i * DownScaleFactor);
+					const uint32 NextMeshX = ((i + 1) == (MaxWidth - 1)) ? Width : ((i + 1) * DownScaleFactor);
+
+					int32 idx[4];
+
+					idx[0] = (MeshX + MeshY * Width);
+					idx[1] = (NextMeshX + MeshY * Width);
+					idx[2] = (MeshX + NextMeshY * Width);
+					idx[3] = (NextMeshX + NextMeshY * Width);
+
+					for (int32 a = 0; a < 4; a++)
+					{
+						if (VIndexMap.Contains(idx[a]))
+						{
+							idx[a] = VIndexMap[idx[a]];
+						}
+						else
+						{
+							idx[a] = -1;
+						}
+					}
+
+					if (idx[0] >= 0 && idx[2] >= 0 && idx[3] >= 0)
+					{
+						Dst.PostAddFace(idx[0], idx[2], idx[3]);
+					}
+
+					if (idx[3] >= 0 && idx[1] >= 0 && idx[0] >= 0)
+					{
+						Dst.PostAddFace(idx[3], idx[1], idx[0]);
+					}
+				}
+			}
+		}
+
+		return true;
+	};
+};
+using namespace UE::DisplayClusterWarp::PFM;
+
+//---------------------------------------------------------------
+// FDisplayClusterWarpBlendExporter_WarpMap
+//---------------------------------------------------------------
+float FDisplayClusterWarpBlendExporter_WarpMap::Get2DProfilePixelsToUnitScale()
+{
+	return 0.1f; // 10 pixels = 1 centimeter
+}
+
+void FDisplayClusterWarpBlendExporter_WarpMap::Get2DProfileGeometry(const FDisplayClusterWarpMPCDIAttributes& InMPCDIAttributes, TArray<FVector>& OutGeometryPoints, TArray<FVector>* OutNormal, TArray<FVector2D>* OutUV)
+{
+	const float PixelsToUnitScale = FDisplayClusterWarpBlendExporter_WarpMap::Get2DProfilePixelsToUnitScale();
+
+	// export each region as a plain
+	const float Width = InMPCDIAttributes.Buffer.Resolution.X * PixelsToUnitScale;
+	const float Height = InMPCDIAttributes.Buffer.Resolution.Y * PixelsToUnitScale;
+
+	// Rectangle stands on the floor, centered at zero
+	const float X0 = Width * (InMPCDIAttributes.Region.Pos.X) - (Width * 0.5f);
+	const float X1 = Width * (InMPCDIAttributes.Region.Pos.X + InMPCDIAttributes.Region.Size.X) - (Width * 0.5f);
+	const float Y0 = Height * (InMPCDIAttributes.Region.Pos.Y);
+	const float Y1 = Height * (InMPCDIAttributes.Region.Pos.Y + InMPCDIAttributes.Region.Size.Y);
+
+	// Create a vertices
+	// Y - right, Z - Up
+	OutGeometryPoints.Add(FVector(0, X0, Y1));
+	OutGeometryPoints.Add(FVector(0, X1, Y1));
+	OutGeometryPoints.Add(FVector(0, X0, Y0));
+	OutGeometryPoints.Add(FVector(0, X1, Y0));
+
+	if (OutNormal)
+	{
+		OutNormal->Add(FVector(-1, 0, 0)); // Fill on face pass
+		OutNormal->Add(FVector(-1, 0, 0)); // Fill on face pass
+		OutNormal->Add(FVector(-1, 0, 0)); // Fill on face pass
+		OutNormal->Add(FVector(-1, 0, 0)); // Fill on face pass
+	}
+
+	if (OutUV)
+	{
+		OutUV->Add(FVector2D(0, 0));
+		OutUV->Add(FVector2D(1, 0));
+		OutUV->Add(FVector2D(0, 1));
+		OutUV->Add(FVector2D(1, 1));
+	}
+}
+
+bool FDisplayClusterWarpBlendExporter_WarpMap::ExportWarpMap(const FDisplayClusterWarpBlend_GeometryContext& InContext, struct FDisplayClusterWarpGeometryOBJ& Dst, uint32 InMaxDimension)
+{
+	switch (InContext.GetWarpProfileType())
+	{
+	case EDisplayClusterWarpProfileType::warp_2D:
+		FDisplayClusterWarpBlendExporter_WarpMap::Get2DProfileGeometry(InContext.GeometryProxy.MPCDIAttributes, Dst.Vertices, &Dst.Normal, &Dst.UV);
+
+		//Create a faces:
+		Dst.PostAddFace(0, 2, 3);
+		Dst.PostAddFace(3, 1, 0);
+
+		return true;
+
+	case EDisplayClusterWarpProfileType::warp_3D:
+		// NOT IMPLEMENTED
+		return false;
+
+	case EDisplayClusterWarpProfileType::warp_A3D:
+		return InContext.GeometryProxy.WarpMapTexture.IsValid() ? ExportProfileA3(InContext.GeometryProxy.WarpMapTexture, Dst, InMaxDimension) : false;
+
+	case EDisplayClusterWarpProfileType::warp_SL:
+		// NOT IMPLEMENTED
+		return false;
+
+	default:
+		break;
+	}
+
+	return false;
+}
diff --git a/Engine/Plugins/Runtime/nDisplay/Source/DisplayClusterWarp/Private/WarpBlend/Exporter/DisplayClusterWarpBlendExporter_WarpMap.h b/Engine/Plugins/Runtime/nDisplay/Source/DisplayClusterWarp/Private/WarpBlend/Exporter/DisplayClusterWarpBlendExporter_WarpMap.h
new file mode 100644
--- /dev/null
+++ b/Engine/Plugins/Runtime/nDisplay/Source/DisplayClusterWarp/Private/WarpBlend/Exporter/DisplayClusterWarpBlendExporter_WarpMap.h
@@ -0,0 +1,22 @@
+// Copyright Epic Games, Inc. All Rights Reserved.
+
+#pragma once
+
+#include "CoreMinimal.h"
+
+/** Auxiliary class for exporting WarpMap geometry. */
+class FDisplayClusterWarpBlendExporter_WarpMap
+{
+public:
+	/**
+	 * Exprot WarpMap texture geometry into the memory structure.
+	 * For a 2D profile the geometry does not exist and is created based on the definition of <buffer> and <region>.
+	 */
+	static bool ExportWarpMap(const class FDisplayClusterWarpBlend_GeometryContext& InContext, struct FDisplayClusterWarpGeometryOBJ& Dst, uint32 InMaxDimension = 0);
+
+	/** Return pixels to Unit scale for fake geometry of MPCDI profile 2D. */
+	static float Get2DProfilePixelsToUnitScale();
+
+	/** [2D] Return fake geometry points for region in mpcdi 2D profile. */
+	static void Get2DProfileGeometry(const struct FDisplayClusterWarpMPCDIAttributes& InMPCDIAttributes, TArray<FVector>& OutGeometryPoints, TArray<FVector>* OutNormal = nullptr, TArray<FVector2D>* OutUV = nullptr);
+};
diff --git a/Engine/Plugins/Runtime/nDisplay/Source/DisplayClusterWarp/Private/WarpBlend/Loader/DisplayClusterWarpBlendLoader_MPCDI.cpp b/Engine/Plugins/Runtime/nDisplay/Source/DisplayClusterWarp/Private/WarpBlend/Loader/DisplayClusterWarpBlendLoader_MPCDI.cpp
new file mode 100644
--- /dev/null
+++ b/Engine/Plugins/Runtime/nDisplay/Source/DisplayClusterWarp/Private/WarpBlend/Loader/DisplayClusterWarpBlendLoader_MPCDI.cpp
@@ -0,0 +1,97 @@
+// Copyright Epic Games, Inc. All Rights Reserved.
+
+#include "DisplayClusterWarpBlendLoader_MPCDI.h"
+#include "DisplayClusterWarpBlendLoader_MeshComponent.h"
+
+#include "WarpBlend/DisplayClusterWarpBlend.h"
+#include "WarpBlend/DisplayClusterWarpBlend_GeometryContext.h"
+#include "WarpBlend/DisplayClusterWarpBlend_GeometryProxy.h"
+
+#include "WarpBlend/Loader/DisplayClusterWarpBlendLoader_MPCDIFileLoader.h"
+#include "WarpBlend/Loader/DisplayClusterWarpBlendLoader_MPCDIRegionLoader.h"
+
+#include "DisplayClusterWarpLog.h"
+
+TSharedPtr<FDisplayClusterWarpBlend, ESPMode::ThreadSafe> FDisplayClusterWarpBlendLoader_MPCDI::Create(const FDisplayClusterWarpInitializer_MPCDIFile& InConstructParameters)
+{
+	FDisplayClusterWarpBlendMPCDIRegionLoader MPCDIRegionLoader;
+	if(MPCDIRegionLoader.LoadRegionFromMPCDIFile(InConstructParameters.MPCDIFileName, InConstructParameters.BufferId, InConstructParameters.RegionId))
+	{
+		//ok, Create and initialize warpblend interface
+		return MPCDIRegionLoader.CreateWarpBlendInterface();
+	}
+
+	UE_LOG(LogDisplayClusterWarpBlend, Error, TEXT("Can't load region '%s' buffer '%s' from mpcdi file '%s'"), *InConstructParameters.RegionId, *InConstructParameters.BufferId, *InConstructParameters.MPCDIFileName);
+
+	return nullptr;
+}
+
+bool FDisplayClusterWarpBlendLoader_MPCDI::ReadMPCDFileStructure(const FString& InMPCDIFileName, TMap<FString, TMap<FString, FDisplayClusterWarpMPCDIAttributes>>& OutMPCDIFileStructure)
+{
+	TSharedPtr<FDisplayClusterWarpBlendMPCDIFileLoader, ESPMode::ThreadSafe> FileLoader = FDisplayClusterWarpBlendMPCDIFileLoader::GetOrCreateCachedMPCDILoader(InMPCDIFileName);
+	if (FileLoader.IsValid())
+	{
+		return FileLoader->ReadMPCDFileStructure(OutMPCDIFileStructure);
+	}
+
+	UE_LOG(LogDisplayClusterWarpBlend, Error, TEXT("Can't load MPCDI file structure from mpcdi file '%s'"), *InMPCDIFileName);
+
+	return false;
+}
+
+TSharedPtr<FDisplayClusterWarpBlend, ESPMode::ThreadSafe> FDisplayClusterWarpBlendLoader_MPCDI::Create(const FDisplayClusterWarpInitializer_PFMFile& InConstructParameters)
+{
+	FDisplayClusterWarpBlendMPCDIRegionLoader RegionLoader;
+
+	// Filling all internal data manually:
+	RegionLoader.AlphaMapGammaEmbedded = InConstructParameters.AlphaMapEmbeddedAlpha;
+
+	RegionLoader.WarpMap = RegionLoader.CreateTextureFromPFMFile(InConstructParameters.PFMFileName, InConstructParameters.PFMScale, InConstructParameters.bIsUnrealGameSpace);
+	if (RegionLoader.WarpMap.IsValid())
+	{
+		if (InConstructParameters.AlphaMapFileName.IsEmpty() == false)
+		{
+			RegionLoader.AlphaMap = RegionLoader.CreateTextureFromFile(InConstructParameters.AlphaMapFileName);
+		}
+
+		if (InConstructParameters.BetaMapFileName.IsEmpty() == false)
+		{
+			RegionLoader.BetaMap = RegionLoader.CreateTextureFromFile(InConstructParameters.BetaMapFileName);
+		}
+
+		// Copy MPCDI attributes
+		RegionLoader.MPCDIAttributes = InConstructParameters.MPCDIAttributes;
+
+		// After all internal data are filled in, the warpblend interface is created and initialized.
+		return RegionLoader.CreateWarpBlendInterface();
+	}
+
+	UE_LOG(LogDisplayClusterWarpBlend, Error, TEXT("Can't load PFM file '%s'"), *InConstructParameters.PFMFileName);
+
+	return nullptr;
+}
+
+TSharedPtr<FDisplayClusterWarpBlend, ESPMode::ThreadSafe> FDisplayClusterWarpBlendLoader_MPCDI::Create(const FDisplayClusterWarpInitializer_MPCDIFile_Profile2DScreen& InConstructParameters)
+{
+	// Create from file
+	const FDisplayClusterWarpInitializer_MPCDIFile& InitializerFromFile = InConstructParameters;
+	TSharedPtr<FDisplayClusterWarpBlend, ESPMode::ThreadSafe> WarpBlend = Create(InitializerFromFile);
+
+	if (WarpBlend.IsValid() && WarpBlend->GetWarpProfileType() == EDisplayClusterWarpProfileType::warp_2D)
+	{
+		// Extra setup for 2D profile:
+		FDisplayClusterWarpBlend_GeometryProxy& Proxy = WarpBlend->GeometryContext.GeometryProxy;
+
+		FDisplayClusterMeshUVs MeshUVs;
+
+		Proxy.MeshComponent = FDisplayClusterWarpBlendLoader_MeshComponent::CreateMeshComponent();
+		Proxy.MeshComponent->AssignStaticMeshComponentRefs(InConstructParameters.StaticMeshComponent, MeshUVs, InConstructParameters.OriginComponent, 0);
+
+		Proxy.WarpMeshUVs = MeshUVs;
+
+		// Use static mesh geometry for frustum calculations
+		Proxy.FrustumGeometryType = EDisplayClusterWarpFrustumGeometryType::WarpMesh;
+	}
+
+	return WarpBlend;
+}
diff --git a/Engine/Plugins/Runtime/nDisplay/Source/DisplayClusterWarp/Private/WarpBlend/Loader/DisplayClusterWarpBlendLoader_MPCDI.h b/Engine/Plugins/Runtime/nDisplay/Source/DisplayClusterWarp/Private/WarpBlend/Loader/DisplayClusterWarpBlendLoader_MPCDI.h
new file mode 100644
--- /dev/null
+++ b/Engine/Plugins/Runtime/nDisplay/Source/DisplayClusterWarp/Private/WarpBlend/Loader/DisplayClusterWarpBlendLoader_MPCDI.h
@@ -0,0 +1,31 @@
+// Copyright Epic Games, Inc. All Rights Reserved.
+
+#pragma once
+#include "Containers/DisplayClusterWarpInitializer.h"
+#include "WarpBlend/DisplayClusterWarpBlend.h"
+
+/**
+ * An auxiliary class for creating the WarpBlend interface from initializing structures associated with MPCDI.
+ */
+class FDisplayClusterWarpBlendLoader_MPCDI
+{
+public:
+	/** Create WarpBlend interface from the MPCDI file. */
+	static TSharedPtr<FDisplayClusterWarpBlend, ESPMode::ThreadSafe> Create(const FDisplayClusterWarpInitializer_MPCDIFile& InConstructParameters);
+
+	/** Create WarpBlend interface from the PFM file. */
+	static TSharedPtr<FDisplayClusterWarpBlend, ESPMode::ThreadSafe> Create(const FDisplayClusterWarpInitializer_PFMFile& InConstructParameters);
+
+	/** Create WarpBlend interface from MPCDI file and DisplayClusterScreen component. */
+	static TSharedPtr<FDisplayClusterWarpBlend, ESPMode::ThreadSafe> Create(const FDisplayClusterWarpInitializer_MPCDIFile_Profile2DScreen& InConstructParameters);
+
+	/**
+	 * Read MPCDI file structure
+	 *
+	 * @param InMPCDIFileName        - (in)  filename of the MPCDI file
+	 * @param OutMPCDIFileStructure  - (out) readed mpcdi structure in format: TMap<BufferId, TArray<RegionId, MPCDIAttributes>>
+	 *
+	 * @return true if success
+	 */
+	static bool ReadMPCDFileStructure(const FString& InMPCDIFileName, TMap<FString, TMap<FString, FDisplayClusterWarpMPCDIAttributes>>& OutMPCDIFileStructure);
+};
diff --git a/Engine/Plugins/Runtime/nDisplay/Source/DisplayClusterWarp/Private/WarpBlend/Loader/DisplayClusterWarpBlendLoader_MPCDIFileLoader.cpp b/Engine/Plugins/Runtime/nDisplay/Source/DisplayClusterWarp/Private/WarpBlend/Loader/DisplayClusterWarpBlendLoader_MPCDIFileLoader.cpp
new file mode 100644
--- /dev/null
+++ b/Engine/Plugins/Runtime/nDisplay/Source/DisplayClusterWarp/Private/WarpBlend/Loader/DisplayClusterWarpBlendLoader_MPCDIFileLoader.cpp
@@ -0,0 +1,279 @@
+// Copyright Epic Games, Inc. All Rights Reserved.
+
+#include "DisplayClusterWarpBlendLoader_MPCDIFileLoader.h"
+
+THIRD_PARTY_INCLUDES_START
+#include "mpcdiProfile.h"
+#include "mpcdiReader.h"
+#include "mpcdiDisplay.h"
+#include "mpcdiBuffer.h"
+#include "mpcdiRegion.h"
+#include "mpcdiAlphaMap.h"
+#include "mpcdiBetaMap.h"
+#include "mpcdiDistortionMap.h"
+#include "mpcdiGeometryWarpFile.h"
+
+#include "IO/mpcdiPfmIO.h"
+#include "mpcdiPNGReadWrite.h"
+THIRD_PARTY_INCLUDES_END
+
+#include "DisplayClusterWarpLog.h"
+
+#include "Misc/DisplayClusterDataCache.h"
+#include "Misc/DisplayClusterHelpers.h"
+#include "Misc/FileHelper.h"
+
+
+int32 GDisplayClusterRender_MPCDILoaderCacheEnable = 1;
+static FAutoConsoleVariableRef CDisplayClusterRender_MPCDILoaderCacheEnable(
+	TEXT("nDisplay.cache.MPCDILoader.enable"),
+	GDisplayClusterRender_MPCDILoaderCacheEnable,
+	TEXT("Enables the use of the MPCDI file loaders cache.\n"),
+	ECVF_Default
+);
+
+int32 GDisplayClusterRender_MPCDILoaderCacheTimeOutInFrames = 5 * 60 * 60; // Timeout is 5 minutes (for 60 frames per second)
+static FAutoConsoleVariableRef CVarDisplayClusterRender_MPCDILoaderCacheTimeOutInFrames(
+	TEXT("nDisplay.cache.MPCDILoader.TimeOut"),
+	GDisplayClusterRender_MPCDILoaderCacheTimeOutInFrames,
+	TEXT("The timeout value in frames  for cached MPCDI file loaders.\n")
+	TEXT("-1 - disable timeout.\n"),
+	ECVF_Default
+);
+
+//--------------------------------------------------------------------------
+// FDisplayClusterWarpBlendMPCDIFileLoader
+//--------------------------------------------------------------------------
+FDisplayClusterWarpBlendMPCDIFileLoader::FDisplayClusterWarpBlendMPCDIFileLoader(const FString& InUniqueName, const FString& InFilePath)
+	: UniqueName(InUniqueName)
+{
+	// Load MPCDI file
+	{
+		Profile = mpcdi::Profile::CreateProfile();
+		if (!Profile)
+		{
+			return;
+		}
+
+		// Read profile data from mpcdi file using mpcdi lib:
+		if (mpcdi::Reader* NewReader = mpcdi::Reader::CreateReader())
+		{
+			UE_LOG(LogDisplayClusterWarpBlend, Log, TEXT("Loading MPCDI file %s."), *InFilePath);
+
+			std::string version = NewReader->GetSupportedVersions();
+			FString Version = version.c_str();
+			UE_LOG(LogDisplayClusterWarpBlend, Verbose, TEXT("MPCDI library version: %s"), *Version);
+
+			mpcdi::MPCDI_Error mpcdi_err = NewReader->Read(TCHAR_TO_ANSI(*InFilePath), Profile);
+			delete NewReader;
+
+			if (MPCDI_FAILED(mpcdi_err))
+			{
+				UE_LOG(LogDisplayClusterWarpBlend, Error, TEXT("Error %d reading MPCDI file"), int32(mpcdi_err));
+				Release();
+			}
+		}
+	}
+}
+
+FDisplayClusterWarpBlendMPCDIFileLoader::~FDisplayClusterWarpBlendMPCDIFileLoader()
+{
+	Release();
+}
+
+void FDisplayClusterWarpBlendMPCDIFileLoader::Release()
+{
+	if (Profile)
+	{
+		delete Profile;
+		Profile = nullptr;
+	}
+}
+
+bool FDisplayClusterWarpBlendMPCDIFileLoader::FindRegion(const FString& InBufferName, const FString& InRegionId, mpcdi::Buffer*& OutBufferData, mpcdi::Region*& OutRegionData) const
+{
+	if (IsValid())
+	{
+		// Find desired region:
+		for (mpcdi::Display::BufferIterator itBuffer = Profile->GetDisplay()->GetBufferBegin(); itBuffer != Profile->GetDisplay()->GetBufferEnd(); ++itBuffer)
+		{
+			mpcdi::Buffer* mpcdiBuffer = itBuffer->second;
+
+			FString BufferId(mpcdiBuffer->GetId().c_str());
+			if (InBufferName.Equals(BufferId, ESearchCase::IgnoreCase))
+			{
+				for (mpcdi::Buffer::RegionIterator it = mpcdiBuffer->GetRegionBegin(); it != mpcdiBuffer->GetRegionEnd(); ++it)
+				{
+					mpcdi::Region* mpcdiRegion = it->second;
+					FString RegionId(mpcdiRegion->GetId().c_str());
+
+					if (InRegionId.Equals(RegionId, ESearchCase::IgnoreCase))
+					{
+						OutBufferData = mpcdiBuffer;
+						OutRegionData = mpcdiRegion;
+
+						return true;
+					}
+				}
+			}
+		}
+	}
+
+	return false;
+}
+
+bool FDisplayClusterWarpBlendMPCDIFileLoader::ReadMPCDFileStructure(TMap<FString, TMap<FString, FDisplayClusterWarpMPCDIAttributes>>& OutMPCDIFileStructure)
+{
+	if (IsValid())
+	{
+		// Find desired region:
+		for (mpcdi::Display::BufferIterator itBuffer = Profile->GetDisplay()->GetBufferBegin(); itBuffer != Profile->GetDisplay()->GetBufferEnd(); ++itBuffer)
+		{
+			mpcdi::Buffer* mpcdiBuffer = itBuffer->second;
+
+			FString BufferId(mpcdiBuffer->GetId().c_str());
+			OutMPCDIFileStructure.Add(BufferId);
+			{
+				for (mpcdi::Buffer::RegionIterator it = mpcdiBuffer->GetRegionBegin(); it != mpcdiBuffer->GetRegionEnd(); ++it)
+				{
+					mpcdi::Region* mpcdiRegion = it->second;
+					FString RegionId(mpcdiRegion->GetId().c_str());
+
+					FDisplayClusterWarpMPCDIAttributes MPCDIAttributes;
+					LoadMPCDIAttributes(mpcdiBuffer, mpcdiRegion, MPCDIAttributes);
+
+					OutMPCDIFileStructure[BufferId].Emplace(RegionId, MPCDIAttributes);
+				}
+			}
+		}
+
+		return true;
+	}
+
+	return false;
+}
+
+void FDisplayClusterWarpBlendMPCDIFileLoader::LoadMPCDIAttributes(mpcdi::Buffer* mpcdiBuffer, mpcdi::Region* mpcdiRegion, FDisplayClusterWarpMPCDIAttributes& MPCDIAttributes) const
+{
+	// Read MPCDI attributes:
+	MPCDIAttributes.ProfileType = GetProfileType();
+
+	MPCDIAttributes.Buffer.Resolution.X = mpcdiBuffer->GetXresolution();
+	MPCDIAttributes.Buffer.Resolution.Y = mpcdiBuffer->GetYresolution();
+
+	MPCDIAttributes.Region.Resolution.X = mpcdiRegion->GetXresolution();
+	MPCDIAttributes.Region.Resolution.Y = mpcdiRegion->GetYresolution();
+
+	MPCDIAttributes.Region.Pos.X = mpcdiRegion->GetX();
+	MPCDIAttributes.Region.Pos.Y = mpcdiRegion->GetY();
+	MPCDIAttributes.Region.Size.X = mpcdiRegion->GetXsize();
+	MPCDIAttributes.Region.Size.Y = mpcdiRegion->GetYsize();
+
+	if (const mpcdi::Frustum* Frustum = mpcdiRegion->GetFrustum())
+	{
+		EnumAddFlags(MPCDIAttributes.Flags, EDisplayClusterWarpMPCDIAttributesFlags::HasFrustum);
+
+		MPCDIAttributes.Frustum.Rotator.Pitch = Frustum->GetPitch();
+		MPCDIAttributes.Frustum.Rotator.Yaw = Frustum->GetYaw();
+		MPCDIAttributes.Frustum.Rotator.Roll = Frustum->GetRoll();
+
+		// Frustum angles XYZW = LRTB
+		MPCDIAttributes.Frustum.Angles.X = Frustum->GetLeftAngle();
+		MPCDIAttributes.Frustum.Angles.Y = Frustum->GetRightAngle();
+		MPCDIAttributes.Frustum.Angles.Z = Frustum->GetUpAngle();
+		MPCDIAttributes.Frustum.Angles.W = Frustum->GetDownAngle();
+	}
+
+	if (mpcdi::CoordinateFrame* CoordinateFrame = mpcdiRegion->GetCoordinateFrame())
+	{
+		EnumAddFlags(MPCDIAttributes.Flags, EDisplayClusterWarpMPCDIAttributesFlags::HasCoordinateFrame);
+
+		MPCDIAttributes.CoordinateFrame.Pos = FVector(CoordinateFrame->GetPosx(), CoordinateFrame->GetPosy(), CoordinateFrame->GetPosz());
+		MPCDIAttributes.CoordinateFrame.Yaw = FVector(CoordinateFrame->GetYawx(), CoordinateFrame->GetYawy(), CoordinateFrame->GetYawz());
+		MPCDIAttributes.CoordinateFrame.Pitch = FVector(CoordinateFrame->GetPitchx(), CoordinateFrame->GetPitchy(), CoordinateFrame->GetPitchz());
+		MPCDIAttributes.CoordinateFrame.Roll = FVector(CoordinateFrame->GetRollx(), CoordinateFrame->GetRolly(), CoordinateFrame->GetRollz());
+	}
+}
+
+EDisplayClusterWarpProfileType FDisplayClusterWarpBlendMPCDIFileLoader::GetProfileType() const
+{
+	if (Profile)
+	{
+		switch (Profile->GetProfileType())
+		{
+		case mpcdi::ProfileType2d: return EDisplayClusterWarpProfileType::warp_2D;
+		case mpcdi::ProfileType3d: return EDisplayClusterWarpProfileType::warp_3D;
+		case mpcdi::ProfileTypea3: return EDisplayClusterWarpProfileType::warp_A3D;
+		case mpcdi::ProfileTypesl: return EDisplayClusterWarpProfileType::warp_SL;
+		default:
+			// Invalid profile type
+			break;
+		}
+	}
+
+	return EDisplayClusterWarpProfileType::Invalid;
+}
+
+FString FDisplayClusterWarpBlendMPCDIFileLoader::GetFullPathToFile(const FString& InFilePath)
+{
+	// Fix relative paths
+	FString OutFileName = InFilePath;
+	if (FPaths::IsRelative(InFilePath))
+	{
+		OutFileName = DisplayClusterHelpers::filesystem::GetFullPathForConfigResource(InFilePath);
+	}
+
+	if (!FPaths::FileExists(OutFileName))
+	{
+		UE_LOG(LogDisplayClusterWarpBlend, Error, TEXT("file '%s' not found"), *OutFileName);
+
+		OutFileName.Empty();
+	}
+
+	return OutFileName;
+}
+
+int32 FDisplayClusterWarpBlendMPCDIFileLoader::GetDataCacheTimeOutInFrames()
+{
+	return FMath::Max(0, GDisplayClusterRender_MPCDILoaderCacheTimeOutInFrames);
+}
+
+bool FDisplayClusterWarpBlendMPCDIFileLoader::IsDataCacheEnabled()
+{
+	return GDisplayClusterRender_MPCDILoaderCacheEnable != 0;
+}
+
+/**
+ * The cache for MPCDI file loader objects. (Singleton)
+ */
+class FDisplayClusterWarpBlendMPCDIFileLoaderCache
+	: public TDisplayClusterDataCache<FDisplayClusterWarpBlendMPCDIFileLoader>
+{
+public:
+	static TSharedPtr<FDisplayClusterWarpBlendMPCDIFileLoader, ESPMode::ThreadSafe> GetOrCreateMPCDIFileLoader(const FString& InFilePath)
+	{
+		static FDisplayClusterWarpBlendMPCDIFileLoaderCache MPCDIFileLoaderCacheSingleton;
+
+		const FString FullPath = FDisplayClusterWarpBlendMPCDIFileLoader::GetFullPathToFile(InFilePath);
+		if (FullPath.IsEmpty())
+		{
+			return nullptr;
+		}
+
+		const FString UniqueName = HashString(FullPath);
+
+		TSharedPtr<FDisplayClusterWarpBlendMPCDIFileLoader, ESPMode::ThreadSafe> TextureRef = MPCDIFileLoaderCacheSingleton.Find(UniqueName);
+		if (!TextureRef.IsValid())
+		{
+			TextureRef = MakeShared<FDisplayClusterWarpBlendMPCDIFileLoader, ESPMode::ThreadSafe>(UniqueName, FullPath);
+			MPCDIFileLoaderCacheSingleton.Add(TextureRef);
+		}
+
+		return TextureRef;
+	}
+};
+
+TSharedPtr<FDisplayClusterWarpBlendMPCDIFileLoader, ESPMode::ThreadSafe> FDisplayClusterWarpBlendMPCDIFileLoader::GetOrCreateCachedMPCDILoader(const FString& InFilePath)
+{
+	return FDisplayClusterWarpBlendMPCDIFileLoaderCache::GetOrCreateMPCDIFileLoader(InFilePath);
+}
diff --git a/Engine/Plugins/Runtime/nDisplay/Source/DisplayClusterWarp/Private/WarpBlend/Loader/DisplayClusterWarpBlendLoader_MPCDIFileLoader.h b/Engine/Plugins/Runtime/nDisplay/Source/DisplayClusterWarp/Private/WarpBlend/Loader/DisplayClusterWarpBlendLoader_MPCDIFileLoader.h
new file mode 100644
--- /dev/null
+++ b/Engine/Plugins/Runtime/nDisplay/Source/DisplayClusterWarp/Private/WarpBlend/Loader/DisplayClusterWarpBlendLoader_MPCDIFileLoader.h
@@ -0,0 +1,97 @@
+// Copyright Epic Games, Inc. All Rights Reserved.
+
+#pragma once
+#include "Containers/DisplayClusterWarpEnums.h"
+#include "Containers/DisplayClusterWarpContainers.h"
+
+namespace mpcdi
+{
+	struct Profile;
+	struct Buffer;
+	struct Region;
+};
+
+/**
+ * Loader and container for mpcdi file.
+ */
+class FDisplayClusterWarpBlendMPCDIFileLoader
+{
+public:
+	FDisplayClusterWarpBlendMPCDIFileLoader(const FString& InUniqueName, const FString& InFilePath);
+	~FDisplayClusterWarpBlendMPCDIFileLoader();
+
+public:
+	/** Get unique name of this file loader. Generated from file name. */
+	const FString& GetName() const
+	{
+		return UniqueName;
+	}
+	
+	/** Find MPCDI buffer and region by name.
+	 * @param InBufferName - buffer name from the xml file inside the mpcdi file
+	 * @param InRegionId - region name under the buffer name from the xml file inside the mpcdi file
+	 * @param OutBufferData - (out) loaded buffer data
+	 * @param OutRegionData - (out) loaded region data
+	 * 
+	 * @return false if no buffer or region exists
+	 */
+	bool FindRegion(const FString& InBufferName, const FString& InRegionId, mpcdi::Buffer*& OutBufferData, mpcdi::Region*& OutRegionData) const;
+
+	/**
+	 * Read MPCDI file structure
+	 *
+	 * @param OutMPCDIFileStructure  - (out) readed mpcdi structure in format: TMap<BufferId, TArray<RegionId, MPCDIAttributes>>
+	 *
+	 * @return true if success
+	 */
+	bool ReadMPCDFileStructure(TMap<FString, TMap<FString, FDisplayClusterWarpMPCDIAttributes>>& OutMPCDIFileStructure);
+	
+	/** True if this data valid. */
+	bool IsValid() const
+	{
+		return Profile != nullptr;
+	}
+
+	/**
+	 * Load attributes
+	 */
+	void LoadMPCDIAttributes(mpcdi::Buffer* mpcdiBuffer, mpcdi::Region* mpcdiRegion, FDisplayClusterWarpMPCDIAttributes& OutMPCDIAttributes) const;
+
+	/** Get mpcdi profile type. */
+	EDisplayClusterWarpProfileType GetProfileType() const;
+
+	/** Convert full file path from relative path to mpcdi files. */
+	static FString GetFullPathToFile(const FString& InFilePath);
+
+	/** Get or create an MPCDI loader in the cache. */
+	static TSharedPtr<FDisplayClusterWarpBlendMPCDIFileLoader, ESPMode::ThreadSafe> GetOrCreateCachedMPCDILoader(const FString& InFilePath);
+
+
+	//~ Begin TDisplayClusterDataCache
+	/** Return DataCache timeout in frames. */
+	static int32 GetDataCacheTimeOutInFrames();
+
+	/** Return true if DataCache is enabled. */
+	static bool IsDataCacheEnabled();
+
+	/** Returns the unique name of this MPCDI file resource for DataCache. */
+	inline const FString& GetDataCacheName() const
+	{
+		return UniqueName;
+	}
+
+	/** Method for releasing a cached data item, called before its destructor. */
+	inline void ReleaseDataCacheItem()
+	{ }
+	// ~~ End TDisplayClusterDataCache
+
+protected:
+	void Release();
+
+private:
+	// Unique name generated from the filename.
+	const FString UniqueName;
+
+	// loaded mpcdi data.
+	mpcdi::Profile* Profile = nullptr;
+};
diff --git a/Engine/Plugins/Runtime/nDisplay/Source/DisplayClusterWarp/Private/WarpBlend/Loader/DisplayClusterWarpBlendLoader_MPCDIRegionLoader.cpp b/Engine/Plugins/Runtime/nDisplay/Source/DisplayClusterWarp/Private/WarpBlend/Loader/DisplayClusterWarpBlendLoader_MPCDIRegionLoader.cpp
new file mode 100644
--- /dev/null
+++ b/Engine/Plugins/Runtime/nDisplay/Source/DisplayClusterWarp/Private/WarpBlend/Loader/DisplayClusterWarpBlendLoader_MPCDIRegionLoader.cpp
@@ -0,0 +1,236 @@
+// Copyright Epic Games, Inc. All Rights Reserved.
+
+#include "DisplayClusterWarpBlendLoader_MPCDIRegionLoader.h"
+
+#include "WarpBlend/DisplayClusterWarpBlend.h"
+#include "WarpBlend/DisplayClusterWarpBlend_GeometryContext.h"
+#include "WarpBlend/DisplayClusterWarpBlend_GeometryProxy.h"
+
+#include "WarpBlend/Loader/DisplayClusterWarpBlendLoader_MPCDIFileLoader.h"
+#include "WarpBlend/Loader/DisplayClusterWarpBlendLoader_WarpMap.h"
+
+#include "DisplayClusterWarpLog.h"
+
+#include "IDisplayCluster.h"
+#include "Render/IDisplayClusterRenderManager.h"
+
+THIRD_PARTY_INCLUDES_START
+#include "mpcdiProfile.h"
+#include "mpcdiReader.h"
+#include "mpcdiDisplay.h"
+#include "mpcdiBuffer.h"
+#include "mpcdiRegion.h"
+#include "mpcdiAlphaMap.h"
+#include "mpcdiBetaMap.h"
+#include "mpcdiDistortionMap.h"
+#include "mpcdiGeometryWarpFile.h"
+
+#include "IO/mpcdiPfmIO.h"
+#include "mpcdiPNGReadWrite.h"
+THIRD_PARTY_INCLUDES_END
+
+namespace UE::DisplayClusterWarp::MPCDIRegionLoader
+{
+	static inline TSharedPtr<IDisplayClusterRender_Texture, ESPMode::ThreadSafe> GetOrCreateCachedTextureImpl(const FString& InUniqueTextureName)
+	{
+		static IDisplayCluster& DisplayClusterAPI = IDisplayCluster::Get();
+
+		return DisplayClusterAPI.GetRenderMgr()->GetOrCreateCachedTexture(InUniqueTextureName);
+	}
+
+	static inline void InitializeTextureFromDataMapImpl(TSharedPtr<IDisplayClusterRender_Texture, ESPMode::ThreadSafe>& InTexture, mpcdi::DataMap* InDataMap)
+	{
+		check(InDataMap);
+
+		if (InTexture.IsValid() && !InTexture->IsEnabled())
+		{
+			// new resource, initialize once
+			const int32 ComponentDepth = InDataMap->GetComponentDepth();
+			const int32 BitDepth = InDataMap->GetBitDepth();
+			const uint32_t Width = InDataMap->GetSizeX();
+			const uint32_t Height = InDataMap->GetSizeY();
+			const void* TextureData = reinterpret_cast<void*>(InDataMap->GetData()->data());
+			const bool bHasCPUAccess = false;
+
+			InTexture->CreateTexture(TextureData, ComponentDepth, BitDepth, Width, Height, bHasCPUAccess);
+		}
+	}
+};
+using namespace UE::DisplayClusterWarp::MPCDIRegionLoader;
+
+//--------------------------------------------------------------
+// FDisplayClusterWarpBlendMPCDIRegionLoader
+//--------------------------------------------------------------
+TSharedPtr<FDisplayClusterWarpBlend, ESPMode::ThreadSafe> FDisplayClusterWarpBlendMPCDIRegionLoader::CreateWarpBlendInterface()
+{
+	TSharedPtr<FDisplayClusterWarpBlend, ESPMode::ThreadSafe> WarpBlend = MakeShared<FDisplayClusterWarpBlend, ESPMode::ThreadSafe>();
+
+	WarpBlend->GeometryContext.GeometryProxy.GeometryType = EDisplayClusterWarpGeometryType::WarpMap;
+
+	// Frustum geometry source
+	switch (MPCDIAttributes.ProfileType)
+	{
+	case EDisplayClusterWarpProfileType::warp_2D:
+	case EDisplayClusterWarpProfileType::warp_3D:
+		// Get frustum from attributes
+		WarpBlend->GeometryContext.GeometryProxy.FrustumGeometryType = EDisplayClusterWarpFrustumGeometryType::MPCDIAttributes;
+		break;
+
+	default:
+		WarpBlend->GeometryContext.GeometryProxy.FrustumGeometryType = EDisplayClusterWarpFrustumGeometryType::WarpMap;
+		break;
+	}
+
+	WarpBlend->GeometryContext.GeometryProxy.MPCDIAttributes = MPCDIAttributes;
+
+	FDisplayClusterWarpBlend_GeometryProxy& Proxy = WarpBlend->GeometryContext.GeometryProxy;
+
+	Proxy.WarpMapTexture = WarpMap;
+
+	Proxy.AlphaMapEmbeddedGamma = AlphaMapGammaEmbedded;
+	Proxy.AlphaMapTexture = AlphaMap;
+	Proxy.BetaMapTexture = BetaMap;
+
+	ReleaseResources();
+
+	return WarpBlend;
+}
+
+void FDisplayClusterWarpBlendMPCDIRegionLoader::ReleaseResources()
+{
+	AlphaMap.Reset();
+	BetaMap.Reset();
+	WarpMap.Reset();
+}
+
+TSharedPtr<IDisplayClusterRender_Texture, ESPMode::ThreadSafe> FDisplayClusterWarpBlendMPCDIRegionLoader::CreateTextureFromDataMap(const FString& InUniqueTextureName, mpcdi::DataMap* InDataMap)
+{
+	check(InDataMap);
+
+	TSharedPtr<IDisplayClusterRender_Texture, ESPMode::ThreadSafe> Texture = GetOrCreateCachedTextureImpl(InUniqueTextureName);
+	InitializeTextureFromDataMapImpl(Texture, InDataMap);
+
+	return  Texture.IsValid() && Texture->IsEnabled() ? Texture : nullptr;
+}
+
+TSharedPtr<IDisplayClusterRender_Texture, ESPMode::ThreadSafe> FDisplayClusterWarpBlendMPCDIRegionLoader::CreateTextureFromGeometry(const FString& InUniqueTextureName, mpcdi::GeometryWarpFile* InGeometryWarp)
+{
+	TSharedPtr<IDisplayClusterRender_Texture, ESPMode::ThreadSafe> Texture = GetOrCreateCachedTextureImpl(InUniqueTextureName);
+	if (Texture && !Texture->IsEnabled())
+	{
+		// new resource, initialize once
+		FDisplayClusterWarpBlendLoader WarpMapData;
+		if (WarpMapData.LoadFromGeometryWarpFile(GetWarpProfileType(), InGeometryWarp))
+		{
+			Texture->CreateTexture(WarpMapData.GetWarpData(), 4, 32, WarpMapData.GetWidth(), WarpMapData.GetHeight(), true);
+		}
+	}
+
+	return  Texture.IsValid() && Texture->IsEnabled() ? Texture : nullptr;
+}
+
+TSharedPtr<IDisplayClusterRender_Texture, ESPMode::ThreadSafe> FDisplayClusterWarpBlendMPCDIRegionLoader::CreateTextureFromFile(const FString& InFileName)
+{
+	// new resource, initialize once
+	const FString TextureFileFullPath = FDisplayClusterWarpBlendMPCDIFileLoader::GetFullPathToFile(InFileName);
+	if (TextureFileFullPath.IsEmpty())
+	{
+		return nullptr;
+	}
+
+	TSharedPtr<IDisplayClusterRender_Texture, ESPMode::ThreadSafe> Texture = GetOrCreateCachedTextureImpl(TextureFileFullPath);
+	if (Texture && !Texture->IsEnabled())
+	{
+		// Currently, only the PNG file format is supported.
+		const std::string PNGFileName = TCHAR_TO_ANSI(*TextureFileFullPath);
+		mpcdi::DataMap* PngData;
+		mpcdi::MPCDI_Error res = mpcdi::PNGReadWrite::Read(PNGFileName, PngData);
+		if (mpcdi::MPCDI_SUCCESS == res && PngData)
+		{
+			InitializeTextureFromDataMapImpl(Texture, PngData);
+		}
+	}
+
+	return  Texture.IsValid() && Texture->IsEnabled() ? Texture : nullptr;
+}
+
+TSharedPtr<IDisplayClusterRender_Texture, ESPMode::ThreadSafe> FDisplayClusterWarpBlendMPCDIRegionLoader::CreateTextureFromPFMFile(const FString& InPFMFile, float PFMScale, bool bIsUnrealGameSpace)
+{
+	// new resource, initialize once
+	const FString PFMFileFullPath = FDisplayClusterWarpBlendMPCDIFileLoader::GetFullPathToFile(InPFMFile);
+	if (PFMFileFullPath.IsEmpty())
+	{
+		return nullptr;
+	}
+
+	TSharedPtr<IDisplayClusterRender_Texture, ESPMode::ThreadSafe> Texture = GetOrCreateCachedTextureImpl(PFMFileFullPath);
+	if (Texture && !Texture->IsEnabled())
+	{
+		bool bResult = false;
+		const std::string FileName = TCHAR_TO_ANSI(*PFMFileFullPath);
+		mpcdi::PFM* PFMData;
+		mpcdi::MPCDI_Error res = mpcdi::PfmIO::Read(FileName, PFMData);
+		if (mpcdi::MPCDI_SUCCESS == res && PFMData)
+		{
+			FDisplayClusterWarpBlendLoader WarpMapData;
+			if (WarpMapData.LoadFromPFM(GetWarpProfileType(), PFMData, PFMScale, bIsUnrealGameSpace))
+			{
+				Texture->CreateTexture(WarpMapData.GetWarpData(), 4, 32, WarpMapData.GetWidth(), WarpMapData.GetHeight(), true);
+			}
+		}
+	}
+
+	return  Texture.IsValid() && Texture->IsEnabled() ? Texture : nullptr;
+}
+
+bool FDisplayClusterWarpBlendMPCDIRegionLoader::LoadRegionFromMPCDIFile(const FString& InFilePath, const FString& InBufferName, const FString& InRegionName)
+{
+	TSharedPtr<FDisplayClusterWarpBlendMPCDIFileLoader, ESPMode::ThreadSafe> FileLoader = FDisplayClusterWarpBlendMPCDIFileLoader::GetOrCreateCachedMPCDILoader(InFilePath);
+	if (FileLoader.IsValid())
+	{
+		mpcdi::Buffer* mpcdiBuffer = nullptr;
+		mpcdi::Region* mpcdiRegion = nullptr;
+		if (FileLoader->FindRegion(InBufferName, InRegionName, mpcdiBuffer, mpcdiRegion))
+		{
+			// Load MPCDI attributes
+			FileLoader->LoadMPCDIAttributes(mpcdiBuffer, mpcdiRegion, MPCDIAttributes);
+
+			// Load and create region resources:
+			const FString UniqueName = FString::Printf(TEXT("%s.%s.%s"), *FileLoader->GetName(), *InBufferName, *InRegionName);
+
+			// ALpha Map
+			mpcdi::AlphaMap* AlphaMapSource = mpcdiRegion->GetFileSet()->GetAlphaMap();
+			if (AlphaMapSource != nullptr)
+			{
+				AlphaMap = CreateTextureFromDataMap(FString::Printf(TEXT("%s.AlphaMap"), *UniqueName), AlphaMapSource);
+				AlphaMapGammaEmbedded = AlphaMapSource->GetGammaEmbedded();
+			}
+
+			// Beta Map
+			mpcdi::BetaMap* BetaMapSource = mpcdiRegion->GetFileSet()->GetBetaMap();
+			if (BetaMapSource != nullptr)
+			{
+				BetaMap = CreateTextureFromDataMap(FString::Printf(TEXT("%s.BetaMap"), *UniqueName), BetaMapSource);
+			}
+
+			// Warp Map
+			if (GetWarpProfileType() != EDisplayClusterWarpProfileType::warp_SL)
+			{
+				mpcdi::GeometryWarpFile* WarpMapSource = mpcdiRegion->GetFileSet()->GetGeometryWarpFile();
+				if (WarpMapSource)
+				{
+					WarpMap = CreateTextureFromGeometry(FString::Printf(TEXT("%s.WarpMap"), *UniqueName), WarpMapSource);
+					if (WarpMap.IsValid())
+					{
+						return true;
+					}
+				}
+			}
+
+			// Release unused resources
+			ReleaseResources();
+		}
+	}
+
+	return false;
+}
diff --git a/Engine/Plugins/Runtime/nDisplay/Source/DisplayClusterWarp/Private/WarpBlend/Loader/DisplayClusterWarpBlendLoader_MPCDIRegionLoader.h b/Engine/Plugins/Runtime/nDisplay/Source/DisplayClusterWarp/Private/WarpBlend/Loader/DisplayClusterWarpBlendLoader_MPCDIRegionLoader.h
new file mode 100644
--- /dev/null
+++ b/Engine/Plugins/Runtime/nDisplay/Source/DisplayClusterWarp/Private/WarpBlend/Loader/DisplayClusterWarpBlendLoader_MPCDIRegionLoader.h
@@ -0,0 +1,75 @@
+// Copyright Epic Games, Inc. All Rights Reserved.
+
+#pragma once
+#include "Containers/DisplayClusterWarpContainers.h"
+#include "Render/Containers/IDisplayClusterRender_Texture.h"
+#include "WarpBlend/DisplayClusterWarpBlend.h"
+
+namespace mpcdi
+{
+	struct GeometryWarpFile;
+	struct DataMap;
+};
+
+/**
+ * MPCDI region loader.
+ * Creates the IDisplayClusterWarpBlend interface.
+ */
+class  FDisplayClusterWarpBlendMPCDIRegionLoader
+{
+public:
+	FDisplayClusterWarpBlendMPCDIRegionLoader() = default;
+	~FDisplayClusterWarpBlendMPCDIRegionLoader()
+	{
+		ReleaseResources();
+	}
+
+public:
+	/** Creates a WarpMap texture from the DataMap. */
+	TSharedPtr<IDisplayClusterRender_Texture, ESPMode::ThreadSafe> CreateTextureFromDataMap(const FString& InUniqueTextureName, mpcdi::DataMap* InDataMap);
+
+	/** Creates a WarpMap texture from the GeometryWarp. */
+	TSharedPtr<IDisplayClusterRender_Texture, ESPMode::ThreadSafe> CreateTextureFromGeometry(const FString& InUniqueTextureName, mpcdi::GeometryWarpFile* InGeometryWarp);
+
+	/** Creates a WarpMap texture from the PFM file. */
+	TSharedPtr<IDisplayClusterRender_Texture, ESPMode::ThreadSafe> CreateTextureFromPFMFile(const FString& InPFMFile, float PFMScale, bool bIsUnrealGameSpace);
+
+	/** Creates a texture from the file.
+	 * Supported image formats: PNG.
+	 */
+	TSharedPtr<IDisplayClusterRender_Texture, ESPMode::ThreadSafe> CreateTextureFromFile(const FString& InFileName);
+
+	/**
+	 * Load region data from the MPCDI file into the internal data of this class.
+	 */
+	bool LoadRegionFromMPCDIFile(const FString& InFilePath, const FString& InBufferName, const FString& InRegionName);
+
+	/**
+	 * Creates the WarpBlend interface from the internal data of this class.
+	 * This function should be called at the end when all internal data is initialized.
+	 */
+	TSharedPtr<FDisplayClusterWarpBlend, ESPMode::ThreadSafe> CreateWarpBlendInterface();
+
+private:
+	/** Release internal resources. */
+	void ReleaseResources();
+
+	inline EDisplayClusterWarpProfileType  GetWarpProfileType() const
+	{
+		return MPCDIAttributes.ProfileType;
+	}
+
+public:
+	// mpcdi attributes from xml
+	FDisplayClusterWarpMPCDIAttributes MPCDIAttributes;
+
+	// WarpMap texture (required)
+	TSharedPtr<IDisplayClusterRender_Texture, ESPMode::ThreadSafe> WarpMap;
+
+	// ALphaMap texture
+	TSharedPtr<IDisplayClusterRender_Texture, ESPMode::ThreadSafe> AlphaMap;
+	float AlphaMapGammaEmbedded = 1.f;
+
+	// BetaMap texture
+	TSharedPtr<IDisplayClusterRender_Texture, ESPMode::ThreadSafe> BetaMap;
+};
diff --git a/Engine/Plugins/Runtime/nDisplay/Source/DisplayClusterWarp/Private/WarpBlend/Loader/DisplayClusterWarpBlendLoader_MeshComponent.cpp b/Engine/Plugins/Runtime/nDisplay/Source/DisplayClusterWarp/Private/WarpBlend/Loader/DisplayClusterWarpBlendLoader_MeshComponent.cpp
new file mode 100644
--- /dev/null
+++ b/Engine/Plugins/Runtime/nDisplay/Source/DisplayClusterWarp/Private/WarpBlend/Loader/DisplayClusterWarpBlendLoader_MeshComponent.cpp
@@ -0,0 +1,87 @@
+// Copyright Epic Games, Inc. All Rights Reserved.
+
+#include "DisplayClusterWarpBlendLoader_MeshComponent.h"
+
+#include "Containers/DisplayClusterWarpInitializer.h"
+#include "WarpBlend/DisplayClusterWarpBlend.h"
+#include "WarpBlend/DisplayClusterWarpBlend_GeometryContext.h"
+#include "WarpBlend/DisplayClusterWarpBlend_GeometryProxy.h"
+
+#include "DisplayClusterWarpLog.h"
+
+#include "Render/IDisplayClusterRenderManager.h"
+
+#include "Misc/DisplayClusterHelpers.h"
+#include "Misc/FileHelper.h"
+#include "Misc/ConfigCacheIni.h"
+#include "Misc/Paths.h"
+
+#include "Stats/Stats.h"
+#include "Engine/Engine.h"
+
+//---------------------------------------------------------------------
+// FDisplayClusterWarpBlendLoader_MeshComponent
+//---------------------------------------------------------------------
+TSharedPtr<IDisplayClusterRender_MeshComponent, ESPMode::ThreadSafe> FDisplayClusterWarpBlendLoader_MeshComponent::CreateMeshComponent()
+{
+	static IDisplayCluster& DisplayClusterAPI = IDisplayCluster::Get();
+
+	return DisplayClusterAPI.GetRenderMgr()->CreateMeshComponent();
+}
+
+TSharedPtr<FDisplayClusterWarpBlend, ESPMode::ThreadSafe> FDisplayClusterWarpBlendLoader_MeshComponent::Create(const FDisplayClusterWarpInitializer_StaticMesh& InConstructParameters)
+{
+	if (InConstructParameters.StaticMeshComponent != nullptr)
+	{
+		//ok, Create and initialize warpblend interface:
+		TSharedPtr<FDisplayClusterWarpBlend, ESPMode::ThreadSafe> WarpBlend = MakeShared<FDisplayClusterWarpBlend, ESPMode::ThreadSafe>();
+
+		WarpBlend->GeometryContext.SetWarpProfileType(EDisplayClusterWarpProfileType::warp_A3D);
+
+		FDisplayClusterWarpBlend_GeometryProxy& Proxy = WarpBlend->GeometryContext.GeometryProxy;
+
+		FDisplayClusterMeshUVs MeshUVs(InConstructParameters.BaseUVIndex, InConstructParameters.ChromakeyUVIndex);
+
+		Proxy.MeshComponent = CreateMeshComponent();
+		Proxy.MeshComponent->AssignStaticMeshComponentRefs(InConstructParameters.StaticMeshComponent, MeshUVs, InConstructParameters.OriginComponent, InConstructParameters.StaticMeshComponentLODIndex);
+
+		Proxy.StaticMeshComponentLODIndex = InConstructParameters.StaticMeshComponentLODIndex;
+		Proxy.WarpMeshUVs = MeshUVs;
+
+		Proxy.GeometryType = EDisplayClusterWarpGeometryType::WarpMesh;
+		Proxy.FrustumGeometryType = EDisplayClusterWarpFrustumGeometryType::WarpMesh;
+
+		return WarpBlend;
+	}
+
+	return nullptr;
+}
+
+TSharedPtr<FDisplayClusterWarpBlend, ESPMode::ThreadSafe> FDisplayClusterWarpBlendLoader_MeshComponent::Create(const FDisplayClusterWarpInitializer_ProceduralMesh& InConstructParameters)
+{
+	if (InConstructParameters.ProceduralMeshComponent != nullptr)
+	{
+		//ok, Create and initialize warpblend interface:
+		TSharedPtr<FDisplayClusterWarpBlend, ESPMode::ThreadSafe> WarpBlend = MakeShared<FDisplayClusterWarpBlend, ESPMode::ThreadSafe>();
+
+		WarpBlend->GeometryContext.SetWarpProfileType(EDisplayClusterWarpProfileType::warp_A3D);
+
+		FDisplayClusterWarpBlend_GeometryProxy& Proxy = WarpBlend->GeometryContext.GeometryProxy;
+
+		FDisplayClusterMeshUVs MeshUVs(InConstructParameters.BaseUVIndex, InConstructParameters.ChromakeyUVIndex);
+
+		// Assign procedural mesh
+		Proxy.MeshComponent = CreateMeshComponent();
+		Proxy.MeshComponent->AssignProceduralMeshComponentRefs(InConstructParameters.ProceduralMeshComponent, MeshUVs, InConstructParameters.OriginComponent, InConstructParameters.ProceduralMeshComponentSectionIndex);
+
+		Proxy.ProceduralMeshComponentSectionIndex = InConstructParameters.ProceduralMeshComponentSectionIndex;
+		Proxy.WarpMeshUVs = MeshUVs;
+
+		Proxy.GeometryType = EDisplayClusterWarpGeometryType::WarpProceduralMesh;
+		Proxy.FrustumGeometryType = EDisplayClusterWarpFrustumGeometryType::WarpProceduralMesh;
+
+		return WarpBlend;
+	}
+
+	return nullptr;
+}
diff --git a/Engine/Plugins/Runtime/nDisplay/Source/DisplayClusterWarp/Private/WarpBlend/Loader/DisplayClusterWarpBlendLoader_MeshComponent.h b/Engine/Plugins/Runtime/nDisplay/Source/DisplayClusterWarp/Private/WarpBlend/Loader/DisplayClusterWarpBlendLoader_MeshComponent.h
new file mode 100644
--- /dev/null
+++ b/Engine/Plugins/Runtime/nDisplay/Source/DisplayClusterWarp/Private/WarpBlend/Loader/DisplayClusterWarpBlendLoader_MeshComponent.h
@@ -0,0 +1,23 @@
+// Copyright Epic Games, Inc. All Rights Reserved.
+
+#pragma once
+#include "Containers/DisplayClusterWarpInitializer.h"
+
+class IDisplayClusterRender_MeshComponent;
+class FDisplayClusterWarpBlend;
+
+/**
+ * An auxiliary class for creating the WarpBlend interface from initializing structures associated with UE components.
+ */
+class FDisplayClusterWarpBlendLoader_MeshComponent
+{
+public:
+	/** Create mesh component instance. */
+	static TSharedPtr<IDisplayClusterRender_MeshComponent, ESPMode::ThreadSafe> CreateMeshComponent();
+
+	/** Create WarpBlend interface from the StaticMesh component. */
+	static TSharedPtr<FDisplayClusterWarpBlend, ESPMode::ThreadSafe> Create(const FDisplayClusterWarpInitializer_StaticMesh& InConstructParameters);
+
+	/** Create WarpBlend interface from the ProceduralMesh component. */
+	static TSharedPtr<FDisplayClusterWarpBlend, ESPMode::ThreadSafe> Create(const FDisplayClusterWarpInitializer_ProceduralMesh& InConstructParameters);
+};
diff --git a/Engine/Plugins/Runtime/nDisplay/Source/DisplayClusterWarp/Private/WarpBlend/Loader/DisplayClusterWarpBlendLoader_WarpMap.cpp b/Engine/Plugins/Runtime/nDisplay/Source/DisplayClusterWarp/Private/WarpBlend/Loader/DisplayClusterWarpBlendLoader_WarpMap.cpp
new file mode 100644
--- /dev/null
+++ b/Engine/Plugins/Runtime/nDisplay/Source/DisplayClusterWarp/Private/WarpBlend/Loader/DisplayClusterWarpBlendLoader_WarpMap.cpp
@@ -0,0 +1,364 @@
+// Copyright Epic Games, Inc. All Rights Reserved.
+
+#include "DisplayClusterWarpBlendLoader_WarpMap.h"
+#include "DisplayClusterWarpLog.h"
+
+THIRD_PARTY_INCLUDES_START
+#include "mpcdiProfile.h"
+#include "mpcdiReader.h"
+#include "mpcdiDisplay.h"
+#include "mpcdiBuffer.h"
+#include "mpcdiRegion.h"
+#include "mpcdiAlphaMap.h"
+#include "mpcdiBetaMap.h"
+#include "mpcdiDistortionMap.h"
+#include "mpcdiGeometryWarpFile.h"
+THIRD_PARTY_INCLUDES_END
+
+#include "RHIGlobals.h"
+
+namespace UE::DisplayClusterWarp::PFMHelpers
+{
+	static constexpr float kEpsilon = 0.00001f;
+
+	static inline float GetUnitToCentimeter(mpcdi::GeometricUnit GeomUnit)
+	{
+		switch (GeomUnit)
+		{
+		case mpcdi::GeometricUnitmm:
+			return 1.f / 10.f;
+		case mpcdi::GeometricUnitcm:
+			return 1.f;
+		case mpcdi::GeometricUnitdm:
+			return 10.f;
+		case mpcdi::GeometricUnitm:
+			return 100.f;
+		case mpcdi::GeometricUnitin:
+			return 2.54f;
+		case mpcdi::GeometricUnitft:
+			return 30.48f;
+		case mpcdi::GeometricUnityd:
+			return 91.44f;
+		case mpcdi::GeometricUnitunkown:
+			return 1.f;
+		default:
+			check(false);
+			return 1.f;
+		}
+	}
+
+	static inline FMatrix GetMPCDI2UEConventionMatrix(const EDisplayClusterWarpProfileType InProfileType, const float UnitToCentemeter)
+	{
+		switch (InProfileType)
+		{
+		case EDisplayClusterWarpProfileType::warp_A3D:
+		case EDisplayClusterWarpProfileType::warp_SL:
+			// Unreal is in cm, so we need to convert to cm.
+			// Convert from MPCDI convention to Unreal convention
+			// MPCDI is Right Handed (Y is up, X is left, Z is in the screen)
+			// Unreal is Left Handed (Z is up, X in the screen, Y is right)
+			return FMatrix(
+				FPlane(0.f, UnitToCentemeter, 0.f, 0.f),
+				FPlane(0.f, 0.f, UnitToCentemeter, 0.f),
+				FPlane(-UnitToCentemeter, 0.f, 0.f, 0.f),
+				FPlane(0.f, 0.f, 0.f, 1.f));
+
+		default:
+			break;
+		};
+
+		return FMatrix::Identity;
+	}
+
+	static inline FMatrix GetUE2UEConventionMatrix(const float WorldScale)
+	{
+		return FMatrix(
+			FPlane(WorldScale, 0.f, 0.f, 0.f),
+			FPlane(0.f, WorldScale, 0.f, 0.f),
+			FPlane(0.f, 0.f, WorldScale, 0.f),
+			FPlane(0.f, 0.f, 0.f, 1.f));
+	}
+
+	static inline FVector ImportMPCDINode(const EDisplayClusterWarpProfileType InProfileType, const mpcdi::NODE& InNode)
+	{
+		switch (InProfileType)
+		{
+		case EDisplayClusterWarpProfileType::warp_A3D:
+		case EDisplayClusterWarpProfileType::warp_SL:
+			return FVector(InNode.r, InNode.g, InNode.b);
+
+		default:
+			break;
+		}
+
+		// ignore Z value for 2d profiles
+		return FVector(InNode.r, InNode.g, 0);
+	}
+
+	// This is just a copy of the old rules:
+	//     if ((!(FMath::Abs(t.X) < kEpsilon && FMath::Abs(t.Y) < kEpsilon && FMath::Abs(t.Z) < kEpsilon))
+	//      && (!FMath::IsNaN(t.X) && !FMath::IsNaN(t.Y) && !FMath::IsNaN(t.Z)))
+	// 
+	// Todo: improve it if necessary
+	static inline bool IsMPCDINodeValid(const FVector& InNode)
+	{
+		if (FMath::Abs(InNode.X) < kEpsilon && FMath::Abs(InNode.Y) < kEpsilon && FMath::Abs(InNode.Z) < kEpsilon)
+		{
+			return false;
+		}
+
+		if(FMath::IsNaN(InNode.X) || FMath::IsNaN(InNode.Y) || FMath::IsNaN(InNode.Z))
+		{
+			return false;
+		}
+
+		return true;
+	}
+};
+using namespace UE::DisplayClusterWarp::PFMHelpers;
+
+//----------------------------------------------------------------------
+// FDisplayClusterWarpBlendLoader
+//----------------------------------------------------------------------
+bool FDisplayClusterWarpBlendLoader::LoadFromGeometryWarpFile(const EDisplayClusterWarpProfileType InProfileType, mpcdi::GeometryWarpFile* SourceWarpMap)
+{
+	check(SourceWarpMap);
+
+	if (!InitializeWarpDataImpl(SourceWarpMap->GetSizeX(), SourceWarpMap->GetSizeY()))
+	{
+		return false;
+	}
+
+	const FMatrix ConventionMatrix = GetMPCDI2UEConventionMatrix(InProfileType, GetUnitToCentimeter(SourceWarpMap->GetGeometricUnit()));
+	for (int32 WarpMapY = 0; WarpMapY < Height; ++WarpMapY)
+	{
+		for (int32 WarpMapX = 0; WarpMapX < Width; ++WarpMapX)
+		{
+			FVector4f& OutPts = WarpData[WarpMapX + WarpMapY * Width];
+
+			FVector InputPts = ImportMPCDINode(InProfileType, (*SourceWarpMap)(WarpMapX, WarpMapY));
+			if (IsMPCDINodeValid(InputPts))
+			{
+				const FVector ScaledPts = ConventionMatrix.TransformPosition(InputPts);
+
+				OutPts.X = ScaledPts.X;
+				OutPts.Y = ScaledPts.Y;
+				OutPts.Z = ScaledPts.Z;
+				OutPts.W = 1;
+			}
+			else
+			{
+				OutPts = FVector4f(0.f, 0.f, 0.f, -1.f);
+			}
+		}
+	}
+
+	// Remove noise from warp mesh (small areas less than 3*3 quads)
+	ClearNoise(InProfileType);
+
+	return true;
+}
+
+bool FDisplayClusterWarpBlendLoader::LoadFromPoint(const EDisplayClusterWarpProfileType InProfileType, const TArray<FVector>& InPoints, int32 WarpX, int32 WarpY, float WorldScale, bool bIsUnrealGameSpace)
+{
+	if (!InitializeWarpDataImpl(WarpX, WarpY))
+	{
+		return false;
+	}
+
+	LoadGeometryImpl(InProfileType, InPoints, WorldScale, bIsUnrealGameSpace);
+
+	return true;
+}
+
+bool FDisplayClusterWarpBlendLoader::LoadFromPFM(const EDisplayClusterWarpProfileType InProfileType, mpcdi::PFM* SourcePFM, float PFMScale, bool bIsUnrealGameSpace)
+{
+	check(SourcePFM);
+
+	if (!InitializeWarpDataImpl(SourcePFM->GetSizeX(), SourcePFM->GetSizeY()))
+	{
+		return false;
+	}
+
+	TArray<FVector> WarpMeshPoints;
+	WarpMeshPoints.Reserve(Width * Height);
+
+	for (int32 WarpMapY = 0; WarpMapY < Height; ++WarpMapY)
+	{
+		for (int32 WarpMapX = 0; WarpMapX < Width; ++WarpMapX)
+		{
+			WarpMeshPoints.Add(ImportMPCDINode(InProfileType, SourcePFM->operator()(WarpMapX, WarpMapY)));
+		}
+	}
+
+	LoadGeometryImpl(InProfileType, WarpMeshPoints, PFMScale, bIsUnrealGameSpace);
+
+	return true;
+}
+
+bool FDisplayClusterWarpBlendLoader::InitializeWarpDataImpl(int32 InWidth, int32 InHeight)
+{
+	if (FMath::Min(InWidth, InHeight) <= 1 || FMath::Max(InWidth, InHeight) >= GMaxTextureDimensions)
+	{
+		UE_LOG(LogDisplayClusterWarpBlend, Error, TEXT("Invalid PFM warpmap data size '%d x %d'"), InWidth, InHeight);
+
+		return false;
+	}
+
+	Width = InWidth;
+	Height = InHeight;
+	WarpData = new FVector4f[Width * Height];
+
+	return true;
+}
+
+FDisplayClusterWarpBlendLoader::~FDisplayClusterWarpBlendLoader()
+{
+	if (WarpData != nullptr)
+	{
+		delete WarpData;
+		WarpData = nullptr;
+	}
+}
+
+void FDisplayClusterWarpBlendLoader::LoadGeometryImpl(EDisplayClusterWarpProfileType ProfileType, const TArray<FVector>& InPoints, float WorldScale, bool bIsUnrealGameSpace)
+{
+	check(InPoints.Num() == (Width * Height));
+
+	FMatrix ConventionMatrix = bIsUnrealGameSpace ? GetUE2UEConventionMatrix(WorldScale) : GetMPCDI2UEConventionMatrix(ProfileType, WorldScale);
+
+	FVector4f* WarpDataPointIt = WarpData;
+	for (const FVector& PointIt : InPoints)
+	{
+		const FVector& InPts = PointIt;
+
+		FVector4f& OutPts = *WarpDataPointIt;
+		WarpDataPointIt++;
+
+		if(IsMPCDINodeValid(InPts))
+		{
+			const FVector ScaledPts = ConventionMatrix.TransformPosition(InPts);
+
+			OutPts.X = ScaledPts.X;
+			OutPts.Y = ScaledPts.Y;
+			OutPts.Z = ScaledPts.Z;
+			OutPts.W = 1;
+		}
+		else
+		{
+			OutPts = FVector4f(0.f, 0.f, 0.f, -1.f);
+		}
+	}
+}
+
+bool FDisplayClusterWarpBlendLoader::Is3DPointValid(int32 X, int32 Y) const
+{
+	if (WarpData && X >= 0 && X < Width && Y >= 0 && Y < Height)
+	{
+		const FVector4f& Point = WarpData[X + Y * Width];
+
+		return Point.W > 0;
+	}
+
+	return false;
+}
+
+void FDisplayClusterWarpBlendLoader::ClearNoise(const EDisplayClusterWarpProfileType InProfileType)
+{
+	switch (InProfileType)
+	{
+	case EDisplayClusterWarpProfileType::warp_A3D:
+	case EDisplayClusterWarpProfileType::warp_SL:
+		// only for 3d geometry
+		break;
+
+	default:
+		return;
+	}
+
+	// Todo: Replace these constants with CVAR if necessary
+	const FIntPoint SearchXYDepth(3, 3);
+	const FIntPoint AllowedXYDepthRules(2, 3);
+	const int32 MinGridSize = 10;
+	const int32 MaxLoopsNum = 50;
+
+	if (Width > MinGridSize && Height > MinGridSize)
+	{
+		//Remove noise for large warp mesh
+		int32 MaxLoops = MaxLoopsNum;
+		while (MaxLoops-- > 0)
+		{
+			if (!ClearNoiseImpl(SearchXYDepth, AllowedXYDepthRules))
+			{
+				break;
+			}
+		}
+	}
+}
+
+int32 FDisplayClusterWarpBlendLoader::ClearNoiseImpl(const FIntPoint& SearchLen, const FIntPoint& RemoveRule)
+{
+	const int32 SearchX = SearchLen.X * Width / 100;
+	const int32 SearchY = SearchLen.Y * Height / 100;
+	const int32 Rule1X = RemoveRule.X * Width / 100;
+	const int32 Rule1Y = RemoveRule.Y * Height / 100;
+
+	int32 TotalChangesCount = 0;
+	static const int32 DirIndexValue[] = { -1, 1 };
+
+	for (int32 Y = 0; Y < Height; ++Y)
+	{
+		for (int32 X = 0; X < Width; ++X)
+		{
+			if (Is3DPointValid(X, Y))
+			{
+				int32 XLen = 0;
+				int32 YLen = 0;
+
+				for (int32 DirIndex = 0; DirIndex < 2; DirIndex++)
+				{
+					int32 dx = 0;
+					int32 dy = 0;
+
+					for (int32 Offset = 1; Offset <= SearchX; Offset++)
+					{
+						if (Is3DPointValid(X + DirIndexValue[DirIndex] * Offset, Y))
+						{
+							dx++;
+						}
+						else
+						{
+							break;
+						}
+					}
+					for (int32 Offset = 1; Offset <= SearchY; Offset++)
+					{
+						if (Is3DPointValid(X, Y + DirIndexValue[DirIndex] * Offset))
+						{
+							dy++;
+						}
+						else
+						{
+							break;
+						}
+					}
+
+					XLen = FMath::Max(XLen, dx);
+					YLen = FMath::Max(YLen, dy);
+				}
+
+				const bool Test1 = XLen >= Rule1X && YLen >= Rule1Y;
+				const bool Test2 = YLen >= Rule1X && XLen >= Rule1Y;
+
+				if (!Test1 && !Test2)
+				{
+					// Both test failed, remove it
+					WarpData[X + Y * Width] = FVector4f(0.f, 0.f, 0.f, -1.f);
+					TotalChangesCount++;
+				}
+			}
+		}
+	}
+
+	return TotalChangesCount;
+}
diff --git a/Engine/Plugins/Runtime/nDisplay/Source/DisplayClusterWarp/Private/WarpBlend/Loader/DisplayClusterWarpBlendLoader_WarpMap.h b/Engine/Plugins/Runtime/nDisplay/Source/DisplayClusterWarp/Private/WarpBlend/Loader/DisplayClusterWarpBlendLoader_WarpMap.h
new file mode 100644
--- /dev/null
+++ b/Engine/Plugins/Runtime/nDisplay/Source/DisplayClusterWarp/Private/WarpBlend/Loader/DisplayClusterWarpBlendLoader_WarpMap.h
@@ -0,0 +1,77 @@
+// Copyright Epic Games, Inc. All Rights Reserved.
+
+#pragma once
+#include "Containers/DisplayClusterWarpEnums.h"
+
+namespace mpcdi
+{
+	struct GeometryWarpFile;
+	struct PFM;
+};
+
+/**
+ * WarpMap data loader.
+ */
+class FDisplayClusterWarpBlendLoader
+{
+public:
+	FDisplayClusterWarpBlendLoader() = default;
+	~FDisplayClusterWarpBlendLoader();
+
+public:
+	/** Get the width of the WarpMap. */
+	int32 GetWidth() const
+	{
+		return Width;
+	}
+
+	/** Get the height of the WarpMap. */
+	int32 GetHeight() const
+	{
+		return Height;
+	}
+
+	/** Get a pointer to WarpMap data. */
+	const FVector4f* GetWarpData() const
+	{
+		return WarpData;
+	}
+
+	/** Load WarpMap from GeometryWarpFile structure. */
+	bool LoadFromGeometryWarpFile(const EDisplayClusterWarpProfileType InProfileType, mpcdi::GeometryWarpFile* SourceWarpMap);
+
+	/** Load WarpMap from points. */
+	bool LoadFromPoint(const EDisplayClusterWarpProfileType InProfileType, const TArray<FVector>& InPoints, int32 WarpX, int32 WarpY, float WorldScale, bool bIsUnrealGameSpace);
+
+	/** Load WarpMap from PFM structure. */
+	bool LoadFromPFM(const EDisplayClusterWarpProfileType InProfileType, mpcdi::PFM* SourcePFM, float PFMScale, bool bIsUnrealGameSpace);
+
+protected:
+	/** Initializes WarpData and returns true if successful. */
+	bool InitializeWarpDataImpl(int32 InWidth, int32 InHeight);
+
+	/** Load data from points with scaling and space conversion. */
+	void LoadGeometryImpl(EDisplayClusterWarpProfileType ProfileType, const TArray<FVector>& InPoints, float WorldScale, bool bIsUnrealGameSpace);
+
+	/** Returns true if this point can be used in the warp. */
+	bool Is3DPointValid(int32 X, int32 Y) const;
+
+	/** Helper: The calibration data is not perfect and contains noise.
+	 * This function deletes disconnected point areas.
+	 * Useful for large warp nets.
+	 */
+	void ClearNoise(const EDisplayClusterWarpProfileType InProfileType);
+
+	/** Implementation of the ClearNoice(). */
+	int32 ClearNoiseImpl(const FIntPoint& SearchLen, const FIntPoint& RemoveRule);
+
+protected:
+	// Width of the WarpData
+	int32 Width = 0;
+
+	// Height of the WarpData
+	int32 Height = 0;
+
+	// WarpData
+	FVector4f* WarpData = nullptr;
+};
diff --git a/Engine/Plugins/Runtime/nDisplay/Source/DisplayClusterWarp/Private/WarpBlend/Math/DisplayClusterWarpBlendMath_Frustum.cpp b/Engine/Plugins/Runtime/nDisplay/Source/DisplayClusterWarp/Private/WarpBlend/Math/DisplayClusterWarpBlendMath_Frustum.cpp
new file mode 100644
--- /dev/null
+++ b/Engine/Plugins/Runtime/nDisplay/Source/DisplayClusterWarp/Private/WarpBlend/Math/DisplayClusterWarpBlendMath_Frustum.cpp
@@ -0,0 +1,401 @@
+// Copyright Epic Games, Inc. All Rights Reserved.
+
+#include "WarpBlend/Math/DisplayClusterWarpBlendMath_Frustum.h"
+
+#include "Render/Viewport/IDisplayClusterViewport.h"
+#include "Render/Containers/IDisplayClusterRender_Texture.h"
+#include "Misc/DisplayClusterHelpers.h"
+
+
+#include "StaticMeshResources.h"
+#include "ProceduralMeshComponent.h"
+
+// Select mpcdi frustum calc method
+int32 GDisplayClusterWarpBlendMPCDIFrustumMethod = (int32)EDisplayClusterWarpBlendFrustumType::LOD;
+static FAutoConsoleVariableRef CVarDisplayClusterWarpBlendMPCDIFrustumMethod(
+	TEXT("nDisplay.render.mpcdi.Frustum"),
+	GDisplayClusterWarpBlendMPCDIFrustumMethod,
+	TEXT("Frustum computation method:\n")
+	TEXT(" 0: mesh AABB based, lower quality but fast\n")
+	TEXT(" 1: mesh vertices based, best quality but slow\n")
+	TEXT(" 2: LOD, get A*B distributed points from texture, fast, good quality for flat panels\n"),
+	ECVF_Default
+);
+
+// Select mpcdi stereo mode
+int32 GDisplayClusterWarpBlendMPCDIStereoMode = (int32)EDisplayClusterWarpBlendStereoMode::AsymmetricAABB;
+static FAutoConsoleVariableRef CVarDisplayClusterWarpBlendMPCDIStereoMode(
+	TEXT("nDisplay.render.mpcdi.StereoMode"),
+	GDisplayClusterWarpBlendMPCDIStereoMode,
+	TEXT("Stereo mode:\n")
+	TEXT(" 0: Asymmetric to AABB center\n")
+	TEXT(" 1: Symmetric to AABB center\n"),
+	ECVF_Default
+);
+
+// Select mpcdi projection mode
+int32 GDisplayClusterWarpBlendMPCDIProjectionMode = (int32)EDisplayClusterWarpBlendProjectionType::StaticSurfaceNormal;
+static FAutoConsoleVariableRef CVarDisplayClusterWarpBlendMPCDIProjectionMode(
+	TEXT("nDisplay.render.mpcdi.Projection"),
+	GDisplayClusterWarpBlendMPCDIProjectionMode,
+	TEXT("Projection method:\n")
+	TEXT(" 0: Static, aligned to average region surface normal\n")
+	TEXT(" 1: Static, aligned to average region surface corners plane\n")
+	TEXT(" 2: Dynamic, to view target center\n"),
+	ECVF_Default
+);
+
+// Frustum projection fix (back-side view planes)
+int32 GDisplayClusterWarpBlendMPCDIProjectionAuto = 1;
+static FAutoConsoleVariableRef CVarDisplayClusterWarpBlendMPCDIProjectionAuto(
+	TEXT("nDisplay.render.mpcdi.ProjectionAuto"),
+	GDisplayClusterWarpBlendMPCDIProjectionAuto,
+	TEXT("Runtime frustum method, fix back-side view projection.\n")
+	TEXT(" 0: Disabled\n")
+	TEXT(" 1: Enabled (default)\n"),
+	ECVF_Default
+);
+
+int32 GDisplayClusterWarpBlendRotateFrustumToFitContextSize = 1;
+static FAutoConsoleVariableRef CVarDisplayClusterWarpBlendRotateFrustumToFitContextSize(
+	TEXT("nDisplay.render.mpcdi.RotateFrustumToFitContextSize"),
+	GDisplayClusterWarpBlendRotateFrustumToFitContextSize,
+	TEXT("Enables the method that rotates the frustum according to the size of the viewport area. (0 - disable)\n"),
+	ECVF_Default
+);
+
+float GDisplayClusterFrustumOrientationFitHysteresis = 0.05;
+static FAutoConsoleVariableRef CVarDisplayClusterFrustumOrientationFitHysteresis(
+	TEXT("nDisplay.render.mpcdi.FrustumOrientationFitHysteresis"),
+	GDisplayClusterFrustumOrientationFitHysteresis,
+	TEXT(
+		"Hysteresis margin when fitting frustum orientation with render target aspect ratio.\n"
+		"Values closer to 0 will allow immediate fit, while values close to 1 will require a \n"
+		"higher degree of mismatch before reorienting the frustum for a better fit."
+	),
+	ECVF_Default
+);
+
+//------------------------------------------------------------
+// FDisplayClusterWarpBlendMath_Frustum
+//------------------------------------------------------------
+FDisplayClusterWarpBlendMath_Frustum::FDisplayClusterWarpBlendMath_Frustum(FDisplayClusterWarpData& InWarpData, FDisplayClusterWarpBlend_GeometryContext& InGeometryContext)
+	: GeometryContext(InGeometryContext)
+	, WarpData(InWarpData)
+{
+	check(WarpData.WarpEye.IsValid());
+
+	// Configure logic from CVar:
+	WarpData.bEnabledRotateFrustumToFitContextSize = GDisplayClusterWarpBlendRotateFrustumToFitContextSize != 0;
+	WarpData.bFindBestProjectionType = GDisplayClusterWarpBlendMPCDIProjectionAuto != 0;
+	WarpData.FrustumType = (EDisplayClusterWarpBlendFrustumType)(GDisplayClusterWarpBlendMPCDIFrustumMethod);
+	WarpData.StereoMode = (EDisplayClusterWarpBlendStereoMode)(GDisplayClusterWarpBlendMPCDIStereoMode);
+
+	SetParameter((EDisplayClusterWarpBlendProjectionType)(GDisplayClusterWarpBlendMPCDIProjectionMode));
+}
+
+bool FDisplayClusterWarpBlendMath_Frustum::CalcFrustum()
+{
+	// Initialize warp projection data
+	InitializeWarpProjectionData();
+
+	// Override frustum
+	if (ShouldHandleOverrideCalcFrustum())
+	{
+		if (OverrideCalcFrustum())
+		{
+			return EndCalcFrustum();
+		}
+
+		return false;
+	}
+
+	switch (GeometryContext.GetWarpProfileType())
+	{
+	case EDisplayClusterWarpProfileType::warp_2D:
+		return CalcFrustum2D();
+
+	case EDisplayClusterWarpProfileType::warp_3D:
+		return CalcFrustum3D();
+
+	case EDisplayClusterWarpProfileType::warp_A3D:
+		return CalcFrustumA3D();
+
+	case EDisplayClusterWarpProfileType::warp_SL:
+		return CalcFrustumSL();
+
+	default:
+		break;
+	}
+
+	return false;
+}
+
+FMatrix FDisplayClusterWarpBlendMath_Frustum::CalcProjectionMatrix(const FDisplayClusterWarpProjection& InWarpProjection)
+{
+	check(WarpData.WarpEye.IsValid())
+
+	if (IDisplayClusterViewport* Viewport = WarpData.WarpEye->GetViewport())
+	{
+		const uint32 ContextNum = WarpData.WarpEye->ContextNum;
+		if (Viewport->GetContexts().IsValidIndex(ContextNum))
+		{
+			// the input data can be in degrees
+			const bool bIsAnglesInput = InWarpProjection.DataType == EDisplayClusterWarpAngleUnit::Degrees;
+
+			Viewport->CalculateProjectionMatrix(ContextNum,
+				InWarpProjection.Left,  InWarpProjection.Right,
+				InWarpProjection.Top,   InWarpProjection.Bottom,
+				InWarpProjection.ZNear, InWarpProjection.ZFar, bIsAnglesInput);
+
+			return Viewport->GetContexts()[ContextNum].ProjectionMatrix;
+		}
+	}
+
+	return FMatrix::Identity;
+}
+
+bool FDisplayClusterWarpBlendMath_Frustum::ShouldRotateFrustumToFitContextSize() const
+{
+	return WarpData.bEnabledRotateFrustumToFitContextSize && WarpData.bRotateFrustumToFitContextSize;
+}
+
+bool FDisplayClusterWarpBlendMath_Frustum::EndCalcFrustum()
+{
+	// Fit frustum to context size
+	ImplFitFrustumToContextSize();
+
+	// Update Camera
+	WarpData.GeometryWarpProjection.CameraLocation = WarpData.GeometryWarpProjection.EyeLocation;
+	WarpData.GeometryWarpProjection.CameraRotation = WarpData.Local2World.Rotator();
+
+	// Use geometry projection for render. This projection can be overriden
+	WarpData.WarpProjection = WarpData.GeometryWarpProjection;
+
+	check(WarpData.WarpEye.IsValid());
+
+	if (WarpData.WarpEye->WarpPolicy.IsValid())
+	{
+		if (IDisplayClusterViewport* Viewport = WarpData.WarpEye->GetViewport())
+		{
+			WarpData.WarpEye->WarpPolicy->EndCalcFrustum(Viewport, WarpData.WarpEye->ContextNum);
+		}
+	}
+
+	// After updating the warp data in the warp policy, we build the final warp context in the code below:
+
+	FRotator CameraRotation = WarpData.WarpProjection.CameraRotation;
+	FMatrix RegionMatrix = GeometryContext.GetRegionMatrix();
+
+	// Apply frustum rotation to fit context size:
+	if (ShouldRotateFrustumToFitContextSize())
+	{
+		static FQuat PostRotation = FRotator(0, 0, -90).Quaternion();
+
+		// Rotate camera view 90 degree 
+		CameraRotation = (CameraRotation.Quaternion() * PostRotation).Rotator();
+
+		// Rotate projection angles
+		WarpData.WarpProjection.RotateProjectionAngles90Degree();
+
+		// Rotate input texture
+		static FMatrix Rotate90DegreeCCW(
+			FPlane(0, 1, 0, 0),
+			FPlane(-1, 0, 0, 0),
+			FPlane(0, 0, 1, 0),
+			FPlane(1, 0, 0, 1));
+
+		RegionMatrix = RegionMatrix * Rotate90DegreeCCW;
+	}
+
+	// Compute projection matricies
+	WarpData.WarpContext.GeometryProjectionMatrix = CalcProjectionMatrix(WarpData.GeometryWarpProjection);
+	WarpData.WarpContext.ProjectionMatrix = CalcProjectionMatrix(WarpData.WarpProjection);
+
+	// Use camera for rendering, but eye for math
+	WarpData.WarpContext.Rotation = CameraRotation;
+	WarpData.WarpContext.Location = WarpData.WarpProjection.CameraLocation;
+
+	WarpData.WarpContext.UVMatrix = GetUVMatrix(WarpData.WarpContext.GeometryProjectionMatrix);
+
+	WarpData.WarpContext.TextureMatrix = GeometryContext.GetTextureMatrix();
+	WarpData.WarpContext.RegionMatrix = RegionMatrix;
+
+	WarpData.WarpContext.MeshToStageMatrix = GeometryContext.Context.GeometryToOrigin;
+
+	WarpData.WarpContext.bIsValid = true;
+
+	return true;
+}
+
+void FDisplayClusterWarpBlendMath_Frustum::ImplFitFrustumToContextSize()
+{
+	check(WarpData.WarpEye.IsValid());
+
+	const uint32 ContextNum = WarpData.WarpEye->ContextNum;
+	const IDisplayClusterViewport* Viewport = WarpData.WarpEye->GetViewport();
+	if (!Viewport || !Viewport->GetContexts().IsValidIndex(ContextNum))
+	{
+		return;
+	}
+
+	const FIntPoint ContextSize = Viewport->GetContexts()[ContextNum].ContextSize;
+
+	// See if it should roll by 90 degrees to get a better fit with the context size aspect ratio.
+	FDisplayClusterWarpProjection CurrentWarpProjection = WarpData.GeometryWarpProjection;
+	if (ShouldRotateFrustumToFitContextSize())
+	{
+		CurrentWarpProjection.RotateProjectionAngles90Degree();
+	}
+
+	// Calculate the frustum Horizontal to Vertical aspect ratio.
+	const float HFoV = FMath::Abs(CurrentWarpProjection.Right - CurrentWarpProjection.Left);
+	const float VFoV = FMath::Abs(CurrentWarpProjection.Top - CurrentWarpProjection.Bottom);
+
+	// Validate the values to protect from division by zero
+
+	const bool bValidValues =
+		!FMath::IsNearlyZero(HFoV)
+		&& !FMath::IsNearlyZero(VFoV)
+		&& (ContextSize.X > 0)
+		&& (ContextSize.Y > 0);
+
+	if (!bValidValues)
+	{
+		return;
+	}
+
+	const bool bFrustumIsWide = HFoV > VFoV;
+	const bool bContextSizeIsWide = ContextSize.X > ContextSize.Y;
+
+	// If both aspect ratios are of the same type already, then there is nothing else to do.
+	if (bContextSizeIsWide == bFrustumIsWide)
+	{
+		return;
+	}
+
+	// If we are here, we must decide if we should toggle the 90 degree roll.
+	// The toggle will have hysteresis to avoid jumping back and forth.
+	// We will use the frustum aspect ratio to determine this
+
+	// Don't flip the rotation if we're inside hystersis
+	if (FMath::IsNearlyEqual(HFoV / VFoV, 1.0, GDisplayClusterFrustumOrientationFitHysteresis))
+	{
+		return;
+	}
+
+	// If we're here, it is because we should toggle the rotation fit.
+	WarpData.bRotateFrustumToFitContextSize = !WarpData.bRotateFrustumToFitContextSize;
+}
+
+FMatrix FDisplayClusterWarpBlendMath_Frustum::GetUVMatrix(const FMatrix& InProjectionMatrix) const
+{
+	// These matrices were copied from LocalPlayer.cpp.
+	// They change the coordinate system from the Unreal "Game" coordinate system to the Unreal "Render" coordinate system
+	static const FMatrix Game2Render(
+		FPlane(0, 0, 1, 0),
+		FPlane(1, 0, 0, 0),
+		FPlane(0, 1, 0, 0),
+		FPlane(0, 0, 0, 1));
+
+	switch (GeometryContext.GetWarpProfileType())
+	{
+	case EDisplayClusterWarpProfileType::warp_A3D:
+		return  WarpData.Local2World.Inverse() * Game2Render * InProjectionMatrix;
+
+	default:
+		break;
+	}
+
+	return FMatrix::Identity;
+}
+
+bool FDisplayClusterWarpBlendMath_Frustum::ShouldUseCustomViewDirection() const
+{
+	return WarpData.WarpEye->OverrideViewTarget.IsSet() || WarpData.WarpEye->OverrideViewDirection.IsSet();
+}
+
+FVector FDisplayClusterWarpBlendMath_Frustum::GetViewTarget() const
+{
+	return WarpData.WarpEye->OverrideViewTarget.IsSet()  ? WarpData.WarpEye->OverrideViewTarget.GetValue() : GeometryContext.Context.AABBox.GetCenter();
+}
+
+FVector FDisplayClusterWarpBlendMath_Frustum::GetViewDirection()
+{
+	check(WarpData.WarpEye.IsValid());
+
+	if (WarpData.WarpEye->OverrideViewDirection.IsSet())
+	{
+		return WarpData.WarpEye->OverrideViewDirection.GetValue();
+	}
+
+	// Get view target location
+	const FVector ViewTargetLocation = GetViewTarget();
+
+	// Get view direction vector
+	if (WarpData.StereoMode == EDisplayClusterWarpBlendStereoMode::AsymmetricAABB)
+	{
+		// Create view transform matrix from look direction vector:
+		return (ViewTargetLocation - WarpData.WarpEye->ViewPoint.GetEyeLocation()).GetSafeNormal();
+	}
+
+	return (ViewTargetLocation - WarpData.WarpEye->ViewPoint.Location).GetSafeNormal();
+}
+
+bool FDisplayClusterWarpBlendMath_Frustum::GetProjectionClip(const FVector4& InPoint, const FMatrix& InGeometryToFrustumA3D, FDisplayClusterWarpProjection& InOutWarpProjection)
+{
+	if (InPoint.W > 0)
+	{
+		FVector4 ProjectedVertice = InGeometryToFrustumA3D.TransformFVector4(InPoint);
+
+		// Use only points over view plane, ignore backside pts
+		if (isnan(ProjectedVertice.X) || isnan(ProjectedVertice.Y) || isnan(ProjectedVertice.Z) || ProjectedVertice.X <= 0 || FMath::IsNearlyZero(ProjectedVertice.X, (FVector4::FReal)1.e-6f))
+		{
+			// This point out of view plane
+			return false;
+		}
+
+		const double Scale = InOutWarpProjection.ZNear / ProjectedVertice.X;
+
+		ProjectedVertice.Y *= Scale;
+		ProjectedVertice.Z *= Scale;
+
+		InOutWarpProjection.Left = FMath::Min(InOutWarpProjection.Left, ProjectedVertice.Y);
+		InOutWarpProjection.Right = FMath::Max(InOutWarpProjection.Right, ProjectedVertice.Y);
+
+		InOutWarpProjection.Top = FMath::Max(InOutWarpProjection.Top, ProjectedVertice.Z);
+		InOutWarpProjection.Bottom = FMath::Min(InOutWarpProjection.Bottom, ProjectedVertice.Z);
+	}
+
+	return true;
+}
+
+bool FDisplayClusterWarpBlendMath_Frustum::ShouldHandleOverrideCalcFrustum() const
+{
+	check(WarpData.WarpEye.IsValid());
+
+	if (WarpData.WarpEye->WarpPolicy.IsValid())
+	{
+		if (IDisplayClusterViewport* Viewport = WarpData.WarpEye->GetViewport())
+		{
+			return WarpData.WarpEye->WarpPolicy->ShouldOverrideCalcFrustum(Viewport);
+		}
+	}
+
+	return false;
+}
+
+bool FDisplayClusterWarpBlendMath_Frustum::OverrideCalcFrustum()
+{
+	check(WarpData.WarpEye.IsValid());
+
+	if (WarpData.WarpEye->WarpPolicy.IsValid())
+	{
+		if (IDisplayClusterViewport* Viewport = WarpData.WarpEye->GetViewport())
+		{
+			return WarpData.WarpEye->WarpPolicy->OverrideCalcFrustum(Viewport, WarpData.WarpEye->ContextNum);
+		}
+	}
+
+	return false;
+}
diff --git a/Engine/Plugins/Runtime/nDisplay/Source/DisplayClusterWarp/Private/WarpBlend/Math/DisplayClusterWarpBlendMath_Frustum.h b/Engine/Plugins/Runtime/nDisplay/Source/DisplayClusterWarp/Private/WarpBlend/Math/DisplayClusterWarpBlendMath_Frustum.h
new file mode 100644
--- /dev/null
+++ b/Engine/Plugins/Runtime/nDisplay/Source/DisplayClusterWarp/Private/WarpBlend/Math/DisplayClusterWarpBlendMath_Frustum.h
@@ -0,0 +1,132 @@
+// Copyright Epic Games, Inc. All Rights Reserved.
+
+#pragma once
+#include "Containers/DisplayClusterWarpContext.h"
+#include "Containers/DisplayClusterWarpEye.h"
+#include "WarpBlend/DisplayClusterWarpBlend_GeometryContext.h"
+
+class IDisplayClusterViewport;
+
+/**
+ * The WarpBlend math implementation.
+ */
+class FDisplayClusterWarpBlendMath_Frustum
+{
+public:
+	FDisplayClusterWarpBlendMath_Frustum(FDisplayClusterWarpData& InWarpData, FDisplayClusterWarpBlend_GeometryContext& InGeometryContext);
+
+public:
+	/** Entry point. */
+	bool CalcFrustum();
+
+private:
+	/** Math for mpcdi 2D profile. */
+	bool CalcFrustum2D();
+
+	/** Math for mpcdi 3D profile. */
+	bool CalcFrustum3D();
+
+	/** Math for mpcdi A3D profile. */
+	bool CalcFrustumA3D();
+
+	/** Math for mpcdi SL profile. */
+	bool CalcFrustumSL();
+
+protected:
+	/** Calculate projection matrix from angles*/
+	FMatrix CalcProjectionMatrix(const FDisplayClusterWarpProjection& InWarpProjection);
+
+	/** Get warp profile type. */
+	inline EDisplayClusterWarpProfileType  GetWarpProfileType() const
+	{
+		return GeometryContext.GetWarpProfileType();
+	}
+
+public:
+	inline void SetParameter(EDisplayClusterWarpBlendProjectionType  InParameterValue)
+	{ 
+		WarpData.ProjectionType = InParameterValue;
+
+		// Protect runtime methods:
+		if (WarpData.ProjectionType >= EDisplayClusterWarpBlendProjectionType::RuntimeProjectionModes)
+		{
+			// Runtime method disabled for console vars
+			WarpData.ProjectionType = EDisplayClusterWarpBlendProjectionType::DynamicAABBCenter;
+		}
+	}
+
+private:
+	bool ShouldRotateFrustumToFitContextSize() const;
+
+	inline bool ImplUpdateProjectionType()
+	{
+		switch (WarpData.ProjectionType)
+		{
+		// Accept any frustum for aabb center method
+		case EDisplayClusterWarpBlendProjectionType::DynamicAABBCenter:
+			break;
+
+		case EDisplayClusterWarpBlendProjectionType::StaticSurfaceNormal:
+			WarpData.ProjectionType = EDisplayClusterWarpBlendProjectionType::RuntimeStaticSurfaceNormalInverted; // Make mirror for backside eye position
+			return true;
+
+		case EDisplayClusterWarpBlendProjectionType::StaticSurfacePlane:
+			WarpData.ProjectionType = EDisplayClusterWarpBlendProjectionType::RuntimeStaticSurfacePlaneInverted; // Make mirror for backside eye position
+			return true;
+
+		default:
+			WarpData.ProjectionType = EDisplayClusterWarpBlendProjectionType::DynamicAABBCenter;
+			return true;
+		}
+
+		// Accept current projection
+		return false;
+	}
+
+	FMatrix GetUVMatrix(const FMatrix& InProjectionMatrix) const;
+
+	bool ShouldUseCustomViewDirection() const;
+
+	FVector GetViewTarget() const;
+
+	/** Get view direction. */
+	FVector GetViewDirection();
+
+	/** Initialize warp projection data before search. */
+	void InitializeWarpProjectionData();
+
+	/** build projection plane matrices. */
+	void ImplCalcViewProjectionMatrices();
+
+	// Frustum build logics:
+	bool ImplCalcFrustumProjection();
+	bool ImplCalcFrustumProjection_AABB(const FMatrix& InGeometryToFrustumMatrix);
+	bool ImplCalcFrustumProjection_WarpMap(const FMatrix& InGeometryToFrustumMatrix);
+	bool ImplCalcFrustumProjection_WarpMesh(const FMatrix& InGeometryToFrustumMatrix);
+	bool ImplCalcFrustumProjection_WarpProceduralMesh(const FMatrix& InGeometryToFrustumMatrix);
+	bool ImplCalcFrustumProjection_MPCDIAttributes(const FMatrix& InGeometryToFrustumMatrix);
+
+	bool EndCalcFrustum();
+
+	// WarpBlend
+
+	bool ShouldHandleOverrideCalcFrustum() const;
+	bool OverrideCalcFrustum();
+
+	/** 
+	 * When necessary, rotates by 90 deg. the frustum to better fit the aspect ratio of the viewport.
+	 */
+	void ImplFitFrustumToContextSize();
+
+private:
+	static bool GetProjectionClip(const FVector4& InPoint, const FMatrix& InGeometryToFrustumA3D, FDisplayClusterWarpProjection& InOutWarpProjection);
+
+private:
+	FDisplayClusterWarpBlend_GeometryContext& GeometryContext;
+
+	// context data
+	FDisplayClusterWarpData& WarpData;
+
+	// legacy WarpMap, now always use x16 downscale ratio
+	const int32 WarpMapLODRatio = 16;
+};
diff --git a/Engine/Plugins/Runtime/nDisplay/Source/DisplayClusterWarp/Private/WarpBlend/Math/DisplayClusterWarpBlendMath_FrustumProjection.cpp b/Engine/Plugins/Runtime/nDisplay/Source/DisplayClusterWarp/Private/WarpBlend/Math/DisplayClusterWarpBlendMath_FrustumProjection.cpp
new file mode 100644
--- /dev/null
+++ b/Engine/Plugins/Runtime/nDisplay/Source/DisplayClusterWarp/Private/WarpBlend/Math/DisplayClusterWarpBlendMath_FrustumProjection.cpp
@@ -0,0 +1,297 @@
+// Copyright Epic Games, Inc. All Rights Reserved.
+
+#include "WarpBlend/Math/DisplayClusterWarpBlendMath_Frustum.h"
+
+#include "Render/Containers/IDisplayClusterRender_Texture.h"
+#include "Render/Viewport/IDisplayClusterViewport.h"
+#include "Misc/DisplayClusterHelpers.h"
+
+#include "WarpBlend/Exporter/DisplayClusterWarpBlendExporter_WarpMap.h"
+
+#include "StaticMeshResources.h"
+#include "ProceduralMeshComponent.h"
+
+//---------------------------------------------------------------------------------------
+// FDisplayClusterWarpBlendMath_Frustum
+//---------------------------------------------------------------------------------------
+bool FDisplayClusterWarpBlendMath_Frustum::ImplCalcFrustumProjection()
+{
+	// build projection plane
+	ImplCalcViewProjectionMatrices();
+
+	// Reset extend of the frustum
+	WarpData.GeometryWarpProjection.ResetProjectionAngles();
+
+	const FMatrix GeometryToFrustumMatrix = GeometryContext.Context.GeometryToOrigin * WarpData.Local2World.Inverse();
+
+	switch (GeometryContext.GeometryProxy.FrustumGeometryType)
+	{
+	case EDisplayClusterWarpFrustumGeometryType::WarpMap:
+		return ImplCalcFrustumProjection_WarpMap(GeometryToFrustumMatrix);
+
+	case EDisplayClusterWarpFrustumGeometryType::WarpMesh:
+		return ImplCalcFrustumProjection_WarpMesh(GeometryToFrustumMatrix);
+
+	case EDisplayClusterWarpFrustumGeometryType::WarpProceduralMesh:
+		return ImplCalcFrustumProjection_WarpProceduralMesh(GeometryToFrustumMatrix);
+
+	case EDisplayClusterWarpFrustumGeometryType::MPCDIAttributes:
+		return ImplCalcFrustumProjection_MPCDIAttributes(GeometryToFrustumMatrix);
+
+	default:
+		break;
+	}
+
+	return false;
+}
+
+bool FDisplayClusterWarpBlendMath_Frustum::ImplCalcFrustumProjection_MPCDIAttributes(const FMatrix& InGeometryToFrustumMatrix)
+{
+	const FDisplayClusterWarpMPCDIAttributes& MPCDIAttributes = GeometryContext.GeometryProxy.MPCDIAttributes;
+	switch (GetWarpProfileType())
+	{
+	case EDisplayClusterWarpProfileType::warp_2D:
+	{
+		// Create geometry from attributes
+		TArray<FVector> ScreenPoints;
+		FDisplayClusterWarpBlendExporter_WarpMap::Get2DProfileGeometry(MPCDIAttributes, ScreenPoints);
+
+		bool bOutOfFrustumPoints = false;
+
+		// Search a camera space frustum
+		for (const FVector& Pts : ScreenPoints)
+		{
+			bOutOfFrustumPoints |= !GetProjectionClip(Pts, InGeometryToFrustumMatrix, WarpData.GeometryWarpProjection);
+		}
+
+		return bOutOfFrustumPoints == false;
+	}
+
+	case EDisplayClusterWarpProfileType::warp_3D:
+	{
+		// Todo: need to be implemeted. This is an experimental code
+		const float ZNear = WarpData.GeometryWarpProjection.ZFar;
+
+		// Frustum angles XYZW = LRTB in degrees
+		WarpData.GeometryWarpProjection.Left = (ZNear * FMath::Tan(FMath::DegreesToRadians(MPCDIAttributes.Frustum.Angles.X)));
+		WarpData.GeometryWarpProjection.Right = (ZNear * FMath::Tan(FMath::DegreesToRadians(MPCDIAttributes.Frustum.Angles.Y)));
+		WarpData.GeometryWarpProjection.Top = (ZNear * FMath::Tan(FMath::DegreesToRadians(MPCDIAttributes.Frustum.Angles.Z)));
+		WarpData.GeometryWarpProjection.Bottom = (ZNear * FMath::Tan(FMath::DegreesToRadians(MPCDIAttributes.Frustum.Angles.W)));
+	}
+	return true;
+
+	default:
+		break;
+	}
+
+	return false;
+}
+
+void FDisplayClusterWarpBlendMath_Frustum::ImplCalcViewProjectionMatrices()
+{
+	// Customize view direction
+	FVector ViewDir = GetViewDirection();
+
+	if (!ShouldUseCustomViewDirection())
+	{
+		switch (WarpData.ProjectionType)
+		{
+		case EDisplayClusterWarpBlendProjectionType::StaticSurfacePlane:
+			ViewDir = GeometryContext.Context.SurfaceViewPlane;
+			break;
+
+		case EDisplayClusterWarpBlendProjectionType::RuntimeStaticSurfacePlaneInverted:
+			// Frustum projection fix (back-side view planes)
+			ViewDir = -GeometryContext.Context.SurfaceViewPlane;
+			break;
+
+		case EDisplayClusterWarpBlendProjectionType::StaticSurfaceNormal:
+			// Use fixed surface view normal
+			ViewDir = GeometryContext.Context.SurfaceViewNormal;
+			break;
+
+		case EDisplayClusterWarpBlendProjectionType::RuntimeStaticSurfaceNormalInverted:
+			// Frustum projection fix (back-side view planes)
+			ViewDir = -GeometryContext.Context.SurfaceViewNormal;
+			break;
+
+		default:
+			break;
+		}
+	}
+
+	FVector const NewX = ViewDir.GetSafeNormal();
+	if (FMath::Abs(NewX.Z) < (1.f - KINDA_SMALL_NUMBER))
+	{
+		WarpData.Local2World = FRotationMatrix::MakeFromXZ(NewX, FVector(0.f, 0.f, 1.f));
+	}
+	else
+	{
+		WarpData.Local2World = FRotationMatrix::MakeFromXY(NewX, FVector(0.f, 1.f, 0.f));
+	}
+
+	WarpData.Local2World.SetOrigin(WarpData.GeometryWarpProjection.EyeLocation); // Finally set view origin to eye location
+}
+
+void FDisplayClusterWarpBlendMath_Frustum::InitializeWarpProjectionData()
+{
+	check(WarpData.WarpEye);
+
+	// Reset extend of the frustum
+	WarpData.GeometryWarpProjection = FDisplayClusterWarpProjection();
+
+	if (IDisplayClusterViewport* Viewport = WarpData.WarpEye->GetViewport())
+	{
+		const FVector2D ClippingPLanes = Viewport->GetClippingPlanes();
+
+		WarpData.GeometryWarpProjection.ZNear = ClippingPLanes.X;
+		WarpData.GeometryWarpProjection.ZFar  = ClippingPLanes.Y;
+	}
+	else
+	{
+		WarpData.GeometryWarpProjection.ZNear = GNearClippingPlane;
+		WarpData.GeometryWarpProjection.ZFar = WarpData.GeometryWarpProjection.ZNear;
+	}
+
+	WarpData.GeometryWarpProjection.WorldScale = WarpData.WarpEye->WorldScale;
+
+	WarpData.GeometryWarpProjection.EyeLocation = WarpData.WarpEye->ViewPoint.GetEyeLocation();
+}
+
+bool FDisplayClusterWarpBlendMath_Frustum::ImplCalcFrustumProjection_AABB(const FMatrix& InGeometryToFrustumMatrix)
+{
+	bool bOutOfFrustumPoints = false;
+
+	// Search a camera space frustum
+	for (int32 AABBPointIndex = 0; AABBPointIndex < 8; AABBPointIndex++)
+	{
+		bOutOfFrustumPoints |= !GetProjectionClip(GeometryContext.Context.AABBox.GetAABBPts(AABBPointIndex), InGeometryToFrustumMatrix, WarpData.GeometryWarpProjection);
+	}
+
+	return bOutOfFrustumPoints == false;
+}
+
+bool FDisplayClusterWarpBlendMath_Frustum::ImplCalcFrustumProjection_WarpMap(const FMatrix& InGeometryToFrustumMatrix)
+{
+	const TSharedPtr<IDisplayClusterRender_Texture, ESPMode::ThreadSafe>& WarpMapTexture = GeometryContext.GeometryProxy.WarpMapTexture;
+	if (const FVector4f* SourcePts = (FVector4f*)(WarpMapTexture.IsValid() ? WarpMapTexture->GetData() : nullptr))
+	{
+		bool bAllPointsInFrustum = true;
+
+		const int32 PointsAmount = WarpMapTexture->GetWidth() * WarpMapTexture->GetHeight();
+
+		switch (WarpData.FrustumType)
+		{
+		case EDisplayClusterWarpBlendFrustumType::AABB:
+			// WarpMap supports AABB frustum
+			return ImplCalcFrustumProjection_AABB(InGeometryToFrustumMatrix);
+
+		case EDisplayClusterWarpBlendFrustumType::FULL:
+			// Search a camera space frustum
+			for (int32 PointIndex = 0; PointIndex < PointsAmount; ++PointIndex)
+			{
+				const FVector4 Pts(SourcePts[PointIndex]);
+				if (!GetProjectionClip(Pts, InGeometryToFrustumMatrix, WarpData.GeometryWarpProjection))
+				{
+					bAllPointsInFrustum = false;
+				}
+			}
+			return bAllPointsInFrustum;
+
+		case EDisplayClusterWarpBlendFrustumType::LOD:
+		{
+			const TArray<int32>& IndexLOD = GeometryContext.GeometryProxy.GeometryCache.IndexLOD;
+			if (IndexLOD.Num() == 0)
+			{
+				check(WarpMapLODRatio > 0);
+
+				GeometryContext.GeometryProxy.UpdateFrustumGeometryLOD(FIntPoint(WarpMapLODRatio));
+			}
+
+			// Search a camera space frustum
+			for (const int32& PtsIndex : IndexLOD)
+			{
+				if (PtsIndex >= 0 && PtsIndex < PointsAmount)
+				{
+					const FVector4 Pts(SourcePts[PtsIndex]);
+					if (!GetProjectionClip(Pts, InGeometryToFrustumMatrix, WarpData.GeometryWarpProjection))
+					{
+						bAllPointsInFrustum = false;
+					}
+				}
+			}
+			return bAllPointsInFrustum;
+		}
+
+		default:
+			break;
+		}
+	}
+
+	return false;
+}
+
+bool FDisplayClusterWarpBlendMath_Frustum::ImplCalcFrustumProjection_WarpMesh(const FMatrix& InGeometryToFrustumMatrix)
+{
+	if (!GeometryContext.GeometryProxy.MeshComponent.IsValid())
+	{
+		return false;
+	}
+
+	if (WarpData.FrustumType == EDisplayClusterWarpBlendFrustumType::AABB)
+	{
+		// WarpMesh supports AABB frustum
+		return ImplCalcFrustumProjection_AABB(InGeometryToFrustumMatrix);
+	}
+
+	bool bAllPointsInFrustum = true;
+
+	const FStaticMeshLODResources* StaticMeshLODResources = GeometryContext.GeometryProxy.GetStaticMeshComponentLODResources();
+	if (StaticMeshLODResources == nullptr)
+	{
+		return false;
+	}
+
+	const FPositionVertexBuffer& VertexPosition = StaticMeshLODResources->VertexBuffers.PositionVertexBuffer;
+	for (uint32 i = 0; i < VertexPosition.GetNumVertices(); i++)
+	{
+		if (!GetProjectionClip(FVector4(FVector(VertexPosition.VertexPosition(i)), 1.f), InGeometryToFrustumMatrix, WarpData.GeometryWarpProjection))
+		{
+			bAllPointsInFrustum = false;
+		}
+	}
+
+	return bAllPointsInFrustum;
+}
+
+bool FDisplayClusterWarpBlendMath_Frustum::ImplCalcFrustumProjection_WarpProceduralMesh(const FMatrix& InGeometryToFrustumMatrix)
+{
+	if (!GeometryContext.GeometryProxy.MeshComponent.IsValid())
+	{
+		return false;
+	}
+
+	if (WarpData.FrustumType == EDisplayClusterWarpBlendFrustumType::AABB)
+	{
+		// WarpProceduralMesh supports AABB frustum
+		return ImplCalcFrustumProjection_AABB(InGeometryToFrustumMatrix);
+	}
+
+	bool bAllPointsInFrustum = true;
+
+	const FProcMeshSection* ProcMeshSection = GeometryContext.GeometryProxy.MeshComponent->GetProceduralMeshComponentSection();
+	if (ProcMeshSection == nullptr)
+	{
+		return false;
+	}
+
+	for (const FProcMeshVertex& VertexIt : ProcMeshSection->ProcVertexBuffer)
+	{
+		if (!GetProjectionClip(FVector4(VertexIt.Position, 1), InGeometryToFrustumMatrix, WarpData.GeometryWarpProjection))
+		{
+			bAllPointsInFrustum = false;
+		}
+	}
+
+	return bAllPointsInFrustum;
+}
diff --git a/Engine/Plugins/Runtime/nDisplay/Source/DisplayClusterWarp/Private/WarpBlend/Math/DisplayClusterWarpBlendMath_Frustum_2D.cpp b/Engine/Plugins/Runtime/nDisplay/Source/DisplayClusterWarp/Private/WarpBlend/Math/DisplayClusterWarpBlendMath_Frustum_2D.cpp
new file mode 100644
--- /dev/null
+++ b/Engine/Plugins/Runtime/nDisplay/Source/DisplayClusterWarp/Private/WarpBlend/Math/DisplayClusterWarpBlendMath_Frustum_2D.cpp
@@ -0,0 +1,20 @@
+// Copyright Epic Games, Inc. All Rights Reserved.
+
+#include "WarpBlend/Math/DisplayClusterWarpBlendMath_Frustum.h"
+#include "WarpBlend/Exporter/DisplayClusterWarpBlendExporter_WarpMap.h"
+#include "WarpBlend/DisplayClusterWarpBlend.h"
+
+#include "Render/Viewport/IDisplayClusterViewport.h"
+#include "Render/Containers/IDisplayClusterRender_Texture.h"
+#include "Misc/DisplayClusterHelpers.h"
+
+//------------------------------------------------------------
+// FDisplayClusterWarpBlendMath_Frustum
+//------------------------------------------------------------
+bool FDisplayClusterWarpBlendMath_Frustum::CalcFrustum2D()
+{
+	TRACE_CPUPROFILER_EVENT_SCOPE(DCWarpBlendMath_Frustum::CalcFrustum2D);
+
+	// Calc projection
+	return ImplCalcFrustumProjection() && EndCalcFrustum();
+}
diff --git a/Engine/Plugins/Runtime/nDisplay/Source/DisplayClusterWarp/Private/WarpBlend/Math/DisplayClusterWarpBlendMath_Frustum_3D.cpp b/Engine/Plugins/Runtime/nDisplay/Source/DisplayClusterWarp/Private/WarpBlend/Math/DisplayClusterWarpBlendMath_Frustum_3D.cpp
new file mode 100644
--- /dev/null
+++ b/Engine/Plugins/Runtime/nDisplay/Source/DisplayClusterWarp/Private/WarpBlend/Math/DisplayClusterWarpBlendMath_Frustum_3D.cpp
@@ -0,0 +1,22 @@
+// Copyright Epic Games, Inc. All Rights Reserved.
+
+#include "WarpBlend/Math/DisplayClusterWarpBlendMath_Frustum.h"
+
+#include "Render/Containers/IDisplayClusterRender_Texture.h"
+#include "Render/Viewport/IDisplayClusterViewport.h"
+#include "Misc/DisplayClusterHelpers.h"
+
+//------------------------------------------------------------
+// FDisplayClusterWarpBlendMath_Frustum
+//------------------------------------------------------------
+
+bool FDisplayClusterWarpBlendMath_Frustum::CalcFrustum3D()
+{
+	TRACE_CPUPROFILER_EVENT_SCOPE(DCWarpBlendMath_Frustum::CalcFrustum3D);
+
+	// NOT IMPLEMETNED
+	// Currently we have implemented only the math of the A3D profile. We have to write a new implementation to support the 3D profile.
+	// Some of the math from the A3D profile can be reused.
+
+	return false;
+}
diff --git a/Engine/Plugins/Runtime/nDisplay/Source/DisplayClusterWarp/Private/WarpBlend/Math/DisplayClusterWarpBlendMath_Frustum_A3D.cpp b/Engine/Plugins/Runtime/nDisplay/Source/DisplayClusterWarp/Private/WarpBlend/Math/DisplayClusterWarpBlendMath_Frustum_A3D.cpp
new file mode 100644
--- /dev/null
+++ b/Engine/Plugins/Runtime/nDisplay/Source/DisplayClusterWarp/Private/WarpBlend/Math/DisplayClusterWarpBlendMath_Frustum_A3D.cpp
@@ -0,0 +1,86 @@
+// Copyright Epic Games, Inc. All Rights Reserved.
+
+#include "WarpBlend/Math/DisplayClusterWarpBlendMath_Frustum.h"
+
+#include "Render/Containers/IDisplayClusterRender_Texture.h"
+#include "Render/Viewport/IDisplayClusterViewport.h"
+#include "Misc/DisplayClusterHelpers.h"
+
+#include "StaticMeshResources.h"
+#include "ProceduralMeshComponent.h"
+
+//---------------------------------------------------------------------------------------
+// FDisplayClusterWarpBlendMath_Frustum
+//---------------------------------------------------------------------------------------
+bool FDisplayClusterWarpBlendMath_Frustum::CalcFrustumA3D()
+{
+	TRACE_CPUPROFILER_EVENT_SCOPE(DCWarpBlendMath_Frustum::CalcFrustumA3D);
+
+	if (!GeometryContext.GeometryProxy.bIsGeometryValid)
+	{
+		// Bad source geometry
+		return false;
+	}
+
+	// When a custom ViewTarget is used, we do not use adaptation algorithms.
+	if (ShouldUseCustomViewDirection())
+	{
+		return ImplCalcFrustumProjection() && EndCalcFrustum();
+	}
+
+	if (WarpData.bFindBestProjectionType == false)
+	{
+		ImplCalcFrustumProjection();
+	}
+	else
+	{
+		// Frustum projection with auto-fix (back-side view planes)
+		// Projection type changed runtime. Validate frustum points, all must be over view plane:
+		EDisplayClusterWarpBlendProjectionType BaseProjectionType = WarpData.ProjectionType;
+		EDisplayClusterWarpBlendFrustumType    BaseFrustumType = WarpData.FrustumType;
+
+		// Optimize for PerfectCPU:
+		if (WarpData.FrustumType == EDisplayClusterWarpBlendFrustumType::FULL && GeometryContext.GeometryProxy.FrustumGeometryType == EDisplayClusterWarpFrustumGeometryType::WarpMap)
+		{
+			while (true)
+			{
+				WarpData.FrustumType = EDisplayClusterWarpBlendFrustumType::LOD;
+
+				// Fast check for bad view
+				if (ImplCalcFrustumProjection())
+				{
+					break;
+				}
+
+				if (!ImplUpdateProjectionType())
+				{
+					break;
+				}
+			}
+		}
+
+		// Restore base parameter
+		WarpData.FrustumType = BaseFrustumType;
+
+		// Search valid projection mode:
+		if (BaseProjectionType == WarpData.ProjectionType)
+		{
+
+			while (true)
+			{
+				// Full check for bad projection:
+				if (ImplCalcFrustumProjection())
+				{
+					break;
+				}
+
+				if (!ImplUpdateProjectionType())
+				{
+					break;
+				}
+			}
+		}
+	}
+
+	return EndCalcFrustum();
+}
diff --git a/Engine/Plugins/Runtime/nDisplay/Source/DisplayClusterWarp/Private/WarpBlend/Math/DisplayClusterWarpBlendMath_Frustum_SL.cpp b/Engine/Plugins/Runtime/nDisplay/Source/DisplayClusterWarp/Private/WarpBlend/Math/DisplayClusterWarpBlendMath_Frustum_SL.cpp
new file mode 100644
--- /dev/null
+++ b/Engine/Plugins/Runtime/nDisplay/Source/DisplayClusterWarp/Private/WarpBlend/Math/DisplayClusterWarpBlendMath_Frustum_SL.cpp
@@ -0,0 +1,21 @@
+// Copyright Epic Games, Inc. All Rights Reserved.
+
+#include "WarpBlend/Math/DisplayClusterWarpBlendMath_Frustum.h"
+
+#include "Render/Containers/IDisplayClusterRender_Texture.h"
+#include "Render/Viewport/IDisplayClusterViewport.h"
+#include "Misc/DisplayClusterHelpers.h"
+
+//------------------------------------------------------------
+// FDisplayClusterWarpBlendMath_Frustum
+//------------------------------------------------------------
+bool FDisplayClusterWarpBlendMath_Frustum::CalcFrustumSL()
+{
+	TRACE_CPUPROFILER_EVENT_SCOPE(DCWarpBlendMath_Frustum::CalcFrustumSL);
+
+	// NOT IMPLEMENTED
+	// Currently we have implemented only the math of the A3D profile. We have to write a new implementation to support the SL profile.
+	// Some of the math from the A3D profile can be reused.
+
+	return false;
+}
diff --git a/Engine/Plugins/Runtime/nDisplay/Source/DisplayClusterWarp/Private/WarpBlend/Math/DisplayClusterWarpBlendMath_WarpMap.h b/Engine/Plugins/Runtime/nDisplay/Source/DisplayClusterWarp/Private/WarpBlend/Math/DisplayClusterWarpBlendMath_WarpMap.h
new file mode 100644
--- /dev/null
+++ b/Engine/Plugins/Runtime/nDisplay/Source/DisplayClusterWarp/Private/WarpBlend/Math/DisplayClusterWarpBlendMath_WarpMap.h
@@ -0,0 +1,225 @@
+// Copyright Epic Games, Inc. All Rights Reserved.
+
+#pragma once
+
+#include "CoreMinimal.h"
+#include "Containers/DisplayClusterWarpAABB.h"
+
+/**
+ * Implement math for WarpMap texture as a source of geometry
+ */
+class FDisplayClusterWarpBlendMath_WarpMap
+{
+public:
+	FDisplayClusterWarpBlendMath_WarpMap(const IDisplayClusterRender_Texture& InWarpMap)
+		: Width(InWarpMap.GetWidth())
+		, Height(InWarpMap.GetHeight())
+		, Data((FVector4f*)InWarpMap.GetData())
+	{ }
+
+public:
+	/** Build a bbox from valid points. */
+	FBox GetAABBox() const
+	{
+		FDisplayClusterWarpAABB WarpAABB;
+		for (int32 Y = 0; Y < Height; ++Y)
+		{
+			for (int32 X = 0; X < Width; ++X)
+			{
+				WarpAABB.UpdateAABB(GetPoint(X, Y));
+			}
+		}
+
+		return WarpAABB;
+	}
+
+	void BuildIndexLOD(int32 DivX, int32 DivY, TArray<int32>& OutIndexLOD)
+	{
+		check(DivX > 0);
+		check(DivY > 0);
+
+		OutIndexLOD.Reserve(DivX * DivY);
+
+		// Generate valid points for texturebox method:
+		for (int32 low_y = 0; low_y < DivY; low_y++)
+		{
+			int32 y = (Height - 1) * (float(low_y) / (DivY - 1));
+
+			for (int32 low_x = 0; low_x < DivX; low_x++)
+			{
+				int32 x = (Width - 1) * (float(low_x) / (DivX - 1));
+
+				if (IsValidPoint(x, y))
+				{
+					// Just use direct point
+					OutIndexLOD.Add(GetPointIndex(x, y));
+				}
+				else
+				{
+					// Search for nearset valid point
+					if (FindValidPoint(x, y))
+					{
+						OutIndexLOD.Add(GetSavedPointIndex());
+					}
+				}
+			}
+		}
+	}
+
+public:
+	int32 GetSavedPointIndex()
+	{
+		return GetPointIndex(ValidPointX, ValidPointY);
+	}
+
+	FORCEINLINE int32 GetPointIndex(int32 InX, int32 InY) const
+	{
+		return InX + InY * Width;
+	}
+
+	FORCEINLINE const FVector4f& GetPoint(int32 InX, int32 InY) const
+	{
+		return Data[GetPointIndex(InX, InY)];
+	}
+
+	FORCEINLINE const FVector4f& GetPoint(int32 PointIndex) const
+	{
+		return Data[PointIndex];
+	}
+
+	bool IsValidPoint(int32 InX, int32 InY)
+	{
+		return GetPoint(InX, InY).W > 0;
+	}
+
+	bool FindValidPoint(int32 InX, int32 InY)
+	{
+		X0 = InX;
+		Y0 = InY;
+
+		for (int32 Range = 1; Range < Width; Range++)
+		{
+			if (FindValidPointInRange(Range))
+			{
+				return true;
+			}
+		}
+
+		return false;
+	}
+
+	FVector GetSurfaceViewNormal()
+	{
+		int32 Ncount = 0;
+		double Nxyz[3] = { 0,0,0 };
+
+		for (int32 ItY = 0; ItY < (Height - 2); ++ItY)
+		{
+			for (int32 ItX = 0; ItX < (Width - 2); ++ItX)
+			{
+				const FVector4f& Pts0 = GetPoint(ItX, ItY);
+				const FVector4f& Pts1 = GetPoint(ItX + 1, ItY);
+				const FVector4f& Pts2 = GetPoint(ItX, ItY + 1);
+
+				if (Pts0.W > 0 && Pts1.W > 0 && Pts2.W > 0)
+				{
+					const FVector N1 = FVector4(Pts1 - Pts0);
+					const FVector N2 = FVector4(Pts2 - Pts0);
+					const FVector N = FVector::CrossProduct(N2, N1).GetSafeNormal();
+
+					for (int32 AxisIndex = 0; AxisIndex < 3; AxisIndex++)
+					{
+						Nxyz[AxisIndex] += N[AxisIndex];
+					}
+
+					Ncount++;
+				}
+			}
+		}
+
+		double Scale = double(1) / Ncount;
+		for (int32 AxisIndex = 0; AxisIndex < 3; AxisIndex++)
+		{
+			Nxyz[AxisIndex] *= Scale;
+		}
+
+		return FVector(Nxyz[0], Nxyz[1], Nxyz[2]).GetSafeNormal();
+	}
+
+	FVector GetSurfaceViewPlane()
+	{
+		const FVector4f& Pts0 = GetValidPoint(0, 0);
+		const FVector4f& Pts1 = GetValidPoint(Width - 1, 0);
+		const FVector4f& Pts2 = GetValidPoint(0, Height - 1);
+
+		const FVector N1 = FVector4(Pts1 - Pts0);
+		const FVector N2 = FVector4(Pts2 - Pts0);
+		return FVector::CrossProduct(N2, N1).GetSafeNormal();
+	}
+
+private:
+	const FVector4f& GetValidPoint(int32 InX, int32 InY)
+	{
+		if (!IsValidPoint(InX, InY))
+		{
+			if (FindValidPoint(InX, InY))
+			{
+				return GetPoint(GetSavedPointIndex());
+			}
+		}
+
+		return GetPoint(InX, InY);
+	}
+
+	bool FindValidPointInRange(int32 Range)
+	{
+		for (int32 RangeIt = -Range; RangeIt <= Range; RangeIt++)
+		{
+			// Top or bottom rows
+			if (IsValid(X0 + RangeIt, Y0 - Range) || IsValid(X0 + RangeIt, Y0 + Range))
+			{
+				return true;
+			}
+
+			// Left or Right columns
+			if (IsValid(X0 - Range, Y0 + RangeIt) || IsValid(X0 + Range, Y0 + RangeIt))
+			{
+				return true;
+			}
+		}
+
+		return false;
+	}
+
+	bool IsValid(int32 newX, int32 newY)
+	{
+		if (newX < 0 || newY < 0 || newX >= Width || newY >= Height)
+		{
+			// Out of texture
+			return false;
+		}
+
+		if (Data[GetPointIndex(newX, newY)].W > 0)
+		{
+			// Store valid result
+			ValidPointX = newX;
+			ValidPointY = newY;
+
+			return true;
+		}
+
+		return false;
+	}
+
+private:
+	const int32 Width;
+	const int32 Height;
+	const FVector4f* Data;
+
+	// Internal logic
+	int32 ValidPointX = 0;
+	int32 ValidPointY = 0;
+
+	int32 X0 = 0;
+	int32 Y0 = 0;
+};
diff --git a/Engine/Plugins/Runtime/nDisplay/Source/DisplayClusterWarp/Private/WarpBlend/Math/DisplayClusterWarpBlendMath_WarpMesh.h b/Engine/Plugins/Runtime/nDisplay/Source/DisplayClusterWarp/Private/WarpBlend/Math/DisplayClusterWarpBlendMath_WarpMesh.h
new file mode 100644
--- /dev/null
+++ b/Engine/Plugins/Runtime/nDisplay/Source/DisplayClusterWarp/Private/WarpBlend/Math/DisplayClusterWarpBlendMath_WarpMesh.h
@@ -0,0 +1,74 @@
+// Copyright Epic Games, Inc. All Rights Reserved.
+
+#pragma once
+
+#include "CoreMinimal.h"
+#include "StaticMeshResources.h"
+#include "Containers/DisplayClusterWarpAABB.h"
+
+/**
+ * Implement math for StaticMesh component as a source of geometry
+ */
+class FDisplayClusterWarpBlendMath_WarpMesh
+{
+public:
+	FDisplayClusterWarpBlendMath_WarpMesh(const FStaticMeshLODResources& InMeshLODResources)
+		: MeshLODResources(InMeshLODResources)
+	{ }
+
+public:
+	FBox CalcAABBox()
+	{
+		FDisplayClusterWarpAABB WarpAABB;
+		const FPositionVertexBuffer& PositionBuffer = MeshLODResources.VertexBuffers.PositionVertexBuffer;
+		for (uint32 VertIdx = 0; VertIdx < PositionBuffer.GetNumVertices(); ++VertIdx)
+		{
+			WarpAABB.UpdateAABB((FVector)PositionBuffer.VertexPosition(VertIdx));
+		}
+
+		return WarpAABB;
+	}
+
+	void CalcSurfaceVectors(FVector& OutSurfaceViewNormal, FVector& OutSurfaceViewPlane)
+	{
+		// Calc static normal and plane
+		const int32 IdxNum = MeshLODResources.IndexBuffer.GetNumIndices();
+		const int32 TriNum = IdxNum / 3;
+
+		if (TriNum <= 0)
+		{
+			return;
+		}
+
+		FVector SurfaceNormalViewSum = FVector::ZeroVector;
+		const FPositionVertexBuffer& PositionBuffer = MeshLODResources.VertexBuffers.PositionVertexBuffer;
+
+		for (int32 TriIdx = 0; TriIdx < TriNum; ++TriIdx)
+		{
+			const int32 Index0 = MeshLODResources.IndexBuffer.GetIndex(TriIdx * 3 + 0);
+			const int32 Index1 = MeshLODResources.IndexBuffer.GetIndex(TriIdx * 3 + 1);
+			const int32 Index2 = MeshLODResources.IndexBuffer.GetIndex(TriIdx * 3 + 2);
+
+			const FVector& Pts1 = (FVector)PositionBuffer.VertexPosition(Index0);
+			const FVector& Pts0 = (FVector)PositionBuffer.VertexPosition(Index1);
+			const FVector& Pts2 = (FVector)PositionBuffer.VertexPosition(Index2);
+
+			const FVector N1 = Pts1 - Pts0;
+			const FVector N2 = Pts2 - Pts0;
+
+			const FVector Normal = FVector::CrossProduct(N2, N1);
+			const double Magnitude = Normal.Length();
+
+			SurfaceNormalViewSum += (Normal * Magnitude);
+		}
+
+		OutSurfaceViewNormal = SurfaceNormalViewSum.GetSafeNormal();
+
+		//@todo: MeshSurfaceViewPlane not implemented, use MeshSurfaceViewNormal
+		// this is for the case when normals are set in the side direction to have same sampling quality over the multiple heavily curved surfaces
+		OutSurfaceViewPlane = OutSurfaceViewNormal;
+	}
+
+private:
+	const FStaticMeshLODResources& MeshLODResources;
+};
diff --git a/Engine/Plugins/Runtime/nDisplay/Source/DisplayClusterWarp/Private/WarpBlend/Math/DisplayClusterWarpBlendMath_WarpProceduralMesh.h b/Engine/Plugins/Runtime/nDisplay/Source/DisplayClusterWarp/Private/WarpBlend/Math/DisplayClusterWarpBlendMath_WarpProceduralMesh.h
new file mode 100644
--- /dev/null
+++ b/Engine/Plugins/Runtime/nDisplay/Source/DisplayClusterWarp/Private/WarpBlend/Math/DisplayClusterWarpBlendMath_WarpProceduralMesh.h
@@ -0,0 +1,74 @@
+// Copyright Epic Games, Inc. All Rights Reserved.
+
+#pragma once
+
+#include "CoreMinimal.h"
+#include "ProceduralMeshComponent.h"
+#include "Containers/DisplayClusterWarpAABB.h"
+
+/**
+ * Implement math for ProceduralMesh component as a source of geometry
+ */
+class FDisplayClusterWarpBlendMath_WarpProceduralMesh
+{
+public:
+	FDisplayClusterWarpBlendMath_WarpProceduralMesh(const FProcMeshSection& InProcMeshSection)
+		: ProcMeshSection(InProcMeshSection)
+	{ }
+
+public:
+	FBox CalcAABBox()
+	{
+		FDisplayClusterWarpAABB WarpAABB;
+		for (const FProcMeshVertex& VertexIt : ProcMeshSection.ProcVertexBuffer)
+		{
+			WarpAABB.UpdateAABB(VertexIt.Position);
+		}
+
+		return WarpAABB;
+	}
+
+	void CalcSurfaceVectors(FVector& OutSurfaceViewNormal, FVector& OutSurfaceViewPlane)
+	{
+		// Calc static normal and plane
+		const int32 IdxNum = ProcMeshSection.ProcIndexBuffer.Num();
+		const int32 TriNum = IdxNum / 3;
+
+		if (TriNum <= 0)
+		{
+			return;
+		}
+
+		// Calc static normal and plane
+		bool bAverageNormalValid = false;
+		FVector AverageNormal;
+
+		for (int32 TriIdx = 0; TriIdx < TriNum; ++TriIdx)
+		{
+			const uint32& Index0 = ProcMeshSection.ProcIndexBuffer[TriIdx * 3 + 0];
+			const uint32& Index1 = ProcMeshSection.ProcIndexBuffer[TriIdx * 3 + 1];
+			const uint32& Index2 = ProcMeshSection.ProcIndexBuffer[TriIdx * 3 + 2];
+
+			const FVector& Pts0 = ProcMeshSection.ProcVertexBuffer[Index0].Position;
+			const FVector& Pts1 = ProcMeshSection.ProcVertexBuffer[Index1].Position;
+			const FVector& Pts2 = ProcMeshSection.ProcVertexBuffer[Index2].Position;
+
+			const FVector N1 = Pts0 - Pts1;
+			const FVector N2 = Pts2 - Pts1;
+
+			const FVector N = FVector::CrossProduct(N2, N1).GetSafeNormal();
+
+			AverageNormal = bAverageNormalValid ? (N + AverageNormal).GetSafeNormal() : N;
+
+			bAverageNormalValid = true;
+		}
+
+		OutSurfaceViewNormal = AverageNormal;
+
+		//@todo: MeshSurfaceViewPlane not implemented, use MeshSurfaceViewNormal
+		OutSurfaceViewPlane = OutSurfaceViewNormal;
+	}
+
+private:
+	const FProcMeshSection& ProcMeshSection;
+};
diff --git a/Engine/Plugins/Runtime/nDisplay/Source/DisplayClusterWarp/Public/Blueprints/DisplayClusterWarpBlueprint_Enums.h b/Engine/Plugins/Runtime/nDisplay/Source/DisplayClusterWarp/Public/Blueprints/DisplayClusterWarpBlueprint_Enums.h
new file mode 100644
--- /dev/null
+++ b/Engine/Plugins/Runtime/nDisplay/Source/DisplayClusterWarp/Public/Blueprints/DisplayClusterWarpBlueprint_Enums.h
@@ -0,0 +1,32 @@
+// Copyright Epic Games, Inc. All Rights Reserved.
+
+#pragma once
+#include "CoreMinimal.h"
+
+#include "DisplayClusterWarpBlueprint_Enums.generated.h"
+
+/**
+ * Projection mode for the camera that is used as an image source
+ * This projection does not directly use slices from the camera image,
+ * but calculates the camera sub-frustum used to render the sub-images of camera for a particular viewport.
+ * Since the aspect ratio of the AABB geometry in the viewport cannot be equal
+ * to the aspect ratio of the camera, we must use several methods to fit the camera to the geometry.
+ */
+UENUM()
+enum class EDisplayClusterWarpCameraProjectionMode : uint8
+{
+	/** Use camera as source of image.Preserving the camera aspect ratio.
+	 * The camera frame touches the geometry inside
+	 */
+	TouchInside UMETA(DisplayName = "Touch Inside"),
+
+	/** Use camera as source of image.Preserving the camera aspect ratio.
+	 * The camera frame touches the geometry outside
+	 */
+	TouchOutside  UMETA(DisplayName = "Touch Outside"),
+
+	/** Use camera as source of image.
+	 * Stretch the camera image to geometry (the camera aspect ratio will be lost)
+	 */
+	Stretch  UMETA(DisplayName = "Stretch"),
+};
diff --git a/Engine/Plugins/Runtime/nDisplay/Source/DisplayClusterWarp/Public/Blueprints/DisplayClusterWarpGeometry.h b/Engine/Plugins/Runtime/nDisplay/Source/DisplayClusterWarp/Public/Blueprints/DisplayClusterWarpGeometry.h
new file mode 100644
--- /dev/null
+++ b/Engine/Plugins/Runtime/nDisplay/Source/DisplayClusterWarp/Public/Blueprints/DisplayClusterWarpGeometry.h
@@ -0,0 +1,57 @@
+// Copyright Epic Games, Inc. All Rights Reserved.
+
+#pragma once
+#include "CoreMinimal.h"
+#include "DisplayClusterWarpGeometry.generated.h"
+
+/**
+ * 3D geometry that can be used for warping, in an PFM-like format
+ * UE scale used: 1 unit = 1 centimeter
+ */
+USTRUCT(BlueprintType, Category = "nDisplay|WarpGeometry")
+struct FDisplayClusterWarpGeometryPFM
+{
+	GENERATED_BODY()
+
+public:
+	/** Number of vertices horizontally. */
+	UPROPERTY(EditAnywhere, BlueprintReadWrite, Category = "nDisplay|WarpGeometry")
+	int32 Width = 0;
+
+	/** Number of vertices vertically. */
+	UPROPERTY(EditAnywhere, BlueprintReadWrite, Category = "nDisplay|WarpGeometry")
+	int32 Height = 0;
+
+	/** An array with vertices. The total number of vertices in this array must be equal to Width*Height. */
+	UPROPERTY(EditAnywhere, BlueprintReadWrite, Category = "nDisplay|WarpGeometry")
+	TArray<FVector>   Vertices;
+};
+
+/**
+ * 3D geometry that can be used for warping, in an OBJ-like format
+ */
+USTRUCT(BlueprintType, Category = "nDisplay|WarpGeometry")
+struct FDisplayClusterWarpGeometryOBJ
+{
+	GENERATED_BODY()
+
+public:
+	void PostAddFace(int32 f0, int32 f1, int32 f2);
+
+public:
+	/** Vertices of the mesh. */
+	UPROPERTY(EditAnywhere, BlueprintReadWrite, Category = "nDisplay|WarpGeometry")
+	TArray<FVector>   Vertices;
+
+	/** Normal of the vertices of the mesh. */
+	UPROPERTY(EditAnywhere, BlueprintReadWrite, Category = "nDisplay|WarpGeometry")
+	TArray<FVector>   Normal;
+
+	/** UV of the vertices of the mesh. */
+	UPROPERTY(EditAnywhere, BlueprintReadWrite, Category = "nDisplay|WarpGeometry")
+	TArray<FVector2D> UV;
+
+	/** Triangles with mesh vertex indices. */
+	UPROPERTY(EditAnywhere, BlueprintReadWrite, Category = "nDisplay|WarpGeometry")
+	TArray<int32>   Triangles;
+};
diff --git a/Engine/Plugins/Runtime/nDisplay/Source/DisplayClusterWarp/Public/Containers/DisplayClusterWarpAABB.h b/Engine/Plugins/Runtime/nDisplay/Source/DisplayClusterWarp/Public/Containers/DisplayClusterWarpAABB.h
new file mode 100644
--- /dev/null
+++ b/Engine/Plugins/Runtime/nDisplay/Source/DisplayClusterWarp/Public/Containers/DisplayClusterWarpAABB.h
@@ -0,0 +1,92 @@
+// Copyright Epic Games, Inc. All Rights Reserved.
+
+#pragma once
+
+#include "CoreMinimal.h"
+
+/**
+ * Implement AABB math
+ */
+struct FDisplayClusterWarpAABB
+	: public FBox
+{
+	FDisplayClusterWarpAABB()
+		: FBox(FVector(FLT_MAX, FLT_MAX, FLT_MAX), FVector(-FLT_MAX, -FLT_MAX, -FLT_MAX))
+	{ }
+
+	template<typename FArg>
+	inline void UpdateAABB(const FArg X, const FArg Y, const FArg Z)
+	{
+		Min.X = FMath::Min(Min.X, X);
+		Min.Y = FMath::Min(Min.Y, Y);
+		Min.Z = FMath::Min(Min.Z, Z);
+
+		Max.X = FMath::Max(Max.X, X);
+		Max.Y = FMath::Max(Max.Y, Y);
+		Max.Z = FMath::Max(Max.Z, Z);
+	}
+
+	inline void UpdateAABB(const FVector4f& InPts)
+	{
+		if (InPts.W > 0)
+		{
+			UpdateAABB(InPts.X, InPts.Y, InPts.Z);
+		}
+	}
+
+	inline void UpdateAABB(const FVector& InPts)
+	{
+		UpdateAABB(InPts.X, InPts.Y, InPts.Z);
+	}
+
+	inline void UpdateAABB(const TArray<FVector>& InPoints)
+	{
+		for (const FVector& Pts : InPoints)
+		{
+			UpdateAABB(Pts);
+		}
+	}
+
+	inline void UpdateAABB(const FDisplayClusterWarpAABB& InAABB)
+	{
+		operator+=(InAABB);
+	}
+
+	/** The AABB defined by 8 points
+	 * Get AABB 3d-cube points:
+	 */
+	inline FVector GetAABBPts(const uint32 InPtsIndex) const
+	{
+		switch (InPtsIndex)
+		{
+		case 0:
+			return FVector(Max.X, Max.Y, Max.Z);
+
+		case 1:
+			return FVector(Max.X, Max.Y, Min.Z);
+
+		case 2:
+			return FVector(Min.X, Max.Y, Min.Z);
+
+		case 3:
+			return FVector(Min.X, Max.Y, Max.Z);
+
+		case 4:
+			return FVector(Max.X, Min.Y, Max.Z);
+
+		case 5:
+			return FVector(Max.X, Min.Y, Min.Z);
+
+		case 6:
+			return FVector(Min.X, Min.Y, Min.Z);
+
+		case 7:
+			return FVector(Min.X, Min.Y, Max.Z);
+
+		default:
+			break;
+		}
+
+		return GetCenter();
+	}
+};
diff --git a/Engine/Plugins/Runtime/nDisplay/Source/DisplayClusterWarp/Public/Containers/DisplayClusterWarpContainers.h b/Engine/Plugins/Runtime/nDisplay/Source/DisplayClusterWarp/Public/Containers/DisplayClusterWarpContainers.h
new file mode 100644
--- /dev/null
+++ b/Engine/Plugins/Runtime/nDisplay/Source/DisplayClusterWarp/Public/Containers/DisplayClusterWarpContainers.h
@@ -0,0 +1,207 @@
+// Copyright Epic Games, Inc. All Rights Reserved.
+
+#pragma once
+#include "CoreMinimal.h"
+#include "Containers/DisplayClusterWarpEnums.h"
+
+/**
+ * Frustum projection angles.
+ */
+struct FDisplayClusterWarpProjection
+{
+	inline void ResetProjectionAngles()
+	{
+		Left = DBL_MAX;
+		Right = -DBL_MAX;
+		Top = -DBL_MAX;
+		Bottom = DBL_MAX;
+	}
+
+	inline void ExpandProjectionAngles(const FDisplayClusterWarpProjection& In)
+	{
+		Left   = FMath::Min(Left,   In.Left);
+		Right  = FMath::Max(Right,  In.Right);
+		Top    = FMath::Max(Top,    In.Top);
+		Bottom = FMath::Min(Bottom, In.Bottom);
+
+		ZNear = In.ZNear;
+		ZFar = In.ZFar;
+	}
+
+	inline void RotateProjectionAngles90Degree()
+	{
+		double InLeft   = Left;
+		double InRight  = Right;
+		double InTop    = Top;
+		double InBottom = Bottom;
+
+		Left   = InBottom;
+		Right  = InTop;
+		Top    = -InLeft;
+		Bottom = -InRight;
+	}
+
+	inline double ConvertDegreesToProjection(double InAngle) const
+	{
+		const double ProjectedAngle = ZNear * FMath::Tan(FMath::DegreesToRadians(FMath::Abs(InAngle)));
+
+		return (InAngle < 0) ? -ProjectedAngle : ProjectedAngle;
+	}
+
+	inline double ConvertProjectionToDegrees(double InProjectedAngle) const
+	{
+		const double Angle = FMath::RadiansToDegrees(FMath::Atan(FMath::Abs(InProjectedAngle) / ZNear));
+
+		return (InProjectedAngle < 0) ? -Angle : Angle;
+	}
+
+	inline void ConvertProjectionAngles(const EDisplayClusterWarpAngleUnit InUnitType)
+	{
+		if (InUnitType != DataType)
+		{
+			if (DataType == EDisplayClusterWarpAngleUnit::Degrees && InUnitType == EDisplayClusterWarpAngleUnit::Default)
+			{
+				Left   = ConvertDegreesToProjection(Left);
+				Right  = ConvertDegreesToProjection(Right);
+				Bottom = ConvertDegreesToProjection(Bottom);
+				Top    = ConvertDegreesToProjection(Top);
+
+				DataType = InUnitType;
+
+				return;
+			}
+
+			if (DataType == EDisplayClusterWarpAngleUnit::Default && InUnitType == EDisplayClusterWarpAngleUnit::Degrees)
+			{
+				Left   = ConvertProjectionToDegrees(Left);
+				Right  = ConvertProjectionToDegrees(Right);
+				Bottom = ConvertProjectionToDegrees(Bottom);
+				Top    = ConvertProjectionToDegrees(Top);
+
+				DataType = InUnitType;
+
+				return;
+			}
+		}
+	}
+
+	// Unit type for values
+	EDisplayClusterWarpAngleUnit DataType = EDisplayClusterWarpAngleUnit::Default;
+
+	// Projection angles
+	double Left = DBL_MAX;
+	double Right = -DBL_MAX;
+	double Top = -DBL_MAX;
+	double Bottom = DBL_MAX;
+
+	// Clipping planes
+	double ZNear = 1.f;
+	double ZFar = 1.f;
+
+	// Scale
+	double WorldScale = 1.f;
+
+	// Warp projection ViewPoint
+	FVector EyeLocation = FVector::ZeroVector;
+
+	// Camera ViewPoint
+	FRotator CameraRotation = FRotator::ZeroRotator;
+	FVector  CameraLocation = FVector::ZeroVector;
+};
+
+/**
+ * WarpBlend ViewPoint data
+ */
+struct FDisplayClusterWarpViewPoint
+{
+	inline FVector GetEyeLocation() const
+	{
+		return Location + EyeOffset;
+	}
+
+	inline bool IsEqual(const FDisplayClusterWarpViewPoint& InWarpViewPoint, float Precision) const
+	{
+		return (GetEyeLocation() - InWarpViewPoint.GetEyeLocation()).Size() < Precision;
+	}
+
+public:
+	// ViewPoint location
+	FVector Location = FVector::ZeroVector;
+
+	// Eye offset
+	FVector EyeOffset = FVector::ZeroVector;
+
+	// ViewPoint rotation
+	FRotator Rotation = FRotator::ZeroRotator;
+};
+
+/**
+ * MPCDI attributes
+ */
+struct FDisplayClusterWarpMPCDIAttributes
+{
+	// MPCDI profile type
+	EDisplayClusterWarpProfileType ProfileType = EDisplayClusterWarpProfileType::Invalid;
+
+	// Special flags
+	EDisplayClusterWarpMPCDIAttributesFlags Flags = EDisplayClusterWarpMPCDIAttributesFlags::None;
+
+	// additional mpcdi attributes: <Buffer>
+	struct FBuffer
+	{
+		FBuffer()
+			: Resolution(1024, 1024)
+		{ }
+
+		FIntPoint Resolution;
+
+	} Buffer;
+
+	// additional mpcdi attributes: <Region>
+	struct FRegion
+	{
+		FRegion()
+			: Resolution(1024, 1024)
+			, Pos(0.0, 0.0)
+			, Size(1.0, 1.0)
+		{ }
+
+		FIntPoint Resolution;
+		FVector2D Pos;
+		FVector2D Size;
+
+	} Region;
+
+	// additional mpcdi attributes <frustum>
+	struct FFrustum
+	{
+		FFrustum()
+			: Rotator(FRotator::ZeroRotator)
+			, Angles(FVector4::Zero())
+		{ }
+
+		// Frustum direction
+		FRotator Rotator;
+
+		// Frustum angles XYZW = LRTB
+		FVector4 Angles;
+
+	} Frustum;
+
+	// additional mpcdi attributes <coordinateFrame>
+	struct FCoordinateFrame
+	{
+		FCoordinateFrame()
+			: Pos(FVector::ZeroVector)
+			, Yaw(FVector::ZeroVector)
+			, Pitch(FVector::ZeroVector)
+			, Roll(FVector::ZeroVector)
+		{ }
+
+		FVector Pos;
+		FVector Yaw;
+		FVector Pitch;
+		FVector Roll;
+
+	} CoordinateFrame;
+};
diff --git a/Engine/Plugins/Runtime/nDisplay/Source/DisplayClusterWarp/Public/Containers/DisplayClusterWarpContext.h b/Engine/Plugins/Runtime/nDisplay/Source/DisplayClusterWarp/Public/Containers/DisplayClusterWarpContext.h
new file mode 100644
--- /dev/null
+++ b/Engine/Plugins/Runtime/nDisplay/Source/DisplayClusterWarp/Public/Containers/DisplayClusterWarpContext.h
@@ -0,0 +1,91 @@
+// Copyright Epic Games, Inc. All Rights Reserved.
+
+#pragma once
+#include "Containers/DisplayClusterWarpContainers.h"
+#include "Containers/DisplayClusterWarpAABB.h"
+
+class FDisplayClusterWarpEye;
+
+/**
+ *  Warp context
+ */
+struct FDisplayClusterWarpContext
+{
+	bool bIsValid = false;
+
+	// Camera
+	FRotator Rotation = FRotator::ZeroRotator;
+	FVector  Location = FVector::ZeroVector;
+
+	// Frustum projection matrix
+	FMatrix  GeometryProjectionMatrix = FMatrix::Identity;
+
+	// Frustum projection matrix
+	FMatrix  ProjectionMatrix = FMatrix::Identity;
+
+	// From the texture's perspective
+	FMatrix  UVMatrix = FMatrix::Identity;
+
+	FMatrix  TextureMatrix = FMatrix::Identity;
+	FMatrix  RegionMatrix = FMatrix::Identity;
+
+	// From the mesh local space to cave
+	FMatrix  MeshToStageMatrix = FMatrix::Identity;
+};
+
+/** Internal warp data for the one context of the viewport. */
+struct FDisplayClusterWarpData
+{
+	// Input data from projection policy
+	TSharedPtr<FDisplayClusterWarpEye, ESPMode::ThreadSafe> WarpEye;
+
+	// Warp projection data for geometry
+	FDisplayClusterWarpProjection GeometryWarpProjection;
+
+	// Warp projection for render
+	FDisplayClusterWarpProjection WarpProjection;
+
+	// Results
+	FDisplayClusterWarpContext WarpContext;
+
+	// Projection plane matrices
+	FMatrix Local2World = FMatrix::Identity;
+
+	// Current world scale
+	double WorldScale = 1;
+
+	// Indicates if the frustum was rotated from base algorithm to better fit context size
+	// Used to implement hysteresis and avoid potential jumps back and forth.
+	bool bRotateFrustumToFitContextSize = false;
+
+	// Allows to use bRotateFrustumToFitContextSize
+	bool bEnabledRotateFrustumToFitContextSize = true;
+
+	// Other warp settings:
+	EDisplayClusterWarpBlendFrustumType       FrustumType = EDisplayClusterWarpBlendFrustumType::AABB;
+	EDisplayClusterWarpBlendStereoMode         StereoMode = EDisplayClusterWarpBlendStereoMode::AsymmetricAABB;
+
+	// Control projection type
+	bool bFindBestProjectionType = false;
+	EDisplayClusterWarpBlendProjectionType ProjectionType = EDisplayClusterWarpBlendProjectionType::DynamicAABBCenter;
+};
+
+struct FDisplayClusterWarpGeometryContext
+{
+	FMatrix GeometryToOrigin = FMatrix::Identity;
+
+	// AABB of geometry
+	FDisplayClusterWarpAABB AABBox;
+
+	// Static surface average normal for this region
+	FVector SurfaceViewNormal;
+
+	// Static surface average normal from 4 corner points
+	FVector SurfaceViewPlane;
+
+	// Force update caches
+	bool bGeometryChanged = false;
+
+	// Allows you to use a matrix of regions (MPCDI 2D profile area defines a sub-rectangle of the input image to be cut and used for warp).
+	bool bUseRegionMatrix = false;
+};
diff --git a/Engine/Plugins/Runtime/nDisplay/Source/DisplayClusterWarp/Public/Containers/DisplayClusterWarpEnums.h b/Engine/Plugins/Runtime/nDisplay/Source/DisplayClusterWarp/Public/Containers/DisplayClusterWarpEnums.h
new file mode 100644
--- /dev/null
+++ b/Engine/Plugins/Runtime/nDisplay/Source/DisplayClusterWarp/Public/Containers/DisplayClusterWarpEnums.h
@@ -0,0 +1,140 @@
+// Copyright Epic Games, Inc. All Rights Reserved.
+
+#pragma once
+#include "CoreMinimal.h"
+
+/**
+ * Angle values data type.
+ */
+enum class EDisplayClusterWarpAngleUnit : uint8
+{
+	// in projected space
+	Default = 0,
+
+	// In degress 0..360
+	Degrees
+};
+
+/**
+ * MPCDI profile type
+ */
+enum class EDisplayClusterWarpProfileType : uint8
+{
+	// 2D mode
+	warp_2D = 0,
+	// 3D mode (2D + static Frustum)
+	warp_3D,
+	// Advanced 3D mode
+	warp_A3D,
+	// Shader lamps
+	warp_SL,
+
+	Invalid
+};
+
+/**
+ * Type of geometry source
+ */
+enum class EDisplayClusterWarpGeometryType : uint8
+{
+	// Use texture warpmap as geometry source
+	WarpMap = 0,
+	
+	// Use mesh component as geometry source
+	WarpMesh,
+
+	// Use procedural mesh component as geometry source
+	WarpProceduralMesh,
+
+	// Undefined source
+	Invalid
+};
+
+/**
+ * Type of frustum geometry source
+ */
+enum class EDisplayClusterWarpFrustumGeometryType : uint8
+{
+	// Use mpcdi attributes  to build virtual frustum geometry (profile 2D)
+	MPCDIAttributes = 0,
+
+	// Use texture warpmap as frustum geometry source
+	WarpMap,
+
+	// Use mesh component as frustum geometry source
+	WarpMesh,
+
+	// Use procedural mesh component as frustum geometry source
+	WarpProceduralMesh,
+
+	// Undefined source
+	Invalid
+};
+
+/**
+ * The type of data stored in the texture (MPCDI or PFM only)
+ */
+enum class EDisplayClusterWarpBlendTextureType :uint8
+{
+	// Geometry data will be stored as a three-component (RGB) PFM file (Portable Float Map Format)
+	WarpMap = 0,
+
+	// Blend grayscale texture, color multiplier. Use in basic color correction: color = color * alpha;
+	AlphaMap,
+
+	// Bate grayscale texture, for advanced collor correction. Use with AlphaMap: color = (color*alpha*(1 - beta)) + beta;
+	BetaMap
+};
+
+/**
+ * The method used to calc a frustum from geometric points
+ */
+enum class EDisplayClusterWarpBlendFrustumType : uint8
+{
+	// Build frustum from 8 points of AABB
+	AABB = 0,
+
+	// Build frustum from all geometry points
+	FULL,
+
+	// Performance: Build frustum from part of geometry (LOD)
+	LOD,
+};
+
+/**
+ * Type of normal used to make the projection surface
+ */
+enum class EDisplayClusterWarpBlendProjectionType : uint8
+{
+	StaticSurfaceNormal = 0,
+	StaticSurfacePlane,
+	DynamicAABBCenter,
+
+	RuntimeProjectionModes,
+	RuntimeStaticSurfaceNormalInverted,
+	RuntimeStaticSurfacePlaneInverted,
+};
+
+/**
+ * The method used to calculate a stereo projection
+ */
+enum class EDisplayClusterWarpBlendStereoMode : uint8
+{
+	AsymmetricAABB = 0,
+	SymmetricAABB,
+};
+
+/**
+ * MPCDI attributes flags
+ */
+enum class EDisplayClusterWarpMPCDIAttributesFlags : uint8
+{
+	None = 0,
+
+	// Frustum data defined
+	HasFrustum = 1 << 0,
+
+	// CoordinateFrame data defined
+	HasCoordinateFrame = 1 << 1,
+};
+ENUM_CLASS_FLAGS(EDisplayClusterWarpMPCDIAttributesFlags);
diff --git a/Engine/Plugins/Runtime/nDisplay/Source/DisplayClusterWarp/Public/Containers/DisplayClusterWarpEye.h b/Engine/Plugins/Runtime/nDisplay/Source/DisplayClusterWarp/Public/Containers/DisplayClusterWarpEye.h
new file mode 100644
--- /dev/null
+++ b/Engine/Plugins/Runtime/nDisplay/Source/DisplayClusterWarp/Public/Containers/DisplayClusterWarpEye.h
@@ -0,0 +1,60 @@
+// Copyright Epic Games, Inc. All Rights Reserved.
+
+#pragma once
+#include "Containers/DisplayClusterWarpContainers.h"
+#include "Render/Warp/IDisplayClusterWarpPolicy.h"
+#include "Templates/SharedPointer.h"
+
+class IDisplayClusterViewport;
+
+/**
+ * WarpBlend eye data
+ */
+class FDisplayClusterWarpEye
+	: public TSharedFromThis<FDisplayClusterWarpEye, ESPMode::ThreadSafe>
+{
+public:
+	FDisplayClusterWarpEye(const TSharedPtr<IDisplayClusterViewport, ESPMode::ThreadSafe>& InViewport, const uint32 InContextNum)
+		: ViewportWeakPtr(InViewport)
+		, ContextNum(InContextNum)
+	{ }
+
+	inline bool IsEqual(const FDisplayClusterWarpEye& InEye, float Precision) const
+	{
+		if (FMath::IsNearlyEqual(WorldScale, InEye.WorldScale))
+		{
+			return ViewPoint.IsEqual(InEye.ViewPoint, Precision);
+		}
+
+		return false;
+	}
+
+	inline IDisplayClusterViewport* GetViewport() const
+	{
+		return ViewportWeakPtr.IsValid() ? ViewportWeakPtr.Pin().Get() : nullptr;
+	}
+
+public:
+	const TWeakPtr<IDisplayClusterViewport, ESPMode::ThreadSafe> ViewportWeakPtr;
+	const uint32 ContextNum;
+
+	FDisplayClusterWarpViewPoint ViewPoint;
+
+	// Current scene additional settings
+	float WorldScale = 1.f;
+
+	// DCRA world2local
+	FTransform World2LocalTransform;
+
+	// Warp policy
+	TSharedPtr<IDisplayClusterWarpPolicy, ESPMode::ThreadSafe> WarpPolicy;
+
+	// Is geometry context data must should be updated
+	bool bUpdateGeometryContext = true;
+
+	// (optional) this a view target point. By default used AABB center.
+	TOptional<FVector> OverrideViewTarget;
+
+	// (Optional) is the view direction (when this value is used, it also overrides the OverrideViewTarget parameter).
+	TOptional<FVector> OverrideViewDirection;
+};
diff --git a/Engine/Plugins/Runtime/nDisplay/Source/DisplayClusterWarp/Public/Containers/DisplayClusterWarpInitializer.h b/Engine/Plugins/Runtime/nDisplay/Source/DisplayClusterWarp/Public/Containers/DisplayClusterWarpInitializer.h
new file mode 100644
--- /dev/null
+++ b/Engine/Plugins/Runtime/nDisplay/Source/DisplayClusterWarp/Public/Containers/DisplayClusterWarpInitializer.h
@@ -0,0 +1,103 @@
+// Copyright Epic Games, Inc. All Rights Reserved.
+
+#pragma once
+#include "Containers/DisplayClusterWarpContainers.h"
+
+class USceneComponent;
+class UDisplayClusterScreenComponent;
+class UStaticMeshComponent;
+class UProceduralMeshComponent;
+
+/**
+ * Initialize WarpBlend from MPCDI file
+ */
+struct FDisplayClusterWarpInitializer_MPCDIFile
+{
+	// The file with MPCDI calibration data (extension of this file '.mpcdi')
+	FString MPCDIFileName;
+
+	// Name of buffer (mpcdi internal topology)
+	FString BufferId;
+
+	// Name of region (mpcdi internal topology)
+	FString RegionId;
+};
+
+/**
+ * Initialize WarpBlend from MPCDI file and DisplayClusterScreen component
+ */
+struct FDisplayClusterWarpInitializer_MPCDIFile_Profile2DScreen
+	: public FDisplayClusterWarpInitializer_MPCDIFile
+{
+	// Origin component ptr
+	USceneComponent* OriginComponent = nullptr;
+
+	// StaticMesh component ptr (warp geometry data source)
+	UStaticMeshComponent* StaticMeshComponent = nullptr;
+};
+
+/**
+ * Initialize WarpBlend from PFM file
+ */
+struct FDisplayClusterWarpInitializer_PFMFile
+{
+	// The file with region calibration data (extension of this file '.pfm')
+	FString PFMFileName;
+
+	// Geometry scale of this file to UE space (1 = 1 centemeter)
+	float   PFMScale = 1.f;
+
+	// Wrap axis to mpcdi space for unreal 
+	bool    bIsUnrealGameSpace = false;
+
+	// AlphaMap file (supported format: PNG)
+	FString AlphaMapFileName;
+
+	// Control alpha multiplier for AlphaMap texture
+	float   AlphaMapEmbeddedAlpha = 0.f;
+
+	// BetaMap file (supported format: PNG)
+	FString BetaMapFileName;
+
+	// Additional attributes from xml inside MPCDI file
+	FDisplayClusterWarpMPCDIAttributes MPCDIAttributes;
+};
+
+/**
+ * Base struct, not used as initialized
+ */
+struct FDisplayClusterWarpInitializer_BaseMesh
+{
+	// Origin component ptr
+	USceneComponent* OriginComponent = nullptr;
+
+	// Map source mesh geometry UVs to WarpMesh UVs
+	int32 BaseUVIndex = INDEX_NONE;
+	int32 ChromakeyUVIndex = INDEX_NONE;
+};
+
+/**
+ * Initialize WarpBlend from static mesh component
+ */
+struct FDisplayClusterWarpInitializer_StaticMesh
+	: public FDisplayClusterWarpInitializer_BaseMesh
+{
+	// StaticMesh component ptr (warp geometry data source)
+	UStaticMeshComponent* StaticMeshComponent = nullptr;
+
+	// (optional) get geometry data from specified LOD index
+	int32 StaticMeshComponentLODIndex = 0;
+};
+
+/**
+ * Initialize WarpBlend from procedural mesh component
+ */
+struct FDisplayClusterWarpInitializer_ProceduralMesh
+	: public FDisplayClusterWarpInitializer_BaseMesh
+{
+	// ProceduralMesh component ptr (warp geometry data source)
+	UProceduralMeshComponent* ProceduralMeshComponent = nullptr;
+
+	// (optional) get geometry data from specified SectionIndex
+	int32 ProceduralMeshComponentSectionIndex = 0;
+};
diff --git a/Engine/Plugins/Runtime/nDisplay/Source/DisplayClusterWarp/Public/DisplayClusterWarpStrings.h b/Engine/Plugins/Runtime/nDisplay/Source/DisplayClusterWarp/Public/DisplayClusterWarpStrings.h
new file mode 100644
--- /dev/null
+++ b/Engine/Plugins/Runtime/nDisplay/Source/DisplayClusterWarp/Public/DisplayClusterWarpStrings.h
@@ -0,0 +1,13 @@
+// Copyright Epic Games, Inc. All Rights Reserved.
+
+#pragma once
+
+#include "CoreMinimal.h"
+
+namespace DisplayClusterWarpStrings
+{
+	namespace warp
+	{
+		static constexpr const TCHAR* InFrustumFit = TEXT("InFrustumFit");
+	}
+};
diff --git a/Engine/Plugins/Runtime/nDisplay/Source/DisplayClusterWarp/Public/IDisplayClusterWarp.h b/Engine/Plugins/Runtime/nDisplay/Source/DisplayClusterWarp/Public/IDisplayClusterWarp.h
new file mode 100644
--- /dev/null
+++ b/Engine/Plugins/Runtime/nDisplay/Source/DisplayClusterWarp/Public/IDisplayClusterWarp.h
@@ -0,0 +1,78 @@
+// Copyright Epic Games, Inc. All Rights Reserved.
+
+#pragma once
+
+#include "Modules/ModuleInterface.h"
+#include "Modules/ModuleManager.h"
+
+#include "Containers/DisplayClusterWarpInitializer.h"
+#include "IDisplayClusterWarpBlend.h"
+
+class IDisplayClusterWarp : public IModuleInterface
+{
+public:
+	static constexpr const TCHAR* ModuleName = TEXT("DisplayClusterWarp");
+
+public:
+	virtual ~IDisplayClusterWarp() = default;
+
+public:
+	/**
+	 * Singleton-like access to this module's interface.  This is just for convenience!
+	 * Beware of calling this during the shutdown phase, though.  Your module might have been unloaded already.
+	 *
+	 * @return Returns singleton instance, loading the module on demand if needed
+	 */
+	static inline IDisplayClusterWarp& Get()
+	{
+		return FModuleManager::LoadModuleChecked<IDisplayClusterWarp>(IDisplayClusterWarp::ModuleName);
+	}
+
+	/**
+	 * Checks to see if this module is loaded and ready.  It is only valid to call Get() if IsAvailable() returns true.
+	 *
+	 * @return True if the module is loaded and ready to use
+	 */
+	static inline bool IsAvailable()
+	{
+		return FModuleManager::Get().IsModuleLoaded(IDisplayClusterWarp::ModuleName);
+	}
+
+public:
+	/**
+	 * Create new WarpBlend interface from MPCDI File
+	 */
+	virtual TSharedPtr<IDisplayClusterWarpBlend, ESPMode::ThreadSafe> Create(const FDisplayClusterWarpInitializer_MPCDIFile& InConstructParameters) const = 0;
+
+	
+	/**
+	 * Create WarpBlend interface from MPCDI file and DisplayClusterScreen component
+	 */
+	virtual TSharedPtr<IDisplayClusterWarpBlend, ESPMode::ThreadSafe> Create(const FDisplayClusterWarpInitializer_MPCDIFile_Profile2DScreen& InConstructParameters) const = 0;
+
+	/**
+	 * Create new WarpBlend interface from PFM File
+	 */
+	virtual TSharedPtr<IDisplayClusterWarpBlend, ESPMode::ThreadSafe> Create(const FDisplayClusterWarpInitializer_PFMFile& InConstructParameters) const = 0;
+
+	/**
+	 * Create new WarpBlend interface from static mesh
+	 */
+	virtual TSharedPtr<IDisplayClusterWarpBlend, ESPMode::ThreadSafe> Create(const FDisplayClusterWarpInitializer_StaticMesh& InConstructParameters) const = 0;
+
+	/**
+	 * Create new WarpBlend interface from procedural mesh
+	 */
+	virtual TSharedPtr<IDisplayClusterWarpBlend, ESPMode::ThreadSafe> Create(const FDisplayClusterWarpInitializer_ProceduralMesh& InConstructParameters) const = 0;
+
+	/**
+	 * Read MPCDI file structure
+	 * (This function uses MPCDICachedLoader's internal implementation, which greatly reduces the load time of huge MPCDI files.)
+	 * 
+	 * @param InMPCDIFileName        - (in)  filename of the MPCDI file
+	 * @param OutMPCDIFileStructure  - (out) readed mpcdi structure in format: TMap<BufferId, TMap<RegionId, MPCDIAttributes>>
+	 * 
+	 * @return true if success
+	 */
+	virtual bool ReadMPCDFileStructure(const FString& InMPCDIFileName, TMap<FString, TMap<FString, FDisplayClusterWarpMPCDIAttributes>>& OutMPCDIFileStructure) const = 0;
+};
diff --git a/Engine/Plugins/Runtime/nDisplay/Source/DisplayClusterWarp/Public/IDisplayClusterWarpBlend.h b/Engine/Plugins/Runtime/nDisplay/Source/DisplayClusterWarp/Public/IDisplayClusterWarpBlend.h
new file mode 100644
--- /dev/null
+++ b/Engine/Plugins/Runtime/nDisplay/Source/DisplayClusterWarp/Public/IDisplayClusterWarpBlend.h
@@ -0,0 +1,98 @@
+// Copyright Epic Games, Inc. All Rights Reserved.
+
+#pragma once
+#include "CoreMinimal.h"
+#include "Containers/DisplayClusterWarpContext.h"
+#include "Containers/DisplayClusterWarpEye.h"
+
+class UMeshComponent;
+/**
+ * WarpBlend interface for MPCDI and mesh projection policies
+ */
+class IDisplayClusterWarpBlend
+{
+public:
+	virtual ~IDisplayClusterWarpBlend() = default;
+
+public:
+	/** Get TSharedPtr from self. */
+	virtual TSharedPtr<IDisplayClusterWarpBlend, ESPMode::ThreadSafe> ToSharedPtr() = 0;
+	virtual TSharedPtr<const IDisplayClusterWarpBlend, ESPMode::ThreadSafe> ToSharedPtr() const = 0;
+
+	/** Update internal geometry cached data
+	 */
+	virtual bool UpdateGeometryContext(const float InWorldScale) = 0;
+
+	/** Get geometry context data. */
+	virtual const FDisplayClusterWarpGeometryContext& GetGeometryContext() const = 0;
+
+	/**
+	* Calculate warp context data for new eye
+	*
+	* @param InEye - Current eye and scene
+	*
+	* @return - true if the context calculated successfully
+	*/
+	virtual bool CalcFrustumContext(const TSharedPtr<FDisplayClusterWarpEye, ESPMode::ThreadSafe>& InWarpEye) = 0;
+
+	/**
+	 * Access to resources by type
+	 */
+	virtual class FRHITexture* GetTexture(EDisplayClusterWarpBlendTextureType InTextureType) const = 0;
+
+	/**
+	 * Return AlphaMap embedded gamma value
+	 */
+	virtual float GetAlphaMapEmbeddedGamma() const = 0;
+
+	/**
+	 * Get mesh proxy
+	 */
+	virtual const class IDisplayClusterRender_MeshComponentProxy* GetWarpMeshProxy_RenderThread() const = 0;
+
+	/**
+	 * MPCDI profile type
+	 */
+	virtual EDisplayClusterWarpProfileType  GetWarpProfileType() const = 0;
+
+	/**
+	 * Get MPCDI attributes (or their default values if they are not defined)
+	 */
+	virtual const FDisplayClusterWarpMPCDIAttributes& GetMPCDIAttributes() const = 0;
+
+	/**
+	 * Warp geometry type
+	 */
+	virtual EDisplayClusterWarpGeometryType GetWarpGeometryType() const = 0;
+
+	/**
+	 * Warp frustum geometry type
+	 */
+	virtual EDisplayClusterWarpFrustumGeometryType GetWarpFrustumGeometryType() const = 0;
+
+	/**
+	 * Export warp geometry to memory
+	 */
+	virtual bool ExportWarpMapGeometry(struct FDisplayClusterWarpGeometryOBJ& OutMeshData, uint32 InMaxDimension = 0) const = 0;
+
+	/**
+	* Mark internal component ref as dirty for geometry update
+	*
+	* @param InComponentName - (optional) the name of the internal geometry ref component. Empty string for any component name
+	*
+	* @return - true, if there is a marked.
+	*/
+	virtual bool MarkWarpGeometryComponentDirty(const FName& InComponentName) = 0;
+
+	/** Return true, if this WarpBlend support ICVFX pipeline. */
+	virtual bool ShouldSupportICVFX() const = 0;
+
+	/** Get internal warp context data. */
+	virtual FDisplayClusterWarpData& GetWarpData(const uint32 ContextNum) = 0;
+
+	/** Get internal warp context data. */
+	virtual const FDisplayClusterWarpData& GetWarpData(const uint32 ContextNum) const = 0;
+
+	/** Get frustum mesh component. */
+	virtual UMeshComponent* GetStaticMeshComponent() const = 0;
+};
diff --git a/Engine/Plugins/Runtime/nDisplay/nDisplay.uplugin b/Engine/Plugins/Runtime/nDisplay/nDisplay.uplugin
--- a/Engine/Plugins/Runtime/nDisplay/nDisplay.uplugin
+++ b/Engine/Plugins/Runtime/nDisplay/nDisplay.uplugin
@@ -55,6 +55,15 @@
 				"Linux"
 			]
 		},
+		{
+			"Name": "DisplayClusterWarp",
+			"Type": "Runtime",
+			"LoadingPhase": "Default",
+			"PlatformAllowList": [
+				"Win64",
+				"Linux"
+			]
+		},
 		{
 			"Name": "DisplayClusterShaders",
 			"Type": "Runtime",
