nDisplay: typo corrected.

diff --git a/Engine/Plugins/Runtime/nDisplay/Source/DisplayClusterProjection/Public/DisplayClusterProjectionStrings.h b/Engine/Plugins/Runtime/nDisplay/Source/DisplayClusterProjection/Public/DisplayClusterProjectionStrings.h
--- a/Engine/Plugins/Runtime/nDisplay/Source/DisplayClusterProjection/Public/DisplayClusterProjectionStrings.h
+++ b/Engine/Plugins/Runtime/nDisplay/Source/DisplayClusterProjection/Public/DisplayClusterProjectionStrings.h
@@ -5,7 +5,7 @@
 #include "CoreMinimal.h"
 
 /**
-* Textual contants used by projection policies.
+* Text constants used by the projection policies.
 */
 namespace DisplayClusterProjectionStrings
 {
