nDisplay: Fix 'Add New Viewport' and 'Add New ClusterNode' Add/Cancel buttons being cut off.

diff --git a/Engine/Plugins/Runtime/nDisplay/Source/DisplayClusterConfigurator/Private/ClusterConfiguration/SDisplayClusterConfiguratorNewClusterItemDialog.cpp b/Engine/Plugins/Runtime/nDisplay/Source/DisplayClusterConfigurator/Private/ClusterConfiguration/SDisplayClusterConfiguratorNewClusterItemDialog.cpp
--- a/Engine/Plugins/Runtime/nDisplay/Source/DisplayClusterConfigurator/Private/ClusterConfiguration/SDisplayClusterConfiguratorNewClusterItemDialog.cpp
+++ b/Engine/Plugins/Runtime/nDisplay/Source/DisplayClusterConfigurator/Private/ClusterConfiguration/SDisplayClusterConfiguratorNewClusterItemDialog.cpp
@@ -261,7 +261,7 @@ void SDisplayClusterConfiguratorNewClusterItemDialog::Construct(const FArguments
 			.AutoHeight()
 			.Padding(2)
 			[
-				SNew(SUniformGridPanel)
+				SAssignNew(AddCancelButtons, SUniformGridPanel)
 				.MinDesiredSlotHeight(FAppStyle::Get().GetFloat("StandardDialog.MinDesiredSlotHeight"))
 				.MinDesiredSlotWidth(FAppStyle::Get().GetFloat("StandardDialog.MinDesiredSlotWidth"))
 				.SlotPadding(FAppStyle::Get().GetMargin("StandardDialog.SlotPadding"))
@@ -387,7 +387,7 @@ FOptionalSize SDisplayClusterConfiguratorNewClusterItemDialog::GetDetailsMaxDesi
 {
 	// The maximum size of the details panel box is dynamically computed as the max window size minus the size of all
 	// other elements on the window, including any footer content.
-	return MaxWindowHeight - 128 - FooterContent->GetTickSpaceGeometry().GetLocalSize().Y;
+	return MaxWindowHeight - 128 - FooterContent->GetTickSpaceGeometry().GetLocalSize().Y - AddCancelButtons->GetTickSpaceGeometry().GetLocalSize().Y;
 }
 
 #undef LOCTEXT_NAMESPACE
\ No newline at end of file
diff --git a/Engine/Plugins/Runtime/nDisplay/Source/DisplayClusterConfigurator/Private/ClusterConfiguration/SDisplayClusterConfiguratorNewClusterItemDialog.h b/Engine/Plugins/Runtime/nDisplay/Source/DisplayClusterConfigurator/Private/ClusterConfiguration/SDisplayClusterConfiguratorNewClusterItemDialog.h
--- a/Engine/Plugins/Runtime/nDisplay/Source/DisplayClusterConfigurator/Private/ClusterConfiguration/SDisplayClusterConfiguratorNewClusterItemDialog.h
+++ b/Engine/Plugins/Runtime/nDisplay/Source/DisplayClusterConfigurator/Private/ClusterConfiguration/SDisplayClusterConfiguratorNewClusterItemDialog.h
@@ -93,6 +93,8 @@ private:
 private:
 	TWeakPtr<SWindow> ParentWindow;
 	TSharedPtr<SWidget> FooterContent;
+	/** The container for the Add and Cancel buttons. */
+	TSharedPtr<SWidget> AddCancelButtons;
 	TSharedPtr<IDetailsView> DetailsView;
 	TSharedPtr<SComboBox<TSharedPtr<FString>>> ParentComboBox;
 	TSharedPtr<SComboBox<TSharedPtr<FDisplayClusterConfiguratorPresetSize>>> PresetsComboBox;
