nDisplay: fixed -server editor mode for mesh projection policy

diff --git a/Engine/Plugins/Runtime/nDisplay/Source/DisplayCluster/Private/Game/EngineClasses/Scene/DisplayClusterRootActorEditor.cpp b/Engine/Plugins/Runtime/nDisplay/Source/DisplayCluster/Private/Game/EngineClasses/Scene/DisplayClusterRootActorEditor.cpp
--- a/Engine/Plugins/Runtime/nDisplay/Source/DisplayCluster/Private/Game/EngineClasses/Scene/DisplayClusterRootActorEditor.cpp
+++ b/Engine/Plugins/Runtime/nDisplay/Source/DisplayCluster/Private/Game/EngineClasses/Scene/DisplayClusterRootActorEditor.cpp
@@ -406,6 +406,12 @@ void ADisplayClusterRootActor::ResetClusterNodePreviewRendering_Editor()
 
 bool ADisplayClusterRootActor::ImplUpdatePreviewRenderFrame_Editor(const FString& InClusterNodeId)
 {
+	// Skip rendering on dedicated server
+	if (GetGameInstance() && GetGameInstance()->IsDedicatedServerInstance())
+	{
+		return false;
+	}
+
 	// Update cluster node for render:
 	if (PreviewRenderFrameClusterNodeId != InClusterNodeId)
 	{
