Fix false positive static analysis warning
#rnx

diff --git a/Engine/Plugins/Runtime/nDisplay/Source/DisplayCluster/Private/Render/Viewport/Resource/DisplayClusterViewportResourceSettings.cpp b/Engine/Plugins/Runtime/nDisplay/Source/DisplayCluster/Private/Render/Viewport/Resource/DisplayClusterViewportResourceSettings.cpp
--- a/Engine/Plugins/Runtime/nDisplay/Source/DisplayCluster/Private/Render/Viewport/Resource/DisplayClusterViewportResourceSettings.cpp
+++ b/Engine/Plugins/Runtime/nDisplay/Source/DisplayCluster/Private/Render/Viewport/Resource/DisplayClusterViewportResourceSettings.cpp
@@ -14,8 +14,9 @@
 FDisplayClusterViewportResourceSettings::FDisplayClusterViewportResourceSettings(const FDisplayClusterRenderFrameSettings& InRenderFrameSettings, FViewport* InViewport)
 	: ClusterNodeId(InRenderFrameSettings.ClusterNodeId)
 {
-	if (FRHITexture2D* ViewportTexture = InViewport ? InViewport->GetRenderTargetTexture() : nullptr)
+	if (InViewport && InViewport->GetRenderTargetTexture())
 	{
+		FRHITexture2D* ViewportTexture = InViewport->GetRenderTargetTexture();
 		Format = ViewportTexture->GetFormat();
 
 		if (EnumHasAnyFlags(ViewportTexture->GetFlags(), TexCreate_SRGB))
