nDisplay: fixed compilation issue #327709 for Horde

After the warpblend interface was moved to a separate DisplayClusterWarp module, all of its log variables were moved there as well.
However, one log variable was not removed, and this issue will be fixed in this CL.

diff --git a/Engine/Plugins/Runtime/nDisplay/Source/DisplayClusterShaders/Private/DisplayClusterShadersLog.cpp b/Engine/Plugins/Runtime/nDisplay/Source/DisplayClusterShaders/Private/DisplayClusterShadersLog.cpp
--- a/Engine/Plugins/Runtime/nDisplay/Source/DisplayClusterShaders/Private/DisplayClusterShadersLog.cpp
+++ b/Engine/Plugins/Runtime/nDisplay/Source/DisplayClusterShaders/Private/DisplayClusterShadersLog.cpp
@@ -3,4 +3,3 @@
 #include "DisplayClusterShadersLog.h"
 
 DEFINE_LOG_CATEGORY(LogDisplayClusterShaders);
-DEFINE_LOG_CATEGORY(LogDisplayClusterWarpBlend);
diff --git a/Engine/Plugins/Runtime/nDisplay/Source/DisplayClusterShaders/Private/DisplayClusterShadersLog.h b/Engine/Plugins/Runtime/nDisplay/Source/DisplayClusterShaders/Private/DisplayClusterShadersLog.h
--- a/Engine/Plugins/Runtime/nDisplay/Source/DisplayClusterShaders/Private/DisplayClusterShadersLog.h
+++ b/Engine/Plugins/Runtime/nDisplay/Source/DisplayClusterShaders/Private/DisplayClusterShadersLog.h
@@ -6,4 +6,3 @@
 
 // Plugin-wide log categories
 DECLARE_LOG_CATEGORY_EXTERN(LogDisplayClusterShaders, Log, All);
-DECLARE_LOG_CATEGORY_EXTERN(LogDisplayClusterWarpBlend, Log, All);
