nDisplay, TextureShare:
- Added a new function to the IDisplayClusterPostProcessFactory interface
- slightly revised logic for PP on the nDisplay side
- The TextureShareDisplayCluster postprocess is redesigned as a singleton.

diff --git a/Engine/Plugins/Runtime/nDisplay/Source/DisplayCluster/Private/Render/Viewport/Configuration/DisplayClusterViewportConfigurationBase.cpp b/Engine/Plugins/Runtime/nDisplay/Source/DisplayCluster/Private/Render/Viewport/Configuration/DisplayClusterViewportConfigurationBase.cpp
--- a/Engine/Plugins/Runtime/nDisplay/Source/DisplayCluster/Private/Render/Viewport/Configuration/DisplayClusterViewportConfigurationBase.cpp
+++ b/Engine/Plugins/Runtime/nDisplay/Source/DisplayCluster/Private/Render/Viewport/Configuration/DisplayClusterViewportConfigurationBase.cpp
@@ -250,10 +250,7 @@ FDisplayClusterViewportConfigurationInstanceData const* FDisplayClusterViewportC
 
 void FDisplayClusterViewportConfigurationBase::AddInternalPostprocess(const FString& InPostprocessName)
 {
-	if (DisabledPostprocessNames.Find(InPostprocessName) == INDEX_NONE)
-	{
-		InternalPostprocessNames.AddUnique(InPostprocessName);
-	}
+	InternalPostprocessNames.AddUnique(InPostprocessName);
 }
 
 void FDisplayClusterViewportConfigurationBase::UpdateClusterNodePostProcess(const FString& InClusterNodeId)
@@ -309,6 +306,12 @@ void FDisplayClusterViewportConfigurationBase::UpdateClusterNodePostProcess(cons
 			// Create InternalPostprocess
 			for (const FString& InternalPostprocessId : InternalPostprocessNames)
 			{
+				if (DisabledPostprocessNames.Find(InternalPostprocessId) != INDEX_NONE)
+				{
+					// Always skip a postprocess that cannot be initialized
+					continue;
+				}
+
 				TSharedPtr<IDisplayClusterPostProcess, ESPMode::ThreadSafe> ExistPostProcess = PPManager->FindPostProcess(InternalPostprocessId);
 				if (!ExistPostProcess.IsValid())
 				{
@@ -316,7 +319,7 @@ void FDisplayClusterViewportConfigurationBase::UpdateClusterNodePostProcess(cons
 					FDisplayClusterConfigurationPostprocess ConfigurationPostprocess;
 					ConfigurationPostprocess.Type = InternalPostprocessId;
 
-					if (!PPManager->CreatePostprocess(InternalPostprocessId, &ConfigurationPostprocess))
+					if (PPManager->CanBeCreated(&ConfigurationPostprocess) && !PPManager->CreatePostprocess(InternalPostprocessId, &ConfigurationPostprocess))
 					{
 						// Can't create... Disable this postprocess
 						DisabledPostprocessNames.AddUnique(InternalPostprocessId);
@@ -329,6 +332,12 @@ void FDisplayClusterViewportConfigurationBase::UpdateClusterNodePostProcess(cons
 			// Create and update PP
 			for (const TPair<FString, FDisplayClusterConfigurationPostprocess>& It : ClusterNode->Postprocess)
 			{
+				if (DisabledPostprocessNames.Find(It.Value.Type) != INDEX_NONE)
+				{
+					// Always skip a postprocess that cannot be initialized
+					continue;
+				}
+
 				TSharedPtr<IDisplayClusterPostProcess, ESPMode::ThreadSafe> ExistPostProcess = PPManager->FindPostProcess(It.Key);
 				if (ExistPostProcess.IsValid())
 				{
@@ -339,7 +348,13 @@ void FDisplayClusterViewportConfigurationBase::UpdateClusterNodePostProcess(cons
 				}
 				else
 				{
-					PPManager->CreatePostprocess(It.Key, &It.Value);
+					if (PPManager->CanBeCreated(&It.Value) && !PPManager->CreatePostprocess(It.Key, &It.Value))
+					{
+						// Can't create... Disable this postprocess
+						DisabledPostprocessNames.AddUnique(It.Value.Type);
+
+						UE_LOG(LogDisplayClusterViewport, Error, TEXT("Can't create postprocess '%s' on cluster node '%s': Disabled"), *It.Value.Type, *InClusterNodeId);
+					}
 				}
 			}
 
diff --git a/Engine/Plugins/Runtime/nDisplay/Source/DisplayCluster/Private/Render/Viewport/Postprocess/DisplayClusterViewportPostProcessManager.cpp b/Engine/Plugins/Runtime/nDisplay/Source/DisplayCluster/Private/Render/Viewport/Postprocess/DisplayClusterViewportPostProcessManager.cpp
--- a/Engine/Plugins/Runtime/nDisplay/Source/DisplayCluster/Private/Render/Viewport/Postprocess/DisplayClusterViewportPostProcessManager.cpp
+++ b/Engine/Plugins/Runtime/nDisplay/Source/DisplayCluster/Private/Render/Viewport/Postprocess/DisplayClusterViewportPostProcessManager.cpp
@@ -143,6 +143,19 @@ const TArray<FString> FDisplayClusterViewportPostProcessManager::GetPostprocess(
 	return ExistPostProcess;
 }
 
+bool FDisplayClusterViewportPostProcessManager::CanBeCreated(const FDisplayClusterConfigurationPostprocess* InConfigurationPostprocess) const
+{
+	check(IsInGameThread());
+	check(InConfigurationPostprocess != nullptr);
+
+	IDisplayClusterRenderManager* const DCRenderManager = IDisplayCluster::Get().GetRenderMgr();
+	check(DCRenderManager);
+
+	const TSharedPtr<IDisplayClusterPostProcessFactory> PostProcessFactory = DCRenderManager->GetPostProcessFactory(InConfigurationPostprocess->Type);
+
+	return PostProcessFactory.IsValid() && PostProcessFactory->CanBeCreated(InConfigurationPostprocess);
+}
+
 bool FDisplayClusterViewportPostProcessManager::CreatePostprocess(const FString& InPostprocessId, const FDisplayClusterConfigurationPostprocess* InConfigurationPostprocess)
 {
 	check(IsInGameThread());
diff --git a/Engine/Plugins/Runtime/nDisplay/Source/DisplayCluster/Private/Render/Viewport/Postprocess/DisplayClusterViewportPostProcessManager.h b/Engine/Plugins/Runtime/nDisplay/Source/DisplayCluster/Private/Render/Viewport/Postprocess/DisplayClusterViewportPostProcessManager.h
--- a/Engine/Plugins/Runtime/nDisplay/Source/DisplayCluster/Private/Render/Viewport/Postprocess/DisplayClusterViewportPostProcessManager.h
+++ b/Engine/Plugins/Runtime/nDisplay/Source/DisplayCluster/Private/Render/Viewport/Postprocess/DisplayClusterViewportPostProcessManager.h
@@ -59,6 +59,20 @@ public:
 	{ return OutputRemap; }
 
 public:
+	/** Special rules for creating a PP instance.
+	* Note: This function must be called before Create() on the nDisplay side.
+	* (useful to avoid unnecessary nDisplay log messages).
+	*
+	* For example, some postprocess types may use an instance as a singleton.
+	* In this case, Create() returns nullptr if the PP instance already exists,
+	* and as a result nDisplay issues an error message to the log, and then avoids creating this PP in the future.
+	*
+	* @param InConfigurationPostProcess - postprocess configuration
+	*
+	* @return true if a postprocess with this type can be created at this time.
+	*/
+	bool CanBeCreated(const FDisplayClusterConfigurationPostprocess* InConfigurationPostprocess) const;
+
 	bool CreatePostprocess(const FString& InPostprocessId, const FDisplayClusterConfigurationPostprocess* InConfigurationPostprocess);
 	bool RemovePostprocess(const FString& InPostprocessId);
 	bool UpdatePostprocess(const FString& InPostprocessId, const FDisplayClusterConfigurationPostprocess* InConfigurationPostprocess);
diff --git a/Engine/Plugins/Runtime/nDisplay/Source/DisplayCluster/Public/Render/PostProcess/IDisplayClusterPostProcessFactory.h b/Engine/Plugins/Runtime/nDisplay/Source/DisplayCluster/Public/Render/PostProcess/IDisplayClusterPostProcessFactory.h
--- a/Engine/Plugins/Runtime/nDisplay/Source/DisplayCluster/Public/Render/PostProcess/IDisplayClusterPostProcessFactory.h
+++ b/Engine/Plugins/Runtime/nDisplay/Source/DisplayCluster/Public/Render/PostProcess/IDisplayClusterPostProcessFactory.h
@@ -19,7 +19,6 @@ public:
 	/**
 	* Creates a output post process instance
 	*
-	* @param PostProcessType - PostProcess type, same as specified on registration (useful if the same factory is registered for multiple postprocess types)
 	* @param PostProcessPolicyId - ID of a postprocess
 	* @param InConfigurationPostProcess - postprocess configuration
 	*
@@ -27,4 +26,20 @@ public:
 	*/
 	virtual TSharedPtr<IDisplayClusterPostProcess, ESPMode::ThreadSafe> Create(const FString& PostProcessId, const struct FDisplayClusterConfigurationPostprocess* InConfigurationPostProcess) = 0;
 	
+	/** Special rules for creating a PP instance.
+	* Note: This function must be called before Create() on the nDisplay side.
+	* (useful to avoid unnecessary nDisplay log messages).
+	*
+	* For example, some postprocess types may use an instance as a singleton.
+	* In this case, Create() returns nullptr if the PP instance already exists,
+	* and as a result nDisplay issues an error message to the log, and then avoids creating this PP in the future.
+	*
+	* @param InConfigurationPostProcess - postprocess configuration
+	*
+	* @return true if a postprocess with this type can be created at this time.
+	*/
+	virtual bool CanBeCreated(const FDisplayClusterConfigurationPostprocess* InConfigurationPostProcess) const
+	{
+		return true;
+	}
 };
diff --git a/Engine/Plugins/VirtualProduction/TextureShare/Source/TextureShareDisplayCluster/Private/PostProcess/TextureSharePostprocess.cpp b/Engine/Plugins/VirtualProduction/TextureShare/Source/TextureShareDisplayCluster/Private/PostProcess/TextureSharePostprocess.cpp
--- a/Engine/Plugins/VirtualProduction/TextureShare/Source/TextureShareDisplayCluster/Private/PostProcess/TextureSharePostprocess.cpp
+++ b/Engine/Plugins/VirtualProduction/TextureShare/Source/TextureShareDisplayCluster/Private/PostProcess/TextureSharePostprocess.cpp
@@ -63,9 +63,12 @@ FTextureSharePostprocess::FTextureSharePostprocess(const FString& PostprocessId,
 {
 	if (TextureShareAPI().IsObjectExist(DisplayClusterStrings::DefaultShareName))
 	{
-		UE_LOG(LogTextureShareDisplayClusterPostProcess, Error, TEXT("TextureShareDisplayCluster: Failed - Already exist"));
+		// The old object can still exists, when referenced.
+		// In this case, we show this warning to know that the TS object is still referenced by someone.
+		UE_LOG(LogTextureShareDisplayClusterPostProcess, Warning, TEXT("TextureShareDisplayCluster: TS object for nDisplay are still referenced by someone"));
 	}
-	else
+
+	// Re-use TextureShare object for nDisplay
 	{
 		Object = TextureShareAPI().GetOrCreateObject(DisplayClusterStrings::DefaultShareName);
 		if (Object.IsValid())
diff --git a/Engine/Plugins/VirtualProduction/TextureShare/Source/TextureShareDisplayCluster/Private/PostProcess/TextureSharePostprocessFactory.cpp b/Engine/Plugins/VirtualProduction/TextureShare/Source/TextureShareDisplayCluster/Private/PostProcess/TextureSharePostprocessFactory.cpp
--- a/Engine/Plugins/VirtualProduction/TextureShare/Source/TextureShareDisplayCluster/Private/PostProcess/TextureSharePostprocessFactory.cpp
+++ b/Engine/Plugins/VirtualProduction/TextureShare/Source/TextureShareDisplayCluster/Private/PostProcess/TextureSharePostprocessFactory.cpp
@@ -13,7 +13,24 @@ TSharedPtr<IDisplayClusterPostProcess, ESPMode::ThreadSafe> FTextureSharePostpro
 {
 	check(InConfigurationPostProcess != nullptr);
 
+	if (ExistsTextureSharePostprocess.IsValid())
+	{
+		UE_LOG(LogTextureShareDisplayClusterPostProcess, Error, TEXT("Can't Instantiating postprocess <%s> id='%s' : There can only be one TextureShare postprocess per application. "), *InConfigurationPostProcess->Type, *InPostProcessId);
+
+		return nullptr;
+	}
+
 	UE_LOG(LogTextureShareDisplayClusterPostProcess, Verbose, TEXT("Instantiating postprocess <%s> id='%s'"), *InConfigurationPostProcess->Type, *InPostProcessId);
+	TSharedPtr<IDisplayClusterPostProcess, ESPMode::ThreadSafe> Result = MakeShared<FTextureSharePostprocess, ESPMode::ThreadSafe>(InPostProcessId, InConfigurationPostProcess);
+
+	// Save weak ref to TS postprocess instance
+	ExistsTextureSharePostprocess = Result;
 
-	return  MakeShared<FTextureSharePostprocess, ESPMode::ThreadSafe>(InPostProcessId, InConfigurationPostProcess);
+	return Result;
+}
+
+bool FTextureSharePostprocessFactory::CanBeCreated(const FDisplayClusterConfigurationPostprocess* InConfigurationPostProcess) const
+{
+	// This postprocess used as a singletone
+	return !ExistsTextureSharePostprocess.IsValid();
 }
diff --git a/Engine/Plugins/VirtualProduction/TextureShare/Source/TextureShareDisplayCluster/Private/PostProcess/TextureSharePostprocessFactory.h b/Engine/Plugins/VirtualProduction/TextureShare/Source/TextureShareDisplayCluster/Private/PostProcess/TextureSharePostprocessFactory.h
--- a/Engine/Plugins/VirtualProduction/TextureShare/Source/TextureShareDisplayCluster/Private/PostProcess/TextureSharePostprocessFactory.h
+++ b/Engine/Plugins/VirtualProduction/TextureShare/Source/TextureShareDisplayCluster/Private/PostProcess/TextureSharePostprocessFactory.h
@@ -13,4 +13,10 @@ class FTextureSharePostprocessFactory
 public:
 	// Constructor for postprocess
 	virtual TSharedPtr<IDisplayClusterPostProcess, ESPMode::ThreadSafe> Create(const FString& PostProcessId, const struct FDisplayClusterConfigurationPostprocess* InConfigurationPostProcess) override;
+
+	virtual bool CanBeCreated(const FDisplayClusterConfigurationPostprocess* InConfigurationPostProcess) const;
+
+private:
+	// There can only be one TextureShare postprocess per application
+	TWeakPtr<IDisplayClusterPostProcess, ESPMode::ThreadSafe> ExistsTextureSharePostprocess;
 };
