nDisplay: Reworked overscan and custom frustum. Added new features to the projection policy interface.

- Overscan and CustomFrustum settings are now publicly available
- Cleaning and reworking the code related to Overscan and CustomFrustum
- Add new functions to IDisplayClusterProjectionPolicy
	* Ability to override viewport rendering settings from the projection policy
	* Ability to override ViewportProxy internal function that copies 'InternalRenderTargetResource' to 'InputShaderResource'.
- Fixed bug: Viewport override doesn't work properly

diff --git a/Engine/Plugins/Runtime/nDisplay/Source/DisplayCluster/Private/Game/EngineClasses/Scene/DisplayClusterICVFXCameraComponent.cpp b/Engine/Plugins/Runtime/nDisplay/Source/DisplayCluster/Private/Game/EngineClasses/Scene/DisplayClusterICVFXCameraComponent.cpp
--- a/Engine/Plugins/Runtime/nDisplay/Source/DisplayCluster/Private/Game/EngineClasses/Scene/DisplayClusterICVFXCameraComponent.cpp
+++ b/Engine/Plugins/Runtime/nDisplay/Source/DisplayCluster/Private/Game/EngineClasses/Scene/DisplayClusterICVFXCameraComponent.cpp
@@ -6,7 +6,9 @@
 #include "Cluster/IPDisplayClusterClusterManager.h"
 
 #include "Render/Viewport/Containers/DisplayClusterViewport_CameraMotionBlur.h"
-#include "Render/Viewport/Containers/ImplDisplayClusterViewport_CustomFrustum.h"
+#include "Render/Viewport/Containers/DisplayClusterViewport_CustomFrustumRuntimeSettings.h"
+#include "Render/Viewport/Configuration/DisplayClusterViewportConfigurationHelpers_ICVFX.h"
+
 #include "Components/DisplayClusterCameraComponent.h"
 #include "DisplayClusterRootActor.h"
 
@@ -183,61 +185,28 @@ void UDisplayClusterICVFXCameraComponent::UpdateOverscanEstimatedFrameSize()
 	const FDisplayClusterConfigurationICVFX_StageSettings& StageSettings = RootActor->GetStageSettings();
 
 	const float CameraBufferRatio = CameraSettings.GetCameraBufferRatio(StageSettings);
-	const FIntPoint InnerFrustumSize = CameraSettings.GetCameraFrameSize(StageSettings);
-
-	float InnerFrustumResolutionWidth = InnerFrustumSize.X * CameraBufferRatio;
-	float InnerFrustumResolutionHeight = InnerFrustumSize.Y * CameraBufferRatio;
+	const FIntPoint InnerFrustumResolution = CameraSettings.GetCameraFrameSize(StageSettings) * CameraBufferRatio;
+	const FIntPoint UpscaledInnerFrustumResolution = InnerFrustumResolution * CameraSettings.CustomFrustum.FieldOfViewMultiplier;
 
-	float EstimatedOverscanResolutionWidth = InnerFrustumResolutionWidth;
-	float EstimatedOverscanResolutionHeight = InnerFrustumResolutionHeight;
-
-	//calculate estimate
-	{
-		EstimatedOverscanResolutionWidth = InnerFrustumResolutionWidth * CameraSettings.CustomFrustum.FieldOfViewMultiplier;
-		EstimatedOverscanResolutionHeight = InnerFrustumResolutionHeight * CameraSettings.CustomFrustum.FieldOfViewMultiplier;
+	// Read configuration data
+	FDisplayClusterViewport_CustomFrustumSettings CustomFrustumSettings;
+	FDisplayClusterViewportConfigurationHelpers_ICVFX::UpdateCameraCustomFrustum(CameraSettings.CustomFrustum, CustomFrustumSettings);
 
-		FDisplayClusterViewport_CustomFrustumSettings CustomFrustumSettings;
+	// Calculate overscan to temp var
+	FIntRect FinalViewportRect(FIntPoint(0,0), UpscaledInnerFrustumResolution);
+	FDisplayClusterViewport_CustomFrustumRuntimeSettings CustomFrustumRuntimeSettings;
+	FDisplayClusterViewport_CustomFrustumRuntimeSettings::UpdateCustomFrustumSettings(GetName(), CustomFrustumSettings, CustomFrustumRuntimeSettings, FinalViewportRect);
 
-		if (CameraSettings.CustomFrustum.Mode == EDisplayClusterConfigurationViewportCustomFrustumMode::Percent)
-		{
-			const float ConvertToPercent = 0.01;
-			CustomFrustumSettings.CustomFrustumPercent.Left = FDisplayClusterViewport_OverscanSettings::ClampPercent(CameraSettings.CustomFrustum.Left * ConvertToPercent);
-			CustomFrustumSettings.CustomFrustumPercent.Right = FDisplayClusterViewport_OverscanSettings::ClampPercent(CameraSettings.CustomFrustum.Right * ConvertToPercent);
-			CustomFrustumSettings.CustomFrustumPercent.Top = FDisplayClusterViewport_OverscanSettings::ClampPercent(CameraSettings.CustomFrustum.Top * ConvertToPercent);
-			CustomFrustumSettings.CustomFrustumPercent.Bottom = FDisplayClusterViewport_OverscanSettings::ClampPercent(CameraSettings.CustomFrustum.Bottom * ConvertToPercent);
-		}
-		else if (CameraSettings.CustomFrustum.Mode == EDisplayClusterConfigurationViewportCustomFrustumMode::Pixels)
-		{
-			CustomFrustumSettings.CustomFrustumPercent.Left = FDisplayClusterViewport_OverscanSettings::ClampPercent(CameraSettings.CustomFrustum.Left / EstimatedOverscanResolutionWidth);
-			CustomFrustumSettings.CustomFrustumPercent.Right = FDisplayClusterViewport_OverscanSettings::ClampPercent(CameraSettings.CustomFrustum.Right / EstimatedOverscanResolutionWidth);
-			CustomFrustumSettings.CustomFrustumPercent.Top = FDisplayClusterViewport_OverscanSettings::ClampPercent(CameraSettings.CustomFrustum.Top / EstimatedOverscanResolutionHeight);
-			CustomFrustumSettings.CustomFrustumPercent.Bottom = FDisplayClusterViewport_OverscanSettings::ClampPercent(CameraSettings.CustomFrustum.Bottom / EstimatedOverscanResolutionHeight);
-		}
+	// Show calculated values
+	CameraSettings.CustomFrustum.InnerFrustumResolution = InnerFrustumResolution;
+	CameraSettings.CustomFrustum.EstimatedOverscanResolution = FinalViewportRect.Size();
 
-		// Calc pixels from percent
-		CustomFrustumSettings.CustomFrustumPixels.Left = FMath::RoundToInt(EstimatedOverscanResolutionWidth * CustomFrustumSettings.CustomFrustumPercent.Left);
-		CustomFrustumSettings.CustomFrustumPixels.Right = FMath::RoundToInt(EstimatedOverscanResolutionWidth * CustomFrustumSettings.CustomFrustumPercent.Right);
-		CustomFrustumSettings.CustomFrustumPixels.Top = FMath::RoundToInt(EstimatedOverscanResolutionHeight * CustomFrustumSettings.CustomFrustumPercent.Top);
-		CustomFrustumSettings.CustomFrustumPixels.Bottom = FMath::RoundToInt(EstimatedOverscanResolutionHeight * CustomFrustumSettings.CustomFrustumPercent.Bottom);
+	const int32 EstimatedPixel = CameraSettings.CustomFrustum.EstimatedOverscanResolution.X * CameraSettings.CustomFrustum.EstimatedOverscanResolution.Y;
+	const int32 BasePixels = CameraSettings.CustomFrustum.InnerFrustumResolution.X * CameraSettings.CustomFrustum.InnerFrustumResolution.Y;
 
-		const FIntPoint AdjustmentSize = CustomFrustumSettings.CustomFrustumPixels.Size();
-		EstimatedOverscanResolutionWidth += AdjustmentSize.X;
-		EstimatedOverscanResolutionHeight += AdjustmentSize.Y;
-	}
-
-	if (CameraSettings.CustomFrustum.bEnable && CameraSettings.CustomFrustum.bAdaptResolution)
-	{
-		InnerFrustumResolutionWidth = EstimatedOverscanResolutionWidth;
-		InnerFrustumResolutionHeight = EstimatedOverscanResolutionHeight;
-	}
-
-	CameraSettings.CustomFrustum.InnerFrustumResolution = FIntPoint(InnerFrustumResolutionWidth, InnerFrustumResolutionHeight);
-	CameraSettings.CustomFrustum.EstimatedOverscanResolution = FIntPoint(EstimatedOverscanResolutionWidth, EstimatedOverscanResolutionHeight);
-
-	CameraSettings.CustomFrustum.OverscanPixelsIncrease = ((float)(EstimatedOverscanResolutionWidth * EstimatedOverscanResolutionHeight) / (float)(InnerFrustumResolutionWidth * InnerFrustumResolutionHeight));
+	CameraSettings.CustomFrustum.OverscanPixelsIncrease = ((float)(EstimatedPixel) / (float)(BasePixels));
 }
 
-
 FDisplayClusterViewport_CameraMotionBlur UDisplayClusterICVFXCameraComponent::GetMotionBlurParameters()
 {
 	FDisplayClusterViewport_CameraMotionBlur OutParameters;
diff --git a/Engine/Plugins/Runtime/nDisplay/Source/DisplayCluster/Private/Render/Viewport/Configuration/DisplayClusterViewportConfigurationHelpers.cpp b/Engine/Plugins/Runtime/nDisplay/Source/DisplayCluster/Private/Render/Viewport/Configuration/DisplayClusterViewportConfigurationHelpers.cpp
--- a/Engine/Plugins/Runtime/nDisplay/Source/DisplayCluster/Private/Render/Viewport/Configuration/DisplayClusterViewportConfigurationHelpers.cpp
+++ b/Engine/Plugins/Runtime/nDisplay/Source/DisplayCluster/Private/Render/Viewport/Configuration/DisplayClusterViewportConfigurationHelpers.cpp
@@ -6,6 +6,7 @@
 
 #include "Render/Viewport/DisplayClusterViewport.h"
 #include "Render/Viewport/DisplayClusterViewportManager.h"
+#include "Render/Viewport/DisplayClusterViewportHelpers.h"
 
 #include "DisplayClusterRootActor.h"
 
@@ -80,7 +81,7 @@ void FDisplayClusterViewportConfigurationHelpers::UpdateViewportSetting_Override
 				DstViewport.PostRenderSettings.Replace.TextureRHI = TextureRHI;
 				FIntVector Size = TextureRHI->GetSizeXYZ();
 
-				DstViewport.PostRenderSettings.Replace.Rect = DstViewport.GetValidRect((InOverride.bShouldUseTextureRegion) ? InOverride.TextureRegion.ToRect() : FIntRect(FIntPoint(0, 0), FIntPoint(Size.X, Size.Y)), TEXT("Configuration Override"));
+				DstViewport.PostRenderSettings.Replace.Rect = FDisplayClusterViewportHelpers::GetValidViewportRect((InOverride.bShouldUseTextureRegion) ? InOverride.TextureRegion.ToRect() : FIntRect(FIntPoint(0, 0), FIntPoint(Size.X, Size.Y)), DstViewport.GetId(), TEXT("Configuration Override"));
 			}
 		}
 	}
@@ -105,41 +106,40 @@ void FDisplayClusterViewportConfigurationHelpers::UpdateViewportSetting_Postproc
 	DstViewport.PostRenderSettings.PostprocessBlur.KernelScale = InBlurPostprocess.KernelScale;
 };
 
-void FDisplayClusterViewportConfigurationHelpers::UpdateViewportSetting_Overscan(FDisplayClusterViewport& DstViewport, const FDisplayClusterConfigurationViewport_Overscan& InOverscan)
+void FDisplayClusterViewportConfigurationHelpers::UpdateViewportSetting_Overscan(const FDisplayClusterConfigurationViewport_Overscan& InOverscan, FDisplayClusterViewport_OverscanSettings& OutOverscanSettings)
 {
-	FImplDisplayClusterViewport_OverscanSettings OverscanSettings;
-	OverscanSettings.bEnabled = InOverscan.bEnabled;
-	OverscanSettings.bOversize = InOverscan.bOversize;
+	OutOverscanSettings.bEnabled = false;
+	OutOverscanSettings.bOversize = InOverscan.bOversize;
 	
-	if (OverscanSettings.bEnabled)
+	if (InOverscan.bEnabled)
 	{
 		switch (InOverscan.Mode)
 		{
 		case EDisplayClusterConfigurationViewportOverscanMode::Percent:
-			OverscanSettings.Mode = EDisplayClusterViewport_OverscanMode::Percent;
+			OutOverscanSettings.bEnabled = InOverscan.bEnabled;
+			OutOverscanSettings.Unit = EDisplayClusterViewport_FrustumUnit::Percent;
 
 			// Scale 0..100% to 0..1 range
-			OverscanSettings.Left = .01f * InOverscan.Left;
-			OverscanSettings.Right = .01f * InOverscan.Right;
-			OverscanSettings.Top = .01f * InOverscan.Top;
-			OverscanSettings.Bottom = .01f * InOverscan.Bottom;
+			OutOverscanSettings.Left = .01f * InOverscan.Left;
+			OutOverscanSettings.Right = .01f * InOverscan.Right;
+			OutOverscanSettings.Top = .01f * InOverscan.Top;
+			OutOverscanSettings.Bottom = .01f * InOverscan.Bottom;
 			break;
 
 		case EDisplayClusterConfigurationViewportOverscanMode::Pixels:
-			OverscanSettings.Mode = EDisplayClusterViewport_OverscanMode::Pixels;
+			OutOverscanSettings.bEnabled = InOverscan.bEnabled;
+			OutOverscanSettings.Unit = EDisplayClusterViewport_FrustumUnit::Pixels;
 
-			OverscanSettings.Left = InOverscan.Left;
-			OverscanSettings.Right = InOverscan.Right;
-			OverscanSettings.Top = InOverscan.Top;
-			OverscanSettings.Bottom = InOverscan.Bottom;
+			OutOverscanSettings.Left = InOverscan.Left;
+			OutOverscanSettings.Right = InOverscan.Right;
+			OutOverscanSettings.Top = InOverscan.Top;
+			OutOverscanSettings.Bottom = InOverscan.Bottom;
 			break;
 
 		default:
 			break;
 		}
 	}
-
-	DstViewport.OverscanRendering.Set(OverscanSettings);
 };
 
 void FDisplayClusterViewportConfigurationHelpers::UpdateViewportSetting_GenerateMips(FDisplayClusterViewport& DstViewport, const FDisplayClusterConfigurationPostRender_GenerateMips& InGenerateMips)
@@ -175,7 +175,7 @@ void FDisplayClusterViewportConfigurationHelpers::UpdateBaseViewportSetting(FDis
 		}
 
 		DstViewport.RenderSettings.CameraId = InConfigurationViewport.Camera;
-		DstViewport.RenderSettings.Rect = DstViewport.GetValidRect(InConfigurationViewport.Region.ToRect(), TEXT("Configuration Region"));
+		DstViewport.RenderSettings.Rect = FDisplayClusterViewportHelpers::GetValidViewportRect(InConfigurationViewport.Region.ToRect(), DstViewport.GetId(), TEXT("Configuration Region"));
 
 		DstViewport.RenderSettings.bEnableCrossGPUTransfer = InConfigurationViewport.RenderSettings.bEnableCrossGPUTransfer;
 
@@ -198,7 +198,7 @@ void FDisplayClusterViewportConfigurationHelpers::UpdateBaseViewportSetting(FDis
 	{
 		DstViewport.SetViewportBufferRatio(InRenderSettings.BufferRatio);
 
-		UpdateViewportSetting_Overscan(DstViewport, InRenderSettings.Overscan);
+		UpdateViewportSetting_Overscan(InRenderSettings.Overscan, DstViewport.RenderSettings.OverscanSettings);
 
 		UpdateViewportSetting_Override(DstViewport, InRenderSettings.Replace);
 		UpdateViewportSetting_PostprocessBlur(DstViewport, InRenderSettings.PostprocessBlur);
diff --git a/Engine/Plugins/Runtime/nDisplay/Source/DisplayCluster/Private/Render/Viewport/Configuration/DisplayClusterViewportConfigurationHelpers.h b/Engine/Plugins/Runtime/nDisplay/Source/DisplayCluster/Private/Render/Viewport/Configuration/DisplayClusterViewportConfigurationHelpers.h
--- a/Engine/Plugins/Runtime/nDisplay/Source/DisplayCluster/Private/Render/Viewport/Configuration/DisplayClusterViewportConfigurationHelpers.h
+++ b/Engine/Plugins/Runtime/nDisplay/Source/DisplayCluster/Private/Render/Viewport/Configuration/DisplayClusterViewportConfigurationHelpers.h
@@ -16,6 +16,7 @@ struct FDisplayClusterConfigurationPostRender_BlurPostprocess;
 struct FDisplayClusterConfigurationViewport_Overscan;
 struct FDisplayClusterConfigurationPostRender_GenerateMips;
 struct FDisplayClusterConfigurationProjection;
+struct FDisplayClusterViewport_OverscanSettings;
 
 class FDisplayClusterViewportConfigurationHelpers
 {
@@ -26,7 +27,7 @@ public:
 	static void UpdateViewportSetting_Override(FDisplayClusterViewport& DstViewport, const FDisplayClusterConfigurationPostRender_Override& InOverride);
 
 	static void UpdateViewportSetting_PostprocessBlur(FDisplayClusterViewport& DstViewport, const FDisplayClusterConfigurationPostRender_BlurPostprocess& InBlurPostprocess);
-	static void UpdateViewportSetting_Overscan(FDisplayClusterViewport& DstViewport, const FDisplayClusterConfigurationViewport_Overscan& InOverscan);
+	static void UpdateViewportSetting_Overscan(const FDisplayClusterConfigurationViewport_Overscan& InOverscan, FDisplayClusterViewport_OverscanSettings& OutOverscanSettings);
 	static void UpdateViewportSetting_GenerateMips(FDisplayClusterViewport& DstViewport, const FDisplayClusterConfigurationPostRender_GenerateMips& InGenerateMips);
 
 	static void UpdateBaseViewportSetting(FDisplayClusterViewport& DstViewport, ADisplayClusterRootActor& RootActor, const UDisplayClusterConfigurationViewport& InConfigurationViewport);
diff --git a/Engine/Plugins/Runtime/nDisplay/Source/DisplayCluster/Private/Render/Viewport/Configuration/DisplayClusterViewportConfigurationHelpers_ICVFX.cpp b/Engine/Plugins/Runtime/nDisplay/Source/DisplayCluster/Private/Render/Viewport/Configuration/DisplayClusterViewportConfigurationHelpers_ICVFX.cpp
--- a/Engine/Plugins/Runtime/nDisplay/Source/DisplayCluster/Private/Render/Viewport/Configuration/DisplayClusterViewportConfigurationHelpers_ICVFX.cpp
+++ b/Engine/Plugins/Runtime/nDisplay/Source/DisplayCluster/Private/Render/Viewport/Configuration/DisplayClusterViewportConfigurationHelpers_ICVFX.cpp
@@ -10,6 +10,7 @@
 #include "Render/Viewport/DisplayClusterViewport.h"
 #include "Render/Viewport/DisplayClusterViewportManager.h"
 #include "Render/Viewport/DisplayClusterViewportStrings.h"
+#include "Render/Viewport/DisplayClusterViewportHelpers.h"
 
 #include "DisplayClusterRootActor.h"
 
@@ -29,7 +30,6 @@
 #include "Render/Viewport/Containers/DisplayClusterViewport_RenderSettings.h"
 #include "Render/Viewport/Containers/DisplayClusterViewport_RenderSettingsICVFX.h"
 #include "Render/Viewport/Containers/DisplayClusterViewport_PostRenderSettings.h"
-#include "Render/Viewport/Containers/ImplDisplayClusterViewport_CustomFrustum.h"
 
 #include "Render/Viewport/RenderFrame/DisplayClusterRenderFrameSettings.h"
 #include "Render/Viewport/LightCard/DisplayClusterViewportLightCardManager.h"
@@ -551,7 +551,7 @@ void FDisplayClusterViewportConfigurationHelpers_ICVFX::UpdateCameraViewportSett
 	// UDisplayClusterConfigurationICVFX_CameraRenderSettings
 	const FIntPoint DesiredSize = CameraSettings.GetCameraFrameSize(StageSettings);
 
-	DstViewport.RenderSettings.Rect = DstViewport.GetValidRect(FIntRect(FIntPoint(0, 0), DesiredSize), TEXT("Configuration Camera Frame Size"));
+	DstViewport.RenderSettings.Rect = FDisplayClusterViewportHelpers::GetValidViewportRect(FIntRect(FIntPoint(0, 0), DesiredSize), DstViewport.GetId(), TEXT("Configuration Camera Frame Size"));
 
 	// Apply postprocess for camera
 	FDisplayClusterViewportConfigurationHelpers_Postprocess::UpdateCameraPostProcessSettings(DstViewport, RootActor, InCameraComponent);
@@ -572,7 +572,7 @@ void FDisplayClusterViewportConfigurationHelpers_ICVFX::UpdateCameraViewportSett
 	}
 
 	// Support inner camera custom frustum
-	FDisplayClusterViewportConfigurationHelpers_ICVFX::UpdateCameraCustomFrustum(DstViewport, CameraSettings.CustomFrustum);
+	FDisplayClusterViewportConfigurationHelpers_ICVFX::UpdateCameraCustomFrustum(CameraSettings.CustomFrustum, DstViewport.RenderSettings.CustomFrustumSettings);
 
 	// Set RenderTargetAdaptRatio
 	DstViewport.RenderSettings.RenderTargetAdaptRatio = CameraSettings.GetCameraAdaptResolutionRatio(StageSettings);
@@ -642,7 +642,7 @@ void FDisplayClusterViewportConfigurationHelpers_ICVFX::UpdateChromakeyViewportS
 			DesiredSize.X = ChromakeyRenderSettings->CustomSize.CustomWidth;
 			DesiredSize.Y = ChromakeyRenderSettings->CustomSize.CustomHeight;
 
-			DstViewport.RenderSettings.Rect = DstViewport.GetValidRect(FIntRect(FIntPoint(0, 0), DesiredSize), TEXT("Configuration custom chromakey Frame Size"));
+			DstViewport.RenderSettings.Rect = FDisplayClusterViewportHelpers::GetValidViewportRect(FIntRect(FIntPoint(0, 0), DesiredSize), DstViewport.GetId(), TEXT("Configuration custom chromakey Frame Size"));
 		}
 
 		// Debug: override the texture of the target viewport from this chromakeyRTT
@@ -657,7 +657,7 @@ void FDisplayClusterViewportConfigurationHelpers_ICVFX::UpdateChromakeyViewportS
 	// FDisplayClusterViewportConfigurationHelpers_OpenColorIO::UpdateChromakeyViewport(DstViewport, RootActor, InCameraComponent);
 
 	// Support inner camera custom frustum
-	FDisplayClusterViewportConfigurationHelpers_ICVFX::UpdateCameraCustomFrustum(DstViewport, CameraSettings.CustomFrustum);
+	FDisplayClusterViewportConfigurationHelpers_ICVFX::UpdateCameraCustomFrustum(CameraSettings.CustomFrustum, DstViewport.RenderSettings.CustomFrustumSettings);
 
 	// Attach to parent viewport
 	DstViewport.RenderSettings.AssignParentViewport(InCameraViewport.GetId(), InCameraViewport.RenderSettings);
@@ -809,38 +809,41 @@ void FDisplayClusterViewportConfigurationHelpers_ICVFX::UpdateLightcardViewportS
 	}
 }
 
-void FDisplayClusterViewportConfigurationHelpers_ICVFX::UpdateCameraCustomFrustum(FDisplayClusterViewport& DstViewport, const FDisplayClusterConfigurationICVFX_CameraCustomFrustum& InCameraCustomFrustumConfiguration)
+void FDisplayClusterViewportConfigurationHelpers_ICVFX::UpdateCameraCustomFrustum(
+	const FDisplayClusterConfigurationICVFX_CameraCustomFrustum& InCameraCustomFrustumConfiguration,
+	FDisplayClusterViewport_CustomFrustumSettings& OutCustomFrustumSettings)
 {
+	OutCustomFrustumSettings.bEnabled = false;
+
 	if (InCameraCustomFrustumConfiguration.bEnable)
 	{
-		FImplDisplayClusterViewport_CustomFrustumSettings CustomFrustumSettings;
-		CustomFrustumSettings.bAdaptResolution = InCameraCustomFrustumConfiguration.bAdaptResolution;
+		OutCustomFrustumSettings.bAdaptResolution = InCameraCustomFrustumConfiguration.bAdaptResolution;
 
 		switch (InCameraCustomFrustumConfiguration.Mode)
 		{
 		case EDisplayClusterConfigurationViewportCustomFrustumMode::Percent:
-			CustomFrustumSettings.Mode = EDisplayClusterViewport_CustomFrustumMode::Percent;
+			OutCustomFrustumSettings.bEnabled = true;
+			OutCustomFrustumSettings.Unit = EDisplayClusterViewport_FrustumUnit::Percent;
 
 			// Scale 0..100% to 0..1 range
-			CustomFrustumSettings.Left = .01f * InCameraCustomFrustumConfiguration.Left;
-			CustomFrustumSettings.Right = .01f * InCameraCustomFrustumConfiguration.Right;
-			CustomFrustumSettings.Top = .01f * InCameraCustomFrustumConfiguration.Top;
-			CustomFrustumSettings.Bottom = .01f * InCameraCustomFrustumConfiguration.Bottom;
+			OutCustomFrustumSettings.Left = .01f * InCameraCustomFrustumConfiguration.Left;
+			OutCustomFrustumSettings.Right = .01f * InCameraCustomFrustumConfiguration.Right;
+			OutCustomFrustumSettings.Top = .01f * InCameraCustomFrustumConfiguration.Top;
+			OutCustomFrustumSettings.Bottom = .01f * InCameraCustomFrustumConfiguration.Bottom;
 			break;
 
 		case EDisplayClusterConfigurationViewportCustomFrustumMode::Pixels:
-			CustomFrustumSettings.Mode = EDisplayClusterViewport_CustomFrustumMode::Pixels;
+			OutCustomFrustumSettings.bEnabled = true;
+			OutCustomFrustumSettings.Unit = EDisplayClusterViewport_FrustumUnit::Pixels;
 
-			CustomFrustumSettings.Left = InCameraCustomFrustumConfiguration.Left;
-			CustomFrustumSettings.Right = InCameraCustomFrustumConfiguration.Right;
-			CustomFrustumSettings.Top = InCameraCustomFrustumConfiguration.Top;
-			CustomFrustumSettings.Bottom = InCameraCustomFrustumConfiguration.Bottom;
+			OutCustomFrustumSettings.Left = InCameraCustomFrustumConfiguration.Left;
+			OutCustomFrustumSettings.Right = InCameraCustomFrustumConfiguration.Right;
+			OutCustomFrustumSettings.Top = InCameraCustomFrustumConfiguration.Top;
+			OutCustomFrustumSettings.Bottom = InCameraCustomFrustumConfiguration.Bottom;
 			break;
 
 		default:
 			break;
 		}
-
-		DstViewport.CustomFrustumRendering.Set(CustomFrustumSettings);
 	}
 };
diff --git a/Engine/Plugins/Runtime/nDisplay/Source/DisplayCluster/Private/Render/Viewport/Configuration/DisplayClusterViewportConfigurationHelpers_ICVFX.h b/Engine/Plugins/Runtime/nDisplay/Source/DisplayCluster/Private/Render/Viewport/Configuration/DisplayClusterViewportConfigurationHelpers_ICVFX.h
--- a/Engine/Plugins/Runtime/nDisplay/Source/DisplayCluster/Private/Render/Viewport/Configuration/DisplayClusterViewportConfigurationHelpers_ICVFX.h
+++ b/Engine/Plugins/Runtime/nDisplay/Source/DisplayCluster/Private/Render/Viewport/Configuration/DisplayClusterViewportConfigurationHelpers_ICVFX.h
@@ -19,6 +19,7 @@ struct FDisplayClusterConfigurationICVFX_ChromakeyMarkers;
 struct FDisplayClusterConfigurationICVFX_LightcardSettings;
 struct FDisplayClusterConfigurationICVFX_CameraCustomFrustum;
 struct FDisplayClusterConfigurationICVFX_CameraSettings;
+struct FDisplayClusterViewport_CustomFrustumSettings;
 
 struct FCameraContext_ICVFX
 {
@@ -74,7 +75,7 @@ public:
 	static void UpdateCameraSettings_Chromakey(FDisplayClusterShaderParameters_ICVFX::FCameraSettings& InOutCameraSettings, const FDisplayClusterConfigurationICVFX_StageSettings& InStageSettings, const FDisplayClusterConfigurationICVFX_CameraSettings& InCameraSettings, bool bEnableChromakeyMarkers, const FString& InChromakeyViewportId);
 	static void UpdateCameraSettings_OverlapChromakey(FDisplayClusterShaderParameters_ICVFX::FCameraSettings& InOutCameraSettings, const FDisplayClusterConfigurationICVFX_StageSettings& InStageSettings, const FDisplayClusterConfigurationICVFX_CameraSettings& InCameraSettings, bool bEnableChromakeyMarkers);
 
-	static void UpdateCameraCustomFrustum(FDisplayClusterViewport& DstViewport, const FDisplayClusterConfigurationICVFX_CameraCustomFrustum& InCameraCustomFrustum);
+	static void UpdateCameraCustomFrustum(const FDisplayClusterConfigurationICVFX_CameraCustomFrustum& InCameraCustomFrustumConfiguration, FDisplayClusterViewport_CustomFrustumSettings& OutCustomFrustumSettings);
 
 	static FDisplayClusterViewport* ImplFindViewport(ADisplayClusterRootActor& RootActor, const FString& InViewportId, const FString& InResourceId);
 private:
diff --git a/Engine/Plugins/Runtime/nDisplay/Source/DisplayCluster/Private/Render/Viewport/Containers/DisplayClusterViewportProxyData.cpp b/Engine/Plugins/Runtime/nDisplay/Source/DisplayCluster/Private/Render/Viewport/Containers/DisplayClusterViewportProxyData.cpp
--- a/Engine/Plugins/Runtime/nDisplay/Source/DisplayCluster/Private/Render/Viewport/Containers/DisplayClusterViewportProxyData.cpp
+++ b/Engine/Plugins/Runtime/nDisplay/Source/DisplayCluster/Private/Render/Viewport/Containers/DisplayClusterViewportProxyData.cpp
@@ -22,7 +22,7 @@ FDisplayClusterViewportProxyData::FDisplayClusterViewportProxyData(const TShared
 	PostRenderSettings.SetParameters(SrcViewport->PostRenderSettings);
 
 	// Additional parameters
-	OverscanSettings = SrcViewport->OverscanRendering.Get();
+	OverscanRuntimeSettings = SrcViewport->OverscanRuntimeSettings;
 
 	RemapMesh = SrcViewport->ViewportRemap.GetRemapMesh();
 
@@ -40,7 +40,7 @@ void FDisplayClusterViewportProxyData::UpdateProxy_RenderThread() const
 
 	DstViewportProxy->OpenColorIO = OpenColorIO;
 
-	DstViewportProxy->OverscanSettings = OverscanSettings;
+	DstViewportProxy->OverscanRuntimeSettings = OverscanRuntimeSettings;
 
 	DstViewportProxy->RemapMesh = RemapMesh;
 
diff --git a/Engine/Plugins/Runtime/nDisplay/Source/DisplayCluster/Private/Render/Viewport/Containers/DisplayClusterViewportProxyData.h b/Engine/Plugins/Runtime/nDisplay/Source/DisplayCluster/Private/Render/Viewport/Containers/DisplayClusterViewportProxyData.h
--- a/Engine/Plugins/Runtime/nDisplay/Source/DisplayCluster/Private/Render/Viewport/Containers/DisplayClusterViewportProxyData.h
+++ b/Engine/Plugins/Runtime/nDisplay/Source/DisplayCluster/Private/Render/Viewport/Containers/DisplayClusterViewportProxyData.h
@@ -6,7 +6,7 @@
 #include "Render/Viewport/Containers/DisplayClusterViewport_RenderSettings.h"
 #include "Render/Viewport/Containers/DisplayClusterViewport_RenderSettingsICVFX.h"
 #include "Render/Viewport/Containers/DisplayClusterViewport_PostRenderSettings.h"
-#include "Render/Viewport/Containers/DisplayClusterViewport_OverscanSettings.h"
+#include "Render/Viewport/Containers/DisplayClusterViewport_OverscanRuntimeSettings.h"
 
 class IDisplayClusterProjectionPolicy;
 class FDisplayClusterViewportResource;
@@ -33,7 +33,7 @@ private:
 
 	TSharedPtr<IDisplayClusterRender_MeshComponent, ESPMode::ThreadSafe> RemapMesh;
 
-	FDisplayClusterViewport_OverscanSettings     OverscanSettings;
+	FDisplayClusterViewport_OverscanRuntimeSettings OverscanRuntimeSettings;
 
 	// Viewport render params
 	FDisplayClusterViewport_RenderSettings       RenderSettings;
diff --git a/Engine/Plugins/Runtime/nDisplay/Source/DisplayCluster/Private/Render/Viewport/Containers/DisplayClusterViewport_CustomFrustumRuntimeSettings.cpp b/Engine/Plugins/Runtime/nDisplay/Source/DisplayCluster/Private/Render/Viewport/Containers/DisplayClusterViewport_CustomFrustumRuntimeSettings.cpp
new file mode 100644
--- /dev/null
+++ b/Engine/Plugins/Runtime/nDisplay/Source/DisplayCluster/Private/Render/Viewport/Containers/DisplayClusterViewport_CustomFrustumRuntimeSettings.cpp
@@ -0,0 +1,121 @@
+// Copyright Epic Games, Inc. All Rights Reserved.
+
+#include "Render/Viewport/Containers/DisplayClusterViewport_CustomFrustumRuntimeSettings.h"
+#include "Render/Viewport/Containers/DisplayClusterViewport_CustomFrustumSettings.h"
+#include "Render/Viewport/DisplayClusterViewportHelpers.h"
+
+#include "HAL/IConsoleManager.h"
+
+int32 GDisplayClusterRenderCustomFrustumEnable = 1;
+static FAutoConsoleVariableRef CVarDisplayClusterRenderCustomFrustumEnable(
+	TEXT("nDisplay.render.custom_frustum.enable"),
+	GDisplayClusterRenderCustomFrustumEnable,
+	TEXT("Enable custom frustum feature.\n")
+	TEXT(" 0 - to disable.\n"),
+	ECVF_Default
+);
+
+int32 GDisplayClusterRenderCustomFrustumMaxValue = 50;
+static FAutoConsoleVariableRef CVarDisplayClusterRenderCustomFrustumMaxValue(
+	TEXT("nDisplay.render.custom_frustum.max_percent"),
+	GDisplayClusterRenderCustomFrustumMaxValue,
+	TEXT("Max percent for custom frustum (default 50).\n"),
+	ECVF_Default
+);
+
+namespace UE::DisplayCluster::Viewport::CustomFrustumHelpers
+{
+	/** Clamp percent for custom frustum settings. */
+	static inline float ClampPercent(float InValue)
+	{
+		const float MaxCustomFrustumValue = float(GDisplayClusterRenderCustomFrustumMaxValue) / 100;
+
+		return FMath::Clamp(InValue, -MaxCustomFrustumValue, MaxCustomFrustumValue);
+	}
+};
+using namespace UE::DisplayCluster::Viewport;
+
+///////////////////////////////////////////////////////////////////////////////////////////
+// FDisplayClusterViewport_CustomFrustumRuntimeSettings
+///////////////////////////////////////////////////////////////////////////////////////////
+bool FDisplayClusterViewport_CustomFrustumRuntimeSettings::UpdateProjectionAngles(
+	const FDisplayClusterViewport_CustomFrustumRuntimeSettings& InRuntimeSettings,
+	float& InOutLeft,
+	float& InOutRight,
+	float& InOutTop,
+	float& InOutBottom)
+{
+	if (InRuntimeSettings.bIsEnabled)
+	{
+		const float Horizontal = InOutRight - InOutLeft;
+		const float Vertical = InOutTop - InOutBottom;
+
+		InOutLeft   -= Horizontal * InRuntimeSettings.CustomFrustumPercent.Left;
+		InOutRight  += Horizontal * InRuntimeSettings.CustomFrustumPercent.Right;
+		InOutBottom -= Vertical * InRuntimeSettings.CustomFrustumPercent.Bottom;
+		InOutTop    += Vertical * InRuntimeSettings.CustomFrustumPercent.Top;
+
+		return true;
+	}
+
+	return false;
+}
+
+void FDisplayClusterViewport_CustomFrustumRuntimeSettings::UpdateCustomFrustumSettings(
+	const FString& InViewportId,
+	const FDisplayClusterViewport_CustomFrustumSettings& InCustomFrustumSettings,
+	FDisplayClusterViewport_CustomFrustumRuntimeSettings& InOutRuntimeSettings,
+	FIntRect& InOutRenderTargetRect)
+{
+	const FIntPoint Size = InOutRenderTargetRect.Size();
+
+	// Disable viewport CustomFrustum feature
+	if (!GDisplayClusterRenderCustomFrustumEnable || !InCustomFrustumSettings.bEnabled)
+	{
+		return;
+	}
+
+	switch (InCustomFrustumSettings.Unit)
+	{
+	case EDisplayClusterViewport_FrustumUnit::Percent:
+	{
+		InOutRuntimeSettings.bIsEnabled = true;
+
+		InOutRuntimeSettings.CustomFrustumPercent.Left   = CustomFrustumHelpers::ClampPercent(InCustomFrustumSettings.Left);
+		InOutRuntimeSettings.CustomFrustumPercent.Right  = CustomFrustumHelpers::ClampPercent(InCustomFrustumSettings.Right);
+		InOutRuntimeSettings.CustomFrustumPercent.Top    = CustomFrustumHelpers::ClampPercent(InCustomFrustumSettings.Top);
+		InOutRuntimeSettings.CustomFrustumPercent.Bottom = CustomFrustumHelpers::ClampPercent(InCustomFrustumSettings.Bottom);
+		break;
+	}
+
+	case EDisplayClusterViewport_FrustumUnit::Pixels:
+	{
+		InOutRuntimeSettings.bIsEnabled = true;
+
+		InOutRuntimeSettings.CustomFrustumPercent.Left   = CustomFrustumHelpers::ClampPercent(InCustomFrustumSettings.Left   / Size.X);
+		InOutRuntimeSettings.CustomFrustumPercent.Right  = CustomFrustumHelpers::ClampPercent(InCustomFrustumSettings.Right  / Size.X);
+		InOutRuntimeSettings.CustomFrustumPercent.Top    = CustomFrustumHelpers::ClampPercent(InCustomFrustumSettings.Top    / Size.Y);
+		InOutRuntimeSettings.CustomFrustumPercent.Bottom = CustomFrustumHelpers::ClampPercent(InCustomFrustumSettings.Bottom / Size.Y);
+
+		break;
+	}
+
+	default:
+		return;
+	}
+
+	// Calc pixels from percent
+	InOutRuntimeSettings.CustomFrustumPixels.Left   = FMath::RoundToInt(Size.X * InOutRuntimeSettings.CustomFrustumPercent.Left);
+	InOutRuntimeSettings.CustomFrustumPixels.Right  = FMath::RoundToInt(Size.X * InOutRuntimeSettings.CustomFrustumPercent.Right);
+	InOutRuntimeSettings.CustomFrustumPixels.Top    = FMath::RoundToInt(Size.Y * InOutRuntimeSettings.CustomFrustumPercent.Top);
+	InOutRuntimeSettings.CustomFrustumPixels.Bottom = FMath::RoundToInt(Size.Y * InOutRuntimeSettings.CustomFrustumPercent.Bottom);
+
+	// Update RTT size for CustomFrustum when we need to scale target resolution
+	if (InCustomFrustumSettings.bAdaptResolution)
+	{
+		const FIntPoint CustomFrustumSize = Size + InOutRuntimeSettings.CustomFrustumPixels.Size();
+		const FIntPoint ValidCustomFrustumSize = FDisplayClusterViewportHelpers::GetValidViewportRect(FIntRect(FIntPoint(0, 0), CustomFrustumSize), InViewportId, TEXT("CustomFrustum")).Size();
+
+		InOutRenderTargetRect.Max = ValidCustomFrustumSize;
+	}
+}
diff --git a/Engine/Plugins/Runtime/nDisplay/Source/DisplayCluster/Private/Render/Viewport/Containers/DisplayClusterViewport_CustomFrustumRuntimeSettings.h b/Engine/Plugins/Runtime/nDisplay/Source/DisplayCluster/Private/Render/Viewport/Containers/DisplayClusterViewport_CustomFrustumRuntimeSettings.h
new file mode 100644
--- /dev/null
+++ b/Engine/Plugins/Runtime/nDisplay/Source/DisplayCluster/Private/Render/Viewport/Containers/DisplayClusterViewport_CustomFrustumRuntimeSettings.h
@@ -0,0 +1,74 @@
+// Copyright Epic Games, Inc. All Rights Reserved.
+
+#pragma once
+
+#include "CoreMinimal.h"
+#include "Render/Viewport/Containers/DisplayClusterViewport_CustomFrustumSettings.h"
+
+/**
+* Runtime custom frustum settings
+*/
+struct FDisplayClusterViewport_CustomFrustumRuntimeSettings
+{
+	/** Update custom frustum settings
+	*
+	* @param InViewport            - owner viewport
+	* @param InOutRuntimeSettings  - the CustomFrustum runtime settings.
+	* @param InOutRenderTargetRect - Viewport rect, changeable during CustomFrustum
+	*/
+	static void UpdateCustomFrustumSettings(const FString& InViewportId, const FDisplayClusterViewport_CustomFrustumSettings& InCustomFrustumSettings, FDisplayClusterViewport_CustomFrustumRuntimeSettings& InOutRuntimeSettings, FIntRect& InOutRenderTargetRect);
+	
+	/** Update projection angles by custom frustum
+	*
+	* @param InOverscanRuntimeSettings - the CustomFrustum runtime settings.
+	* @param InOutLeft                 - the value of the left projection plane that you want to change
+	* @param InOutRight                - the value of the right projection plane that you want to change
+	* @param InOutTop                  - the value of the top projection plane that you want to change
+	* @parma InOutBottom               - the value of the bottom projection plane that you want to change
+	*/
+	static bool UpdateProjectionAngles(const FDisplayClusterViewport_CustomFrustumRuntimeSettings& InRuntimeSettings, float& InOutLeft, float& InOutRight, float& InOutTop, float& InOutBottom);
+
+	/**
+	* Values in percent
+	*/
+	struct FCustomFrustumPercent
+	{
+		float Left = 0;
+		float Right = 0;
+		float Top = 0;
+		float Bottom = 0;
+	};
+
+	/**
+	* Values in pixels
+	*/
+	struct FCustomFrustumPixels
+	{
+		inline FIntRect GetInnerRect(const FIntRect& InRect) const
+		{
+			const FIntPoint InnerSize = InRect.Size() - Size();
+			const FIntPoint InnerPos = FIntPoint(Left, Top);
+
+			return FIntRect(InnerPos, InnerPos + InnerSize);
+		}
+
+		inline FIntPoint Size() const
+		{
+			return FIntPoint(Left + Right, Top + Bottom);
+		}
+
+		int32 Left = 0;
+		int32 Right = 0;
+		int32 Top = 0;
+		int32 Bottom = 0;
+	};
+
+	// Enable custom frustum
+	bool bIsEnabled = false;
+
+	// CustomFrustum sides in percent
+	FCustomFrustumPercent CustomFrustumPercent;
+
+	// CustomFrustum sides in pixels
+	FCustomFrustumPixels  CustomFrustumPixels;
+};
diff --git a/Engine/Plugins/Runtime/nDisplay/Source/DisplayCluster/Private/Render/Viewport/Containers/DisplayClusterViewport_OverscanRuntimeSettings.cpp b/Engine/Plugins/Runtime/nDisplay/Source/DisplayCluster/Private/Render/Viewport/Containers/DisplayClusterViewport_OverscanRuntimeSettings.cpp
new file mode 100644
--- /dev/null
+++ b/Engine/Plugins/Runtime/nDisplay/Source/DisplayCluster/Private/Render/Viewport/Containers/DisplayClusterViewport_OverscanRuntimeSettings.cpp
@@ -0,0 +1,141 @@
+// Copyright Epic Games, Inc. All Rights Reserved.
+
+#include "Render/Viewport/Containers/DisplayClusterViewport_OverscanRuntimeSettings.h"
+#include "Render/Viewport/DisplayClusterViewportHelpers.h"
+
+#include "HAL/IConsoleManager.h"
+
+int32 GDisplayClusterRenderOverscanEnable = 1;
+static FAutoConsoleVariableRef CVarDisplayClusterRenderOverscanEnable(
+	TEXT("nDisplay.render.overscan.enable"),
+	GDisplayClusterRenderOverscanEnable,
+	TEXT("Enable overscan feature.\n")
+	TEXT(" 0 - to disable.\n"),
+	ECVF_Default
+);
+
+int32 GDisplayClusterRenderOverscanMaxValue = 50;
+static FAutoConsoleVariableRef CVarDisplayClusterRenderOverscanMaxValue(
+	TEXT("nDisplay.render.overscan.max_percent"),
+	GDisplayClusterRenderOverscanMaxValue,
+	TEXT("Max percent for overscan (default 50).\n"),
+	ECVF_Default
+);
+
+namespace UE::DisplayCluster::Viewport::OverscanHelpers
+{
+	/** Clamp percent for overscan settings. */
+	static inline float ClampPercent(float InValue)
+	{
+		const float MaxCustomFrustumValue = float(GDisplayClusterRenderOverscanMaxValue) / 100;
+
+		return FMath::Clamp(InValue, -MaxCustomFrustumValue, MaxCustomFrustumValue);
+	}
+};
+using namespace UE::DisplayCluster::Viewport;
+
+///////////////////////////////////////////////////////////////////////////////////////////
+// FDisplayClusterViewport_OverscanRuntimeSettings
+///////////////////////////////////////////////////////////////////////////////////////////
+bool FDisplayClusterViewport_OverscanRuntimeSettings::UpdateProjectionAngles(
+	const FDisplayClusterViewport_OverscanRuntimeSettings& InOverscanRuntimeSettings,
+	float& InOutLeft,
+	float& InOutRight,
+	float& InOutTop,
+	float& InOutBottom)
+{
+	if (InOverscanRuntimeSettings.bIsEnabled)
+	{
+		float Horizontal = InOutRight - InOutLeft;
+		float Vertical = InOutTop - InOutBottom;
+
+		InOutLeft   -= Horizontal * InOverscanRuntimeSettings.OverscanPercent.Left;
+		InOutRight  += Horizontal * InOverscanRuntimeSettings.OverscanPercent.Right;
+		InOutBottom -= Vertical * InOverscanRuntimeSettings.OverscanPercent.Bottom;
+		InOutTop    += Vertical * InOverscanRuntimeSettings.OverscanPercent.Top;
+
+		return true;
+	}
+
+	return false;
+}
+
+void FDisplayClusterViewport_OverscanRuntimeSettings::UpdateOverscanSettings(
+	const FString& InViewportId,
+	const FDisplayClusterViewport_OverscanSettings& InOverscanSettings,
+	FDisplayClusterViewport_OverscanRuntimeSettings& InOutOverscanRuntimeSettings,
+	FIntRect& InOutRenderTargetRect)
+{
+	const FIntPoint Size = InOutRenderTargetRect.Size();
+
+	// Disable viewport overscan feature
+	if (!GDisplayClusterRenderOverscanEnable || !InOverscanSettings.bEnabled)
+	{
+		return;
+	}
+
+	switch (InOverscanSettings.Unit)
+	{
+	case EDisplayClusterViewport_FrustumUnit::Percent:
+	{
+		InOutOverscanRuntimeSettings.bIsEnabled = true;
+
+		InOutOverscanRuntimeSettings.OverscanPercent.Left   = OverscanHelpers::ClampPercent(InOverscanSettings.Left);
+		InOutOverscanRuntimeSettings.OverscanPercent.Right  = OverscanHelpers::ClampPercent(InOverscanSettings.Right);
+		InOutOverscanRuntimeSettings.OverscanPercent.Top    = OverscanHelpers::ClampPercent(InOverscanSettings.Top);
+		InOutOverscanRuntimeSettings.OverscanPercent.Bottom = OverscanHelpers::ClampPercent(InOverscanSettings.Bottom);
+		break;
+	}
+
+	case EDisplayClusterViewport_FrustumUnit::Pixels:
+	{
+		InOutOverscanRuntimeSettings.bIsEnabled = true;
+
+		InOutOverscanRuntimeSettings.OverscanPercent.Left   = OverscanHelpers::ClampPercent(InOverscanSettings.Left   / Size.X);
+		InOutOverscanRuntimeSettings.OverscanPercent.Right  = OverscanHelpers::ClampPercent(InOverscanSettings.Right  / Size.X);
+		InOutOverscanRuntimeSettings.OverscanPercent.Top    = OverscanHelpers::ClampPercent(InOverscanSettings.Top    / Size.Y);
+		InOutOverscanRuntimeSettings.OverscanPercent.Bottom = OverscanHelpers::ClampPercent(InOverscanSettings.Bottom / Size.Y);
+
+		break;
+	}
+
+	default:
+		break;
+	}
+
+	// Update RTT size for overscan
+	if (InOutOverscanRuntimeSettings.bIsEnabled)
+	{
+		// Calc pixels from percent
+		InOutOverscanRuntimeSettings.OverscanPixels.Left   = Size.X * InOutOverscanRuntimeSettings.OverscanPercent.Left;
+		InOutOverscanRuntimeSettings.OverscanPixels.Right  = Size.X * InOutOverscanRuntimeSettings.OverscanPercent.Right;
+		InOutOverscanRuntimeSettings.OverscanPixels.Top    = Size.Y * InOutOverscanRuntimeSettings.OverscanPercent.Top;
+		InOutOverscanRuntimeSettings.OverscanPixels.Bottom = Size.Y * InOutOverscanRuntimeSettings.OverscanPercent.Bottom;
+
+		const FIntPoint OverscanSize = Size + InOutOverscanRuntimeSettings.OverscanPixels.Size();
+		const FIntPoint ValidOverscanSize = FDisplayClusterViewportHelpers::GetValidViewportRect(FIntRect(FIntPoint(0, 0), OverscanSize), InViewportId, TEXT("Overscan")).Size();
+		
+		bool bOversize = InOverscanSettings.bOversize;
+
+		if (OverscanSize != ValidOverscanSize)
+		{
+			// can't use overscan with extra size, disable oversize
+			bOversize = false;
+		}
+
+		if (bOversize)
+		{
+			InOutRenderTargetRect.Max = OverscanSize;
+		}
+		else
+		{
+			float scaleX = float(Size.X) / OverscanSize.X;
+			float scaleY = float(Size.Y) / OverscanSize.Y;
+
+			InOutOverscanRuntimeSettings.OverscanPixels.Left   *= scaleX;
+			InOutOverscanRuntimeSettings.OverscanPixels.Right  *= scaleX;
+			InOutOverscanRuntimeSettings.OverscanPixels.Top    *= scaleY;
+			InOutOverscanRuntimeSettings.OverscanPixels.Bottom *= scaleY;
+		}
+	}
+}
diff --git a/Engine/Plugins/Runtime/nDisplay/Source/DisplayCluster/Private/Render/Viewport/Containers/DisplayClusterViewport_OverscanRuntimeSettings.h b/Engine/Plugins/Runtime/nDisplay/Source/DisplayCluster/Private/Render/Viewport/Containers/DisplayClusterViewport_OverscanRuntimeSettings.h
new file mode 100644
--- /dev/null
+++ b/Engine/Plugins/Runtime/nDisplay/Source/DisplayCluster/Private/Render/Viewport/Containers/DisplayClusterViewport_OverscanRuntimeSettings.h
@@ -0,0 +1,74 @@
+// Copyright Epic Games, Inc. All Rights Reserved.
+
+#pragma once
+
+#include "CoreMinimal.h"
+#include "Render/Viewport/Containers/DisplayClusterViewport_OverscanSettings.h"
+
+/**
+* Runtime overscan settings
+*/
+struct FDisplayClusterViewport_OverscanRuntimeSettings
+{
+	/** Update overscan settings
+	* 
+	* @param InViewport            - owner viewport
+	* @param InOutRuntimeSettings  - the overscan runtime settings.
+	* @param InOutRenderTargetRect - Viewport rect, changeable during overscanning
+	*/
+	static void UpdateOverscanSettings(const FString& InViewportId, const FDisplayClusterViewport_OverscanSettings& InOverscanSettings, FDisplayClusterViewport_OverscanRuntimeSettings& InOutOverscanRuntimeSettings, FIntRect& InOutRenderTargetRect);
+
+	/** Update projection angles by overscan
+	* 
+	* @param InOverscanRuntimeSettings - the overscan runtime settings.
+	* @param InOutLeft                 - the value of the left projection plane that you want to change
+	* @param InOutRight                - the value of the right projection plane that you want to change
+	* @param InOutTop                  - the value of the top projection plane that you want to change
+	* @parma InOutBottom               - the value of the bottom projection plane that you want to change
+	*/
+	static bool UpdateProjectionAngles(const FDisplayClusterViewport_OverscanRuntimeSettings& InOverscanRuntimeSettings, float& InOutLeft, float& InOutRight, float& InOutTop, float& InOutBottom);
+
+	/**
+	* Overscan values in percent
+	*/
+	struct FOverscanPercent
+	{
+		float Left = 0;
+		float Right = 0;
+		float Top = 0;
+		float Bottom = 0;
+	};
+
+	/**
+	* Overscan values in pixels
+	*/
+	struct FOverscanPixels
+	{
+		inline FIntRect GetInnerRect(const FIntRect& InRect) const
+		{
+			const FIntPoint InnerSize = InRect.Size() - Size();
+			const FIntPoint InnerPos = FIntPoint(Left, Top);
+
+			return FIntRect(InnerPos, InnerPos + InnerSize);
+		}
+
+		inline FIntPoint Size() const
+		{
+			return FIntPoint(Left + Right, Top + Bottom);
+		}
+
+		int32 Left = 0;
+		int32 Right = 0;
+		int32 Top = 0;
+		int32 Bottom = 0;
+	};
+
+	// Enable overscan
+	bool bIsEnabled = false;
+
+	// Overscan sides in percent
+	FOverscanPercent OverscanPercent;
+
+	// Overscan sides in pixels
+	FOverscanPixels  OverscanPixels;
+};
diff --git a/Engine/Plugins/Runtime/nDisplay/Source/DisplayCluster/Private/Render/Viewport/Containers/ImplDisplayClusterViewport_CustomFrustum.cpp b/Engine/Plugins/Runtime/nDisplay/Source/DisplayCluster/Private/Render/Viewport/Containers/ImplDisplayClusterViewport_CustomFrustum.cpp
deleted file mode 100644
diff --git a/Engine/Plugins/Runtime/nDisplay/Source/DisplayCluster/Private/Render/Viewport/Containers/ImplDisplayClusterViewport_CustomFrustum.h b/Engine/Plugins/Runtime/nDisplay/Source/DisplayCluster/Private/Render/Viewport/Containers/ImplDisplayClusterViewport_CustomFrustum.h
deleted file mode 100644
diff --git a/Engine/Plugins/Runtime/nDisplay/Source/DisplayCluster/Private/Render/Viewport/Containers/ImplDisplayClusterViewport_Overscan.cpp b/Engine/Plugins/Runtime/nDisplay/Source/DisplayCluster/Private/Render/Viewport/Containers/ImplDisplayClusterViewport_Overscan.cpp
deleted file mode 100644
diff --git a/Engine/Plugins/Runtime/nDisplay/Source/DisplayCluster/Private/Render/Viewport/Containers/ImplDisplayClusterViewport_Overscan.h b/Engine/Plugins/Runtime/nDisplay/Source/DisplayCluster/Private/Render/Viewport/Containers/ImplDisplayClusterViewport_Overscan.h
deleted file mode 100644
diff --git a/Engine/Plugins/Runtime/nDisplay/Source/DisplayCluster/Private/Render/Viewport/DisplayClusterViewport.cpp b/Engine/Plugins/Runtime/nDisplay/Source/DisplayCluster/Private/Render/Viewport/DisplayClusterViewport.cpp
--- a/Engine/Plugins/Runtime/nDisplay/Source/DisplayCluster/Private/Render/Viewport/DisplayClusterViewport.cpp
+++ b/Engine/Plugins/Runtime/nDisplay/Source/DisplayCluster/Private/Render/Viewport/DisplayClusterViewport.cpp
@@ -430,38 +430,6 @@ inline void AdjustRect(FIntRect& InOutRect, const float multX, const float multY
 	InOutRect.Max.Y *= multY;
 }
 
-FIntRect FDisplayClusterViewport::GetValidRect(const FIntRect& InRect, const TCHAR* DbgSourceName)
-{
-	// The target always needs be within GMaxTextureDimensions, larger dimensions are not supported by the engine
-	const int32 MaxTextureSize = FDisplayClusterViewportHelpers::GetMaxTextureDimension();
-	const int32 MinTextureSize = FDisplayClusterViewportHelpers::GetMinTextureDimension();
-
-	int32 Width  = FMath::Max(MinTextureSize, InRect.Width());
-	int32 Height = FMath::Max(MinTextureSize, InRect.Height());
-
-	FIntRect OutRect(InRect.Min, InRect.Min + FIntPoint(Width, Height));
-
-	float RectScale = 1;
-
-	// Make sure the rect doesn't exceed the maximum resolution, and preserve its aspect ratio if it needs to be clamped
-	int32 RectMaxSize = OutRect.Max.GetMax();
-	if (RectMaxSize > MaxTextureSize)
-	{
-		RectScale = float(MaxTextureSize) / RectMaxSize;
-		UE_LOG(LogDisplayClusterViewport, Error, TEXT("The viewport '%s' rect '%s' size %dx%d clamped: max texture dimensions is %d"), *GetId(), (DbgSourceName==nullptr) ? TEXT("none") : DbgSourceName, InRect.Max.X, InRect.Max.Y, MaxTextureSize);
-	}
-
-	OutRect.Min.X = FMath::Min(OutRect.Min.X, MaxTextureSize);
-	OutRect.Min.Y = FMath::Min(OutRect.Min.Y, MaxTextureSize);
-
-	const FIntPoint ScaledRectMax = FDisplayClusterViewportHelpers::ScaleTextureSize(OutRect.Max, RectScale);
-
-	OutRect.Max.X = FMath::Clamp(ScaledRectMax.X, OutRect.Min.X, MaxTextureSize);
-	OutRect.Max.Y = FMath::Clamp(ScaledRectMax.Y, OutRect.Min.Y, MaxTextureSize);
-
-	return OutRect;
-}
-
 float FDisplayClusterViewport::GetClusterRenderTargetRatioMult(const FDisplayClusterRenderFrameSettings& InFrameSettings) const
 {
 	float ClusterRenderTargetRatioMult = InFrameSettings.ClusterRenderTargetRatioMult;
@@ -590,7 +558,7 @@ bool FDisplayClusterViewport::UpdateFrameContexts(const uint32 InStereoViewIndex
 				ContextIt.StereoscopicPass = FDisplayClusterViewportStereoscopicPass::EncodeStereoscopicPass(ContextIt.ContextNum, ViewportContextAmount, InFrameSettings);
 				ContextIt.StereoViewIndex = (int32)(InStereoViewIndex + ContextIt.ContextNum);
 				ContextIt.bDisableRender = true;
-				ContextIt.FrameTargetRect = GetValidRect(DesiredFrameTargetRect, TEXT("Context Frame"));
+				ContextIt.FrameTargetRect = FDisplayClusterViewportHelpers::GetValidViewportRect(DesiredFrameTargetRect, GetId(), TEXT("Context Frame"));
 			}
 
 			return true;
@@ -620,7 +588,7 @@ bool FDisplayClusterViewport::UpdateFrameContexts(const uint32 InStereoViewIndex
 	}
 
 	// Make sure the frame target rect doesn't exceed the maximum resolution, and preserve its aspect ratio if it needs to be clamped
-	FIntRect FrameTargetRect = GetValidRect(DesiredFrameTargetRect, TEXT("Context Frame"));
+	FIntRect FrameTargetRect = FDisplayClusterViewportHelpers::GetValidViewportRect(DesiredFrameTargetRect, GetId(), TEXT("Context Frame"));
 
 	// Exclude zero-size viewports from render
 	if (FrameTargetRect.Size().GetMin() <= 0)
@@ -656,12 +624,18 @@ bool FDisplayClusterViewport::UpdateFrameContexts(const uint32 InStereoViewIndex
 	FIntRect RenderTargetRect = FIntRect(FIntPoint(0, 0), DesiredContextSize);
 
 	// Support custom frustum rendering feature
-	CustomFrustumRendering.Update(*this, RenderTargetRect);
+	if (!RenderSettings.bDisableCustomFrustumFeature)
+	{
+		FDisplayClusterViewport_CustomFrustumRuntimeSettings::UpdateCustomFrustumSettings(GetId(), RenderSettings.CustomFrustumSettings, CustomFrustumRuntimeSettings, RenderTargetRect);
+	}
 
 	FIntPoint ContextSize = RenderTargetRect.Size();
 
 	// Support overscan rendering feature
-	OverscanRendering.Update(*this, RenderTargetRect);
+	if (!RenderSettings.bDisableFrustumOverscanFeature)
+	{
+		FDisplayClusterViewport_OverscanRuntimeSettings::UpdateOverscanSettings(GetId(), RenderSettings.OverscanSettings, OverscanRuntimeSettings, RenderTargetRect);
+	}
 
 	const float BaseCustomBufferRatio = GetCustomBufferRatio(InFrameSettings);
 
diff --git a/Engine/Plugins/Runtime/nDisplay/Source/DisplayCluster/Private/Render/Viewport/DisplayClusterViewport.h b/Engine/Plugins/Runtime/nDisplay/Source/DisplayCluster/Private/Render/Viewport/DisplayClusterViewport.h
--- a/Engine/Plugins/Runtime/nDisplay/Source/DisplayCluster/Private/Render/Viewport/DisplayClusterViewport.h
+++ b/Engine/Plugins/Runtime/nDisplay/Source/DisplayCluster/Private/Render/Viewport/DisplayClusterViewport.h
@@ -5,8 +5,8 @@
 #include "Render/Viewport/IDisplayClusterViewport.h"
 
 #include "Render/Viewport/Containers/ImplDisplayClusterViewport_CameraMotionBlur.h"
-#include "Render/Viewport/Containers/ImplDisplayClusterViewport_CustomFrustum.h"
-#include "Render/Viewport/Containers/ImplDisplayClusterViewport_Overscan.h"
+#include "Render/Viewport/Containers/DisplayClusterViewport_CustomFrustumRuntimeSettings.h"
+#include "Render/Viewport/Containers/DisplayClusterViewport_OverscanRuntimeSettings.h"
 #include "Render/Viewport/Containers/DisplayClusterViewportRemap.h"
 #include "Render/Viewport/Containers/DisplayClusterViewport_Enums.h"
 
@@ -218,8 +218,9 @@ public:
 		PostRenderSettings.BeginUpdateSettings();
 		VisibilitySettings.ResetConfiguration();
 		CameraMotionBlur.ResetConfiguration();
-		OverscanRendering.ResetConfiguration();
-		CustomFrustumRendering.ResetConfiguration();
+
+		OverscanRuntimeSettings = FDisplayClusterViewport_OverscanRuntimeSettings();
+		CustomFrustumRuntimeSettings = FDisplayClusterViewport_CustomFrustumRuntimeSettings();
 	}
 
 	// Active view extension for this viewport
@@ -240,9 +241,6 @@ public:
 	/** Compare OCIO with another viewport, return true if they are equal. */
 	bool IsOpenColorIOEquals(const FDisplayClusterViewport& InViewport) const;
 
-public:
-	FIntRect GetValidRect(const FIntRect& InRect, const TCHAR* DbgSourceName);
-
 private:
 	float GetClusterRenderTargetRatioMult(const FDisplayClusterRenderFrameSettings& InFrameSettings) const;
 	FIntPoint GetDesiredContextSize(const FIntPoint& InSize, const FDisplayClusterRenderFrameSettings& InFrameSettings) const;
@@ -270,8 +268,9 @@ public:
 
 	// Additional features:
 	FImplDisplayClusterViewport_CameraMotionBlur CameraMotionBlur;
-	FImplDisplayClusterViewport_Overscan         OverscanRendering;
-	FImplDisplayClusterViewport_CustomFrustum    CustomFrustumRendering;
+
+	FDisplayClusterViewport_OverscanRuntimeSettings      OverscanRuntimeSettings;
+	FDisplayClusterViewport_CustomFrustumRuntimeSettings CustomFrustumRuntimeSettings;
 
 	// viewport OutputRemap feature
 	FDisplayClusterViewportRemap ViewportRemap;
diff --git a/Engine/Plugins/Runtime/nDisplay/Source/DisplayCluster/Private/Render/Viewport/DisplayClusterViewportHelpers.cpp b/Engine/Plugins/Runtime/nDisplay/Source/DisplayCluster/Private/Render/Viewport/DisplayClusterViewportHelpers.cpp
--- a/Engine/Plugins/Runtime/nDisplay/Source/DisplayCluster/Private/Render/Viewport/DisplayClusterViewportHelpers.cpp
+++ b/Engine/Plugins/Runtime/nDisplay/Source/DisplayCluster/Private/Render/Viewport/DisplayClusterViewportHelpers.cpp
@@ -4,6 +4,8 @@
 #include "Render/Viewport/RenderTarget/DisplayClusterRenderTargetResource.h"
 #include "Render/Viewport/RenderFrame/DisplayClusterRenderFrameSettings.h"
 
+#include "Misc/DisplayClusterLog.h"
+
 #include "Engine/TextureRenderTarget2D.h"
 
 ////////////////////////////////////////////////////////////////////////////////
@@ -126,6 +128,38 @@ int32 FDisplayClusterViewportHelpers::GetMinTextureDimension()
 	return MinTextureDimension;
 }
 
+FIntRect FDisplayClusterViewportHelpers::GetValidViewportRect(const FIntRect& InRect, const FString& InViewportId, const TCHAR* InResourceName)
+{
+	// The target always needs be within GMaxTextureDimensions, larger dimensions are not supported by the engine
+	const int32 MaxTextureSize = FDisplayClusterViewportHelpers::GetMaxTextureDimension();
+	const int32 MinTextureSize = FDisplayClusterViewportHelpers::GetMinTextureDimension();
+
+	int32 Width = FMath::Max(MinTextureSize, InRect.Width());
+	int32 Height = FMath::Max(MinTextureSize, InRect.Height());
+
+	FIntRect OutRect(InRect.Min, InRect.Min + FIntPoint(Width, Height));
+
+	float RectScale = 1;
+
+	// Make sure the rect doesn't exceed the maximum resolution, and preserve its aspect ratio if it needs to be clamped
+	int32 RectMaxSize = OutRect.Max.GetMax();
+	if (RectMaxSize > MaxTextureSize)
+	{
+		RectScale = float(MaxTextureSize) / RectMaxSize;
+		UE_LOG(LogDisplayClusterViewport, Error, TEXT("The viewport '%s' rect '%s' size %dx%d clamped: max texture dimensions is %d"), *InViewportId, (InResourceName == nullptr) ? TEXT("none") : InResourceName, InRect.Max.X, InRect.Max.Y, MaxTextureSize);
+	}
+
+	OutRect.Min.X = FMath::Min(OutRect.Min.X, MaxTextureSize);
+	OutRect.Min.Y = FMath::Min(OutRect.Min.Y, MaxTextureSize);
+
+	const FIntPoint ScaledRectMax = FDisplayClusterViewportHelpers::ScaleTextureSize(OutRect.Max, RectScale);
+
+	OutRect.Max.X = FMath::Clamp(ScaledRectMax.X, OutRect.Min.X, MaxTextureSize);
+	OutRect.Max.Y = FMath::Clamp(ScaledRectMax.Y, OutRect.Min.Y, MaxTextureSize);
+
+	return OutRect;
+}
+
 bool FDisplayClusterViewportHelpers::IsValidTextureSize(const FIntPoint& InSize)
 {
 	return InSize.GetMin() >= GetMinTextureDimension() && InSize.GetMax() <= GetMaxTextureDimension();
diff --git a/Engine/Plugins/Runtime/nDisplay/Source/DisplayCluster/Private/Render/Viewport/DisplayClusterViewportHelpers.h b/Engine/Plugins/Runtime/nDisplay/Source/DisplayCluster/Private/Render/Viewport/DisplayClusterViewportHelpers.h
--- a/Engine/Plugins/Runtime/nDisplay/Source/DisplayCluster/Private/Render/Viewport/DisplayClusterViewportHelpers.h
+++ b/Engine/Plugins/Runtime/nDisplay/Source/DisplayCluster/Private/Render/Viewport/DisplayClusterViewportHelpers.h
@@ -19,6 +19,16 @@ public:
 	/** Get the minimum allowable texture size used for the nDisplay viewport. */
 	static int32 GetMaxTextureDimension();
 
+	/** Get the valid viewport size.
+	* 
+	* @param InRect         - source region
+	* @param InViewportId   - The viewport name (debug purpose)
+	* @param InResourceName - The resource name (debug purpose)
+	* 
+	* @return the region of the viewport that can be used
+	*/
+	static FIntRect GetValidViewportRect(const FIntRect& InRect, const FString& InViewportId, const TCHAR* InResourceName = nullptr);
+
 	/** Return true, if size is valid.
 	* 
 	* @param InSize - the size to be checked.
diff --git a/Engine/Plugins/Runtime/nDisplay/Source/DisplayCluster/Private/Render/Viewport/DisplayClusterViewportManager.cpp b/Engine/Plugins/Runtime/nDisplay/Source/DisplayCluster/Private/Render/Viewport/DisplayClusterViewportManager.cpp
--- a/Engine/Plugins/Runtime/nDisplay/Source/DisplayCluster/Private/Render/Viewport/DisplayClusterViewportManager.cpp
+++ b/Engine/Plugins/Runtime/nDisplay/Source/DisplayCluster/Private/Render/Viewport/DisplayClusterViewportManager.cpp
@@ -643,8 +643,19 @@ bool FDisplayClusterViewportManager::BeginNewFrame(FViewport* InViewport, UWorld
 			// Save orig viewport contexts
 			TArray<FDisplayClusterViewport_Context> PrevContexts;
 			PrevContexts.Append(ViewportIt->GetContexts());
+
+			if (ViewportIt->ProjectionPolicy.IsValid())
+			{
+				ViewportIt->ProjectionPolicy->BeginUpdateFrameContexts(ViewportIt.Get());
+			}
+
 			ViewportIt->UpdateFrameContexts(FirstViewportStereoViewIndex, RenderFrameSettings);
 
+			if (ViewportIt->ProjectionPolicy.IsValid())
+			{
+				ViewportIt->ProjectionPolicy->EndUpdateFrameContexts(ViewportIt.Get());
+			}
+
 			HandleViewportRTTChanges(PrevContexts, ViewportIt->GetContexts());
 		}
 		else
diff --git a/Engine/Plugins/Runtime/nDisplay/Source/DisplayCluster/Private/Render/Viewport/DisplayClusterViewportProxy.cpp b/Engine/Plugins/Runtime/nDisplay/Source/DisplayCluster/Private/Render/Viewport/DisplayClusterViewportProxy.cpp
--- a/Engine/Plugins/Runtime/nDisplay/Source/DisplayCluster/Private/Render/Viewport/DisplayClusterViewportProxy.cpp
+++ b/Engine/Plugins/Runtime/nDisplay/Source/DisplayCluster/Private/Render/Viewport/DisplayClusterViewportProxy.cpp
@@ -641,6 +641,60 @@ bool FDisplayClusterViewportProxy::ImplGetResourcesWithRects_RenderThread(const
 	return true;
 }
 
+bool FDisplayClusterViewportProxy::ApplyOCIO_RenderThread(FRHICommandListImmediate& RHICmdList, const FDisplayClusterViewportProxy& InSrcViewportProxy, const EDisplayClusterViewportResourceType InSrcResourceType) const
+{
+	if (GetOpenColorIOMode() != EDisplayClusterViewportOpenColorIOMode::Resolved)
+	{
+		return false;
+	}
+
+	const EDisplayClusterViewportResourceType DestResourceType = (InSrcResourceType == EDisplayClusterViewportResourceType::InternalRenderTargetResource)
+		? EDisplayClusterViewportResourceType::InputShaderResource
+		: EDisplayClusterViewportResourceType::AdditionalTargetableResource;
+
+	TArray<FRHITexture2D*> Input, Output;
+	TArray<FIntRect> InputRects, OutputRects;
+	if (!InSrcViewportProxy.GetResourcesWithRects_RenderThread(InSrcResourceType, Input, InputRects)
+		|| !GetResourcesWithRects_RenderThread(DestResourceType, Output, OutputRects)
+		|| Input.Num() != Output.Num())
+	{
+		return false;
+	};
+
+	FRDGBuilder GraphBuilder(RHICmdList);
+
+	bool bResult = false;
+	for (int32 ContextNum = 0; ContextNum < Input.Num(); ContextNum++)
+	{
+		const bool bUnpremultiply = EnumHasAnyFlags(RenderSettingsICVFX.RuntimeFlags, EDisplayClusterViewportRuntimeICVFXFlags::UVLightcard | EDisplayClusterViewportRuntimeICVFXFlags::Lightcard);
+		const bool bInvertAlpha = !EnumHasAnyFlags(RenderSettingsICVFX.RuntimeFlags, EDisplayClusterViewportRuntimeICVFXFlags::UVLightcard);
+
+		if (OpenColorIO->AddPass_RenderThread(
+			GraphBuilder,
+			InSrcViewportProxy.GetContexts_RenderThread()[ContextNum],
+			Input[ContextNum], InputRects[ContextNum],
+			Output[ContextNum], OutputRects[ContextNum],
+			bUnpremultiply,
+			bInvertAlpha))
+		{
+			bResult = true;
+		}
+	}
+
+	if (bResult)
+	{
+		GraphBuilder.Execute();
+
+		// copy OCIO results back
+		if (DestResourceType != EDisplayClusterViewportResourceType::InputShaderResource)
+		{
+			ResolveResources_RenderThread(RHICmdList, DestResourceType, EDisplayClusterViewportResourceType::InputShaderResource);
+		}
+	}
+
+	return bResult;
+}
+
 void FDisplayClusterViewportProxy::UpdateDeferredResources(FRHICommandListImmediate& RHICmdList) const
 {
 	check(IsInRenderingThread());
@@ -664,47 +718,22 @@ void FDisplayClusterViewportProxy::UpdateDeferredResources(FRHICommandListImmedi
 		return;
 	}
 
-	bool bPass0Applied = false;
+	EDisplayClusterViewportResourceType SrcResourceType = EDisplayClusterViewportResourceType::InternalRenderTargetResource;
 
-	// OCIO support on the first pass for an resolved RTT
-	if (GetOpenColorIOMode() == EDisplayClusterViewportOpenColorIOMode::Resolved)
+	// pre-Pass 0 (Projection policy):The projection policy can use its own method to resolve 'InternalRenderTargetResource' to 'InputShaderResource'
+	if (ProjectionPolicy.IsValid() && ProjectionPolicy->ResolveInternalRenderTargetResource_RenderThread(RHICmdList, this, &SourceViewportProxy))
 	{
-		TArray<FRHITexture2D*> Input, Output;
-		TArray<FIntRect> InputRects, OutputRects;
-		if (SourceViewportProxy.GetResourcesWithRects_RenderThread(EDisplayClusterViewportResourceType::InternalRenderTargetResource, Input, InputRects)
-			&& GetResourcesWithRects_RenderThread(EDisplayClusterViewportResourceType::InputShaderResource, Output, OutputRects)
-			&& Input.Num() == Output.Num())
-		{
-			FRDGBuilder GraphBuilder(RHICmdList);
-
-			for (int32 ContextNum = 0; ContextNum < Input.Num(); ContextNum++)
-			{
-				const bool bUnpremultiply = EnumHasAnyFlags(RenderSettingsICVFX.RuntimeFlags, EDisplayClusterViewportRuntimeICVFXFlags::UVLightcard | EDisplayClusterViewportRuntimeICVFXFlags::Lightcard);
-				const bool bInvertAlpha = !EnumHasAnyFlags(RenderSettingsICVFX.RuntimeFlags, EDisplayClusterViewportRuntimeICVFXFlags::UVLightcard);
-
-				if (OpenColorIO->AddPass_RenderThread(
-					GraphBuilder,
-					SourceViewportProxy.GetContexts_RenderThread()[ContextNum],
-					Input[ContextNum], InputRects[ContextNum],
-					Output[ContextNum], OutputRects[ContextNum],
-					bUnpremultiply,
-					bInvertAlpha))
-				{
-					bPass0Applied = true;
-				}
-			}
-
-			if (bPass0Applied)
-			{
-				GraphBuilder.Execute();
-			}
-		}
+		SrcResourceType = EDisplayClusterViewportResourceType::InputShaderResource;
 	}
-	
-	if(!bPass0Applied)
+
+	// Pass 0 (OCIO): OCIO support on the first pass for an resolved RTT
+	if(!ApplyOCIO_RenderThread(RHICmdList, SourceViewportProxy, SrcResourceType))
 	{
-		// Pass 0: Resolve from RTT region to separated viewport context resource:
-		ResolveResources_RenderThread(RHICmdList, EDisplayClusterViewportResourceType::InternalRenderTargetResource, EDisplayClusterViewportResourceType::InputShaderResource);
+		// Pass 0 (default): Resolve from RTT region to separated viewport context resource:
+		if (SrcResourceType == EDisplayClusterViewportResourceType::InternalRenderTargetResource)
+		{
+			ImplResolveResources_RenderThread(RHICmdList, &SourceViewportProxy, SrcResourceType, EDisplayClusterViewportResourceType::InputShaderResource);
+		}
 	}
 
 	// Pass 1: Generate blur postprocess effect for render target texture rect for all contexts
@@ -852,10 +881,10 @@ FIntRect FDisplayClusterViewportProxy::GetFinalContextRect(const EDisplayCluster
 	switch (InResourceType)
 	{
 	case EDisplayClusterViewportResourceType::InternalRenderTargetResource:
-		if (OverscanSettings.bIsEnabled && CVarDisplayClusterRenderOverscanResolve.GetValueOnRenderThread() != 0)
+		if (OverscanRuntimeSettings.bIsEnabled && CVarDisplayClusterRenderOverscanResolve.GetValueOnRenderThread() != 0)
 		{
 			// Support overscan crop
-			return OverscanSettings.OverscanPixels.GetInnerRect(InRect);
+			return OverscanRuntimeSettings.OverscanPixels.GetInnerRect(InRect);
 		}
 		break;
 	default:
@@ -871,6 +900,13 @@ bool FDisplayClusterViewportProxy::ResolveResources_RenderThread(FRHICommandList
 	return ImplResolveResources_RenderThread(RHICmdList, this, InExtResourceType, OutExtResourceType, InContextNum);
 }
 
+bool FDisplayClusterViewportProxy::ResolveResources_RenderThread(FRHICommandListImmediate& RHICmdList, IDisplayClusterViewportProxy* InputResourceViewportProxy, const EDisplayClusterViewportResourceType InExtResourceType, const EDisplayClusterViewportResourceType OutExtResourceType, const int32 InContextNum) const
+{
+	const FDisplayClusterViewportProxy* SourceProxy = static_cast<FDisplayClusterViewportProxy*>(InputResourceViewportProxy);
+
+	return ImplResolveResources_RenderThread(RHICmdList, SourceProxy , InExtResourceType, OutExtResourceType, InContextNum);
+}
+
 bool FDisplayClusterViewportProxy::ImplResolveResources_RenderThread(FRHICommandListImmediate& RHICmdList, FDisplayClusterViewportProxy const* SourceProxy, const EDisplayClusterViewportResourceType InExtResourceType, const EDisplayClusterViewportResourceType OutExtResourceType, const int32 InContextNum) const
 {
 	check(IsInRenderingThread());
diff --git a/Engine/Plugins/Runtime/nDisplay/Source/DisplayCluster/Private/Render/Viewport/DisplayClusterViewportProxy.h b/Engine/Plugins/Runtime/nDisplay/Source/DisplayCluster/Private/Render/Viewport/DisplayClusterViewportProxy.h
--- a/Engine/Plugins/Runtime/nDisplay/Source/DisplayCluster/Private/Render/Viewport/DisplayClusterViewportProxy.h
+++ b/Engine/Plugins/Runtime/nDisplay/Source/DisplayCluster/Private/Render/Viewport/DisplayClusterViewportProxy.h
@@ -11,7 +11,7 @@
 
 #include "Render/Viewport/IDisplayClusterViewportProxy.h"
 #include "Render/Viewport/DisplayClusterViewportResources.h"
-#include "Render/Viewport/Containers/DisplayClusterViewport_OverscanSettings.h"
+#include "Render/Viewport/Containers/DisplayClusterViewport_OverscanRuntimeSettings.h"
 
 class FDisplayClusterViewportProxyData;
 class IDisplayClusterViewportManagerProxy;
@@ -141,7 +141,7 @@ public:
 
 	// Resolve resource contexts
 	virtual bool ResolveResources_RenderThread(FRHICommandListImmediate& RHICmdList, const EDisplayClusterViewportResourceType InputResourceType, const EDisplayClusterViewportResourceType OutputResourceType, const int32 InContextNum = INDEX_NONE) const override;
-
+	virtual bool ResolveResources_RenderThread(FRHICommandListImmediate& RHICmdList, IDisplayClusterViewportProxy* InputResourceViewportProxy, const EDisplayClusterViewportResourceType InputResourceType, const EDisplayClusterViewportResourceType OutputResourceType, const int32 InContextNum = INDEX_NONE) const override;
 
 	virtual const class IDisplayClusterViewportManagerProxy* GetViewportManagerProxy_RenderThread() const override;
 
@@ -285,7 +285,7 @@ private:
 	void ImplViewportRemap_RenderThread(FRHICommandListImmediate& RHICmdList) const;
 	void ImplPreviewReadPixels_RenderThread(FRHICommandListImmediate& RHICmdList) const;
 
-	bool ImplResolveResources_RenderThread(FRHICommandListImmediate& RHICmdList, FDisplayClusterViewportProxy const* SourceProxy, const EDisplayClusterViewportResourceType InputResourceType, const EDisplayClusterViewportResourceType OutputResourceType, const int32 InContextNum) const;
+	bool ImplResolveResources_RenderThread(FRHICommandListImmediate& RHICmdList, FDisplayClusterViewportProxy const* SourceProxy, const EDisplayClusterViewportResourceType InputResourceType, const EDisplayClusterViewportResourceType OutputResourceType, const int32 InContextNum = INDEX_NONE) const;
 
 	/** When a resource by type can be overridden from another viewport, true is returned. */
 	bool ShouldOverrideViewportResource(const EDisplayClusterViewportResourceType InResourceType) const;
@@ -333,6 +333,9 @@ private:
 	 */
 	const FDisplayClusterViewportProxy& GetRenderingViewportProxy() const;
 
+	/** Apply OCIO to the image from the 'InSrcResourceType' and put the result image to the EDisplayClusterViewportResourceType::InputShaderResource. */
+	bool ApplyOCIO_RenderThread(FRHICommandListImmediate& RHICmdList, const FDisplayClusterViewportProxy& InSrcViewportProxy, const EDisplayClusterViewportResourceType InSrcResourceType) const;
+
 	/** Check if there is an RTT source (internal or external) in this viewport proxy. */
 	bool IsInputRenderTargetResourceExists() const;
 
@@ -361,7 +364,7 @@ protected:
 	FDisplayClusterViewport_PostRenderSettings   PostRenderSettings;
 
 	// Additional parameters
-	FDisplayClusterViewport_OverscanSettings     OverscanSettings;
+	FDisplayClusterViewport_OverscanRuntimeSettings OverscanRuntimeSettings;
 
 	TSharedPtr<IDisplayClusterRender_MeshComponent, ESPMode::ThreadSafe> RemapMesh;
 
diff --git a/Engine/Plugins/Runtime/nDisplay/Source/DisplayCluster/Private/Render/Viewport/DisplayClusterViewport_Math.cpp b/Engine/Plugins/Runtime/nDisplay/Source/DisplayCluster/Private/Render/Viewport/DisplayClusterViewport_Math.cpp
--- a/Engine/Plugins/Runtime/nDisplay/Source/DisplayCluster/Private/Render/Viewport/DisplayClusterViewport_Math.cpp
+++ b/Engine/Plugins/Runtime/nDisplay/Source/DisplayCluster/Private/Render/Viewport/DisplayClusterViewport_Math.cpp
@@ -159,7 +159,7 @@ bool FDisplayClusterViewport::GetProjectionMatrix(const uint32 InContextNum, FMa
 				Contexts[InContextNum].ProjectionMatrix = OutPrjMatrix;
 				EnumAddFlags(Contexts[InContextNum].ContextState, EDisplayClusterViewportContextState::HasCalculatedProjectionMatrix);
 
-				if (OverscanRendering.IsEnabled())
+				if (OverscanRuntimeSettings.bIsEnabled)
 				{
 					// use overscan proj matrix for rendering
 					OutPrjMatrix = Contexts[InContextNum].OverscanProjectionMatrix;
@@ -231,7 +231,7 @@ void FDisplayClusterViewport::CalculateProjectionMatrix(const uint32 InContextNu
 
 	// Support custom frustum rendering
 	const float OrigValues[] = {l, r, t, b};
-	if (CustomFrustumRendering.UpdateProjectionAngles(l, r, t, b))
+	if (FDisplayClusterViewport_CustomFrustumRuntimeSettings::UpdateProjectionAngles(CustomFrustumRuntimeSettings, l, r, t, b))
 	{
 		const bool bIsValidLimits =  FMath::IsWithin(l, -MaxValue, MaxValue)
 							&& FMath::IsWithin(r, -MaxValue, MaxValue)
@@ -241,7 +241,7 @@ void FDisplayClusterViewport::CalculateProjectionMatrix(const uint32 InContextNu
 		if (!bIsValidLimits)
 		{
 			// overscan out of frustum : disable
-			CustomFrustumRendering.Disable();
+			CustomFrustumRuntimeSettings.bIsEnabled = false;
 
 			// restore orig values
 			l = OrigValues[0];
@@ -256,7 +256,7 @@ void FDisplayClusterViewport::CalculateProjectionMatrix(const uint32 InContextNu
 
 	Contexts[InContextNum].ProjectionMatrix = IDisplayClusterViewport::MakeProjectionMatrix(l, r, t, b, n, f);
 
-	if (OverscanRendering.UpdateProjectionAngles(l, r, t, b))
+	if (FDisplayClusterViewport_OverscanRuntimeSettings::UpdateProjectionAngles(OverscanRuntimeSettings, l, r, t, b))
 	{
 		if (FMath::IsWithin(l, -MaxValue, MaxValue) &&
 			FMath::IsWithin(r, -MaxValue, MaxValue) &&
@@ -271,7 +271,7 @@ void FDisplayClusterViewport::CalculateProjectionMatrix(const uint32 InContextNu
 	}
 
 	// overscan out of frustum: disable
-	OverscanRendering.Disable();
+	OverscanRuntimeSettings.bIsEnabled = false;
 }
 
 ///////////////////////////////////////////////////////////////////////////////////////
diff --git a/Engine/Plugins/Runtime/nDisplay/Source/DisplayCluster/Public/Render/Projection/IDisplayClusterProjectionPolicy.h b/Engine/Plugins/Runtime/nDisplay/Source/DisplayCluster/Public/Render/Projection/IDisplayClusterProjectionPolicy.h
--- a/Engine/Plugins/Runtime/nDisplay/Source/DisplayCluster/Public/Render/Projection/IDisplayClusterProjectionPolicy.h
+++ b/Engine/Plugins/Runtime/nDisplay/Source/DisplayCluster/Public/Render/Projection/IDisplayClusterProjectionPolicy.h
@@ -82,6 +82,19 @@ public:
 	virtual void HandleEndScene(IDisplayClusterViewport* InViewport)
 	{ }
 
+	/**
+	 * Called before FDisplayClusterViewport::UpdateFrameContexts()
+	 * From this function, the policy can override any viewport settings (custom overscan, etc).
+	 */
+	virtual void BeginUpdateFrameContexts(IDisplayClusterViewport* InViewport) const
+	{ }
+
+	/**
+	 * Called after FDisplayClusterViewport::UpdateFrameContexts()
+	 */
+	virtual void EndUpdateFrameContexts(IDisplayClusterViewport* InViewport) const
+	{ }
+
 	/**
 	* Set warp policy for this projection
 	*
@@ -256,6 +269,21 @@ public:
 		return false;
 	}
 
+	/**
+	* Override copying 'InternalRenderTargetResource' to 'InputShaderResource'.
+	* The same behavior as for function bellow  is expected:
+	*     ResolveResources_RenderThread(RHICmdList, InSourceViewportProxy, InSrcResourceType, InDstResourceType);
+	*
+	* @param InViewportProxy       - This ViewportProxy will get the result
+	* @param InSourceViewportProxy - This ViewportProxy will provide the input resource
+	* 
+	* @return - true, if copying is overridden.
+	*/
+	virtual bool ResolveInternalRenderTargetResource_RenderThread(FRHICommandListImmediate& RHICmdList, const IDisplayClusterViewportProxy* InViewportProxy, const IDisplayClusterViewportProxy* InSourceViewportProxy)
+	{
+		return false;
+	}
+
 #if WITH_EDITOR
 	/**
 	* Ask projection policy instance if it has any mesh based preview
diff --git a/Engine/Plugins/Runtime/nDisplay/Source/DisplayCluster/Public/Render/Viewport/Containers/DisplayClusterViewport_CustomFrustumSettings.h b/Engine/Plugins/Runtime/nDisplay/Source/DisplayCluster/Public/Render/Viewport/Containers/DisplayClusterViewport_CustomFrustumSettings.h
new file mode 100644
--- /dev/null
+++ b/Engine/Plugins/Runtime/nDisplay/Source/DisplayCluster/Public/Render/Viewport/Containers/DisplayClusterViewport_CustomFrustumSettings.h
@@ -0,0 +1,27 @@
+// Copyright Epic Games, Inc. All Rights Reserved.
+
+#pragma once
+
+#include "CoreMinimal.h"
+#include "DisplayClusterViewport_Enums.h"
+
+/**
+* Custom frustum settings of viewport
+*/
+struct FDisplayClusterViewport_CustomFrustumSettings
+{
+	// Enable custom frustum
+	bool bEnabled = false;
+
+	// Enable adaptive resolution
+	bool bAdaptResolution = false;
+
+	// Units type of values
+	EDisplayClusterViewport_FrustumUnit Unit = EDisplayClusterViewport_FrustumUnit::Pixels;
+
+	// custum frustum values
+	float Left = 0;
+	float Right = 0;
+	float Top = 0;
+	float Bottom = 0;
+};
diff --git a/Engine/Plugins/Runtime/nDisplay/Source/DisplayCluster/Public/Render/Viewport/Containers/DisplayClusterViewport_Enums.h b/Engine/Plugins/Runtime/nDisplay/Source/DisplayCluster/Public/Render/Viewport/Containers/DisplayClusterViewport_Enums.h
--- a/Engine/Plugins/Runtime/nDisplay/Source/DisplayCluster/Public/Render/Viewport/Containers/DisplayClusterViewport_Enums.h
+++ b/Engine/Plugins/Runtime/nDisplay/Source/DisplayCluster/Public/Render/Viewport/Containers/DisplayClusterViewport_Enums.h
@@ -119,3 +119,15 @@ enum class EDisplayClusterViewportOverrideMode : uint8
 	// Override all - clone viewport (Dont render; Don't create resources;)
 	All
 };
+
+/**
+* Type of unit for frustum
+*/
+enum class EDisplayClusterViewport_FrustumUnit: uint8
+{
+	// 1 unit = 1 pixel
+	Pixels = 0,
+
+	// 1 unit = 1 per cent
+	Percent
+};
diff --git a/Engine/Plugins/Runtime/nDisplay/Source/DisplayCluster/Public/Render/Viewport/Containers/DisplayClusterViewport_OverscanSettings.h b/Engine/Plugins/Runtime/nDisplay/Source/DisplayCluster/Public/Render/Viewport/Containers/DisplayClusterViewport_OverscanSettings.h
--- a/Engine/Plugins/Runtime/nDisplay/Source/DisplayCluster/Public/Render/Viewport/Containers/DisplayClusterViewport_OverscanSettings.h
+++ b/Engine/Plugins/Runtime/nDisplay/Source/DisplayCluster/Public/Render/Viewport/Containers/DisplayClusterViewport_OverscanSettings.h
@@ -3,51 +3,25 @@
 #pragma once
 
 #include "CoreMinimal.h"
+#include "DisplayClusterViewport_Enums.h"
 
+/**
+* Overscan settings of viewport
+*/
 struct FDisplayClusterViewport_OverscanSettings
 {
-	struct FOverscanPercent
-	{
-		float Left = 0;
-		float Right = 0;
-		float Top = 0;
-		float Bottom = 0;
-	};
+	// Enable overscan
+	bool bEnabled = false;
 
-	struct FOverscanPixels
-	{
-		inline FIntRect GetInnerRect(const FIntRect& InRect) const
-		{
-			const FIntPoint InnerSize = InRect.Size() - Size();
-			const FIntPoint InnerPos = FIntPoint(Left, Top);
+	// Set to True to render at the overscan resolution, set to false to render at the resolution in the configuration and scale for overscan
+	bool bOversize = false;
 
-			return FIntRect(InnerPos, InnerPos + InnerSize);
-		}
+	// Units type of overscan values
+	EDisplayClusterViewport_FrustumUnit Unit = EDisplayClusterViewport_FrustumUnit::Pixels;
 
-		inline FIntPoint Size() const
-		{
-			return FIntPoint(Left + Right, Top + Bottom);
-		}
-
-		int32 Left = 0;
-		int32 Right = 0;
-		int32 Top = 0;
-		int32 Bottom = 0;
-	};
-
-	bool bIsEnabled = false;
-
-	// Overscan sides in percent
-	FOverscanPercent OverscanPercent;
-
-	// Overscan sides in pixels
-	FOverscanPixels  OverscanPixels;
-
-	static float ClampPercent(float InValue)
-	{
-		static const float MaxCustomFrustumValue = 5.f;
-
-		return FMath::Clamp(InValue, -MaxCustomFrustumValue, MaxCustomFrustumValue);
-	}
+	// Overscan values
+	float Left = 0;
+	float Right = 0;
+	float Top = 0;
+	float Bottom = 0;
 };
-
diff --git a/Engine/Plugins/Runtime/nDisplay/Source/DisplayCluster/Public/Render/Viewport/Containers/DisplayClusterViewport_RenderSettings.h b/Engine/Plugins/Runtime/nDisplay/Source/DisplayCluster/Public/Render/Viewport/Containers/DisplayClusterViewport_RenderSettings.h
--- a/Engine/Plugins/Runtime/nDisplay/Source/DisplayCluster/Public/Render/Viewport/Containers/DisplayClusterViewport_RenderSettings.h
+++ b/Engine/Plugins/Runtime/nDisplay/Source/DisplayCluster/Public/Render/Viewport/Containers/DisplayClusterViewport_RenderSettings.h
@@ -4,7 +4,8 @@
 
 #include "CoreMinimal.h"
 
-#include "DisplayClusterViewport_Enums.h"
+#include "DisplayClusterViewport_OverscanSettings.h"
+#include "DisplayClusterViewport_CustomFrustumSettings.h"
 
 /**
  * nDisplay viewport render settings.
@@ -43,9 +44,15 @@ public:
 	// Disable CustomFrustum feature from viewport settings
 	bool bDisableCustomFrustumFeature = false;
 
+	// Viewport custom frustum settings
+	FDisplayClusterViewport_CustomFrustumSettings CustomFrustumSettings;
+
 	// Disable viewport overscan feature from settings
 	bool bDisableFrustumOverscanFeature = false;
 
+	// Viewport frustum overscan settings
+	FDisplayClusterViewport_OverscanSettings OverscanSettings;
+
 	// Read viewport pixels for preview (this flag is cleared at the end of the frame)
 	bool bPreviewReadPixels = false;
 
@@ -98,6 +105,9 @@ public:
 		bFreezeRendering = false;
 		bWarpBlendRenderAlphaChannel = false;
 
+		OverscanSettings = FDisplayClusterViewport_OverscanSettings();
+		CustomFrustumSettings = FDisplayClusterViewport_CustomFrustumSettings();
+
 		CaptureMode = EDisplayClusterViewportCaptureMode::Default;
 
 		ViewportOverrideMode = EDisplayClusterViewportOverrideMode::None;
diff --git a/Engine/Plugins/Runtime/nDisplay/Source/DisplayCluster/Public/Render/Viewport/IDisplayClusterViewportProxy.h b/Engine/Plugins/Runtime/nDisplay/Source/DisplayCluster/Public/Render/Viewport/IDisplayClusterViewportProxy.h
--- a/Engine/Plugins/Runtime/nDisplay/Source/DisplayCluster/Public/Render/Viewport/IDisplayClusterViewportProxy.h
+++ b/Engine/Plugins/Runtime/nDisplay/Source/DisplayCluster/Public/Render/Viewport/IDisplayClusterViewportProxy.h
@@ -75,6 +75,18 @@ public:
 	 */
 	virtual bool ResolveResources_RenderThread(FRHICommandListImmediate& RHICmdList, const EDisplayClusterViewportResourceType InputResourceType, const EDisplayClusterViewportResourceType OutputResourceType, const int32 InContextNum = INDEX_NONE) const = 0;
 
+	/** Copy resource contexts by type
+	 *
+	 * @param RHICmdList                 - RHIinterface
+	 * @param InputResourceViewportProxy - The 'InputResourceType' resource will be obtained from this object
+	 * @param InputResourceType          - Input resource type (RTT, Shader, MIPS, etc)
+	 * @param OutputResourceType         - Output resource type (RTT, Shader, MIPS, etc)
+	 * @param InContextNum               - [optional] the type of context to copy (by default, all contexts are copied).
+	 *
+	 * @return - true if success
+	 */
+	virtual bool ResolveResources_RenderThread(FRHICommandListImmediate& RHICmdList, IDisplayClusterViewportProxy* InputResourceViewportProxy, const EDisplayClusterViewportResourceType InputResourceType, const EDisplayClusterViewportResourceType OutputResourceType, const int32 InContextNum = INDEX_NONE) const = 0;
+
 	/** Return output resource type (support preview, remap, etc). */
 	UE_DEPRECATED(5.3, "This function has beend deprecate. Please use the new enumeration values.")
 	virtual EDisplayClusterViewportResourceType GetOutputResourceType_RenderThread() const
