nDisplay: Fixed compiler issue  #334021

diff --git a/Engine/Plugins/Runtime/nDisplay/Source/DisplayClusterProjection/Private/Policy/VIOSO/DisplayClusterProjectionVIOSOGeometryExportData.cpp b/Engine/Plugins/Runtime/nDisplay/Source/DisplayClusterProjection/Private/Policy/VIOSO/DisplayClusterProjectionVIOSOGeometryExportData.cpp
--- a/Engine/Plugins/Runtime/nDisplay/Source/DisplayClusterProjection/Private/Policy/VIOSO/DisplayClusterProjectionVIOSOGeometryExportData.cpp
+++ b/Engine/Plugins/Runtime/nDisplay/Source/DisplayClusterProjection/Private/Policy/VIOSO/DisplayClusterProjectionVIOSOGeometryExportData.cpp
@@ -4,6 +4,8 @@
 #include "Policy/VIOSO/Windows/DisplayClusterProjectionVIOSOWarper.h"
 #include "Misc/DisplayClusterDataCache.h"
 
+#include "HAL/IConsoleManager.h"
+
 int32 GDisplayClusterRender_VIOSOPreviewMeshCacheEnable = 1;
 static FAutoConsoleVariableRef CDisplayClusterRender_VIOSOPreviewMeshCacheEnable(
 	TEXT("nDisplay.cache.VIOSO.PreviewMesh.enable"),
