[nDisplay] Fixed SyncID generation to preserve ID uniqness in synchronization components

diff --git a/Engine/Plugins/Runtime/nDisplay/Source/DisplayCluster/Private/Game/EngineClasses/Scene/DisplayClusterSceneComponentSync.cpp b/Engine/Plugins/Runtime/nDisplay/Source/DisplayCluster/Private/Game/EngineClasses/Scene/DisplayClusterSceneComponentSync.cpp
--- a/Engine/Plugins/Runtime/nDisplay/Source/DisplayCluster/Private/Game/EngineClasses/Scene/DisplayClusterSceneComponentSync.cpp
+++ b/Engine/Plugins/Runtime/nDisplay/Source/DisplayCluster/Private/Game/EngineClasses/Scene/DisplayClusterSceneComponentSync.cpp
@@ -78,7 +78,7 @@ bool UDisplayClusterSceneComponentSync::IsActive() const
 //////////////////////////////////////////////////////////////////////////////////////////////
 FString UDisplayClusterSceneComponentSync::GenerateSyncId()
 {
-	return FString::Printf(TEXT("S_%s"), *GetOwner()->GetName());
+	return FString::Printf(TEXT("S_%s"), *GetFullName());
 }
 
 FString UDisplayClusterSceneComponentSync::SerializeToString() const
diff --git a/Engine/Plugins/Runtime/nDisplay/Source/DisplayCluster/Private/Game/EngineClasses/Scene/DisplayClusterSceneComponentSyncParent.cpp b/Engine/Plugins/Runtime/nDisplay/Source/DisplayCluster/Private/Game/EngineClasses/Scene/DisplayClusterSceneComponentSyncParent.cpp
--- a/Engine/Plugins/Runtime/nDisplay/Source/DisplayCluster/Private/Game/EngineClasses/Scene/DisplayClusterSceneComponentSyncParent.cpp
+++ b/Engine/Plugins/Runtime/nDisplay/Source/DisplayCluster/Private/Game/EngineClasses/Scene/DisplayClusterSceneComponentSyncParent.cpp
@@ -46,7 +46,7 @@ void UDisplayClusterSceneComponentSyncParent::ClearDirty()
 //////////////////////////////////////////////////////////////////////////////////////////////
 FString UDisplayClusterSceneComponentSyncParent::GenerateSyncId()
 {
-	return FString::Printf(TEXT("SP_%s.%s"), *GetOwner()->GetName(), *GetAttachParent()->GetName());
+	return FString::Printf(TEXT("SP_%s"), *GetAttachParent()->GetFullName());
 }
 
 FTransform UDisplayClusterSceneComponentSyncParent::GetSyncTransform() const
diff --git a/Engine/Plugins/Runtime/nDisplay/Source/DisplayCluster/Private/Game/EngineClasses/Scene/DisplayClusterSceneComponentSyncThis.cpp b/Engine/Plugins/Runtime/nDisplay/Source/DisplayCluster/Private/Game/EngineClasses/Scene/DisplayClusterSceneComponentSyncThis.cpp
--- a/Engine/Plugins/Runtime/nDisplay/Source/DisplayCluster/Private/Game/EngineClasses/Scene/DisplayClusterSceneComponentSyncThis.cpp
+++ b/Engine/Plugins/Runtime/nDisplay/Source/DisplayCluster/Private/Game/EngineClasses/Scene/DisplayClusterSceneComponentSyncThis.cpp
@@ -32,7 +32,7 @@ void UDisplayClusterSceneComponentSyncThis::ClearDirty()
 //////////////////////////////////////////////////////////////////////////////////////////////
 FString UDisplayClusterSceneComponentSyncThis::GenerateSyncId()
 {
-	return FString::Printf(TEXT("ST_%s"), *GetOwner()->GetName());
+	return FString::Printf(TEXT("ST_%s"), *GetFullName());
 }
 
 FTransform UDisplayClusterSceneComponentSyncThis::GetSyncTransform() const
