[nDisplay] Ability to pass custom data and customize synchronization in generic barriers

diff --git a/Engine/Plugins/Runtime/nDisplay/Source/DisplayCluster/Private/Cluster/Controller/DisplayClusterClusterNodeCtrlPrimary.cpp b/Engine/Plugins/Runtime/nDisplay/Source/DisplayCluster/Private/Cluster/Controller/DisplayClusterClusterNodeCtrlPrimary.cpp
--- a/Engine/Plugins/Runtime/nDisplay/Source/DisplayCluster/Private/Cluster/Controller/DisplayClusterClusterNodeCtrlPrimary.cpp
+++ b/Engine/Plugins/Runtime/nDisplay/Source/DisplayCluster/Private/Cluster/Controller/DisplayClusterClusterNodeCtrlPrimary.cpp
@@ -86,6 +86,11 @@ EDisplayClusterCommResult FDisplayClusterClusterNodeCtrlPrimary::GetNativeInputD
 //////////////////////////////////////////////////////////////////////////////////////////////
 // IDisplayClusterClusterNodeController
 //////////////////////////////////////////////////////////////////////////////////////////////
+FDisplayClusterService* FDisplayClusterClusterNodeCtrlPrimary::GetGenericBarriersServer() const
+{
+	return GenericBarriersServer.Get();
+}
+
 bool FDisplayClusterClusterNodeCtrlPrimary::DropClusterNode(const FString& NodeId)
 {
 	// Kill all sessions of the requested node
diff --git a/Engine/Plugins/Runtime/nDisplay/Source/DisplayCluster/Private/Cluster/Controller/DisplayClusterClusterNodeCtrlPrimary.h b/Engine/Plugins/Runtime/nDisplay/Source/DisplayCluster/Private/Cluster/Controller/DisplayClusterClusterNodeCtrlPrimary.h
--- a/Engine/Plugins/Runtime/nDisplay/Source/DisplayCluster/Private/Cluster/Controller/DisplayClusterClusterNodeCtrlPrimary.h
+++ b/Engine/Plugins/Runtime/nDisplay/Source/DisplayCluster/Private/Cluster/Controller/DisplayClusterClusterNodeCtrlPrimary.h
@@ -42,6 +42,7 @@ public:
 		return EDisplayClusterNodeRole::Primary;
 	}
 
+	virtual FDisplayClusterService* GetGenericBarriersServer() const override;
 	virtual bool DropClusterNode(const FString& NodeId) override;
 
 protected:
diff --git a/Engine/Plugins/Runtime/nDisplay/Source/DisplayCluster/Private/Cluster/Controller/IDisplayClusterClusterNodeController.h b/Engine/Plugins/Runtime/nDisplay/Source/DisplayCluster/Private/Cluster/Controller/IDisplayClusterClusterNodeController.h
--- a/Engine/Plugins/Runtime/nDisplay/Source/DisplayCluster/Private/Cluster/Controller/IDisplayClusterClusterNodeController.h
+++ b/Engine/Plugins/Runtime/nDisplay/Source/DisplayCluster/Private/Cluster/Controller/IDisplayClusterClusterNodeController.h
@@ -11,6 +11,8 @@
 
 #include "DisplayClusterEnums.h"
 
+class FDisplayClusterService;
+
 
 /**
  * Node controller interface
@@ -40,6 +42,18 @@ public:
 	/** Return controller name */
 	virtual FString GetControllerName() const = 0;
 
+	/**
+	 * @note
+	 * This is a temporary workaround for a quick and safe implementation of the generic barriers expansion,
+	 * that would allow to set custom synchronization delegate. This will be refactored during Failover v2
+	 * implementation as we'll need all the nodes to have full set of servers for smooth primary role transition
+	 * and safe node replacement. Currently, the GB server is hosted on the p-nodes only.
+	 */
+	virtual FDisplayClusterService* GetGenericBarriersServer() const
+	{
+		return nullptr;
+	}
+
 	/** Drop specific cluster node */
 	virtual bool DropClusterNode(const FString& NodeId)
 	{
diff --git a/Engine/Plugins/Runtime/nDisplay/Source/DisplayCluster/Private/Network/Barrier/DisplayClusterBarrier.cpp b/Engine/Plugins/Runtime/nDisplay/Source/DisplayCluster/Private/Network/Barrier/DisplayClusterBarrier.cpp
--- a/Engine/Plugins/Runtime/nDisplay/Source/DisplayCluster/Private/Network/Barrier/DisplayClusterBarrier.cpp
+++ b/Engine/Plugins/Runtime/nDisplay/Source/DisplayCluster/Private/Network/Barrier/DisplayClusterBarrier.cpp
@@ -138,6 +138,9 @@ EDisplayClusterBarrierWaitResult FDisplayClusterBarrier::Wait(const FString& Cal
 				{
 					UE_LOG(LogDisplayClusterBarrier, Verbose, TEXT("Barrier '%s': sync start"), *Name);
 
+					// Prepare for new sync iteration
+					HandleBarrierPreSyncStart();
+
 					BarrierWaitTimeStart = ThreadWaitTimeStart;
 					WatchdogTimer.SetTimer(Timeout);
 				}
@@ -159,6 +162,9 @@ EDisplayClusterBarrierWaitResult FDisplayClusterBarrier::Wait(const FString& Cal
 					// All callers are here, pet the watchdog
 					WatchdogTimer.ResetTimer();
 
+					// Process sync done before allowing the threads to leave
+					HandleBarrierPreSyncEnd();
+
 					// Close input gate, open output gate
 					EventInputGateOpen->Reset();
 					EventOutputGateOpen->Trigger();
@@ -212,6 +218,30 @@ EDisplayClusterBarrierWaitResult FDisplayClusterBarrier::Wait(const FString& Cal
 	return EDisplayClusterBarrierWaitResult::Ok;
 }
 
+EDisplayClusterBarrierWaitResult FDisplayClusterBarrier::WaitWithData(const FString& ThreadMarker, const TArray<uint8>& RequestData, TArray<uint8>& OutResponseData, double* OutThreadWaitTime /* = nullptr */, double* OutBarrierWaitTime /* = nullptr */)
+{
+	{
+		FScopeLock Lock(&CommDataCS);
+		// Store request data so it can be used once all the threads arrived
+		ClientsRequestData.Emplace(ThreadMarker, RequestData);
+	}
+
+	// Wait at the barrier
+	const EDisplayClusterBarrierWaitResult WaitResult = Wait(ThreadMarker, OutThreadWaitTime, OutBarrierWaitTime);
+
+	{
+		FScopeLock Lock(&CommDataCS);
+
+		if (TArray<uint8>* const FoundThreadResponseData = ClientsResponseData.Find(ThreadMarker))
+		{
+			OutResponseData = MoveTemp(*FoundThreadResponseData);
+			ClientsResponseData.Remove(ThreadMarker);
+		}
+	}
+
+	return WaitResult;
+}
+
 void FDisplayClusterBarrier::UnregisterSyncCaller(const FString& CallerId)
 {
 	UE_LOG(LogDisplayClusterBarrier, Log, TEXT("Barrier '%s': unregistering caller '%s'..."), *Name, *CallerId);
@@ -248,6 +278,19 @@ void FDisplayClusterBarrier::UnregisterSyncCaller(const FString& CallerId)
 	}
 }
 
+void FDisplayClusterBarrier::HandleBarrierPreSyncStart()
+{
+}
+
+void FDisplayClusterBarrier::HandleBarrierPreSyncEnd()
+{
+	// Execute sync delegate with data requested
+	BarrierPreSyncEndDelegate.ExecuteIfBound(Name, ClientsRequestData, ClientsResponseData);
+
+	// We can clean request data now before next iteration
+	ClientsRequestData.Empty(CallersAllowed.Num());
+}
+
 void FDisplayClusterBarrier::HandleBarrierTimeout()
 {
 	// Being here means some callers have not come to the barrier yet during specific time period. Those
diff --git a/Engine/Plugins/Runtime/nDisplay/Source/DisplayCluster/Private/Network/Barrier/DisplayClusterBarrier.h b/Engine/Plugins/Runtime/nDisplay/Source/DisplayCluster/Private/Network/Barrier/DisplayClusterBarrier.h
--- a/Engine/Plugins/Runtime/nDisplay/Source/DisplayCluster/Private/Network/Barrier/DisplayClusterBarrier.h
+++ b/Engine/Plugins/Runtime/nDisplay/Source/DisplayCluster/Private/Network/Barrier/DisplayClusterBarrier.h
@@ -36,9 +36,17 @@ public:
 	// Wait until all caller threads arrived
 	virtual EDisplayClusterBarrierWaitResult Wait(const FString& CallerId, double* OutThreadWaitTime = nullptr, double* OutBarrierWaitTime = nullptr) override;
 
+	// Wait until all threads arrive (with data)
+	virtual EDisplayClusterBarrierWaitResult WaitWithData(const FString& ThreadMarker, const TArray<uint8>& RequestData, TArray<uint8>& OutResponseData, double* OutThreadWaitTime = nullptr, double* OutBarrierWaitTime = nullptr) override;
+
 	// Remove specified caller from the sync pipeline
 	virtual void UnregisterSyncCaller(const FString& CallerId) override;
 
+	virtual FDisplayClusterBarrierPreSyncEndDelegate& GetPreSyncEndDelegate() override
+	{
+		return BarrierPreSyncEndDelegate;
+	}
+
 	// Barrier timout notification
 	virtual FDisplayClusterBarrierTimeoutEvent& OnBarrierTimeout() override
 	{
@@ -46,6 +54,12 @@ public:
 	}
 
 private:
+	// Do some job before starting new sync iteration (opening the entrance gate)
+	void HandleBarrierPreSyncStart();
+
+	// Do some job before opening the exit gate
+	void HandleBarrierPreSyncEnd();
+
 	// Handler for barrier timouts
 	void HandleBarrierTimeout();
 
@@ -83,9 +97,16 @@ private:
 	// Watchdog timer to detect barrier waiting timeouts
 	FDisplayClusterWatchdogTimer WatchdogTimer;
 
+	// PreSyncEnd delegate. It's called when all threads arrived before opening the gate.
+	FDisplayClusterBarrierPreSyncEndDelegate BarrierPreSyncEndDelegate;
 	// Barrier timeout event
 	FDisplayClusterBarrierTimeoutEvent BarrierTimeoutEvent;
 
+	// Request data from the callers
+	FClientsCommData ClientsRequestData;
+	// Response data for the callers
+	FClientsCommData ClientsResponseData;
+
 	// Diagnostics data
 	double BarrierWaitTimeStart   = 0;
 	double BarrierWaitTimeFinish  = 0;
@@ -95,4 +116,6 @@ private:
 	mutable FCriticalSection DataCS;
 	// Barrier entrance CS
 	mutable FCriticalSection EntranceCS;
+	// Request/response data CS
+	mutable FCriticalSection CommDataCS;
 };
diff --git a/Engine/Plugins/Runtime/nDisplay/Source/DisplayCluster/Private/Network/Barrier/IDisplayClusterBarrier.h b/Engine/Plugins/Runtime/nDisplay/Source/DisplayCluster/Private/Network/Barrier/IDisplayClusterBarrier.h
--- a/Engine/Plugins/Runtime/nDisplay/Source/DisplayCluster/Private/Network/Barrier/IDisplayClusterBarrier.h
+++ b/Engine/Plugins/Runtime/nDisplay/Source/DisplayCluster/Private/Network/Barrier/IDisplayClusterBarrier.h
@@ -2,7 +2,7 @@
 
 #pragma once
 
-#include "CoreMinimal.h"
+#include "Delegates/DelegateCombinations.h"
 
 class FDisplayClusterWatchdogTimer;
 
@@ -43,9 +43,17 @@ public:
 	// Wait until all threads arrive
 	virtual EDisplayClusterBarrierWaitResult Wait(const FString& ThreadMarker, double* OutThreadWaitTime = nullptr, double* OutBarrierWaitTime = nullptr) = 0;
 
+	// Wait until all threads arrive (with data)
+	virtual EDisplayClusterBarrierWaitResult WaitWithData(const FString& ThreadMarker, const TArray<uint8>& RequestData, TArray<uint8>& OutResponseData, double* OutThreadWaitTime = nullptr, double* OutBarrierWaitTime = nullptr) = 0;
+
 	// Remove specified caller from the sync pipeline
 	virtual void UnregisterSyncCaller(const FString& CallerId) = 0;
 
+	/** Barrier PreSyncEnd delegate. Called when all calling threads arrived right before opening the gate. */
+	typedef TMap<FString, TArray<uint8>> FClientsCommData;
+	DECLARE_DELEGATE_ThreeParams(FDisplayClusterBarrierPreSyncEndDelegate, const FString&, const FClientsCommData&, FClientsCommData&);
+	virtual FDisplayClusterBarrierPreSyncEndDelegate& GetPreSyncEndDelegate() = 0;
+
 	// Barrier timout notification (provides BarrierName and CallersTimedOut in parameters)
 	DECLARE_EVENT_TwoParams(IDisplayClusterBarrier, FDisplayClusterBarrierTimeoutEvent, const FString&, const TArray<FString>&);
 	virtual FDisplayClusterBarrierTimeoutEvent& OnBarrierTimeout() = 0;
diff --git a/Engine/Plugins/Runtime/nDisplay/Source/DisplayCluster/Private/Network/Packet/DisplayClusterPacketInternal.h b/Engine/Plugins/Runtime/nDisplay/Source/DisplayCluster/Private/Network/Packet/DisplayClusterPacketInternal.h
--- a/Engine/Plugins/Runtime/nDisplay/Source/DisplayCluster/Private/Network/Packet/DisplayClusterPacketInternal.h
+++ b/Engine/Plugins/Runtime/nDisplay/Source/DisplayCluster/Private/Network/Packet/DisplayClusterPacketInternal.h
@@ -137,7 +137,7 @@ public:
 	}
 
 	// Set new binary argument to the packet
-	void SetBinArg(const FString& SectionName, const FString& ArgName, const TArray<uint8> ArgVal)
+	void SetBinArg(const FString& SectionName, const FString& ArgName, const TArray<uint8>& ArgVal)
 	{
 		TMap<FString, TArray<uint8>>* Section = BinaryArguments.Find(SectionName);
 		if (!Section)
diff --git a/Engine/Plugins/Runtime/nDisplay/Source/DisplayCluster/Private/Network/Protocol/IDisplayClusterProtocolGenericBarrier.h b/Engine/Plugins/Runtime/nDisplay/Source/DisplayCluster/Private/Network/Protocol/IDisplayClusterProtocolGenericBarrier.h
--- a/Engine/Plugins/Runtime/nDisplay/Source/DisplayCluster/Private/Network/Protocol/IDisplayClusterProtocolGenericBarrier.h
+++ b/Engine/Plugins/Runtime/nDisplay/Source/DisplayCluster/Private/Network/Protocol/IDisplayClusterProtocolGenericBarrier.h
@@ -42,4 +42,7 @@ public:
 
 	// Synchronize calling thread on a specific barrier
 	virtual EDisplayClusterCommResult SyncOnBarrier(const FString& BarrierId, const FString& UniqueThreadMarker, EBarrierControlResult& Result) = 0;
+
+	// Synchronize calling thread on a specific barrier (with custom data)
+	virtual EDisplayClusterCommResult SyncOnBarrierWithData(const FString& BarrierId, const FString& UniqueThreadMarker, const TArray<uint8>& RequestData, TArray<uint8>& OutResponseData, EBarrierControlResult& Result) = 0;
 };
diff --git a/Engine/Plugins/Runtime/nDisplay/Source/DisplayCluster/Private/Network/Service/GenericBarrier/DisplayClusterGenericBarrierClient.cpp b/Engine/Plugins/Runtime/nDisplay/Source/DisplayCluster/Private/Network/Service/GenericBarrier/DisplayClusterGenericBarrierClient.cpp
--- a/Engine/Plugins/Runtime/nDisplay/Source/DisplayCluster/Private/Network/Service/GenericBarrier/DisplayClusterGenericBarrierClient.cpp
+++ b/Engine/Plugins/Runtime/nDisplay/Source/DisplayCluster/Private/Network/Service/GenericBarrier/DisplayClusterGenericBarrierClient.cpp
@@ -1,12 +1,15 @@
 // Copyright Epic Games, Inc. All Rights Reserved.
 
 #include "Network/Service/GenericBarrier/DisplayClusterGenericBarrierClient.h"
+#include "Network/Service/GenericBarrier/DisplayClusterGenericBarrierService.h"
 #include "Network/Service/GenericBarrier/DisplayClusterGenericBarrierStrings.h"
 #include "Network/Packet/DisplayClusterPacketInternal.h"
 #include "Network/Listener/DisplayClusterHelloMessageStrings.h"
 
 #include "Config/IPDisplayClusterConfigManager.h"
+
 #include "Cluster/IPDisplayClusterClusterManager.h"
+#include "Cluster/Controller/IDisplayClusterClusterNodeController.h"
 
 #include "DisplayClusterEnums.h"
 #include "Misc/DisplayClusterGlobals.h"
@@ -105,7 +108,15 @@ bool FDisplayClusterGenericBarrierClient::CreateBarrier(const FString& BarrierId
 {
 	EBarrierControlResult CtrlResult = EBarrierControlResult::UnknownError;
 	const EDisplayClusterCommResult Result = CreateBarrier(BarrierId, UniqueThreadMarkers, Timeout, CtrlResult);
-	return (CtrlResult == EBarrierControlResult::CreatedSuccessfully || CtrlResult == EBarrierControlResult::AlreadyExists) && (Result == EDisplayClusterCommResult::Ok);
+	const bool bResult = (CtrlResult == EBarrierControlResult::CreatedSuccessfully || CtrlResult == EBarrierControlResult::AlreadyExists) && (Result == EDisplayClusterCommResult::Ok);
+
+	if (bResult)
+	{
+		// Setup sync delegate
+		ConfigureBarrierSyncDelegate(BarrierId, true);
+	}
+
+	return bResult;
 }
 
 bool FDisplayClusterGenericBarrierClient::WaitUntilBarrierIsCreated(const FString& BarrierId)
@@ -122,8 +133,16 @@ bool FDisplayClusterGenericBarrierClient::IsBarrierAvailable(const FString& Barr
 	return (CtrlResult == EBarrierControlResult::AlreadyExists) && (Result == EDisplayClusterCommResult::Ok);
 }
 
+IDisplayClusterGenericBarriersClient::FOnGenericBarrierSynchronizationDelegate* FDisplayClusterGenericBarrierClient::GetBarrierSyncDelegate(const FString& BarrierId)
+{
+	return BarrierSyncDelegates.Find(BarrierId);
+}
+
 bool FDisplayClusterGenericBarrierClient::ReleaseBarrier(const FString& BarrierId)
 {
+	// Release sync delegate
+	ConfigureBarrierSyncDelegate(BarrierId, false);
+
 	EBarrierControlResult CtrlResult = EBarrierControlResult::UnknownError;
 	const EDisplayClusterCommResult Result = ReleaseBarrier(BarrierId, CtrlResult);
 	return (CtrlResult == EBarrierControlResult::ReleasedSuccessfully) && (Result == EDisplayClusterCommResult::Ok);
@@ -136,6 +155,13 @@ bool FDisplayClusterGenericBarrierClient::Synchronize(const FString& BarrierId,
 	return (CtrlResult == EBarrierControlResult::SynchronizedSuccessfully) && (Result == EDisplayClusterCommResult::Ok);
 }
 
+bool FDisplayClusterGenericBarrierClient::Synchronize(const FString& BarrierId, const FString& UniqueThreadMarker, const TArray<uint8>& RequestData, TArray<uint8>& OutResponseData)
+{
+	EBarrierControlResult CtrlResult = EBarrierControlResult::UnknownError;
+	const EDisplayClusterCommResult Result = SyncOnBarrierWithData(BarrierId, UniqueThreadMarker, RequestData, OutResponseData, CtrlResult);
+	return (CtrlResult == EBarrierControlResult::SynchronizedSuccessfully) && (Result == EDisplayClusterCommResult::Ok);
+}
+
 
 //////////////////////////////////////////////////////////////////////////////////////////////
 // IDisplayClusterProtocolGenericBarrier
@@ -300,3 +326,100 @@ EDisplayClusterCommResult FDisplayClusterGenericBarrierClient::SyncOnBarrier(con
 
 	return Response->GetCommResult();
 }
+
+EDisplayClusterCommResult FDisplayClusterGenericBarrierClient::SyncOnBarrierWithData(const FString& BarrierId, const FString& UniqueThreadMarker, const TArray<uint8>& RequestData, TArray<uint8>& OutResponseData, EBarrierControlResult& Result)
+{
+	const TSharedPtr<FDisplayClusterPacketInternal> Request = MakeShared<FDisplayClusterPacketInternal>(
+		DisplayClusterGenericBarrierStrings::SyncOnBarrierWithData::Name,
+		DisplayClusterGenericBarrierStrings::TypeRequest,
+		DisplayClusterGenericBarrierStrings::ProtocolName);
+
+	// Request arguments
+	Request->SetTextArg(DisplayClusterGenericBarrierStrings::ArgumentsDefaultCategory, DisplayClusterGenericBarrierStrings::ArgBarrierId, BarrierId);
+	Request->SetTextArg(DisplayClusterGenericBarrierStrings::ArgumentsDefaultCategory, DisplayClusterGenericBarrierStrings::SyncOnBarrierWithData::ArgThreadMarker, UniqueThreadMarker);
+	Request->SetBinArg(DisplayClusterGenericBarrierStrings::ArgumentsDefaultCategory, DisplayClusterGenericBarrierStrings::SyncOnBarrierWithData::ArgRequestData, RequestData);
+
+	TSharedPtr<FDisplayClusterPacketInternal> Response;
+
+	{
+		TRACE_CPUPROFILER_EVENT_SCOPE(CLN_GB::SyncOnBarrierWithData);
+		Response = SendRecvPacket(Request);
+	}
+
+	if (!Response)
+	{
+		UE_LOG(LogDisplayClusterNetwork, Warning, TEXT("Network error on '%s'"), *Request->GetName());
+		return EDisplayClusterCommResult::NetworkError;
+	}
+
+	// Extract response data
+	Response->GetBinArg(DisplayClusterGenericBarrierStrings::ArgumentsDefaultCategory, DisplayClusterGenericBarrierStrings::SyncOnBarrierWithData::ArgResponseData, OutResponseData);
+
+	FString CtrlResult;
+	Response->GetTextArg(DisplayClusterGenericBarrierStrings::ArgumentsDefaultCategory, DisplayClusterGenericBarrierStrings::ArgResult, CtrlResult);
+	Result = static_cast<EBarrierControlResult>(DisplayClusterTypesConverter::template FromString<uint8>(CtrlResult));
+
+	return Response->GetCommResult();
+}
+
+bool FDisplayClusterGenericBarrierClient::ConfigureBarrierSyncDelegate(const FString& BarrierId, bool bSetup)
+{
+	// @note
+	// This is a temporary workaround. It will be refactored during failover v2 implementation in more clear and generic way.
+	//
+	// Once a barrier is created, we can explicitly set custom sync handler to that specific barrier. As it was mentioned above,
+	// this is a temporary workaround. We need to pick GB service on the p-node, and set the synchronization delegate so it's
+	// called when all the barrier clients have arrived.
+	if (GDisplayCluster->GetOperationMode() == EDisplayClusterOperationMode::Cluster)
+	{
+		if (IPDisplayClusterClusterManager* ClusterMgr = GDisplayCluster->GetPrivateClusterMgr())
+		{
+			if (ClusterMgr->IsPrimary())
+			{
+				if (IDisplayClusterClusterNodeController* ClusterNodeCtrl = ClusterMgr->GetClusterNodeController())
+				{
+					if (FDisplayClusterService* Service = ClusterNodeCtrl->GetGenericBarriersServer())
+					{
+						FDisplayClusterGenericBarrierService* const GBService = static_cast<FDisplayClusterGenericBarrierService*>(Service);
+						if (TSharedPtr<IDisplayClusterBarrier, ESPMode::ThreadSafe> Barrier = GBService->GetBarrier(BarrierId))
+						{
+							// Setup
+							if (bSetup)
+							{
+								if (!BarrierSyncDelegates.Contains(BarrierId))
+								{
+									BarrierSyncDelegates.Emplace(BarrierId);
+									Barrier->GetPreSyncEndDelegate().BindRaw(this, &FDisplayClusterGenericBarrierClient::OnPreBarrierSyncEnd);
+								}
+							}
+							// Release
+							else
+							{
+								if (BarrierSyncDelegates.Remove(BarrierId) > 0)
+								{
+									Barrier->GetPreSyncEndDelegate().Unbind();
+								}
+							}
+
+							return true;
+						}
+					}
+				}
+			}
+		}
+	}
+
+	return false;
+}
+
+void FDisplayClusterGenericBarrierClient::OnPreBarrierSyncEnd(const FString& BarrierId, const TMap<FString, TArray<uint8>>& RequestData, TMap<FString, TArray<uint8>>& ResponseData)
+{
+	if (FOnGenericBarrierSynchronizationDelegate* const Delegate = BarrierSyncDelegates.Find(BarrierId))
+	{
+		if (Delegate->IsBound())
+		{
+			FGenericBarrierSynchronizationDelegateData CallbackData{ BarrierId, RequestData, ResponseData };
+			Delegate->Execute(CallbackData);
+		}
+	}
+}
diff --git a/Engine/Plugins/Runtime/nDisplay/Source/DisplayCluster/Private/Network/Service/GenericBarrier/DisplayClusterGenericBarrierClient.h b/Engine/Plugins/Runtime/nDisplay/Source/DisplayCluster/Private/Network/Service/GenericBarrier/DisplayClusterGenericBarrierClient.h
--- a/Engine/Plugins/Runtime/nDisplay/Source/DisplayCluster/Private/Network/Service/GenericBarrier/DisplayClusterGenericBarrierClient.h
+++ b/Engine/Plugins/Runtime/nDisplay/Source/DisplayCluster/Private/Network/Service/GenericBarrier/DisplayClusterGenericBarrierClient.h
@@ -40,8 +40,10 @@ public:
 	virtual bool CreateBarrier(const FString& BarrierId, const TArray<FString>& UniqueThreadMarkers, uint32 Timeout) override;
 	virtual bool WaitUntilBarrierIsCreated(const FString& BarrierId) override;
 	virtual bool IsBarrierAvailable(const FString& BarrierId) override;
+	virtual FOnGenericBarrierSynchronizationDelegate* GetBarrierSyncDelegate(const FString& BarrierId) override;
 	virtual bool ReleaseBarrier(const FString& BarrierId) override;
 	virtual bool Synchronize(const FString& BarrierId, const FString& UniqueThreadMarker) override;
+	virtual bool Synchronize(const FString& BarrierId, const FString& UniqueThreadMarker, const TArray<uint8>& RequestData, TArray<uint8>& OutResponseData) override;
 
 public:
 	//////////////////////////////////////////////////////////////////////////////////////////////
@@ -52,4 +54,16 @@ public:
 	virtual EDisplayClusterCommResult IsBarrierAvailable(const FString& BarrierId, EBarrierControlResult& Result) override;
 	virtual EDisplayClusterCommResult ReleaseBarrier(const FString& BarrierId, EBarrierControlResult& Result) override;
 	virtual EDisplayClusterCommResult SyncOnBarrier(const FString& BarrierId, const FString& UniqueThreadMarker, EBarrierControlResult& Result) override;
+	virtual EDisplayClusterCommResult SyncOnBarrierWithData(const FString& BarrierId, const FString& UniqueThreadMarker, const TArray<uint8>& RequestData, TArray<uint8>& OutResponseData, EBarrierControlResult& Result) override;
+
+private:
+	/** Setup/release sync delegate for a specific barrier */
+	bool ConfigureBarrierSyncDelegate(const FString& BarrierId, bool bSetup);
+
+	/** Callback on barrier sync phase end */
+	void OnPreBarrierSyncEnd(const FString& BarrierId, const TMap<FString, TArray<uint8>>& RequestData, TMap<FString, TArray<uint8>>& ResponseData);
+
+private:
+	// Holds synchronization delegates for the barriers operated by this client
+	TMap<FString, FOnGenericBarrierSynchronizationDelegate> BarrierSyncDelegates;
 };
diff --git a/Engine/Plugins/Runtime/nDisplay/Source/DisplayCluster/Private/Network/Service/GenericBarrier/DisplayClusterGenericBarrierService.cpp b/Engine/Plugins/Runtime/nDisplay/Source/DisplayCluster/Private/Network/Service/GenericBarrier/DisplayClusterGenericBarrierService.cpp
--- a/Engine/Plugins/Runtime/nDisplay/Source/DisplayCluster/Private/Network/Service/GenericBarrier/DisplayClusterGenericBarrierService.cpp
+++ b/Engine/Plugins/Runtime/nDisplay/Source/DisplayCluster/Private/Network/Service/GenericBarrier/DisplayClusterGenericBarrierService.cpp
@@ -56,6 +56,13 @@ FString FDisplayClusterGenericBarrierService::GetProtocolName() const
 	return ProtocolName;
 }
 
+TSharedPtr<IDisplayClusterBarrier, ESPMode::ThreadSafe> FDisplayClusterGenericBarrierService::GetBarrier(const FString& BarrierId)
+{
+	FScopeLock Lock(&BarriersCS);
+	TSharedPtr<IDisplayClusterBarrier, ESPMode::ThreadSafe>* BarrierFound = Barriers.Find(BarrierId);
+	return BarrierFound ? *BarrierFound : nullptr;
+}
+
 TSharedPtr<IDisplayClusterSession> FDisplayClusterGenericBarrierService::CreateSession(FDisplayClusterSessionInfo& SessionInfo)
 {
 	SessionInfo.SessionName = FString::Printf(TEXT("%s_session_%lu_%s"), *GetName(), SessionInfo.SessionId, *SessionInfo.Endpoint.ToString());
@@ -111,6 +118,32 @@ TSharedPtr<FDisplayClusterPacketInternal> FDisplayClusterGenericBarrierService::
 
 		return Response;
 	}
+	else if (Request->GetName().Equals(DisplayClusterGenericBarrierStrings::SyncOnBarrierWithData::Name, ESearchCase::IgnoreCase))
+	{
+		// Extract parameter: barrier ID
+		FString BarrierId;
+		Request->GetTextArg(DisplayClusterGenericBarrierStrings::ArgumentsDefaultCategory, DisplayClusterGenericBarrierStrings::ArgBarrierId, BarrierId);
+
+		// Extract parameter: thread marker
+		FString UniqueThreadMarker;
+		Request->GetTextArg(DisplayClusterGenericBarrierStrings::ArgumentsDefaultCategory, DisplayClusterGenericBarrierStrings::SyncOnBarrierWithData::ArgThreadMarker, UniqueThreadMarker);
+
+		// Extract parameter: request data
+		TArray<uint8> RequestData;
+		Request->GetBinArg(DisplayClusterGenericBarrierStrings::ArgumentsDefaultCategory, DisplayClusterGenericBarrierStrings::SyncOnBarrierWithData::ArgRequestData, RequestData);
+
+		// Process command
+		EBarrierControlResult CtrlResult = EBarrierControlResult::UnknownError;
+		TArray<uint8> ResponseData;
+		const EDisplayClusterCommResult Result = SyncOnBarrierWithData(BarrierId, UniqueThreadMarker, RequestData, ResponseData, CtrlResult);
+
+		// Set response data
+		Response->SetBinArg(DisplayClusterGenericBarrierStrings::ArgumentsDefaultCategory, DisplayClusterGenericBarrierStrings::SyncOnBarrierWithData::ArgResponseData, ResponseData);
+		Response->SetTextArg(DisplayClusterGenericBarrierStrings::ArgumentsDefaultCategory, DisplayClusterGenericBarrierStrings::ArgResult, DisplayClusterTypesConverter::template ToString((uint8)CtrlResult));
+		Response->SetCommResult(Result);
+
+		return Response;
+	}
 	else if (Request->GetName().Equals(DisplayClusterGenericBarrierStrings::CreateBarrier::Name, ESearchCase::IgnoreCase))
 	{
 		// Extract parameter: barrier ID
@@ -373,3 +406,44 @@ EDisplayClusterCommResult FDisplayClusterGenericBarrierService::SyncOnBarrier(co
 	Result = EBarrierControlResult::SynchronizedSuccessfully;
 	return EDisplayClusterCommResult::Ok;
 }
+
+EDisplayClusterCommResult FDisplayClusterGenericBarrierService::SyncOnBarrierWithData(const FString& BarrierId, const FString& UniqueThreadMarker, const TArray<uint8>& RequestData, TArray<uint8>& OutResponseData, EBarrierControlResult& Result)
+{
+	TRACE_CPUPROFILER_EVENT_SCOPE(SRV_GB::SyncOnBarrierWithData);
+
+	// Validate input data
+	if (BarrierId.IsEmpty() || UniqueThreadMarker.IsEmpty())
+	{
+		Result = EBarrierControlResult::UnknownError;
+		return EDisplayClusterCommResult::WrongRequestData;
+	}
+
+	TSharedPtr<IDisplayClusterBarrier, ESPMode::ThreadSafe>* Barrier = nullptr;
+
+	{
+		FScopeLock Lock(&BarriersCS);
+
+		// Get barrier ptr to be able to use outside of the critical section
+		// so the other clients can also access it.
+		Barrier = Barriers.Find(BarrierId);
+	}
+
+	// More validation
+	if (!Barrier)
+	{
+		// Barrier not found
+		Result = EBarrierControlResult::NotFound;
+		return EDisplayClusterCommResult::WrongRequestData;
+	}
+	else if (!(*Barrier)->IsActivated())
+	{
+		Result = EBarrierControlResult::UnknownError;
+		return EDisplayClusterCommResult::NotAllowed;
+	}
+
+	// Sync on the barrier
+	(*Barrier)->WaitWithData(UniqueThreadMarker, RequestData, OutResponseData);
+
+	Result = EBarrierControlResult::SynchronizedSuccessfully;
+	return EDisplayClusterCommResult::Ok;
+}
diff --git a/Engine/Plugins/Runtime/nDisplay/Source/DisplayCluster/Private/Network/Service/GenericBarrier/DisplayClusterGenericBarrierService.h b/Engine/Plugins/Runtime/nDisplay/Source/DisplayCluster/Private/Network/Service/GenericBarrier/DisplayClusterGenericBarrierService.h
--- a/Engine/Plugins/Runtime/nDisplay/Source/DisplayCluster/Private/Network/Service/GenericBarrier/DisplayClusterGenericBarrierService.h
+++ b/Engine/Plugins/Runtime/nDisplay/Source/DisplayCluster/Private/Network/Service/GenericBarrier/DisplayClusterGenericBarrierService.h
@@ -2,13 +2,13 @@
 
 #pragma once
 
-#include "CoreMinimal.h"
-
 #include "Network/Service/DisplayClusterService.h"
 #include "Network/Session/IDisplayClusterSessionPacketHandler.h"
 #include "Network/Protocol/IDisplayClusterProtocolGenericBarrier.h"
 #include "Network/Packet/DisplayClusterPacketInternal.h"
 
+#include "Cluster/IDisplayClusterGenericBarriersClient.h"
+
 class FEvent;
 
 
@@ -31,6 +31,9 @@ public:
 	virtual void Shutdown() override final;
 	virtual FString GetProtocolName() const override;
 
+public:
+	TSharedPtr<IDisplayClusterBarrier, ESPMode::ThreadSafe> GetBarrier(const FString& BarrierId);
+
 protected:
 	// Creates session instance for this service
 	virtual TSharedPtr<IDisplayClusterSession> CreateSession(FDisplayClusterSessionInfo& SessionInfo) override;
@@ -50,10 +53,7 @@ private:
 	virtual EDisplayClusterCommResult IsBarrierAvailable(const FString& BarrierId, EBarrierControlResult& Result) override;
 	virtual EDisplayClusterCommResult ReleaseBarrier(const FString& BarrierId, EBarrierControlResult& Result) override;
 	virtual EDisplayClusterCommResult SyncOnBarrier(const FString& BarrierId, const FString& UniqueThreadMarker, EBarrierControlResult& Result) override;
-
-private:
-	// Unregister cluster node from a barrier
-	void UnregisterClusterNode(const FString& NodeId);
+	virtual EDisplayClusterCommResult SyncOnBarrierWithData(const FString& BarrierId, const FString& UniqueThreadMarker, const TArray<uint8>& RequestData, TArray<uint8>& OutResponseData, EBarrierControlResult& Result) override;
 
 private:
 	// Barriers
diff --git a/Engine/Plugins/Runtime/nDisplay/Source/DisplayCluster/Private/Network/Service/GenericBarrier/DisplayClusterGenericBarrierStrings.h b/Engine/Plugins/Runtime/nDisplay/Source/DisplayCluster/Private/Network/Service/GenericBarrier/DisplayClusterGenericBarrierStrings.h
--- a/Engine/Plugins/Runtime/nDisplay/Source/DisplayCluster/Private/Network/Service/GenericBarrier/DisplayClusterGenericBarrierStrings.h
+++ b/Engine/Plugins/Runtime/nDisplay/Source/DisplayCluster/Private/Network/Service/GenericBarrier/DisplayClusterGenericBarrierStrings.h
@@ -47,8 +47,17 @@ namespace DisplayClusterGenericBarrierStrings
 
 	namespace SyncOnBarrier
 	{
-		constexpr static const TCHAR* Name = TEXT("SyncOnBarrier");
+		constexpr static const TCHAR* Name = TEXT("Sync");
 
 		constexpr static const TCHAR* ArgThreadMarker = TEXT("ThreadMarker");
 	}
+
+	namespace SyncOnBarrierWithData
+	{
+		constexpr static const TCHAR* Name = TEXT("SyncWithData");
+
+		constexpr static const TCHAR* ArgThreadMarker  = TEXT("ThreadMarker");
+		constexpr static const TCHAR* ArgRequestData   = TEXT("ReqData");
+		constexpr static const TCHAR* ArgResponseData  = TEXT("RespData");
+	}
 };
diff --git a/Engine/Plugins/Runtime/nDisplay/Source/DisplayCluster/Public/Cluster/IDisplayClusterGenericBarriersClient.h b/Engine/Plugins/Runtime/nDisplay/Source/DisplayCluster/Public/Cluster/IDisplayClusterGenericBarriersClient.h
--- a/Engine/Plugins/Runtime/nDisplay/Source/DisplayCluster/Public/Cluster/IDisplayClusterGenericBarriersClient.h
+++ b/Engine/Plugins/Runtime/nDisplay/Source/DisplayCluster/Public/Cluster/IDisplayClusterGenericBarriersClient.h
@@ -5,6 +5,22 @@
 #include "CoreMinimal.h"
 
 
+/**
+ * Barreir synchonrization callback data
+ */
+struct FGenericBarrierSynchronizationDelegateData
+{
+	/** Barrier ID */
+	const FString BarrierId;
+
+	/** Binary data provided on sync request (node Id - to - data mapping) */
+	const TMap<FString, TArray<uint8>>& RequestData;
+
+	/** Binary data to respond (node Id - to - data mapping) */
+	TMap<FString, TArray<uint8>>& ResponseData;
+};
+
+
 /**
  * Generic barriers client interface
  */
@@ -38,9 +54,16 @@ public:
 	/** Checks if a specific barrier exists */
 	virtual bool IsBarrierAvailable(const FString& BarrierId) = 0;
 
+	/** Returns a synchronization delegate to a specific barrier or nullptr if not available */
+	DECLARE_DELEGATE_OneParam(FOnGenericBarrierSynchronizationDelegate, FGenericBarrierSynchronizationDelegateData&);
+	virtual FOnGenericBarrierSynchronizationDelegate* GetBarrierSyncDelegate(const FString& BarrierId) = 0;
+
 	/** Releases specific barrier */
 	virtual bool ReleaseBarrier(const FString& BarrierId) = 0;
 
 	/** Synchronize calling thread on a specific barrier */
 	virtual bool Synchronize(const FString& BarrierId, const FString& UniqueThreadMarker) = 0;
+
+	/** Synchronize calling thread on a specific barrier with custom data */
+	virtual bool Synchronize(const FString& BarrierId, const FString& UniqueThreadMarker, const TArray<uint8>& RequestData, TArray<uint8>& OutResponseData) = 0;
 };
