nDisplay: Fix missing post-processing checkbox causing settings to never be applied.

This is a partial undo of CL 26114339, which conflicts with post-process settings customization.

#rnx

diff --git a/Engine/Plugins/Runtime/nDisplay/Source/DisplayClusterConfiguration/Public/DisplayClusterConfigurationTypes_Postprocess.h b/Engine/Plugins/Runtime/nDisplay/Source/DisplayClusterConfiguration/Public/DisplayClusterConfigurationTypes_Postprocess.h
--- a/Engine/Plugins/Runtime/nDisplay/Source/DisplayClusterConfiguration/Public/DisplayClusterConfigurationTypes_Postprocess.h
+++ b/Engine/Plugins/Runtime/nDisplay/Source/DisplayClusterConfiguration/Public/DisplayClusterConfigurationTypes_Postprocess.h
@@ -77,23 +77,23 @@ struct DISPLAYCLUSTERCONFIGURATION_API FDisplayClusterConfigurationViewport_Colo
 	uint8 bOverride_Offset : 1;
 
 	// Saturation
-	UPROPERTY(EditAnywhere, BlueprintReadWrite, Category = "Color Grading", meta = (UIMin = "0.0", UIMax = "2.0", Delta = "0.01", ColorGradingMode = "saturation"))
+	UPROPERTY(EditAnywhere, BlueprintReadWrite, Category = "Color Grading", meta = (UIMin = "0.0", UIMax = "2.0", Delta = "0.01", EditCondition = "bOverride_Saturation", ColorGradingMode = "saturation"))
 	FVector4 Saturation = FVector4(1.0f, 1.0f, 1.0f, 1.0f);
 
 	// Contrast
-	UPROPERTY(EditAnywhere, BlueprintReadWrite, Category = "Color Grading", meta = (UIMin = "0.0", UIMax = "2.0", Delta = "0.01", ColorGradingMode = "contrast"))
+	UPROPERTY(EditAnywhere, BlueprintReadWrite, Category = "Color Grading", meta = (UIMin = "0.0", UIMax = "2.0", Delta = "0.01", EditCondition = "bOverride_Contrast", ColorGradingMode = "contrast"))
 	FVector4 Contrast = FVector4(1.0f, 1.0f, 1.0f, 1.0f);
 
 	// Gamma
-	UPROPERTY(EditAnywhere, BlueprintReadWrite, Category = "Color Grading", meta = (UIMin = "0.0", UIMax = "2.0", Delta = "0.01", ColorGradingMode = "gamma"))
+	UPROPERTY(EditAnywhere, BlueprintReadWrite, Category = "Color Grading", meta = (UIMin = "0.0", UIMax = "2.0", Delta = "0.01", EditCondition = "bOverride_Gamma", ColorGradingMode = "gamma"))
 	FVector4 Gamma = FVector4(1.0f, 1.0f, 1.0f, 1.0f);
 
 	// Gain
-	UPROPERTY(EditAnywhere, BlueprintReadWrite, Category = "Color Grading", meta = (UIMin = "0.0", UIMax = "2.0", Delta = "0.01", ColorGradingMode = "gain"))
+	UPROPERTY(EditAnywhere, BlueprintReadWrite, Category = "Color Grading", meta = (UIMin = "0.0", UIMax = "2.0", Delta = "0.01", EditCondition = "bOverride_Gain", ColorGradingMode = "gain"))
 	FVector4 Gain = FVector4(1.0f, 1.0f, 1.0f, 1.0f);
 
 	// Offset
-	UPROPERTY(EditAnywhere, BlueprintReadWrite, Category = "Color Grading", meta = (UIMin = "-1.0", UIMax = "1.0", Delta = "0.001", ColorGradingMode = "offset"))
+	UPROPERTY(EditAnywhere, BlueprintReadWrite, Category = "Color Grading", meta = (UIMin = "-1.0", UIMax = "1.0", Delta = "0.001", EditCondition = "bOverride_Offset", ColorGradingMode = "offset"))
 	FVector4 Offset = FVector4(0.0f, 0.0f, 0.0f, 0.0f);
 };
 
@@ -120,15 +120,15 @@ struct DISPLAYCLUSTERCONFIGURATION_API FDisplayClusterConfigurationViewport_Colo
 	* White Balance uses the Temperature value to control the virtual camera's White Balance. This is the default selection.
 	* Color Temperature uses the Temperature value to adjust the color temperature of the scene, which is the inverse of the White Balance operation.
 	*/
-	UPROPERTY(interp, BlueprintReadWrite, Category = "Viewport Settings", meta = (DisplayName = "Temperature Type"))
+	UPROPERTY(interp, BlueprintReadWrite, Category = "Viewport Settings", meta = (DisplayName = "Temperature Type", EditCondition = "bOverride_TemperatureType"))
 	TEnumAsByte<enum ETemperatureMethod> TemperatureType = ETemperatureMethod::TEMP_WhiteBalance;
 
 	// White temperature
-	UPROPERTY(EditAnywhere, BlueprintReadWrite, Category = "White Balance", meta = (UIMin = "1500.0", UIMax = "15000.0", DisplayName = "Temp"))
+	UPROPERTY(EditAnywhere, BlueprintReadWrite, Category = "White Balance", meta = (UIMin = "1500.0", UIMax = "15000.0", EditCondition = "bOverride_WhiteTemp", DisplayName = "Temp"))
 	float WhiteTemp = 6500.0f;
 
 	// White tint
-	UPROPERTY(EditAnywhere, BlueprintReadWrite, Category = "White Balance", meta = (UIMin = "-1.0", UIMax = "1.0", DisplayName = "Tint"))
+	UPROPERTY(EditAnywhere, BlueprintReadWrite, Category = "White Balance", meta = (UIMin = "-1.0", UIMax = "1.0", EditCondition = "bOverride_WhiteTint", DisplayName = "Tint"))
 	float WhiteTint = 0.0f;
 };
 
@@ -155,15 +155,15 @@ struct DISPLAYCLUSTERCONFIGURATION_API FDisplayClusterConfigurationViewport_Colo
 	uint8 bOverride_SceneColorTint:1;
 
 	// Correct for artifacts with "electric" blues due to the ACEScg color space. Bright blue desaturates instead of going to violet.
-	UPROPERTY(EditAnywhere, BlueprintReadWrite, Category = "Misc", meta = (ClampMin = "0.0", ClampMax = "1.0"))
+	UPROPERTY(EditAnywhere, BlueprintReadWrite, Category = "Misc", meta = (ClampMin = "0.0", ClampMax = "1.0", EditCondition = "bOverride_BlueCorrection"))
 	float BlueCorrection;
 
 	// Expand bright saturated colors outside the sRGB gamut to fake wide gamut rendering.
-	UPROPERTY(EditAnywhere, BlueprintReadWrite, Category = "Misc", meta = (ClampMin = "0.0", UIMax = "1.0"))
+	UPROPERTY(EditAnywhere, BlueprintReadWrite, Category = "Misc", meta = (ClampMin = "0.0", UIMax = "1.0", EditCondition = "bOverride_ExpandGamut"))
 	float ExpandGamut;
 
 	// Scene tint color
-	UPROPERTY(EditAnywhere, BlueprintReadWrite, Category = "Misc", meta = (HideAlphaChannel))
+	UPROPERTY(EditAnywhere, BlueprintReadWrite, Category = "Misc", meta = (EditCondition = "bOverride_SceneColorTint", HideAlphaChannel))
 	FLinearColor SceneColorTint;
 };
 
@@ -201,7 +201,7 @@ struct DISPLAYCLUSTERCONFIGURATION_API FDisplayClusterConfigurationViewport_Colo
 	float BlendWeight = 1.0f;
 
 	// Exposure compensation
-	UPROPERTY(EditAnywhere, BlueprintReadWrite, Category = "Viewport Settings", meta = (UIMin = "-15.0", UIMax = "15.0", DisplayName = "Exposure Compensation"))
+	UPROPERTY(EditAnywhere, BlueprintReadWrite, Category = "Viewport Settings", meta = (UIMin = "-15.0", UIMax = "15.0", EditCondition = "bOverride_AutoExposureBias", DisplayName = "Exposure Compensation"))
 	float AutoExposureBias = 0.0f;
 
 	// White balance
@@ -216,7 +216,7 @@ struct DISPLAYCLUSTERCONFIGURATION_API FDisplayClusterConfigurationViewport_Colo
 	UPROPERTY(EditAnywhere, BlueprintReadWrite, Category = "Viewport Settings")
 	FDisplayClusterConfigurationViewport_ColorGradingSettings Shadows;
 
-	UPROPERTY(EditAnywhere, BlueprintReadWrite, Category = "Viewport Settings", meta = (UIMin = "-1.0", UIMax = "1.0", DisplayName = "ShadowsMax"))
+	UPROPERTY(EditAnywhere, BlueprintReadWrite, Category = "Viewport Settings", meta = (UIMin = "-1.0", UIMax = "1.0", EditCondition = "bOverride_ColorCorrectionShadowsMax", DisplayName = "ShadowsMax"))
 	float ColorCorrectionShadowsMax;
 
 	// Midtones color grading
@@ -227,10 +227,10 @@ struct DISPLAYCLUSTERCONFIGURATION_API FDisplayClusterConfigurationViewport_Colo
 	UPROPERTY(EditAnywhere, BlueprintReadWrite, Category = "Viewport Settings")
 	FDisplayClusterConfigurationViewport_ColorGradingSettings Highlights;
 
-	UPROPERTY(EditAnywhere, BlueprintReadWrite, Category = "Viewport Settings", meta = (UIMin = "-1.0", UIMax = "1.0", DisplayName = "HighlightsMin"))
+	UPROPERTY(EditAnywhere, BlueprintReadWrite, Category = "Viewport Settings", meta = (UIMin = "-1.0", UIMax = "1.0", EditCondition = "bOverride_ColorCorrectionHighlightsMin", DisplayName = "HighlightsMin"))
 	float ColorCorrectionHighlightsMin;
 
-	UPROPERTY(EditAnywhere, BlueprintReadWrite, Category = "Viewport Settings", meta = (UIMin = "1.0", UIMax = "10.0", DisplayName = "HighlightsMax"))
+	UPROPERTY(EditAnywhere, BlueprintReadWrite, Category = "Viewport Settings", meta = (UIMin = "1.0", UIMax = "10.0", EditCondition = "bOverride_ColorCorrectionHighlightsMax", DisplayName = "HighlightsMax"))
 	float ColorCorrectionHighlightsMax;
 
 	// Highlights color grading misc settings
