[UnrealStage] Fix case where pin positions could be temporarily out of sync with light card's visible position

diff --git a/Engine/Plugins/Runtime/nDisplay/Source/DisplayClusterScenePreview/Private/DisplayClusterScenePreviewModule.cpp b/Engine/Plugins/Runtime/nDisplay/Source/DisplayClusterScenePreview/Private/DisplayClusterScenePreviewModule.cpp
--- a/Engine/Plugins/Runtime/nDisplay/Source/DisplayClusterScenePreview/Private/DisplayClusterScenePreviewModule.cpp
+++ b/Engine/Plugins/Runtime/nDisplay/Source/DisplayClusterScenePreview/Private/DisplayClusterScenePreviewModule.cpp
@@ -338,20 +338,15 @@ bool FDisplayClusterScenePreviewModule::InternalRenderImmediate(FRendererConfig&
 				{
 					PreviewComponentOverrideTextures.Add(PreviewComp, PreviewComp->GetOverrideTexture());
 					PreviewComp->SetOverrideTexture(PreviewComp->GetRenderTargetTexturePostProcess());
-
-					// Force the preview mesh to recreate its render state immediately, in case setting the override texture
-					// caused the preview component to create and set the preview material on the mesh (as is the case when
-					// previews are disabled on the root actor)
-					if (UMeshComponent* PreviewMesh = PreviewComp->GetPreviewMesh())
-					{
-						PreviewMesh->RecreateRenderState_Concurrent();
-					}
 				}
 			}
 		}
 	}
 #endif
 
+	// Push any deferred render state updates to ensure that light card positions, preview meshes modified above, etc. are up to date
+	World->SendAllEndOfFrameUpdates();
+
 	RendererConfig.Renderer->Render(&Canvas, World->Scene, RenderSettings);
 
 #if WITH_EDITOR
