nDisplay: PreviewShare now subscribes to preview with post process.

Lightviz is expected to want all postprocess to be included.

Also fixed not closing media when the NDC is deleted and not just the preview share component, by using the on component destroyed callback.

Also reduced the severity a busy fence in shared memory media since in freerun it is not unexpected.

diff --git a/Engine/Plugins/Runtime/nDisplay/Source/DisplayCluster/Private/Game/EngineClasses/Scene/DisplayClusterPreviewShareComponent.cpp b/Engine/Plugins/Runtime/nDisplay/Source/DisplayCluster/Private/Game/EngineClasses/Scene/DisplayClusterPreviewShareComponent.cpp
--- a/Engine/Plugins/Runtime/nDisplay/Source/DisplayCluster/Private/Game/EngineClasses/Scene/DisplayClusterPreviewShareComponent.cpp
+++ b/Engine/Plugins/Runtime/nDisplay/Source/DisplayCluster/Private/Game/EngineClasses/Scene/DisplayClusterPreviewShareComponent.cpp
@@ -22,6 +22,68 @@
 
 namespace UE::DisplayClusterPreviewShare
 {
+	/** Convenience struct when iterating over config viewports */
+	struct FViewportData
+	{
+		// Id of the node for this viewport
+		FString NodeId;
+
+		// Id of the viewport itself
+		FString ViewportId;
+
+		// Config viewport itself
+		UDisplayClusterConfigurationViewport* Viewport;
+
+		FViewportData(const FString& NodeId, const FString& ViewportId, UDisplayClusterConfigurationViewport* Viewport)
+			: NodeId(NodeId)
+			, ViewportId(ViewportId)
+			, Viewport(Viewport)
+		{}
+	};
+
+	/** Convenience function to get the data needed when iterating over config viewports of a given nDisplay actor */
+	TArray<FViewportData> GetConfigViewportsFromRootActor(const ADisplayClusterRootActor* RootActor)
+	{
+		TArray<FViewportData> ViewportDatas;
+
+		if (!RootActor)
+		{
+			return ViewportDatas;
+		}
+
+		UDisplayClusterConfigurationData* ConfigData = RootActor->GetConfigData();
+
+		if (!ConfigData || !ConfigData->Cluster)
+		{
+			return ViewportDatas;
+		}
+
+		for (const auto& NodeElem : ConfigData->Cluster->Nodes)
+		{
+			if (!NodeElem.Value)
+			{
+				continue;
+			}
+
+			for (const auto& ViewportElem : NodeElem.Value->Viewports)
+			{
+				TObjectPtr<UDisplayClusterConfigurationViewport> Viewport = ViewportElem.Value;
+
+				if (!Viewport)
+				{
+					continue;
+				}
+
+				const FString& NodeId = NodeElem.Key;
+				const FString& ViewportId = ViewportElem.Key;
+
+				ViewportDatas.Emplace(NodeId, ViewportId, Viewport);
+			}
+		}
+
+		return ViewportDatas;
+	}
+
 	/** Retrieves the viewport configuration from the DCRA's current config using the given NodeId and ViewportId */
 	static UDisplayClusterConfigurationViewport* GetViewportFromDCRA(ADisplayClusterRootActor* RootActor, const FString& NodeId, const FString& ViewportId)
 	{
@@ -79,39 +141,46 @@ namespace UE::DisplayClusterPreviewShare
 			return;
 		}
 
-		for (UActorComponent* ActorComponent : RootActor->K2_GetComponentsByClass(UDisplayClusterPreviewComponent::StaticClass()))
+		for (const FViewportData& ViewportData : GetConfigViewportsFromRootActor(RootActor))
 		{
-			UDisplayClusterPreviewComponent* PreviewComponent = Cast<UDisplayClusterPreviewComponent>(ActorComponent);
-
-			if (!PreviewComponent)
+			if (ViewportData.Viewport->RenderSettings.Replace.SourceTexture)
 			{
-				continue;
-			}
-
-			UDisplayClusterConfigurationViewport* Viewport = GetViewportFromDCRA(RootActor, PreviewComponent->GetClusterNodeId(), PreviewComponent->GetViewportId());
-
-			if (!Viewport)
-			{
-				continue;
-			}
-
-			if (Viewport->RenderSettings.Replace.SourceTexture)
-			{
-				const UPackage* TexturePackage = Viewport->RenderSettings.Replace.SourceTexture->GetPackage();
+				const UPackage* TexturePackage = ViewportData.Viewport->RenderSettings.Replace.SourceTexture->GetPackage();
 
 				if (TexturePackage && TexturePackage->ContainsMap() && (TexturePackage != RootActorPackage))
 				{
 					UE_LOG(LogDisplayClusterGame, Warning,
-						TEXT("Removing unexpected external reference to texture package '%s', a map different to the wwning RootActor package '%s'"),
+						TEXT("Removing unexpected external reference to texture package '%s', a map different to the owning RootActor package '%s'"),
 						*TexturePackage->GetName(),
 						*RootActorPackage->GetName()
 					);
 
-					Viewport->RenderSettings.Replace.SourceTexture = nullptr;
+					ViewportData.Viewport->RenderSettings.Replace.SourceTexture = nullptr;
 				}
 			}
 		}
 	}
+
+	/** Retrieves appropriate preview render target from preview component */
+	UTextureRenderTarget2D* GetPreviewRenderTarget(const UDisplayClusterPreviewComponent* PreviewComponent)
+	{
+		if (!PreviewComponent)
+		{
+			return nullptr;
+		}
+
+		// We're primarily interested in the post process render target.
+		UTextureRenderTarget2D* RenderTarget = PreviewComponent->GetRenderTargetTexturePostProcess();
+
+		// Fallback if unavailable.
+		if (!RenderTarget)
+		{
+			RenderTarget = PreviewComponent->GetRenderTargetTexture();
+		}
+
+		return RenderTarget;
+	}
+
 }
 
 #endif // WITH_EDITOR
@@ -190,8 +259,7 @@ void UDisplayClusterPreviewShareComponent::PostEditChangeProperty(FPropertyChang
 	}
 	else if (PropertyChangedEvent.GetPropertyName() == GET_MEMBER_NAME_CHECKED(UDisplayClusterPreviewShareComponent, SourceNDisplayActor))
 	{
-		// Since we're changing the source actor, we should restore the original settings before taking the new ones.
-		RestoreRootActorOriginalSettings();
+		ModeChanged();
 	}
 }
 #endif // WITH_EDITOR
@@ -204,29 +272,18 @@ void UDisplayClusterPreviewShareComponent::ModeChanged()
 	// Restore original root actor settings (e.g. texture replace settings that we overwrote)
 	RestoreRootActorOriginalSettings();
 
-	ADisplayClusterRootActor* RootActor = Cast<ADisplayClusterRootActor>(GetOwner());
+	// Unsubscribe from previews
+	UnsubscribeFromAllPreviews();
+
+	// Default to no tick
+	SetTickEnable(false);
 
 	switch (Mode)
 	{
-	case EDisplayClusterPreviewShareMode::None:
-	{
-		// Tell root actor that it doesn't need to keep up the preview enabled for us.
-		if (RootActor)
-		{
-			RootActor->RemovePreviewEnableOverride(reinterpret_cast<uint8*>(this));
-		}
-
-		SetTickEnable(false);
-
-		break;
-	}
 	case EDisplayClusterPreviewShareMode::PullActor:
 	{
-		if (RootActor)
-		{
-			// Tell root actor that it doesn't need to keep up the preview enabled for us.
-			RootActor->RemovePreviewEnableOverride(reinterpret_cast<uint8*>(this));
-		}
+		// The source root actor should enable the preview since we're pulling from it.
+		SubscribeToPreview(SourceNDisplayActor.Get());
 
 		SetTickEnable(true);
 
@@ -234,12 +291,8 @@ void UDisplayClusterPreviewShareComponent::ModeChanged()
 	}
 	case EDisplayClusterPreviewShareMode::Send:
 	{
-		// If we're sending, we need to tell the root actor to generate the preview textures for us.
-
-		if (RootActor)
-		{
-			RootActor->AddPreviewEnableOverride(reinterpret_cast<uint8*>(this));
-		}
+		// If we're sending, we need to tell the root actor to generate the preview textures for us to send.
+		SubscribeToPreview(Cast<ADisplayClusterRootActor>(GetOwner()));
 
 		SetTickEnable(true);
 
@@ -247,12 +300,6 @@ void UDisplayClusterPreviewShareComponent::ModeChanged()
 	}
 	case EDisplayClusterPreviewShareMode::Receive:
 	{
-		if (RootActor)
-		{
-			// Tell root actor that it doesn't need to keep up the preview enabled for us.
-			RootActor->RemovePreviewEnableOverride(reinterpret_cast<uint8*>(this));
-		}
-
 		SetTickEnable(true);
 
 		break;
@@ -275,7 +322,7 @@ void UDisplayClusterPreviewShareComponent::RestoreRootActorOriginalSettings()
 	using namespace UE::DisplayClusterPreviewShare;
 
 	ADisplayClusterRootActor* RootActor = Cast<ADisplayClusterRootActor>(GetOwner());
-	
+
 	if (!RootActor)
 	{
 		return;
@@ -286,27 +333,17 @@ void UDisplayClusterPreviewShareComponent::RestoreRootActorOriginalSettings()
 		return;
 	}
 
-	for (UActorComponent* ActorComponent : RootActor->K2_GetComponentsByClass(UDisplayClusterPreviewComponent::StaticClass()))
-	{
-		UDisplayClusterPreviewComponent* PreviewComponent = Cast<UDisplayClusterPreviewComponent>(ActorComponent);
+	// Iterate over viewports and restore their settings
 
-		if (!PreviewComponent)
-		{
-			continue;
-		}
-
-		UDisplayClusterConfigurationViewport* Viewport = GetViewportFromDCRA(RootActor, PreviewComponent->GetClusterNodeId(), PreviewComponent->GetViewportId());
-
-		if (!Viewport)
-		{
-			continue;
-		}
+	for (const FViewportData& ViewportData : GetConfigViewportsFromRootActor(RootActor))
+	{
+		check(ViewportData.Viewport); // The iterator implementation guarantees that it is not nullptr
 
-		const FString ViewportKey = GenerateViewportKey(PreviewComponent->GetClusterNodeId(), PreviewComponent->GetViewportId());
+		const FString ViewportKey = GenerateViewportKey(ViewportData.NodeId, ViewportData.ViewportId);
 
 		if (bool* BoolPtr = OriginalTextureReplaces.Find(ViewportKey))
 		{
-			Viewport->RenderSettings.Replace.bAllowReplace = *BoolPtr;
+			ViewportData.Viewport->RenderSettings.Replace.bAllowReplace = *BoolPtr;
 		}
 
 		if (TObjectPtr<UTexture>* OriginalTexturePtr = OriginalSourceTextures.Find(ViewportKey))
@@ -321,7 +358,7 @@ void UDisplayClusterPreviewShareComponent::RestoreRootActorOriginalSettings()
 			//   the original ones when they really aren't. Otherwise we'd be eventually be restoring textures
 			//   that belong to an external map.
 
-			Viewport->RenderSettings.Replace.SourceTexture = *OriginalTexturePtr;
+			ViewportData.Viewport->RenderSettings.Replace.SourceTexture = *OriginalTexturePtr;
 		}
 	}
 
@@ -329,6 +366,34 @@ void UDisplayClusterPreviewShareComponent::RestoreRootActorOriginalSettings()
 	OriginalTextureReplaces.Empty();
 }
 
+void UDisplayClusterPreviewShareComponent::UnsubscribeFromAllPreviews()
+{
+	if (UWorld* World = GetWorld())
+	{
+		for (TActorIterator<ADisplayClusterRootActor> It(World); It; ++It)
+		{
+			UnsubscribeFromPreview(*It);
+		}
+	}
+}
+
+void UDisplayClusterPreviewShareComponent::SubscribeToPreview(ADisplayClusterRootActor* RootActor)
+{
+	if (IsValid(RootActor))
+	{
+		RootActor->SubscribeToPostProcessRenderTarget(reinterpret_cast<uint8*>(this));
+		RootActor->AddPreviewEnableOverride(reinterpret_cast<uint8*>(this));
+	}
+}
+
+void UDisplayClusterPreviewShareComponent::UnsubscribeFromPreview(ADisplayClusterRootActor* RootActor)
+{
+	if (IsValid(RootActor))
+	{
+		RootActor->UnsubscribeFromPostProcessRenderTarget(reinterpret_cast<uint8*>(this));
+		RootActor->RemovePreviewEnableOverride(reinterpret_cast<uint8*>(this));
+	}
+}
 
 void UDisplayClusterPreviewShareComponent::CloseAllMedia()
 {
@@ -387,6 +452,8 @@ FString UDisplayClusterPreviewShareComponent::GenerateMediaUniqueName(const FStr
 
 void UDisplayClusterPreviewShareComponent::TickSend()
 {
+	using namespace UE::DisplayClusterPreviewShare;
+
 	ADisplayClusterRootActor* RootActor = Cast<ADisplayClusterRootActor>(GetOwner());
 
 	if (!RootActor)
@@ -413,7 +480,7 @@ void UDisplayClusterPreviewShareComponent::TickSend()
 		}
 
 		// Get the texture to share
-		UTextureRenderTarget2D* PreviewTexture = PreviewComponent->GetRenderTargetTexture();
+		UTextureRenderTarget2D* PreviewTexture = GetPreviewRenderTarget(PreviewComponent);
 
 		if (!PreviewTexture)
 		{
@@ -645,7 +712,7 @@ void UDisplayClusterPreviewShareComponent::PullPreviewFromSourceActor(const ADis
 		const UDisplayClusterPreviewComponent* SrcPreviewComponent = *SrcPreviewComponentPtr;
 		check(SrcPreviewComponent);
 
-		UTextureRenderTarget2D* SrcPreviewTexture = SrcPreviewComponent->GetRenderTargetTexture();
+		UTextureRenderTarget2D* SrcPreviewTexture = GetPreviewRenderTarget(SrcPreviewComponent);
 
 		if (!SrcPreviewTexture)
 		{
@@ -973,7 +1040,8 @@ void UDisplayClusterPreviewShareComponent::TickComponent(float DeltaTime, ELevel
 	}
 }
 
-void UDisplayClusterPreviewShareComponent::DestroyComponent(bool bPromoteChildren)
+
+void UDisplayClusterPreviewShareComponent::OnComponentDestroyed(bool bDestroyingHierarchy)
 {
 	// Close all media to avoid keeping resources alive even though the component was destroyed (and possibly kept in undo buffer).
 	// This call also restores the dcra original settings.
@@ -981,17 +1049,6 @@ void UDisplayClusterPreviewShareComponent::DestroyComponent(bool bPromoteChildre
 
 #if WITH_EDITOR
 
-	CloseAllMedia();
-
-	// Tell root actor that it doesn't need to keep up the preview enabled for us.
-	if (ADisplayClusterRootActor* RootActor = Cast<ADisplayClusterRootActor>(GetOwner()))
-	{
-		RootActor->RemovePreviewEnableOverride(reinterpret_cast<uint8*>(this));
-	}
-
-	// We try to leave the root actor as it originally was.
-	RestoreRootActorOriginalSettings();
-
 	if (GIsEditor)
 	{
 		// Unregister Map Change Events
@@ -1003,9 +1060,10 @@ void UDisplayClusterPreviewShareComponent::DestroyComponent(bool bPromoteChildre
 
 #endif // WITH_EDITOR
 
-	Super::DestroyComponent(bPromoteChildren);
+	Super::OnComponentDestroyed(bDestroyingHierarchy);
 }
 
+
 void UDisplayClusterPreviewShareComponent::OnRegister()
 {
 	Super::OnRegister();
diff --git a/Engine/Plugins/Runtime/nDisplay/Source/DisplayCluster/Public/Components/DisplayClusterPreviewShareComponent.h b/Engine/Plugins/Runtime/nDisplay/Source/DisplayCluster/Public/Components/DisplayClusterPreviewShareComponent.h
--- a/Engine/Plugins/Runtime/nDisplay/Source/DisplayCluster/Public/Components/DisplayClusterPreviewShareComponent.h
+++ b/Engine/Plugins/Runtime/nDisplay/Source/DisplayCluster/Public/Components/DisplayClusterPreviewShareComponent.h
@@ -78,7 +78,7 @@ public:
 
 	//~ UActorComponent interface begin
 	virtual void TickComponent(float DeltaTime, ELevelTick TickType, FActorComponentTickFunction* ThisTickFunction) override;
-	virtual void DestroyComponent(bool bPromoteChildren = false) override;
+	virtual void OnComponentDestroyed(bool bDestroyingHierarchy) override;
 	//~ UActorComponent interface end
 
 #endif // WITH_EDITOR
@@ -102,6 +102,7 @@ public:
 	UFUNCTION(BlueprintCallable, Category = "Sharing")
 	void SetUniqueName(const FString& NewUniqueName);
 
+
 protected:
 
 #if WITH_EDITOR // Bulk wrap with WITH_EDITOR until preview is supported in other modes.
@@ -177,6 +178,15 @@ private:
 	/** Called when the editor map is changed. Used to remove unwanted references to external maps. */
 	void HandleMapChanged(UWorld* InWorld, EMapChangeType InMapChangeType);
 
+	/** Unsubcribe from preview of all root actors */
+	void UnsubscribeFromAllPreviews();
+
+	/** Subscribe for preview of given root actor */
+	void SubscribeToPreview(ADisplayClusterRootActor* RootActor);
+
+	/** Unsubscribe from preview of given root actor */
+	void UnsubscribeFromPreview(ADisplayClusterRootActor* RootActor);
+
 #endif // WITH_EDITOR
 
 private:
diff --git a/Engine/Plugins/Runtime/nDisplay/Source/SharedMemoryMedia/Private/SharedMemoryMediaCapture.cpp b/Engine/Plugins/Runtime/nDisplay/Source/SharedMemoryMedia/Private/SharedMemoryMediaCapture.cpp
--- a/Engine/Plugins/Runtime/nDisplay/Source/SharedMemoryMedia/Private/SharedMemoryMediaCapture.cpp
+++ b/Engine/Plugins/Runtime/nDisplay/Source/SharedMemoryMedia/Private/SharedMemoryMediaCapture.cpp
@@ -252,7 +252,7 @@ void USharedMemoryMediaCapture::AddCopyToSharedGpuTexturePass(FRDGBuilder& Graph
 			{
 				TRACE_CPUPROFILER_EVENT_SCOPE(SharedMemMediaOutputFenceBusy);
 
-				UE_LOG(LogSharedMemoryMedia, Warning, TEXT("bTextureReadyFenceBusy[%d] for frame %d was busy, so we wait"), SharedTextureIdx, GFrameCounterRenderThread);
+				UE_LOG(LogSharedMemoryMedia, Verbose, TEXT("bTextureReadyFenceBusy[%d] for frame %d was busy, so we wait"), SharedTextureIdx, GFrameCounterRenderThread);
 
 				while (bTextureReadyFenceBusy[SharedTextureIdx])
 				{
