nDisplay: Added ShortName to DisplayClusterLightCardEditorShaders module to avoid long path names when the module is built.

diff --git a/Engine/Plugins/Runtime/nDisplay/Source/DisplayClusterLightCardEditorShaders/DisplayClusterLightCardEditorShaders.Build.cs b/Engine/Plugins/Runtime/nDisplay/Source/DisplayClusterLightCardEditorShaders/DisplayClusterLightCardEditorShaders.Build.cs
--- a/Engine/Plugins/Runtime/nDisplay/Source/DisplayClusterLightCardEditorShaders/DisplayClusterLightCardEditorShaders.Build.cs
+++ b/Engine/Plugins/Runtime/nDisplay/Source/DisplayClusterLightCardEditorShaders/DisplayClusterLightCardEditorShaders.Build.cs
@@ -21,5 +21,7 @@ public class DisplayClusterLightCardEditorShaders : ModuleRules
 		PrivateIncludePaths.AddRange(
 			new string[] {
 			});
+
+		ShortName = "DCLCEShaders";
 	}
 }
