nDisplay: Fixed conflict when merging CL 26622815.

diff --git a/Engine/Plugins/Runtime/nDisplay/Source/DisplayClusterProjection/Private/Policy/Camera/DisplayClusterProjectionCameraPolicy.cpp b/Engine/Plugins/Runtime/nDisplay/Source/DisplayClusterProjection/Private/Policy/Camera/DisplayClusterProjectionCameraPolicy.cpp
--- a/Engine/Plugins/Runtime/nDisplay/Source/DisplayClusterProjection/Private/Policy/Camera/DisplayClusterProjectionCameraPolicy.cpp
+++ b/Engine/Plugins/Runtime/nDisplay/Source/DisplayClusterProjection/Private/Policy/Camera/DisplayClusterProjectionCameraPolicy.cpp
@@ -118,10 +118,10 @@ bool FDisplayClusterProjectionCameraPolicy::GetProjectionMatrix(IDisplayClusterV
 	if (InViewport)
 	{
 		// Support inner camera custom frustum
-		const float HalfFOVH = 0.5 * ClampedCameraFOVDegrees;
-		const float HalfFOVV = HalfFOVH / CameraAspectRatio;
+		const float HalfFOVH = ZNear * FMath::Tan(FMath::DegreesToRadians(0.5 * ClampedCameraFOV));
+		const float HalfFOVV = HalfFOVH / InCameraAspectRatio;
 
-		InViewport->CalculateProjectionMatrix(InContextNum, -HalfFOVH, HalfFOVH, HalfFOVV, -HalfFOVV, ZNear, ZFar, true);
+		InViewport->CalculateProjectionMatrix(InContextNum, -HalfFOVH, HalfFOVH, HalfFOVV, -HalfFOVV, ZNear, ZFar, false);
 		OutPrjMatrix = InViewport->GetContexts()[InContextNum].ProjectionMatrix;
 	}
 	else
