nDisplay: fixed compilation Issue 339020.

diff --git a/Engine/Plugins/Runtime/nDisplay/Source/DisplayClusterProjection/Private/Policy/Camera/DisplayClusterProjectionCameraPolicy.cpp b/Engine/Plugins/Runtime/nDisplay/Source/DisplayClusterProjection/Private/Policy/Camera/DisplayClusterProjectionCameraPolicy.cpp
--- a/Engine/Plugins/Runtime/nDisplay/Source/DisplayClusterProjection/Private/Policy/Camera/DisplayClusterProjectionCameraPolicy.cpp
+++ b/Engine/Plugins/Runtime/nDisplay/Source/DisplayClusterProjection/Private/Policy/Camera/DisplayClusterProjectionCameraPolicy.cpp
@@ -118,8 +118,8 @@ bool FDisplayClusterProjectionCameraPolicy::GetProjectionMatrix(IDisplayClusterV
 	if (InViewport)
 	{
 		// Support inner camera custom frustum
-		const float HalfFOVH = ZNear * FMath::Tan(FMath::DegreesToRadians(0.5 * ClampedCameraFOV));
-		const float HalfFOVV = HalfFOVH / InCameraAspectRatio;
+		const float HalfFOVH = ZNear * FMath::Tan(FMath::DegreesToRadians(0.5 * ClampedCameraFOVDegrees));
+		const float HalfFOVV = HalfFOVH / CameraAspectRatio;
 
 		InViewport->CalculateProjectionMatrix(InContextNum, -HalfFOVH, HalfFOVH, HalfFOVV, -HalfFOVV, ZNear, ZFar, false);
 		OutPrjMatrix = InViewport->GetContexts()[InContextNum].ProjectionMatrix;
