OCIO: Avoid relying on invalid configs/settings to skip rendering.

We update the default enabled value to false, and let users manually enable the transform.

diff --git a/Engine/Plugins/Compositing/OpenColorIO/Source/OpenColorIO/Private/OpenColorIOColorSpace.cpp b/Engine/Plugins/Compositing/OpenColorIO/Source/OpenColorIO/Private/OpenColorIOColorSpace.cpp
--- a/Engine/Plugins/Compositing/OpenColorIO/Source/OpenColorIO/Private/OpenColorIOColorSpace.cpp
+++ b/Engine/Plugins/Compositing/OpenColorIO/Source/OpenColorIO/Private/OpenColorIOColorSpace.cpp
@@ -4,6 +4,7 @@
 
 #include "OpenColorIOConfiguration.h"
 #include "OpenColorIOSettings.h"
+#include "UObject/UE5ReleaseStreamObjectVersion.h"
 
 #include UE_INLINE_GENERATED_CPP_BY_NAME(OpenColorIOColorSpace)
 
@@ -243,3 +244,26 @@ bool FOpenColorIOColorConversionSettings::IsDisplayView() const
 {
 	return DestinationDisplayView.IsValid();
 }
+
+bool FOpenColorIODisplayConfiguration::Serialize(FArchive& Ar)
+{
+	Ar.UsingCustomVersion(FUE5ReleaseStreamObjectVersion::GUID);
+	
+	// Don't actually serialize, just write the custom version for PostSerialize
+	return false;
+}
+
+void FOpenColorIODisplayConfiguration::PostSerialize(const FArchive& Ar)
+{
+	if (Ar.IsLoading())
+	{
+		if (Ar.CustomVer(FUE5ReleaseStreamObjectVersion::GUID) < FUE5ReleaseStreamObjectVersion::OpenColorIODisabledDisplayConfigurationDefault)
+		{
+			// Retain previous behavior: enabled only when the settings were valid
+			if (ColorConfiguration.IsValid())
+			{
+				bIsEnabled = true;
+			}
+		}
+	}
+}
diff --git a/Engine/Plugins/Compositing/OpenColorIO/Source/OpenColorIO/Private/OpenColorIODisplayExtension.cpp b/Engine/Plugins/Compositing/OpenColorIO/Source/OpenColorIO/Private/OpenColorIODisplayExtension.cpp
--- a/Engine/Plugins/Compositing/OpenColorIO/Source/OpenColorIO/Private/OpenColorIODisplayExtension.cpp
+++ b/Engine/Plugins/Compositing/OpenColorIO/Source/OpenColorIO/Private/OpenColorIODisplayExtension.cpp
@@ -29,12 +29,9 @@ FOpenColorIODisplayExtension::FOpenColorIODisplayExtension(const FAutoRegister&
 
 bool FOpenColorIODisplayExtension::IsActiveThisFrame_Internal(const FSceneViewExtensionContext& Context) const
 {
-	if (Context.Viewport && LinkedViewportClient == Context.Viewport->GetClient() && DisplayConfiguration.bIsEnabled)
+	if (Context.Viewport && LinkedViewportClient == Context.Viewport->GetClient())
 	{
-		if (DisplayConfiguration.ColorConfiguration.IsValid())
-		{
-			return DisplayConfiguration.ColorConfiguration.ConfigurationSource->IsTransformReady(DisplayConfiguration.ColorConfiguration);
-		}
+		return DisplayConfiguration.bIsEnabled;
 	}
 
 	return false;
diff --git a/Engine/Plugins/Compositing/OpenColorIO/Source/OpenColorIO/Private/OpenColorIORendering.cpp b/Engine/Plugins/Compositing/OpenColorIO/Source/OpenColorIO/Private/OpenColorIORendering.cpp
--- a/Engine/Plugins/Compositing/OpenColorIO/Source/OpenColorIO/Private/OpenColorIORendering.cpp
+++ b/Engine/Plugins/Compositing/OpenColorIO/Source/OpenColorIO/Private/OpenColorIORendering.cpp
@@ -163,8 +163,6 @@ FOpenColorIORenderPassResources FOpenColorIORendering::GetRenderPassResources(co
 			check(Result.ShaderResource);
 			if (Result.ShaderResource->GetShaderGameThread<FOpenColorIOPixelShader>().IsNull())
 			{
-				ensureMsgf(false, TEXT("Can't apply display look - Shader was invalid for Resource %s"), *Result.ShaderResource->GetFriendlyName());
-
 				//Invalidate shader resource
 				Result.ShaderResource = nullptr;
 			}
diff --git a/Engine/Plugins/Compositing/OpenColorIO/Source/OpenColorIO/Public/OpenColorIOColorSpace.h b/Engine/Plugins/Compositing/OpenColorIO/Source/OpenColorIO/Public/OpenColorIOColorSpace.h
--- a/Engine/Plugins/Compositing/OpenColorIO/Source/OpenColorIO/Public/OpenColorIOColorSpace.h
+++ b/Engine/Plugins/Compositing/OpenColorIO/Source/OpenColorIO/Public/OpenColorIOColorSpace.h
@@ -226,10 +226,22 @@ public:
 	 * dictate whether it's applied or not to it
 	 */
 	UPROPERTY(EditAnywhere, BlueprintReadWrite, Category = ColorSpace)
-	bool bIsEnabled = true;
+	bool bIsEnabled = false;
 	
 	/** Conversion to apply when this display is enabled */
 	UPROPERTY(EditAnywhere, BlueprintReadWrite, Category = ColorSpace)
 	FOpenColorIOColorConversionSettings ColorConfiguration;
+
+	bool Serialize(FArchive& Ar);
+	void PostSerialize(const FArchive& Ar);
+};
+
+template<> struct TStructOpsTypeTraits<FOpenColorIODisplayConfiguration> : public TStructOpsTypeTraitsBase2<FOpenColorIODisplayConfiguration>
+{
+	enum
+	{
+		WithSerializer = true,
+		WithPostSerialize = true,
+	};
 };
 
diff --git a/Engine/Plugins/Experimental/CineCameraSceneCapture/Source/CineCameraSceneCapture/Private/CineCameraSceneCaptureComponent.cpp b/Engine/Plugins/Experimental/CineCameraSceneCapture/Source/CineCameraSceneCapture/Private/CineCameraSceneCaptureComponent.cpp
--- a/Engine/Plugins/Experimental/CineCameraSceneCapture/Source/CineCameraSceneCapture/Private/CineCameraSceneCaptureComponent.cpp
+++ b/Engine/Plugins/Experimental/CineCameraSceneCapture/Source/CineCameraSceneCapture/Private/CineCameraSceneCaptureComponent.cpp
@@ -160,7 +160,6 @@ UCineCaptureComponent2D::UCineCaptureComponent2D(const FObjectInitializer& Objec
 {
 	CaptureSource = ESceneCaptureSource::SCS_FinalToneCurveHDR;
 	bAlwaysPersistRenderingState = true;
-	OCIOConfiguration.bIsEnabled = false;
 }
 
 void UCineCaptureComponent2D::UpdateSceneCaptureContents(FSceneInterface* Scene)
diff --git a/Engine/Plugins/MovieScene/MovieRenderPipeline/Source/MovieRenderPipelineCore/Private/MoviePipelineColorSetting.cpp b/Engine/Plugins/MovieScene/MovieRenderPipeline/Source/MovieRenderPipelineCore/Private/MoviePipelineColorSetting.cpp
--- a/Engine/Plugins/MovieScene/MovieRenderPipeline/Source/MovieRenderPipelineCore/Private/MoviePipelineColorSetting.cpp
+++ b/Engine/Plugins/MovieScene/MovieRenderPipeline/Source/MovieRenderPipelineCore/Private/MoviePipelineColorSetting.cpp
@@ -6,6 +6,9 @@
 
 UMoviePipelineColorSetting::UMoviePipelineColorSetting() : bDisableToneCurve(false)
 {
+	// Note: Since MRQ settings already have enabled/disabled state control, this is slighty redundant.
+	// As such, we retain the old behavior and always default to enabled == true.
+	OCIOConfiguration.bIsEnabled = true;
 }
 
 #if WITH_EDITOR
diff --git a/Engine/Plugins/Runtime/nDisplay/Source/DisplayCluster/Private/Render/Viewport/Configuration/DisplayClusterViewportConfigurationHelpers_OpenColorIO.cpp b/Engine/Plugins/Runtime/nDisplay/Source/DisplayCluster/Private/Render/Viewport/Configuration/DisplayClusterViewportConfigurationHelpers_OpenColorIO.cpp
--- a/Engine/Plugins/Runtime/nDisplay/Source/DisplayCluster/Private/Render/Viewport/Configuration/DisplayClusterViewportConfigurationHelpers_OpenColorIO.cpp
+++ b/Engine/Plugins/Runtime/nDisplay/Source/DisplayCluster/Private/Render/Viewport/Configuration/DisplayClusterViewportConfigurationHelpers_OpenColorIO.cpp
@@ -87,14 +87,6 @@ bool FDisplayClusterViewportConfigurationHelpers_OpenColorIO::UpdateChromakeyVie
 
 void FDisplayClusterViewportConfigurationHelpers_OpenColorIO::ApplyOCIOConfiguration(FDisplayClusterViewport& DstViewport, const FOpenColorIOColorConversionSettings& InConversionSettings)
 {
-	// Check OCIO before apply
-	if(!InConversionSettings.IsValid())
-	{
-		DisableOCIOConfiguration(DstViewport);
-
-		return;
-	}
-
 	if (DstViewport.OpenColorIO.IsValid() && DstViewport.OpenColorIO->IsConversionSettingsEqual(InConversionSettings))
 	{
 		// Already assigned
diff --git a/Engine/Plugins/Runtime/nDisplay/Source/DisplayCluster/Private/Render/Viewport/DisplayClusterViewportProxy.cpp b/Engine/Plugins/Runtime/nDisplay/Source/DisplayCluster/Private/Render/Viewport/DisplayClusterViewportProxy.cpp
--- a/Engine/Plugins/Runtime/nDisplay/Source/DisplayCluster/Private/Render/Viewport/DisplayClusterViewportProxy.cpp
+++ b/Engine/Plugins/Runtime/nDisplay/Source/DisplayCluster/Private/Render/Viewport/DisplayClusterViewportProxy.cpp
@@ -452,7 +452,7 @@ bool FDisplayClusterViewportProxy::ImplGetResources_RenderThread(const EDisplayC
 
 EDisplayClusterViewportOpenColorIOMode FDisplayClusterViewportProxy::GetOpenColorIOMode() const
 {
-	if (OpenColorIO.IsValid() && OpenColorIO->IsEnabled_RenderThread())
+	if (OpenColorIO.IsValid())
 	{
 		if (EnumHasAnyFlags(RenderSettingsICVFX.RuntimeFlags, EDisplayClusterViewportRuntimeICVFXFlags::UVLightcard | EDisplayClusterViewportRuntimeICVFXFlags::Lightcard | EDisplayClusterViewportRuntimeICVFXFlags::Chromakey))
 		{
diff --git a/Engine/Plugins/Runtime/nDisplay/Source/DisplayCluster/Private/Render/Viewport/DisplayClusterViewport_OpenColorIO.cpp b/Engine/Plugins/Runtime/nDisplay/Source/DisplayCluster/Private/Render/Viewport/DisplayClusterViewport_OpenColorIO.cpp
--- a/Engine/Plugins/Runtime/nDisplay/Source/DisplayCluster/Private/Render/Viewport/DisplayClusterViewport_OpenColorIO.cpp
+++ b/Engine/Plugins/Runtime/nDisplay/Source/DisplayCluster/Private/Render/Viewport/DisplayClusterViewport_OpenColorIO.cpp
@@ -39,15 +39,6 @@ void FDisplayClusterViewport_OpenColorIO::SetupSceneView(FSceneViewFamily& InOut
 
 void FDisplayClusterViewport_OpenColorIO::UpdateOpenColorIORenderPassResources()
 {
-	bShaderResourceValid = false;
-	bIsEnabled = ConversionSettings.IsValid() ? ConversionSettings.ConfigurationSource->IsTransformReady(ConversionSettings) : false;
-
-	if (!bIsEnabled)
-	{
-		// The data in this configuration was not correctly initialized in the previous frame. Ignore this OCIO
-		return;
-	}
-
 	FOpenColorIORenderPassResources PassResources = FOpenColorIORendering::GetRenderPassResources(ConversionSettings, GMaxRHIFeatureLevel);
 
 	bShaderResourceValid = PassResources.IsValid();
@@ -61,11 +52,6 @@ void FDisplayClusterViewport_OpenColorIO::UpdateOpenColorIORenderPassResources()
 	);
 }
 
-bool FDisplayClusterViewport_OpenColorIO::IsEnabled_RenderThread() const
-{
-	return bIsEnabled;
-}
-
 bool FDisplayClusterViewport_OpenColorIO::IsConversionSettingsEqual(const FOpenColorIOColorConversionSettings& InConversionSettings) const
 {
 	return ConversionSettings.ToString() == InConversionSettings.ToString();
@@ -86,33 +72,28 @@ float FDisplayClusterViewport_OpenColorIO::GetDisplayGamma(const FDisplayCluster
 bool FDisplayClusterViewport_OpenColorIO::AddPass_RenderThread(FRDGBuilder& GraphBuilder, const FDisplayClusterViewport_Context& InViewportContext,
 	FRHITexture2D* InputTextureRHI, const FIntRect& InputRect, FRHITexture2D* OutputTextureRHI, const FIntRect& OutputRect, bool bUnpremultiply, bool bInvertAlpha) const
 {
-	if (IsEnabled_RenderThread())
-	{
-		FRDGTextureRef InputTexture = GraphBuilder.RegisterExternalTexture(CreateRenderTarget(InputTextureRHI, TEXT("nDisplayOCIOInputTexture")));
-		FRDGTextureRef OutputTexture = GraphBuilder.RegisterExternalTexture(CreateRenderTarget(OutputTextureRHI, TEXT("nDisplayOCIORenderTargetTexture")));
+	FRDGTextureRef InputTexture = GraphBuilder.RegisterExternalTexture(CreateRenderTarget(InputTextureRHI, TEXT("nDisplayOCIOInputTexture")));
+	FRDGTextureRef OutputTexture = GraphBuilder.RegisterExternalTexture(CreateRenderTarget(OutputTextureRHI, TEXT("nDisplayOCIORenderTargetTexture")));
 
-		FScreenPassRenderTarget Output = FScreenPassRenderTarget(OutputTexture, OutputRect, ERenderTargetLoadAction::EClear);
-		EOpenColorIOTransformAlpha TransformAlpha = EOpenColorIOTransformAlpha::None;
-		if (bUnpremultiply)
-		{
-			TransformAlpha = bInvertAlpha ? EOpenColorIOTransformAlpha::InvertUnpremultiply: EOpenColorIOTransformAlpha::Unpremultiply;
-		}
-
-		FOpenColorIORendering::AddPass_RenderThread(
-			GraphBuilder,
-			FScreenPassViewInfo(),
-			GMaxRHIFeatureLevel,
-			FScreenPassTexture(InputTexture),
-			Output,
-			CachedResourcesRenderThread,
-			GetDisplayGamma(InViewportContext),
-			TransformAlpha
-		);
-
-		return true;
+	FScreenPassRenderTarget Output = FScreenPassRenderTarget(OutputTexture, OutputRect, ERenderTargetLoadAction::EClear);
+	EOpenColorIOTransformAlpha TransformAlpha = EOpenColorIOTransformAlpha::None;
+	if (bUnpremultiply)
+	{
+		TransformAlpha = bInvertAlpha ? EOpenColorIOTransformAlpha::InvertUnpremultiply : EOpenColorIOTransformAlpha::Unpremultiply;
 	}
 
-	return false;
+	FOpenColorIORendering::AddPass_RenderThread(
+		GraphBuilder,
+		FScreenPassViewInfo(),
+		GMaxRHIFeatureLevel,
+		FScreenPassTexture(InputTexture),
+		Output,
+		CachedResourcesRenderThread,
+		GetDisplayGamma(InViewportContext),
+		TransformAlpha
+	);
+
+	return true;
 }
 
 // This is a copy of FOpenColorIODisplayExtension::PostProcessPassAfterTonemap_RenderThread()
@@ -133,11 +114,9 @@ FScreenPassTexture FDisplayClusterViewport_OpenColorIO::PostProcessPassAfterTone
 	FOpenColorIORendering::AddPass_RenderThread(
 		GraphBuilder,
 		ViewInfo,
-		ViewInfo.GetFeatureLevel(),
 		SceneColor,
 		Output,
-		CachedResourcesRenderThread,
-		GetDisplayGamma(InViewportContext)
+		CachedResourcesRenderThread
 	);
 
 	return MoveTemp(Output);
diff --git a/Engine/Plugins/Runtime/nDisplay/Source/DisplayCluster/Private/Render/Viewport/DisplayClusterViewport_OpenColorIO.h b/Engine/Plugins/Runtime/nDisplay/Source/DisplayCluster/Private/Render/Viewport/DisplayClusterViewport_OpenColorIO.h
--- a/Engine/Plugins/Runtime/nDisplay/Source/DisplayCluster/Private/Render/Viewport/DisplayClusterViewport_OpenColorIO.h
+++ b/Engine/Plugins/Runtime/nDisplay/Source/DisplayCluster/Private/Render/Viewport/DisplayClusterViewport_OpenColorIO.h
@@ -60,9 +60,6 @@ public:
 	/* This is a copy of FOpenColorIODisplayExtension::PostProcessPassAfterTonemap_RenderThread() */
 	FScreenPassTexture PostProcessPassAfterTonemap_RenderThread(FRDGBuilder& GraphBuilder, const FDisplayClusterViewport_Context& InViewportContext, const FSceneView& View, const FPostProcessMaterialInputs& InOutInputs);
 
-	/** Is OCIO enabled for render thread */
-	bool IsEnabled_RenderThread() const;
-
 	/** Compare two OCIO configurations.
 	 *
 	 * @param InConversionSettings - configuration to compare with.
@@ -90,7 +87,4 @@ private:
 
 	/** Shader resources state. */
 	bool bShaderResourceValid = true;
-
-	/** Configuration data state. */
-	TAtomic<bool> bIsEnabled = true;
 };
diff --git a/Engine/Plugins/Runtime/nDisplay/Source/DisplayClusterConfiguration/Private/DisplayClusterConfigurationTypes_ICVFX.cpp b/Engine/Plugins/Runtime/nDisplay/Source/DisplayClusterConfiguration/Private/DisplayClusterConfigurationTypes_ICVFX.cpp
--- a/Engine/Plugins/Runtime/nDisplay/Source/DisplayClusterConfiguration/Private/DisplayClusterConfigurationTypes_ICVFX.cpp
+++ b/Engine/Plugins/Runtime/nDisplay/Source/DisplayClusterConfiguration/Private/DisplayClusterConfigurationTypes_ICVFX.cpp
@@ -37,24 +37,18 @@ FDisplayClusterConfigurationICVFX_CameraRenderSettings::FDisplayClusterConfigura
 ///////////////////////////////////////////////////////////////////////////////////////
 const FOpenColorIOColorConversionSettings* FDisplayClusterConfigurationICVFX_LightcardCustomOCIO::FindOCIOConfiguration(const FString& InViewportId) const
 {
-	if (AllViewportsOCIOConfiguration.bIsEnabled)
+	// Note: Lightcard OCIO is enabled from the drop-down menu, so we ignore AllViewportsOCIOConfiguration.bIsEnabled (the property isn't exposed)
+	
+	// Per viewport OCIO:
+	for (const FDisplayClusterConfigurationOCIOProfile& OCIOProfileIt : PerViewportOCIOProfiles)
 	{
-		// Per viewport OCIO:
-		for (const FDisplayClusterConfigurationOCIOProfile& OCIOProfileIt : PerViewportOCIOProfiles)
+		if (OCIOProfileIt.IsEnabledForObject(InViewportId))
 		{
-			if (OCIOProfileIt.IsEnabledForObject(InViewportId))
-			{
-				return &OCIOProfileIt.ColorConfiguration;
-			}
-		}
-
-		if (AllViewportsOCIOConfiguration.IsEnabled())
-		{
-			return &AllViewportsOCIOConfiguration.ColorConfiguration;
+			return &OCIOProfileIt.ColorConfiguration;
 		}
 	}
 
-	return nullptr;
+	return &AllViewportsOCIOConfiguration.ColorConfiguration;
 }
 
 ///////////////////////////////////////////////////////////////////////////////////////
@@ -73,10 +67,7 @@ const FOpenColorIOColorConversionSettings* FDisplayClusterConfigurationICVFX_Vie
 			}
 		}
 
-		if (AllViewportsOCIOConfiguration.IsEnabled())
-		{
-			return &AllViewportsOCIOConfiguration.ColorConfiguration;
-		}
+		return &AllViewportsOCIOConfiguration.ColorConfiguration;
 	}
 
 	return nullptr;
@@ -98,10 +89,7 @@ const FOpenColorIOColorConversionSettings* FDisplayClusterConfigurationICVFX_Cam
 			}
 		}
 
-		if (AllNodesOCIOConfiguration.IsEnabled())
-		{
-			return &AllNodesOCIOConfiguration.ColorConfiguration;
-		}
+		return &AllNodesOCIOConfiguration.ColorConfiguration;
 	}
 
 	return nullptr;
@@ -119,7 +107,7 @@ bool FDisplayClusterConfigurationICVFX_CameraOCIO::IsInnerFrustumViewportSetting
 	{
 		for (const FDisplayClusterConfigurationOCIOProfile& OCIOProfileIt : PerNodeOCIOProfiles)
 		{
-			if (OCIOProfileIt.IsEnabled())
+			if (OCIOProfileIt.bIsEnabled)
 			{
 				const FString* CustomNode1 = OCIOProfileIt.ApplyOCIOToObjects.FindByPredicate([ClusterNodeId = InClusterNodeId1](const FString& InClusterNodeId)
 					{
diff --git a/Engine/Plugins/Runtime/nDisplay/Source/DisplayClusterConfiguration/Private/DisplayClusterConfigurationTypes_OCIO.cpp b/Engine/Plugins/Runtime/nDisplay/Source/DisplayClusterConfiguration/Private/DisplayClusterConfigurationTypes_OCIO.cpp
--- a/Engine/Plugins/Runtime/nDisplay/Source/DisplayClusterConfiguration/Private/DisplayClusterConfigurationTypes_OCIO.cpp
+++ b/Engine/Plugins/Runtime/nDisplay/Source/DisplayClusterConfiguration/Private/DisplayClusterConfigurationTypes_OCIO.cpp
@@ -1,10 +1,34 @@
 // Copyright Epic Games, Inc. All Rights Reserved.
 
 #include "DisplayClusterConfigurationTypes_OCIO.h"
+#include "UObject/UE5ReleaseStreamObjectVersion.h"
 
 ///////////////////////////////////////////////////////////////////////////////////////
 // FDisplayClusterConfigurationOCIOConfiguration
 ///////////////////////////////////////////////////////////////////////////////////////
+bool FDisplayClusterConfigurationOCIOConfiguration::Serialize(FArchive& Ar)
+{
+	Ar.UsingCustomVersion(FUE5ReleaseStreamObjectVersion::GUID);
+
+	// Don't actually serialize, just write the custom version for PostSerialize
+	return false;
+}
+
+void FDisplayClusterConfigurationOCIOConfiguration::PostSerialize(const FArchive& Ar)
+{
+	if (Ar.IsLoading())
+	{
+		if (Ar.CustomVer(FUE5ReleaseStreamObjectVersion::GUID) < FUE5ReleaseStreamObjectVersion::OpenColorIODisabledDisplayConfigurationDefault)
+		{
+			// Retain previous behavior: enabled only when the settings were valid
+			if (ColorConfiguration.IsValid())
+			{
+				bIsEnabled = true;
+			}
+		}
+	}
+}
+
 bool FDisplayClusterConfigurationOCIOConfiguration::IsEnabled() const
 {
 	return bIsEnabled && ColorConfiguration.IsValid();
@@ -13,6 +37,29 @@ bool FDisplayClusterConfigurationOCIOConfiguration::IsEnabled() const
 ///////////////////////////////////////////////////////////////////////////////////////
 // FDisplayClusterConfigurationOCIOProfile
 ///////////////////////////////////////////////////////////////////////////////////////
+bool FDisplayClusterConfigurationOCIOProfile::Serialize(FArchive& Ar)
+{
+	Ar.UsingCustomVersion(FUE5ReleaseStreamObjectVersion::GUID);
+
+	// Don't actually serialize, just write the custom version for PostSerialize
+	return false;
+}
+
+void FDisplayClusterConfigurationOCIOProfile::PostSerialize(const FArchive& Ar)
+{
+	if (Ar.IsLoading())
+	{
+		if (Ar.CustomVer(FUE5ReleaseStreamObjectVersion::GUID) < FUE5ReleaseStreamObjectVersion::OpenColorIODisabledDisplayConfigurationDefault)
+		{
+			// Retain previous behavior: enabled only when the settings were valid
+			if (ColorConfiguration.IsValid())
+			{
+				bIsEnabled = true;
+			}
+		}
+	}
+}
+
 bool FDisplayClusterConfigurationOCIOProfile::IsEnabled() const
 {
 	return bIsEnabled && ColorConfiguration.IsValid();
@@ -20,7 +67,7 @@ bool FDisplayClusterConfigurationOCIOProfile::IsEnabled() const
 
 bool FDisplayClusterConfigurationOCIOProfile::IsEnabledForObject(const FString& InObjectId) const
 {
-	if (IsEnabled())
+	if (bIsEnabled)
 	{
 		for (const FString& ViewportNameIt : ApplyOCIOToObjects)
 		{
diff --git a/Engine/Plugins/Runtime/nDisplay/Source/DisplayClusterConfiguration/Public/DisplayClusterConfigurationTypes_OCIO.h b/Engine/Plugins/Runtime/nDisplay/Source/DisplayClusterConfiguration/Public/DisplayClusterConfigurationTypes_OCIO.h
--- a/Engine/Plugins/Runtime/nDisplay/Source/DisplayClusterConfiguration/Public/DisplayClusterConfigurationTypes_OCIO.h
+++ b/Engine/Plugins/Runtime/nDisplay/Source/DisplayClusterConfiguration/Public/DisplayClusterConfigurationTypes_OCIO.h
@@ -19,7 +19,7 @@ struct FDisplayClusterConfigurationOCIOConfiguration
 
 	/**  Enable the application of an OpenColorIO configuration to all viewports. */
 	UPROPERTY(EditAnywhere, BlueprintReadWrite, Category = "OCIO")
-	bool bIsEnabled = true;
+	bool bIsEnabled = false;
 
 	/** "This property has been deprecated. */
 	UPROPERTY(meta = (DeprecatedProperty, DeprecationMessage = "Use the ColorConfiguration property instead"))
@@ -30,10 +30,24 @@ struct FDisplayClusterConfigurationOCIOConfiguration
 	FOpenColorIOColorConversionSettings ColorConfiguration;
 
 public:
+
+	bool Serialize(FArchive& Ar);
+	void PostSerialize(const FArchive& Ar);
+	
+	UE_DEPRECATED(5.3, "This method is deprecated.")
 	/** Return true if configuration valid */
 	bool IsEnabled() const;
 };
 
+template<> struct TStructOpsTypeTraits<FDisplayClusterConfigurationOCIOConfiguration> : public TStructOpsTypeTraitsBase2<FDisplayClusterConfigurationOCIOConfiguration>
+{
+	enum
+	{
+		WithSerializer = true,
+		WithPostSerialize = true,
+	};
+};
+
 /*
  * OCIO profile structure. Can be configured for viewports or cluster nodes.
  * To enable viewport configuration when using as a UPROPERTY set meta = (ConfigurationMode = "Viewports")
@@ -46,7 +60,7 @@ struct FDisplayClusterConfigurationOCIOProfile
 
 	/** Enable the application of an OpenColorIO configuration for the viewport(s) specified. */
 	UPROPERTY(EditAnywhere, BlueprintReadWrite, Category = "OCIO")
-	bool bIsEnabled = true;
+	bool bIsEnabled = false;
 
 	/** "This property has been deprecated. */
 	UPROPERTY(meta = (DeprecatedProperty, DeprecationMessage = "Use the ColorConfiguration property instead"))
@@ -61,10 +75,23 @@ struct FDisplayClusterConfigurationOCIOProfile
 	FOpenColorIOColorConversionSettings ColorConfiguration;
 
 public:
+	
+	bool Serialize(FArchive& Ar);
+	void PostSerialize(const FArchive& Ar);
+
 	/** Return true if the configuration is valid for the input object */
 	bool IsEnabledForObject(const FString& InObjectId) const;
-
+	
+	UE_DEPRECATED(5.3, "This method is deprecated.")
 	/** Return true if configuration valid */
 	bool IsEnabled() const;
 };
 
+template<> struct TStructOpsTypeTraits<FDisplayClusterConfigurationOCIOProfile> : public TStructOpsTypeTraitsBase2<FDisplayClusterConfigurationOCIOProfile>
+{
+	enum
+	{
+		WithSerializer = true,
+		WithPostSerialize = true,
+	};
+};
\ No newline at end of file
diff --git a/Engine/Source/Runtime/Core/Public/UObject/UE5ReleaseStreamObjectVersion.h b/Engine/Source/Runtime/Core/Public/UObject/UE5ReleaseStreamObjectVersion.h
--- a/Engine/Source/Runtime/Core/Public/UObject/UE5ReleaseStreamObjectVersion.h
+++ b/Engine/Source/Runtime/Core/Public/UObject/UE5ReleaseStreamObjectVersion.h
@@ -150,6 +150,9 @@ struct FUE5ReleaseStreamObjectVersion
 
 		// Serialize bActorIsListedInSceneOutliner in WorldPartitionActorDesc
 		WorldPartitionActorDescSerializeActorIsListedInSceneOutliner,
+
+		// Disabled opencolorio display configuration by default
+		OpenColorIODisabledDisplayConfigurationDefault,
 		
 		// -----<new versions can be added above this line>-------------------------------------------------
 		VersionPlusOne,
