nDisplay: Fix undo/redo issues when adding viewports or cluster nodes in the nDisplay blueprint editor.
- Can no longer get into a state which will crash.
- Undo transactions work properly now while in the modal window.
- Cannot undo transactions prior to opening the window until the window is closed.
- Undo performance greatly improved.

diff --git a/Engine/Plugins/Runtime/nDisplay/Source/DisplayClusterConfigurator/Private/ClusterConfiguration/DisplayClusterConfiguratorClusterUtils.cpp b/Engine/Plugins/Runtime/nDisplay/Source/DisplayClusterConfigurator/Private/ClusterConfiguration/DisplayClusterConfiguratorClusterUtils.cpp
--- a/Engine/Plugins/Runtime/nDisplay/Source/DisplayClusterConfigurator/Private/ClusterConfiguration/DisplayClusterConfiguratorClusterUtils.cpp
+++ b/Engine/Plugins/Runtime/nDisplay/Source/DisplayClusterConfigurator/Private/ClusterConfiguration/DisplayClusterConfiguratorClusterUtils.cpp
@@ -18,6 +18,8 @@
 #include "Framework/Application/SlateApplication.h"
 #include "HAL/PlatformApplicationMisc.h"
 #include "UnrealExporter.h"
+#include "Editor/Transactor.h"
+#include "Editor/TransBuffer.h"
 #include "Exporters/Exporter.h"
 #include "Input/DragAndDrop.h"
 #include "Widgets/SBoxPanel.h"
@@ -31,11 +33,76 @@ const FString FDisplayClusterConfiguratorClusterUtils::DefaultNewHostName = TEXT
 const FString FDisplayClusterConfiguratorClusterUtils::DefaultNewClusterNodeName = TEXT("Node");
 const FString FDisplayClusterConfiguratorClusterUtils::DefaultNewViewportName = TEXT("VP");
 
+namespace UE::DisplayClusterConfiguratorClusterUtils
+{
+	// All transactions performed within this scope can be considered temporary and will be cleared when out of scope.
+	// Transactions prior to the scope cannot be undone while in scope.
+	struct FScopedTemporaryTransactions
+	{
+		// The index of the transaction queue when entering the scope
+		int32 StartTransactionIdx = 0;
+		
+		FScopedTemporaryTransactions()
+		{
+			if (GEditor)
+			{
+				// Prevent undoing any actions before the scope started
+				GEditor->Trans->SetUndoBarrier();
+
+				// Record the transaction idx before we start our scoped transactions
+				StartTransactionIdx = GEditor->Trans->GetCurrentUndoBarrier();
+
+				// Broadcasting changes here can take a big performance hit and isn't necessary since we only use this in a modal window
+				GEditor->bSuspendBroadcastPostUndoRedo = true;
+			}
+		}
+
+		~FScopedTemporaryTransactions()
+		{
+			if (GEditor)
+			{
+				GEditor->bSuspendBroadcastPostUndoRedo = false;
+
+				// Clear all transactions that occurred while in scope, accounting for transactions that were
+				// undone. None of these transactions should be allowed to be undone or redone from this point forward
+				{
+					UTransBuffer* TransBuffer = CastChecked<UTransBuffer>(GEditor->Trans);
+					const int32 ScopedTransactionsPerformed = TransBuffer->UndoBuffer.Num() - StartTransactionIdx;
+				
+					if (ScopedTransactionsPerformed > 0
+						&& StartTransactionIdx >= 0
+						&& StartTransactionIdx < TransBuffer->UndoBuffer.Num()
+						&& (StartTransactionIdx + ScopedTransactionsPerformed - 1) < TransBuffer->UndoBuffer.Num())
+					{
+						// Remove only the transactions performed within scope
+						TransBuffer->UndoBuffer.RemoveAt(StartTransactionIdx, ScopedTransactionsPerformed);
+
+						// We can't restore redo items prior to the scope since they would have been cleared while
+						// performing transactions within the scope
+						TransBuffer->UndoCount = 0;
+
+						// Need to broadcast for other systems since we are changing the undo buffer manually
+						TransBuffer->OnUndoBufferChanged().Broadcast();
+					}
+				}
+
+				// Allow undos again from before the scope started
+				GEditor->Trans->RemoveUndoBarrier();
+			}
+		}
+	};
+}
+
 using namespace DisplayClusterConfiguratorPropertyUtils;
 
-UDisplayClusterConfigurationClusterNode* FDisplayClusterConfiguratorClusterUtils::CreateNewClusterNodeFromDialog(const TSharedRef<FDisplayClusterConfiguratorBlueprintEditor>& Toolkit, UDisplayClusterConfigurationCluster* Cluster, const FDisplayClusterConfigurationRectangle& PresetRect, FString PresetHost)
+UDisplayClusterConfigurationClusterNode* FDisplayClusterConfiguratorClusterUtils::CreateNewClusterNodeFromDialog(
+	const TSharedRef<FDisplayClusterConfiguratorBlueprintEditor>& Toolkit,
+	UDisplayClusterConfigurationCluster* Cluster,
+	const FDisplayClusterConfigurationRectangle& PresetRect,
+	TSharedPtr<FScopedTransaction>& OutTransaction,
+	FString PresetHost)
 {
-	UDisplayClusterConfigurationClusterNode* NodeTemplate = NewObject<UDisplayClusterConfigurationClusterNode>(Toolkit->GetBlueprintObj());
+	UDisplayClusterConfigurationClusterNode* NodeTemplate = NewObject<UDisplayClusterConfigurationClusterNode>(Toolkit->GetBlueprintObj(), NAME_None, RF_Transactional | RF_ArchetypeObject | RF_Public);
 	NodeTemplate->WindowRect = FDisplayClusterConfigurationRectangle(PresetRect);
 	NodeTemplate->Host = PresetHost;
 
@@ -46,7 +113,7 @@ UDisplayClusterConfigurationClusterNode* FDisplayClusterConfiguratorClusterUtils
 	
 	bool bAutoposition = true;
 	bool bAddViewport = true;
-	TSharedRef<SWidget> ClusterNodeFooter = SNew(SVerticalBox)
+	const TSharedRef<SWidget> ClusterNodeFooter = SNew(SVerticalBox)
 		+SVerticalBox::Slot()
 		.AutoHeight()
 		.Padding(2.0f)
@@ -75,7 +142,7 @@ UDisplayClusterConfigurationClusterNode* FDisplayClusterConfiguratorClusterUtils
 			]
 		];
 
-	TSharedRef<SDisplayClusterConfiguratorNewClusterItemDialog> DialogContent = SNew(SDisplayClusterConfiguratorNewClusterItemDialog, NodeTemplate)
+	const TSharedRef<SDisplayClusterConfiguratorNewClusterItemDialog> DialogContent = SNew(SDisplayClusterConfiguratorNewClusterItemDialog, NodeTemplate)
 		.ParentItemOptions(ParentItems)
 		.InitiallySelectedParentItem("Cluster")
 		.PresetItemOptions(FDisplayClusterConfiguratorPresetSize::CommonPresets)
@@ -86,20 +153,25 @@ UDisplayClusterConfigurationClusterNode* FDisplayClusterConfiguratorClusterUtils
 		.FooterContent(ClusterNodeFooter)
 		.OnPresetChanged_Lambda([=](FVector2D Size) { NodeTemplate->WindowRect.W = Size.X; NodeTemplate->WindowRect.H = Size.Y; });
 
-	ShowNewClusterItemDialogWindow(DialogContent, nullptr, LOCTEXT("AddNewClusterNode_DialogTitle", "Add New Cluster Node"), NewClusterItemDialogSize);
-
+	{
+		UE::DisplayClusterConfiguratorClusterUtils::FScopedTemporaryTransactions ScopedTransactionsForModal;
+		ShowNewClusterItemDialogWindow(DialogContent, nullptr, LOCTEXT("AddNewClusterNode_DialogTitle", "Add New Cluster Node"), NewClusterItemDialogSize);
+	}
+	
 	UDisplayClusterConfigurationClusterNode* NewNode = nullptr;
 
 	if (DialogContent->WasAccepted())
 	{
+		OutTransaction = MakeShared<FScopedTransaction>(LOCTEXT("AddClusterNode", "Add Cluster Node"));
+		
 		const FString ItemName = DialogContent->GetItemName();
 		NodeTemplate->SetFlags(RF_Transactional);
 
 		if (bAutoposition)
 		{
-			FVector2D DesiredPosition = FVector2D(NodeTemplate->WindowRect.X, NodeTemplate->WindowRect.Y);
-			FVector2D DesiredSize = FVector2D(NodeTemplate->WindowRect.W, NodeTemplate->WindowRect.H);
-			FVector2D NewPosition = FindNextAvailablePositionForClusterNode(Cluster, NodeTemplate->Host, DesiredPosition, DesiredSize);
+			const FVector2D DesiredPosition = FVector2D(NodeTemplate->WindowRect.X, NodeTemplate->WindowRect.Y);
+			const FVector2D DesiredSize = FVector2D(NodeTemplate->WindowRect.W, NodeTemplate->WindowRect.H);
+			const FVector2D NewPosition = FindNextAvailablePositionForClusterNode(Cluster, NodeTemplate->Host, DesiredPosition, DesiredSize);
 
 			NodeTemplate->WindowRect.X = NewPosition.X;
 			NodeTemplate->WindowRect.Y = NewPosition.Y;
@@ -131,7 +203,11 @@ UDisplayClusterConfigurationClusterNode* FDisplayClusterConfiguratorClusterUtils
 	return NewNode;
 }
 
-UDisplayClusterConfigurationViewport* FDisplayClusterConfiguratorClusterUtils::CreateNewViewportFromDialog(const TSharedRef<FDisplayClusterConfiguratorBlueprintEditor>& Toolkit, UDisplayClusterConfigurationClusterNode* ClusterNode, const FDisplayClusterConfigurationRectangle& PresetRect)
+UDisplayClusterConfigurationViewport* FDisplayClusterConfiguratorClusterUtils::CreateNewViewportFromDialog(
+	const TSharedRef<FDisplayClusterConfiguratorBlueprintEditor>& Toolkit,
+	UDisplayClusterConfigurationClusterNode* ClusterNode,
+	const FDisplayClusterConfigurationRectangle& PresetRect,
+	TSharedPtr<FScopedTransaction>& OutTransaction)
 {
 	UDisplayClusterConfigurationViewport* ViewportTemplate = NewObject<UDisplayClusterConfigurationViewport>(Toolkit->GetBlueprintObj(), NAME_None, RF_Transactional | RF_ArchetypeObject | RF_Public);
 	ViewportTemplate->Region = FDisplayClusterConfigurationRectangle(PresetRect);
@@ -196,12 +272,17 @@ UDisplayClusterConfigurationViewport* FDisplayClusterConfiguratorClusterUtils::C
 		.FooterContent(ViewportFooter)
 		.OnPresetChanged_Lambda([=](FVector2D Size) { ViewportTemplate->Region.W = Size.X; ViewportTemplate->Region.H = Size.Y; });
 
-	FDisplayClusterConfiguratorClusterUtils::ShowNewClusterItemDialogWindow(DialogContent, nullptr, LOCTEXT("AddNewViewport_DialogTitle", "Add New Viewport"), NewClusterItemDialogSize);
-
+	{
+		UE::DisplayClusterConfiguratorClusterUtils::FScopedTemporaryTransactions ScopedTransactionsForModal;
+		ShowNewClusterItemDialogWindow(DialogContent, nullptr, LOCTEXT("AddNewViewport_DialogTitle", "Add New Viewport"), NewClusterItemDialogSize);
+	}
+	
 	UDisplayClusterConfigurationViewport* NewViewport = nullptr;
 
 	if (DialogContent->WasAccepted())
 	{
+		OutTransaction = MakeShared<FScopedTransaction>(LOCTEXT("AddViewport", "Add Viewport"));
+		
 		const FString ParentItem = DialogContent->GetSelectedParentItem();
 		const FString ItemName = DialogContent->GetItemName();
 
diff --git a/Engine/Plugins/Runtime/nDisplay/Source/DisplayClusterConfigurator/Private/ClusterConfiguration/DisplayClusterConfiguratorClusterUtils.h b/Engine/Plugins/Runtime/nDisplay/Source/DisplayClusterConfigurator/Private/ClusterConfiguration/DisplayClusterConfiguratorClusterUtils.h
--- a/Engine/Plugins/Runtime/nDisplay/Source/DisplayClusterConfigurator/Private/ClusterConfiguration/DisplayClusterConfiguratorClusterUtils.h
+++ b/Engine/Plugins/Runtime/nDisplay/Source/DisplayClusterConfigurator/Private/ClusterConfiguration/DisplayClusterConfiguratorClusterUtils.h
@@ -8,6 +8,7 @@
 class SWidget;
 class FDragDropOperation;
 class FDisplayClusterConfiguratorBlueprintEditor;
+class FScopedTransaction;
 class UDisplayClusterConfigurationData_Base;
 class UDisplayClusterConfigurationCluster;
 class UDisplayClusterConfigurationHostDisplayData;
@@ -26,10 +27,13 @@ public:
 	 * @param Toolkit - The blueprint editor toolkit being used.
 	 * @param Cluster - The initial cluster the user wants to add the cluster node to.
 	 * @param PresetRect - The initial rectangle to configure the new cluster node with in the dialog box.
+	 * @param OutTransaction - The transaction created when adding a cluster node. May be null if the user canceled the operation.
 	 * @param PresetHost - Optional initial host string to configure the new cluster node with in the dialog box.
 	 * @return The newly created cluster node, or null if the user cancelled out of the dialog box.
 	 */
-	static UDisplayClusterConfigurationClusterNode* CreateNewClusterNodeFromDialog(const TSharedRef<FDisplayClusterConfiguratorBlueprintEditor>& Toolkit, UDisplayClusterConfigurationCluster* Cluster, const FDisplayClusterConfigurationRectangle& PresetRect,
+	static UDisplayClusterConfigurationClusterNode* CreateNewClusterNodeFromDialog(const TSharedRef<FDisplayClusterConfiguratorBlueprintEditor>& Toolkit,
+		UDisplayClusterConfigurationCluster* Cluster, const FDisplayClusterConfigurationRectangle& PresetRect,
+		TSharedPtr<FScopedTransaction>& OutTransaction,
 		FString PresetHost = NDISPLAY_DEFAULT_CLUSTER_HOST);
 
 	/**
@@ -37,9 +41,12 @@ public:
 	 * @param Toolkit - The blueprint editor toolkit being used.
 	 * @param ClusterNode - The initial cluster node the user wants to add the viewport to, or null to use the first cluster node in the cluster.
 	 * @param PresetRect - The initial rectangle to configure the new viewport with in the dialog box.
+	 * @param OutTransaction - The transaction created when adding a viewport. May be null if the user canceled the operation.
 	 * @return The newly created viewport, or null if the user cancelled out of the dialog box.
 	 */
-	static UDisplayClusterConfigurationViewport* CreateNewViewportFromDialog(const TSharedRef<FDisplayClusterConfiguratorBlueprintEditor>& Toolkit, UDisplayClusterConfigurationClusterNode* ClusterNode, const FDisplayClusterConfigurationRectangle& PresetRect);
+	static UDisplayClusterConfigurationViewport* CreateNewViewportFromDialog(const TSharedRef<FDisplayClusterConfiguratorBlueprintEditor>& Toolkit,
+		UDisplayClusterConfigurationClusterNode* ClusterNode, const FDisplayClusterConfigurationRectangle& PresetRect,
+		TSharedPtr<FScopedTransaction>& OutTransaction);
 
 	/**
 	 * Displays a dialog box used to configure new cluster item properties before adding the new item to the cluster configuration.
diff --git a/Engine/Plugins/Runtime/nDisplay/Source/DisplayClusterConfigurator/Private/Views/OutputMapping/DisplayClusterConfiguratorGraphSchema.cpp b/Engine/Plugins/Runtime/nDisplay/Source/DisplayClusterConfigurator/Private/Views/OutputMapping/DisplayClusterConfiguratorGraphSchema.cpp
--- a/Engine/Plugins/Runtime/nDisplay/Source/DisplayClusterConfigurator/Private/Views/OutputMapping/DisplayClusterConfiguratorGraphSchema.cpp
+++ b/Engine/Plugins/Runtime/nDisplay/Source/DisplayClusterConfigurator/Private/Views/OutputMapping/DisplayClusterConfiguratorGraphSchema.cpp
@@ -36,13 +36,11 @@ UEdGraphNode* FDisplayClusterConfiguratorSchemaAction_NewNode::PerformAction(UEd
 	switch (ItemType)
 	{
 	case EClusterItemType::ClusterNode:
-		Transaction = MakeShareable(new FScopedTransaction(LOCTEXT("AddClusterNode", "Add Cluster Node")));
-		CreatedObject = FDisplayClusterConfiguratorClusterUtils::CreateNewClusterNodeFromDialog(Toolkit.ToSharedRef(), Cluster, PresetRect);
+		CreatedObject = FDisplayClusterConfiguratorClusterUtils::CreateNewClusterNodeFromDialog(Toolkit.ToSharedRef(), Cluster, PresetRect, Transaction);
 		break;
 
 	case EClusterItemType::Viewport:
-		Transaction = MakeShareable(new FScopedTransaction(LOCTEXT("AddViewport", "Add Viewport")));
-		CreatedObject = FDisplayClusterConfiguratorClusterUtils::CreateNewViewportFromDialog(Toolkit.ToSharedRef(), nullptr, PresetRect);
+		CreatedObject = FDisplayClusterConfiguratorClusterUtils::CreateNewViewportFromDialog(Toolkit.ToSharedRef(), nullptr, PresetRect, Transaction);
 		break;
 	}
 
@@ -51,10 +49,6 @@ UEdGraphNode* FDisplayClusterConfiguratorSchemaAction_NewNode::PerformAction(UEd
 		Toolkit->GetEditorData()->MarkPackageDirty();
 		Toolkit->ClusterChanged();
 	}
-	else
-	{
-		Transaction->Cancel();
-	}
 
 	Transaction.Reset();
 
diff --git a/Engine/Plugins/Runtime/nDisplay/Source/DisplayClusterConfigurator/Private/Views/OutputMapping/SDisplayClusterConfiguratorGraphEditor.cpp b/Engine/Plugins/Runtime/nDisplay/Source/DisplayClusterConfigurator/Private/Views/OutputMapping/SDisplayClusterConfiguratorGraphEditor.cpp
--- a/Engine/Plugins/Runtime/nDisplay/Source/DisplayClusterConfigurator/Private/Views/OutputMapping/SDisplayClusterConfiguratorGraphEditor.cpp
+++ b/Engine/Plugins/Runtime/nDisplay/Source/DisplayClusterConfigurator/Private/Views/OutputMapping/SDisplayClusterConfiguratorGraphEditor.cpp
@@ -399,30 +399,26 @@ void SDisplayClusterConfiguratorGraphEditor::AddNewClusterNode()
 	const TSet<UObject*>& SelectedNodes = GetSelectedNodes();
 	if (SelectedNodes.Num() == 1)
 	{
-		FScopedTransaction Transaction(LOCTEXT("AddClusterNode", "Add Cluster Node"));
-
 		FString Host;
-		if (UDisplayClusterConfiguratorHostNode* HostNode = Cast<UDisplayClusterConfiguratorHostNode>(*SelectedNodes.begin()))
+		if (const UDisplayClusterConfiguratorHostNode* HostNode = Cast<UDisplayClusterConfiguratorHostNode>(*SelectedNodes.begin()))
 		{
 			Host = HostNode->GetNodeName();
 		}
 
-		TSharedPtr<FDisplayClusterConfiguratorBlueprintEditor> Toolkit = ToolkitPtr.Pin();
+		const TSharedPtr<FDisplayClusterConfiguratorBlueprintEditor> Toolkit = ToolkitPtr.Pin();
 		UDisplayClusterConfigurationCluster* Cluster = Toolkit->GetEditorData()->Cluster;
 		const FVector2D PresetSize = FDisplayClusterConfiguratorPresetSize::CommonPresets[FDisplayClusterConfiguratorPresetSize::DefaultPreset].Size;
 		const FDisplayClusterConfigurationRectangle PresetRect = FDisplayClusterConfigurationRectangle(0, 0, PresetSize.X, PresetSize.Y);
 
-		if (UDisplayClusterConfigurationClusterNode* NewNode = FDisplayClusterConfiguratorClusterUtils::CreateNewClusterNodeFromDialog(Toolkit.ToSharedRef(), Cluster, PresetRect, Host))
+		TSharedPtr<FScopedTransaction> Transaction;
+		if (UDisplayClusterConfigurationClusterNode* NewNode =
+			FDisplayClusterConfiguratorClusterUtils::CreateNewClusterNodeFromDialog(Toolkit.ToSharedRef(), Cluster, PresetRect,Transaction, Host))
 		{
 			// Mark the cluster configuration data as dirty, allowing user to save the changes, and fire off a cluster changed event to let other
 			// parts of the UI update as well
 			Toolkit->GetEditorData()->MarkPackageDirty();
 			Toolkit->ClusterChanged();
 		}
-		else
-		{
-			Transaction.Cancel();
-		}
 	}
 }
 
@@ -443,29 +439,24 @@ void SDisplayClusterConfiguratorGraphEditor::AddNewViewport()
 	const TSet<UObject*>& SelectedNodes = GetSelectedNodes();
 	if (SelectedNodes.Num() == 1)
 	{
-		FScopedTransaction Transaction(LOCTEXT("AddViewport", "Add Viewport"));
-
 		UDisplayClusterConfigurationClusterNode* SelectedClusterNode = nullptr;
-		if (UDisplayClusterConfiguratorWindowNode* WindowNode = Cast<UDisplayClusterConfiguratorWindowNode>(*SelectedNodes.begin()))
+		if (const UDisplayClusterConfiguratorWindowNode* WindowNode = Cast<UDisplayClusterConfiguratorWindowNode>(*SelectedNodes.begin()))
 		{
 			SelectedClusterNode = Cast<UDisplayClusterConfigurationClusterNode>(WindowNode->GetObject());
 		}
 
-		TSharedPtr<FDisplayClusterConfiguratorBlueprintEditor> Toolkit = ToolkitPtr.Pin();
+		const TSharedPtr<FDisplayClusterConfiguratorBlueprintEditor> Toolkit = ToolkitPtr.Pin();
 		const FVector2D PresetSize = FDisplayClusterConfiguratorPresetSize::CommonPresets[FDisplayClusterConfiguratorPresetSize::DefaultPreset].Size;
 		const FDisplayClusterConfigurationRectangle PresetRect = FDisplayClusterConfigurationRectangle(0, 0, PresetSize.X, PresetSize.Y);
 
-		if (UDisplayClusterConfigurationViewport* NewViewport = FDisplayClusterConfiguratorClusterUtils::CreateNewViewportFromDialog(Toolkit.ToSharedRef(), SelectedClusterNode, PresetRect))
+		TSharedPtr<FScopedTransaction> Transaction;
+		if (UDisplayClusterConfigurationViewport* NewViewport = FDisplayClusterConfiguratorClusterUtils::CreateNewViewportFromDialog(Toolkit.ToSharedRef(), SelectedClusterNode, PresetRect, Transaction))
 		{
 			// Mark the cluster configuration data as dirty, allowing user to save the changes, and fire off a cluster changed event to let other
 			// parts of the UI update as well
 			Toolkit->GetEditorData()->MarkPackageDirty();
 			Toolkit->ClusterChanged();
 		}
-		else
-		{
-			Transaction.Cancel();
-		}
 	}
 }
 
diff --git a/Engine/Plugins/Runtime/nDisplay/Source/DisplayClusterConfigurator/Private/Views/TreeViews/Cluster/DisplayClusterConfiguratorViewCluster.cpp b/Engine/Plugins/Runtime/nDisplay/Source/DisplayClusterConfigurator/Private/Views/TreeViews/Cluster/DisplayClusterConfiguratorViewCluster.cpp
--- a/Engine/Plugins/Runtime/nDisplay/Source/DisplayClusterConfigurator/Private/Views/TreeViews/Cluster/DisplayClusterConfiguratorViewCluster.cpp
+++ b/Engine/Plugins/Runtime/nDisplay/Source/DisplayClusterConfigurator/Private/Views/TreeViews/Cluster/DisplayClusterConfiguratorViewCluster.cpp
@@ -560,11 +560,9 @@ bool FDisplayClusterConfiguratorViewCluster::CanSetAsPrimary() const
 
 void FDisplayClusterConfiguratorViewCluster::AddNewClusterNode(FVector2D PresetSize)
 {
-	FScopedTransaction Transaction(LOCTEXT("AddClusterNode", "Add Cluster Node"));
-
-	TSharedPtr<FDisplayClusterConfiguratorBlueprintEditor> Toolkit = ToolkitPtr.Pin();
+	const TSharedPtr<FDisplayClusterConfiguratorBlueprintEditor> Toolkit = ToolkitPtr.Pin();
 	UDisplayClusterConfigurationCluster* Cluster = Toolkit->GetEditorData()->Cluster;
-	FDisplayClusterConfigurationRectangle PresetRect = FDisplayClusterConfigurationRectangle(0, 0, PresetSize.X, PresetSize.Y);
+	const FDisplayClusterConfigurationRectangle PresetRect = FDisplayClusterConfigurationRectangle(0, 0, PresetSize.X, PresetSize.Y);
 	
 	FString HostAddress = NDISPLAY_DEFAULT_CLUSTER_HOST;
 	TArray<TSharedPtr<IDisplayClusterConfiguratorTreeItem>> SelectedTreeItems = ViewTree->GetSelectedItems();
@@ -577,7 +575,8 @@ void FDisplayClusterConfiguratorViewCluster::AddNewClusterNode(FVector2D PresetS
 		}
 	}
 
-	if (UDisplayClusterConfigurationClusterNode* NewNode = FDisplayClusterConfiguratorClusterUtils::CreateNewClusterNodeFromDialog(Toolkit.ToSharedRef(), Cluster, PresetRect, HostAddress))
+	TSharedPtr<FScopedTransaction> Transaction;
+	if (UDisplayClusterConfigurationClusterNode* NewNode = FDisplayClusterConfiguratorClusterUtils::CreateNewClusterNodeFromDialog(Toolkit.ToSharedRef(), Cluster, PresetRect, Transaction, HostAddress))
 	{
 		// Mark the cluster configuration data as dirty, allowing user to save the changes, and fire off a cluster changed event to let other
 		// parts of the UI update as well
@@ -602,10 +601,6 @@ void FDisplayClusterConfiguratorViewCluster::AddNewClusterNode(FVector2D PresetS
 
 		SetSelectedItems(TreeItemsToSelect);
 	}
-	else
-	{
-		Transaction.Cancel();
-	}
 }
 
 bool FDisplayClusterConfiguratorViewCluster::CanAddNewClusterNode() const
@@ -615,8 +610,6 @@ bool FDisplayClusterConfiguratorViewCluster::CanAddNewClusterNode() const
 
 void FDisplayClusterConfiguratorViewCluster::AddNewViewport(FVector2D PresetSize)
 {
-	FScopedTransaction Transaction(LOCTEXT("AddViewport", "Add Viewport"));
-
 	UDisplayClusterConfigurationClusterNode* SelectedClusterNode = nullptr;
 	TArray<TSharedPtr<IDisplayClusterConfiguratorTreeItem>> SelectedTreeItems = ViewTree->GetSelectedItems();
 	if (SelectedTreeItems.Num() == 1)
@@ -632,9 +625,10 @@ void FDisplayClusterConfiguratorViewCluster::AddNewViewport(FVector2D PresetSize
 		}
 	}
 
-	TSharedPtr<FDisplayClusterConfiguratorBlueprintEditor> Toolkit = ToolkitPtr.Pin();
-	FDisplayClusterConfigurationRectangle PresetRect = FDisplayClusterConfigurationRectangle(0, 0, PresetSize.X, PresetSize.Y);
-	if (UDisplayClusterConfigurationViewport* NewViewport = FDisplayClusterConfiguratorClusterUtils::CreateNewViewportFromDialog(Toolkit.ToSharedRef(), SelectedClusterNode, PresetRect))
+	const TSharedPtr<FDisplayClusterConfiguratorBlueprintEditor> Toolkit = ToolkitPtr.Pin();
+	const FDisplayClusterConfigurationRectangle PresetRect = FDisplayClusterConfigurationRectangle(0, 0, PresetSize.X, PresetSize.Y);
+	TSharedPtr<FScopedTransaction> Transaction;
+	if (const UDisplayClusterConfigurationViewport* NewViewport = FDisplayClusterConfiguratorClusterUtils::CreateNewViewportFromDialog(Toolkit.ToSharedRef(), SelectedClusterNode, PresetRect, Transaction))
 	{
 		// Mark the cluster configuration data as dirty, allowing user to save the changes, and fire off a cluster changed event to let other
 		// parts of the UI update as well
@@ -659,10 +653,6 @@ void FDisplayClusterConfiguratorViewCluster::AddNewViewport(FVector2D PresetSize
 
 		SetSelectedItems(TreeItemsToSelect);
 	}
-	else
-	{
-		Transaction.Cancel();
-	}
 }
 
 bool FDisplayClusterConfiguratorViewCluster::CanAddNewViewport() const
