nDisplay: Fix non-unity build.
CIS issue 339687.

diff --git a/Engine/Plugins/Runtime/nDisplay/Source/DisplayClusterConfigurator/Private/ClusterConfiguration/DisplayClusterConfiguratorClusterUtils.cpp b/Engine/Plugins/Runtime/nDisplay/Source/DisplayClusterConfigurator/Private/ClusterConfiguration/DisplayClusterConfiguratorClusterUtils.cpp
--- a/Engine/Plugins/Runtime/nDisplay/Source/DisplayClusterConfigurator/Private/ClusterConfiguration/DisplayClusterConfiguratorClusterUtils.cpp
+++ b/Engine/Plugins/Runtime/nDisplay/Source/DisplayClusterConfigurator/Private/ClusterConfiguration/DisplayClusterConfiguratorClusterUtils.cpp
@@ -14,6 +14,7 @@
 
 #include "Factories.h"
 #include "ISinglePropertyView.h"
+#include "ScopedTransaction.h"
 #include "ObjectTools.h"
 #include "Framework/Application/SlateApplication.h"
 #include "HAL/PlatformApplicationMisc.h"
