[nDisplay] -Techviz- Add display device component which stores the preview and mesh materials used by the preview component. Display devices can be assigned to viewports, and the preview component will utilize the display device when rendering the preview.

diff --git a/Engine/Plugins/Runtime/nDisplay/Source/DisplayCluster/Private/Game/EngineClasses/Scene/DisplayClusterPreviewComponent.cpp b/Engine/Plugins/Runtime/nDisplay/Source/DisplayCluster/Private/Game/EngineClasses/Scene/DisplayClusterPreviewComponent.cpp
--- a/Engine/Plugins/Runtime/nDisplay/Source/DisplayCluster/Private/Game/EngineClasses/Scene/DisplayClusterPreviewComponent.cpp
+++ b/Engine/Plugins/Runtime/nDisplay/Source/DisplayCluster/Private/Game/EngineClasses/Scene/DisplayClusterPreviewComponent.cpp
@@ -2,9 +2,11 @@
 
 #include "Components/DisplayClusterPreviewComponent.h"
 #include "Components/DisplayClusterCameraComponent.h"
+#include "Components/DisplayDevice/DisplayClusterDisplayDeviceBaseComponent.h"
 
 #include "IDisplayCluster.h"
 #include "Render/IDisplayClusterRenderManager.h"
+#include "DisplayDevice/DisplayClusterDisplayDeviceUtils.h"
 #include "DisplayClusterRootActor.h"
 #include "DisplayClusterConfigurationTypes.h"
 
@@ -40,11 +42,6 @@ UDisplayClusterPreviewComponent::UDisplayClusterPreviewComponent(const FObjectIn
 	: Super(ObjectInitializer)
 {
 #if WITH_EDITOR
-	static ConstructorHelpers::FObjectFinder<UMaterial> PreviewMaterialObj(TEXT("/nDisplay/Materials/Preview/M_ProjPolicyPreview"));
-	check(PreviewMaterialObj.Object);
-
-	PreviewMaterial = PreviewMaterialObj.Object;
-
 	bWantsInitializeComponent = true;
 #endif
 }
@@ -102,6 +99,32 @@ bool UDisplayClusterPreviewComponent::InitializePreviewComponent(ADisplayCluster
 	return true;
 }
 
+UDisplayClusterDisplayDeviceBaseComponent* UDisplayClusterPreviewComponent::GetDisplayDevice() const
+{
+	UDisplayClusterDisplayDeviceBaseComponent* DeviceBaseComponent = nullptr;
+	if (IsValid(ViewportConfig) && RootActor)
+	{
+		UDisplayClusterDisplayDeviceBaseComponent* CachedComponent =
+			Cast<UDisplayClusterDisplayDeviceBaseComponent>(CachedDisplayDevice.GetComponent(GetOwner()));
+		if (CachedComponent && CachedComponent->GetName() == ViewportConfig->DisplayDeviceName)
+		{
+			DeviceBaseComponent = CachedComponent;
+		}
+		else
+		{
+			DeviceBaseComponent = UE::DisplayClusterDisplayDeviceUtils::FindAndSyncDisplayDeviceFromViewport(ViewportConfig);
+		}
+	}
+
+	if (bUseDisplayDevice)
+	{
+		CachedDisplayDevice.OverrideComponent = DeviceBaseComponent;
+		CachedDisplayDevice.ComponentProperty = DeviceBaseComponent ? DeviceBaseComponent->GetFName() : NAME_None;
+	}
+	
+	return DeviceBaseComponent;
+}
+
 bool UDisplayClusterPreviewComponent::IsPreviewEnabled() const
 {
 	return ViewportConfig && RootActor && RootActor->IsPreviewEnabled();
@@ -116,11 +139,12 @@ void UDisplayClusterPreviewComponent::RestorePreviewMeshMaterial()
 {
 	UpdatePreviewMeshReference();
 
-	if (PreviewMesh && OriginalMaterial)
+	if (PreviewMesh)
 	{
 		// Restore
-		PreviewMesh->SetMaterial(0, OriginalMaterial);
-		OriginalMaterial = nullptr;
+		CurrentMeshMaterial = GetMeshMaterialFromDisplayDevice();
+		PreviewMesh->SetMaterial(0, CurrentMeshMaterial);
+		PreviewMaterialInstance = nullptr;
 	}
 
 	// Release RTTs
@@ -133,23 +157,18 @@ void UDisplayClusterPreviewComponent::RestorePreviewMeshMaterial()
 	}
 }
 
-void UDisplayClusterPreviewComponent::SetPreviewMeshMaterial()
+void UDisplayClusterPreviewComponent::SetPreviewMeshMaterial(UMaterial* InMaterial)
 {
 	UpdatePreviewMeshReference();
 
 	if (PreviewMesh)
 	{
-		// Save original material
-		if (OriginalMaterial == nullptr)
+		CurrentMeshMaterial = InMaterial;
+
+		if (InMaterial != nullptr && PreviewMaterialInstance == nullptr)
 		{
-			UMaterialInterface* MatInterface = PreviewMesh->GetMaterial(0);
-			if (MatInterface)
-			{
-				OriginalMaterial = MatInterface->GetMaterial();
-			}
+			PreviewMaterialInstance = UMaterialInstanceDynamic::Create(InMaterial, this);
 		}
-
-		InitializePreviewMaterial();
 		UpdatePreviewMaterial();
 
 		// Set preview material
@@ -160,6 +179,18 @@ void UDisplayClusterPreviewComponent::SetPreviewMeshMaterial()
 	}
 }
 
+UMaterial* UDisplayClusterPreviewComponent::GetPreviewMaterialFromDisplayDevice() const
+{
+	const UDisplayClusterDisplayDeviceBaseComponent* DeviceBaseComponent = GetDisplayDevice();
+	return DeviceBaseComponent ? DeviceBaseComponent->GetPreviewMaterial() : nullptr;
+}
+
+UMaterial* UDisplayClusterPreviewComponent::GetMeshMaterialFromDisplayDevice() const
+{
+	const UDisplayClusterDisplayDeviceBaseComponent* DeviceBaseComponent = GetDisplayDevice();
+	return DeviceBaseComponent ? DeviceBaseComponent->GetMeshMaterial() : nullptr;
+}
+
 void UDisplayClusterPreviewComponent::UpdatePreviewMeshReference()
 {
 	if (PreviewMesh && PreviewMesh->GetName().Find(TEXT("TRASH_")) != INDEX_NONE)
@@ -223,10 +254,11 @@ bool UDisplayClusterPreviewComponent::UpdatePreviewMesh()
 					PreviewMesh->SetCastShadow(false);
 				}
 
-				if (OriginalMaterial == nullptr)
+				UMaterial* PreviewMaterial = GetPreviewMaterialFromDisplayDevice();
+				if (CurrentMeshMaterial != PreviewMaterial)
 				{
 					// Assign preview material to mesh
-					SetPreviewMeshMaterial();
+					SetPreviewMeshMaterial(PreviewMaterial);
 				}
 
 				return true;
@@ -251,7 +283,7 @@ void UDisplayClusterPreviewComponent::ReleasePreviewMesh()
 {
 	// Forget old mesh with material
 	PreviewMesh = nullptr;
-	OriginalMaterial = nullptr;
+	CurrentMeshMaterial = nullptr;
 }
 
 void UDisplayClusterPreviewComponent::UpdatePreviewResources()
@@ -269,8 +301,7 @@ void UDisplayClusterPreviewComponent::UpdatePreviewMaterial()
 {
 	if (PreviewMaterialInstance != nullptr)
 	{
-		PreviewMaterialInstance->SetScalarParameterValue(TEXT("Opacity"), 1.0);
-
+		// TODO: Consider moving these parameter values to the display device component.
 		if (OverrideTexture)
 		{
 			PreviewMaterialInstance->SetTextureParameterValue(TEXT("Preview"), OverrideTexture);
@@ -280,14 +311,13 @@ void UDisplayClusterPreviewComponent::UpdatePreviewMaterial()
 			PreviewMaterialInstance->SetTextureParameterValue(TEXT("Preview"),
 				RootActor && RootActor->bPreviewEnablePostProcess ? RenderTargetPostProcess : RenderTarget);
 		}
-	}
-}
 
-void UDisplayClusterPreviewComponent::InitializePreviewMaterial()
-{
-	if (PreviewMaterial != nullptr && PreviewMaterialInstance == nullptr)
-	{
-		PreviewMaterialInstance = UMaterialInstanceDynamic::Create(PreviewMaterial, this);
+		// Allow display device to perform any processing on the material instance.
+		if (UDisplayClusterDisplayDeviceBaseComponent* CachedComponent =
+			Cast<UDisplayClusterDisplayDeviceBaseComponent>(CachedDisplayDevice.GetComponent(GetOwner())))
+		{
+			CachedComponent->OnUpdatePreviewMaterialInstance(PreviewMaterialInstance);
+		}
 	}
 }
 
@@ -455,4 +485,11 @@ void UDisplayClusterPreviewComponent::SetOverrideTexture(UTexture* InOverrideTex
 	}
 }
 
+void UDisplayClusterPreviewComponent::SetUseDisplayDevice(bool bInNewValue)
+{
+	bUseDisplayDevice = bInNewValue;
+	CachedDisplayDevice.OverrideComponent.Reset();
+	CachedDisplayDevice.ComponentProperty = NAME_None;
+}
+
 #endif
diff --git a/Engine/Plugins/Runtime/nDisplay/Source/DisplayCluster/Private/Game/EngineClasses/Scene/DisplayClusterRootActorEditor.cpp b/Engine/Plugins/Runtime/nDisplay/Source/DisplayCluster/Private/Game/EngineClasses/Scene/DisplayClusterRootActorEditor.cpp
--- a/Engine/Plugins/Runtime/nDisplay/Source/DisplayCluster/Private/Game/EngineClasses/Scene/DisplayClusterRootActorEditor.cpp
+++ b/Engine/Plugins/Runtime/nDisplay/Source/DisplayCluster/Private/Game/EngineClasses/Scene/DisplayClusterRootActorEditor.cpp
@@ -3,26 +3,31 @@
 #include "DisplayClusterRootActor.h"
 
 #include "Blueprints/DisplayClusterBlueprint.h"
-#include "Components/SceneComponent.h"
-#include "Components/DisplayClusterOriginComponent.h"
 #include "Components/DisplayClusterCameraComponent.h"
-#include "Components/DisplayClusterScreenComponent.h"
-#include "Components/DisplayClusterSceneComponentSyncParent.h"
-#include "Components/DisplayClusterPreviewComponent.h"
 #include "Components/DisplayClusterICVFXCameraComponent.h"
+#include "Components/DisplayClusterOriginComponent.h"
+#include "Components/DisplayClusterPreviewComponent.h"
+#include "Components/DisplayClusterSceneComponentSyncParent.h"
+#include "Components/DisplayClusterScreenComponent.h"
+#include "Components/DisplayClusterStageGeometryComponent.h"
+#include "Components/DisplayDevice/DisplayClusterDisplayDeviceBaseComponent.h"
 #include "Components/LineBatchComponent.h"
+#include "Components/MeshComponent.h"
+#include "Components/SceneComponent.h"
+#include "DisplayDevice/DisplayClusterDisplayDeviceUtils.h"
 
 #include "Config/IPDisplayClusterConfigManager.h"
 #include "DisplayClusterConfigurationStrings.h"
 
-#include "IDisplayClusterConfiguration.h"
 #include "DisplayClusterConfigurationTypes.h"
+#include "DisplayClusterProjectionStrings.h"
+#include "IDisplayClusterConfiguration.h"
 
 #include "DisplayClusterPlayerInput.h"
 
-#include "GameFramework/PlayerController.h"
 #include "Engine/Engine.h"
 #include "Engine/World.h"
+#include "GameFramework/PlayerController.h"
 
 #include "Misc/DisplayClusterGlobals.h"
 #include "Misc/DisplayClusterHelpers.h"
@@ -32,18 +37,18 @@
 #include "Misc/TransactionObjectEvent.h"
 #include "Render/Viewport/Configuration/DisplayClusterViewportConfigurationHelpers_ICVFX.h"
 #include "Render/Viewport/DisplayClusterViewportStrings.h"
-#include "Render/Viewport/IDisplayClusterViewportManager.h"
 #include "Render/Viewport/IDisplayClusterViewport.h"
+#include "Render/Viewport/IDisplayClusterViewportManager.h"
 #include "Render/Viewport/IDisplayClusterViewportProxy.h"
 #include "Render/Viewport/RenderFrame/DisplayClusterRenderFrame.h"
 #include "Render/Viewport/RenderFrame/DisplayClusterRenderFrameSettings.h"
 
 #include "Engine/TextureRenderTarget2D.h"
 
-#include "Render/Viewport/IDisplayClusterViewport.h"
+#include "Materials/Material.h"
 #include "Render/Viewport/DisplayClusterViewportManager.h"
+#include "Render/Viewport/IDisplayClusterViewport.h"
 #include "TextureResource.h"
-#include "Components/DisplayClusterStageGeometryComponent.h"
 
 //////////////////////////////////////////////////////////////////////////////////////////////
 // IN-EDITOR STUFF
@@ -52,8 +57,8 @@
 #if WITH_EDITOR
 
 #include "Async/Async.h"
-#include "LevelEditor.h"
 #include "EditorSupportDelegates.h"
+#include "LevelEditor.h"
 
 //////////////////////////////////////////////////////////////////////////////////////////////
 int32 GDisplayClusterPreviewAllowMultiGPURendering = 0;
@@ -101,6 +106,12 @@ void ADisplayClusterRootActor::Constructor_Editor()
 	// Allow tick in editor for preview rendering
 	PrimaryActorTick.bStartWithTickEnabled = true;
 
+	// Set to internal default, user may change
+	DefaultDisplayDeviceName = GetInternalDisplayDeviceName();
+
+	// Our internal display device which always exists
+	BasicDisplayDeviceComponent = CreateDefaultSubobject<UDisplayClusterDisplayDeviceComponent>(GetInternalDisplayDeviceName());
+
 	ResetPreviewInternals_Editor();
 
 	FCoreUObjectDelegates::OnPackageReloaded.AddUObject(this, &ADisplayClusterRootActor::HandleAssetReload);
@@ -880,7 +891,8 @@ void ADisplayClusterRootActor::PostEditChangeProperty(FPropertyChangedEvent& Pro
 	{
 		ResetInnerFrustumPriority();
 	}
-	else if (PropertyName == GET_MEMBER_NAME_CHECKED(ADisplayClusterRootActor, bPreviewEnablePostProcess))
+	else if (PropertyName == GET_MEMBER_NAME_CHECKED(ADisplayClusterRootActor, bPreviewEnablePostProcess)
+		|| PropertyName == GET_MEMBER_NAME_CHECKED(ADisplayClusterRootActor, DefaultDisplayDeviceName))
 	{
 		ResetPreviewComponents_Editor(false);
 		PreviewRenderFrame.Reset();
@@ -932,10 +944,51 @@ void ADisplayClusterRootActor::ResetPreviewComponents_Editor(bool bInRestoreScen
 {
 	TArray<UDisplayClusterPreviewComponent*> AllPreviewComponents;
 	GetComponents(AllPreviewComponents);
-
-	for (UDisplayClusterPreviewComponent* ExistingComp : AllPreviewComponents)
+	
+	if (AllPreviewComponents.Num() > 0)
+	{
+		// Preview components exist and can restore the scene material
+		for (UDisplayClusterPreviewComponent* ExistingComp : AllPreviewComponents)
+		{
+			ExistingComp->ResetPreviewComponent(bInRestoreSceneMaterial);
+		}
+	}
+	else if (bInRestoreSceneMaterial && !bPreviewEnable && CurrentConfigData && !IsTemplate())
 	{
-		ExistingComp->ResetPreviewComponent(bInRestoreSceneMaterial);
+		// There are no preview components but the preview is disabled and we want to restore the scene material
+		CurrentConfigData->ForEachViewport([this](const TObjectPtr<UDisplayClusterConfigurationViewport>& Viewport)
+		{
+			// First locate the display device the viewport references
+			if (const UDisplayClusterDisplayDeviceBaseComponent* DisplayDevice = UE::DisplayClusterDisplayDeviceUtils::FindAndSyncDisplayDeviceFromViewport(Viewport))
+			{
+				// Next find the preview mesh component the viewport is assigned to. We can't just use FindPreviewViewport and its referenced mesh
+				// because the preview viewport won't exist if the editor preview is disabled when the root actor is initialized
+				FString ParameterKey;
+				if (Viewport->ProjectionPolicy.Type == DisplayClusterProjectionStrings::projection::Simple)
+				{
+					ParameterKey = DisplayClusterProjectionStrings::cfg::simple::Screen;
+				}
+				else if (Viewport->ProjectionPolicy.Type == DisplayClusterProjectionStrings::projection::Camera)
+				{
+					ParameterKey = DisplayClusterProjectionStrings::cfg::camera::Component;
+				}
+				else if (Viewport->ProjectionPolicy.Type == DisplayClusterProjectionStrings::projection::Mesh)
+				{
+					ParameterKey = DisplayClusterProjectionStrings::cfg::mesh::Component;
+				}
+
+				if (!ParameterKey.IsEmpty())
+				{
+					if (const FString* ComponentName = Viewport->ProjectionPolicy.Parameters.Find(ParameterKey))
+					{
+						if (UMeshComponent* MeshComponent = GetComponentByName<UMeshComponent>(*ComponentName))
+						{
+							MeshComponent->SetMaterial(0, DisplayDevice->GetMeshMaterial());
+						}
+					}
+				}
+			}
+		});
 	}
 }
 
@@ -1093,6 +1146,45 @@ void ADisplayClusterRootActor::RemovePreviewEnableOverride(const uint8* Object)
 	UpdatePreviewComponents();
 }
 
+FName ADisplayClusterRootActor::GetInternalDisplayDeviceName()
+{
+	return FName("BasicDisplayDevice");
+}
+
+UDisplayClusterDisplayDeviceBaseComponent* ADisplayClusterRootActor::GetDefaultDisplayDevice()
+{
+	if (DefaultDisplayDeviceComponent)
+	{
+		// Check already assigned/created
+		if (DefaultDisplayDeviceComponent->GetFName() == DefaultDisplayDeviceName)
+		{
+			return DefaultDisplayDeviceComponent;
+		}
+	}
+
+	DefaultDisplayDeviceComponent = nullptr;
+
+	// User assigned default
+	if (!DefaultDisplayDeviceName.IsNone() && DefaultDisplayDeviceName != GetInternalDisplayDeviceName())
+	{
+		DefaultDisplayDeviceComponent =
+			GetComponentByName<UDisplayClusterDisplayDeviceBaseComponent>(DefaultDisplayDeviceName.ToString());
+
+		if (!DefaultDisplayDeviceComponent)
+		{
+			UE_LOG(LogDisplayClusterGame, Warning, TEXT("Invalid default display device. Using internal nDisplay default device."));
+		}
+	}
+
+	// Fallback to our internal default
+	if (!DefaultDisplayDeviceComponent)
+	{
+		DefaultDisplayDeviceComponent = BasicDisplayDeviceComponent;
+	}
+
+	return DefaultDisplayDeviceComponent;
+}
+
 FString ADisplayClusterRootActor::GeneratePreviewComponentName_Editor(const FString& NodeId, const FString& ViewportId) const
 {
 	return FString::Printf(TEXT("%s_%s"), *NodeId, *ViewportId);
diff --git a/Engine/Plugins/Runtime/nDisplay/Source/DisplayCluster/Private/Game/EngineClasses/Scene/DisplayDevice/DisplayClusterDisplayDeviceComponentBase.cpp b/Engine/Plugins/Runtime/nDisplay/Source/DisplayCluster/Private/Game/EngineClasses/Scene/DisplayDevice/DisplayClusterDisplayDeviceComponentBase.cpp
new file mode 100644
--- /dev/null
+++ b/Engine/Plugins/Runtime/nDisplay/Source/DisplayCluster/Private/Game/EngineClasses/Scene/DisplayDevice/DisplayClusterDisplayDeviceComponentBase.cpp
@@ -0,0 +1,54 @@
+// Copyright Epic Games, Inc. All Rights Reserved.
+
+#include "Components/DisplayDevice/DisplayClusterDisplayDeviceBaseComponent.h"
+
+#include "DisplayClusterRootActor.h"
+
+#include "Materials/Material.h"
+#include "Materials/MaterialInstanceDynamic.h"
+#include "UObject/ConstructorHelpers.h"
+
+inline UDisplayClusterDisplayDeviceBaseComponent::UDisplayClusterDisplayDeviceBaseComponent()
+{
+#if WITH_EDITOR
+	static ConstructorHelpers::FObjectFinder<UMaterial> PreviewMaterialObj(TEXT("/nDisplay/Materials/Preview/M_DisplayDevicePreview"));
+	check(PreviewMaterialObj.Object);
+	PreviewMaterial = PreviewMaterialObj.Object;
+
+	static ConstructorHelpers::FObjectFinder<UMaterial> MeshMaterialObj(TEXT("/nDisplay/Materials/Preview/M_DisplayDeviceMesh"));
+	check(MeshMaterialObj.Object);
+	MeshMaterial = MeshMaterialObj.Object;
+#endif
+}
+
+#if WITH_EDITOR
+void UDisplayClusterDisplayDeviceBaseComponent::PostEditChangeProperty(FPropertyChangedEvent& PropertyChangedEvent)
+{
+	Super::PostEditChangeProperty(PropertyChangedEvent);
+
+	const FName PropertyName = PropertyChangedEvent.GetPropertyName();
+	
+	if (PropertyName == GET_MEMBER_NAME_CHECKED(UDisplayClusterDisplayDeviceBaseComponent, PreviewMaterial)
+		|| PropertyName == GET_MEMBER_NAME_CHECKED(UDisplayClusterDisplayDeviceBaseComponent, MeshMaterial))
+	{
+		if (ADisplayClusterRootActor* RootActor = Cast<ADisplayClusterRootActor>(GetOwner()))
+		{
+			RootActor->ResetPreviewComponents_Editor(true);
+		}
+	}
+}
+#endif
+
+UDisplayClusterDisplayDeviceComponent::UDisplayClusterDisplayDeviceComponent()
+{
+}
+
+void UDisplayClusterDisplayDeviceComponent::OnUpdatePreviewMaterialInstance(UMaterialInstanceDynamic* InMaterialInstance)
+{
+	Super::OnUpdatePreviewMaterialInstance(InMaterialInstance);
+
+	if (InMaterialInstance)
+	{
+		InMaterialInstance->SetScalarParameterValue(TEXT("Exposure"), Exposure);
+	}
+}
diff --git a/Engine/Plugins/Runtime/nDisplay/Source/DisplayCluster/Private/Game/EngineClasses/Scene/DisplayDevice/DisplayClusterDisplayDeviceUtils.cpp b/Engine/Plugins/Runtime/nDisplay/Source/DisplayCluster/Private/Game/EngineClasses/Scene/DisplayDevice/DisplayClusterDisplayDeviceUtils.cpp
new file mode 100644
--- /dev/null
+++ b/Engine/Plugins/Runtime/nDisplay/Source/DisplayCluster/Private/Game/EngineClasses/Scene/DisplayDevice/DisplayClusterDisplayDeviceUtils.cpp
@@ -0,0 +1,91 @@
+// Copyright Epic Games, Inc. All Rights Reserved.
+
+#include "DisplayClusterDisplayDeviceUtils.h"
+
+#include "Components/DisplayDevice/DisplayClusterDisplayDeviceBaseComponent.h"
+#include "DisplayClusterConfigurationTypes_Viewport.h"
+#include "DisplayClusterRootActor.h"
+#include "Misc/DisplayClusterLog.h"
+
+UDisplayClusterDisplayDeviceBaseComponent* UE::DisplayClusterDisplayDeviceUtils::FindAndSyncDisplayDeviceFromViewport(
+	UDisplayClusterConfigurationViewport* InViewport)
+{
+	if (IsValid(InViewport))
+	{
+		if (ADisplayClusterRootActor* RootActor = InViewport->GetTypedOuter<ADisplayClusterRootActor>())
+		{
+			UDisplayClusterDisplayDeviceBaseComponent* DeviceBaseComponent = nullptr;
+
+			const bool bUseDefaultDevice = InViewport->DisplayDeviceName.IsEmpty();
+			if (bUseDefaultDevice)
+			{
+				DeviceBaseComponent = RootActor->GetDefaultDisplayDevice();
+			}
+			else
+			{
+				// Find the component from the device name
+				TArray<UDisplayClusterDisplayDeviceBaseComponent*> Components;
+				RootActor->GetComponents<UDisplayClusterDisplayDeviceBaseComponent>(Components);
+
+				if (UDisplayClusterDisplayDeviceBaseComponent** FoundComponent =
+					Components.FindByPredicate(
+						[&](const UDisplayClusterDisplayDeviceBaseComponent* Component)
+						{
+							return Component->GetName() == InViewport->DisplayDeviceName;
+						}))
+				{
+					DeviceBaseComponent = *FoundComponent;
+				}
+			}
+
+			if (!DeviceBaseComponent && !bUseDefaultDevice)
+			{
+				// Could not locate device on the instance, attempt reset to CDO value first in case there is a desync.
+				// This might be needed if a display device was renamed, the blueprint compiled, then the action undone.
+				if (UDisplayClusterConfigurationClusterNode* ClusterNodeInstance = Cast<UDisplayClusterConfigurationClusterNode>(InViewport->GetOuter()))
+				{
+					if (const FString* ViewportKeyPtr = ClusterNodeInstance->Viewports.FindKey(InViewport))
+					{
+						const FString ViewportId = *ViewportKeyPtr;
+						if (const UDisplayClusterConfigurationCluster* ClusterNodeInstanceParent = Cast<UDisplayClusterConfigurationCluster>(ClusterNodeInstance->GetOuter()))
+						{
+							// Find the cluster key on the instance
+							if (const FString* KeyPtr = ClusterNodeInstanceParent->Nodes.FindKey(ClusterNodeInstance))
+							{
+								const FString ClusterName = *KeyPtr;
+								// Use the instance cluster key to find the correct CDO viewport
+								const ADisplayClusterRootActor* RootActorCDO = CastChecked<ADisplayClusterRootActor>(RootActor->GetClass()->ClassDefaultObject);
+								if (const UDisplayClusterConfigurationData* ConfigDataCDO = RootActorCDO->GetConfigData())
+								{
+									if (UDisplayClusterConfigurationViewport* ViewportCDO = ConfigDataCDO->GetViewport(ClusterName, ViewportId))
+									{
+										if (InViewport->DisplayDeviceName != ViewportCDO->DisplayDeviceName)
+										{
+											UE_LOG(LogDisplayClusterGame, Log, TEXT("Could not find display device '%s' on instance. Using CDO default device '%s'."),
+												*InViewport->DisplayDeviceName, *ViewportCDO->DisplayDeviceName);
+											InViewport->Modify();
+											InViewport->DisplayDeviceName = ViewportCDO->DisplayDeviceName;
+									
+											return FindAndSyncDisplayDeviceFromViewport(InViewport);
+										}
+									}
+								}
+							}
+						}
+					}
+				}
+			
+				UE_LOG(LogDisplayClusterGame, Warning, TEXT("Could not find display device '%s'. Using root actor default device."), *InViewport->DisplayDeviceName);
+			}
+
+			if (!DeviceBaseComponent)
+			{
+				DeviceBaseComponent = RootActor->GetDefaultDisplayDevice();
+			}
+
+			return DeviceBaseComponent;
+		}
+	}
+
+	return nullptr;
+}
diff --git a/Engine/Plugins/Runtime/nDisplay/Source/DisplayCluster/Private/Game/EngineClasses/Scene/DisplayDevice/DisplayClusterDisplayDeviceUtils.h b/Engine/Plugins/Runtime/nDisplay/Source/DisplayCluster/Private/Game/EngineClasses/Scene/DisplayDevice/DisplayClusterDisplayDeviceUtils.h
new file mode 100644
--- /dev/null
+++ b/Engine/Plugins/Runtime/nDisplay/Source/DisplayCluster/Private/Game/EngineClasses/Scene/DisplayDevice/DisplayClusterDisplayDeviceUtils.h
@@ -0,0 +1,19 @@
+// Copyright Epic Games, Inc. All Rights Reserved.
+
+#pragma once
+
+class UDisplayClusterConfigurationViewport;
+class UDisplayClusterDisplayDeviceBaseComponent;
+
+namespace UE::DisplayClusterDisplayDeviceUtils
+{
+	/**
+	 * Attempt to find the display device given a viewport.
+	 * If the display device is not found on an instance, the CDO is checked in case the display device was renamed and the
+	 * viewport on the instance will be synced to the new name.
+	 *
+	 * @param InViewport The viewport object which points to a display device and belongs to a root actor.
+	 * @return The display device this viewport uses or nullptr.
+	 */
+	UDisplayClusterDisplayDeviceBaseComponent* FindAndSyncDisplayDeviceFromViewport(UDisplayClusterConfigurationViewport* InViewport);
+}
diff --git a/Engine/Plugins/Runtime/nDisplay/Source/DisplayCluster/Public/Components/DisplayClusterPreviewComponent.h b/Engine/Plugins/Runtime/nDisplay/Source/DisplayCluster/Public/Components/DisplayClusterPreviewComponent.h
--- a/Engine/Plugins/Runtime/nDisplay/Source/DisplayCluster/Public/Components/DisplayClusterPreviewComponent.h
+++ b/Engine/Plugins/Runtime/nDisplay/Source/DisplayCluster/Public/Components/DisplayClusterPreviewComponent.h
@@ -10,6 +10,7 @@
 
 #include "DisplayClusterPreviewComponent.generated.h"
 
+class UDisplayClusterDisplayDeviceBaseComponent;
 class IDisplayClusterProjectionPolicy;
 class UTextureRenderTarget2D;
 class UTexture;
@@ -89,6 +90,12 @@ public:
 	/** Sets an override texture to display on the viewport instead of the render target */
 	void SetOverrideTexture(UTexture* InOverrideTexture);
 
+	/** Configure if this preview component should use the display device */
+	void SetUseDisplayDevice(bool bInNewValue);
+
+	/** Retrieve the correct display device to use for this preview component */
+	UDisplayClusterDisplayDeviceBaseComponent* GetDisplayDevice() const;
+	
 protected:
 	/** Gets whether the root actor is requesting a preview render */
 	bool IsPreviewEnabled() const;
@@ -113,12 +120,18 @@ protected:
 	void ReleasePreviewMesh();
 	void UpdatePreviewMeshReference();
 
-	void InitializePreviewMaterial();
 	void ReleasePreviewMaterial();
 	void UpdatePreviewMaterial();
 
 	void RestorePreviewMeshMaterial();
-	void SetPreviewMeshMaterial();
+	void SetPreviewMeshMaterial(UMaterial* InMaterial);
+	
+	/** Retrieve the correct preview material based on the display device */
+	UMaterial* GetPreviewMaterialFromDisplayDevice() const;
+
+	/** Retrieve the basic mesh material which should be applied when the preview is off */
+	UMaterial* GetMeshMaterialFromDisplayDevice() const;
+	
 #endif /* WITH_EDITOR */
 
 #if WITH_EDITORONLY_DATA
@@ -145,20 +158,25 @@ private:
 	UPROPERTY(Transient)
 	FString ClusterNodeId;
 
+	/** The last used display device component */
 	UPROPERTY(Transient)
-	TObjectPtr<UDisplayClusterConfigurationViewport> ViewportConfig = nullptr;
+	mutable FComponentReference CachedDisplayDevice;
 
-	UPROPERTY(Transient)
-	TObjectPtr<UMeshComponent> PreviewMesh = nullptr;
+	/** If this preview component should utilize the display device component */
+	bool bUseDisplayDevice = true;
 
 	UPROPERTY(Transient)
-	bool bIsRootActorPreviewMesh = false;
+	TObjectPtr<UDisplayClusterConfigurationViewport> ViewportConfig = nullptr;
 
 	UPROPERTY(Transient)
-	TObjectPtr<UMaterial> OriginalMaterial = nullptr;
+	TObjectPtr<UMeshComponent> PreviewMesh = nullptr;
 
+	/** The current material assigned to the preview mesh */
 	UPROPERTY(Transient)
-	TObjectPtr<UMaterial> PreviewMaterial = nullptr;
+	TObjectPtr<UMaterial> CurrentMeshMaterial = nullptr;
+	
+	UPROPERTY(Transient)
+	bool bIsRootActorPreviewMesh = false;
 
 	UPROPERTY(Transient)
 	TObjectPtr<UMaterialInstanceDynamic> PreviewMaterialInstance = nullptr;
diff --git a/Engine/Plugins/Runtime/nDisplay/Source/DisplayCluster/Public/Components/DisplayDevice/DisplayClusterDisplayDeviceBaseComponent.h b/Engine/Plugins/Runtime/nDisplay/Source/DisplayCluster/Public/Components/DisplayDevice/DisplayClusterDisplayDeviceBaseComponent.h
new file mode 100644
--- /dev/null
+++ b/Engine/Plugins/Runtime/nDisplay/Source/DisplayCluster/Public/Components/DisplayDevice/DisplayClusterDisplayDeviceBaseComponent.h
@@ -0,0 +1,58 @@
+// Copyright Epic Games, Inc. All Rights Reserved.
+
+#pragma once
+
+#include "Components/SceneComponent.h"
+
+#include "DisplayClusterDisplayDeviceBaseComponent.generated.h"
+
+class UMaterial;
+class UMaterialInstanceDynamic;
+
+UCLASS(Abstract, ClassGroup = (DisplayCluster))
+class DISPLAYCLUSTER_API UDisplayClusterDisplayDeviceBaseComponent
+	: public USceneComponent
+{
+	GENERATED_BODY()
+public:
+	UDisplayClusterDisplayDeviceBaseComponent();
+
+	/** Retrieve the preview material used. */
+	TObjectPtr<UMaterial> GetPreviewMaterial() const { return PreviewMaterial; }
+
+	/** Retrieve the base mesh material. */
+	TObjectPtr<UMaterial> GetMeshMaterial() const { return MeshMaterial; }
+
+	/** Perform any operations on the preview material instance, such as setting parameter values. */
+	virtual void OnUpdatePreviewMaterialInstance(UMaterialInstanceDynamic* InMaterialInstance) {}
+
+#if WITH_EDITOR
+	virtual void PostEditChangeProperty(FPropertyChangedEvent& PropertyChangedEvent) override;
+#endif
+	
+protected:
+	/** The material the preview components will create a material instance from when rendering the nDisplay preview. */
+	UPROPERTY(EditAnywhere, Category=Material, NoClear)
+	TObjectPtr<UMaterial> PreviewMaterial = nullptr;
+
+	/** The material to assign to the static mesh when preview is disabled. */
+	UPROPERTY(EditAnywhere, Category=Material, NoClear)
+	TObjectPtr<UMaterial> MeshMaterial = nullptr;
+};
+
+UCLASS(ClassGroup = (DisplayCluster), meta = (BlueprintSpawnableComponent, DisplayName = "NDisplay Display Device"))
+class DISPLAYCLUSTER_API UDisplayClusterDisplayDeviceComponent
+	: public UDisplayClusterDisplayDeviceBaseComponent
+{
+	GENERATED_BODY()
+
+public:
+	UDisplayClusterDisplayDeviceComponent();
+	
+	virtual void OnUpdatePreviewMaterialInstance(UMaterialInstanceDynamic* InMaterialInstance) override;
+	
+protected:
+	/** Adjust the exposure for the emissive input. */
+	UPROPERTY(EditAnywhere, Category=Material)
+	float Exposure = 0.f;
+};
diff --git a/Engine/Plugins/Runtime/nDisplay/Source/DisplayCluster/Public/DisplayClusterRootActor.h b/Engine/Plugins/Runtime/nDisplay/Source/DisplayCluster/Public/DisplayClusterRootActor.h
--- a/Engine/Plugins/Runtime/nDisplay/Source/DisplayCluster/Public/DisplayClusterRootActor.h
+++ b/Engine/Plugins/Runtime/nDisplay/Source/DisplayCluster/Public/DisplayClusterRootActor.h
@@ -27,6 +27,7 @@
 #include "DisplayClusterRootActor.generated.h"
 
 #if WITH_EDITOR
+class UDisplayClusterDisplayDeviceBaseComponent;
 class IDisplayClusterConfiguratorBlueprintEditor;
 class FTransactionObjectEvent;
 #endif
@@ -419,6 +420,18 @@ private:
 	UPROPERTY(Transient)
 	TMap<FString, TObjectPtr<UDisplayClusterPreviewComponent>> PreviewComponents;
 
+	/** The default display device to use for preview rendering */
+	UPROPERTY(Transient, NonTransactional)
+	TObjectPtr<UDisplayClusterDisplayDeviceBaseComponent> DefaultDisplayDeviceComponent;
+
+	/** The included display device nDisplay provides by default */
+	UPROPERTY(VisibleDefaultsOnly, Category = "Editor Preview", DisplayName = "Basic Display Device")
+	TObjectPtr<UDisplayClusterDisplayDeviceBaseComponent> BasicDisplayDeviceComponent;
+
+	/** Select the default display device class to use when a viewport doesn't have one assigned */
+	UPROPERTY(EditDefaultsOnly, Category = "Editor Preview", DisplayName = "Default Display Device")
+	FName DefaultDisplayDeviceName;
+	
 	UPROPERTY(Transient)
 	bool bDeferPreviewGeneration = false;
 #endif
@@ -534,6 +547,12 @@ public:
 		return false;
 	}
 
+	/** The name the internal default display device uses. */
+	static FName GetInternalDisplayDeviceName();
+
+	/** Retrieve the default display device, creating it if it doesn't exist */
+	UDisplayClusterDisplayDeviceBaseComponent* GetDefaultDisplayDevice();
+	
 protected:
 	FString GeneratePreviewComponentName_Editor(const FString& NodeId, const FString& ViewportId) const;
 	void ResetPreviewInternals_Editor();
diff --git a/Engine/Plugins/Runtime/nDisplay/Source/DisplayClusterConfiguration/Private/DisplayClusterConfigurationTypes.cpp b/Engine/Plugins/Runtime/nDisplay/Source/DisplayClusterConfiguration/Private/DisplayClusterConfigurationTypes.cpp
--- a/Engine/Plugins/Runtime/nDisplay/Source/DisplayClusterConfiguration/Private/DisplayClusterConfigurationTypes.cpp
+++ b/Engine/Plugins/Runtime/nDisplay/Source/DisplayClusterConfiguration/Private/DisplayClusterConfigurationTypes.cpp
@@ -60,6 +60,26 @@ UDisplayClusterConfigurationViewport* UDisplayClusterConfigurationData::GetViewp
 	return nullptr;
 }
 
+void UDisplayClusterConfigurationData::ForEachViewport(const TFunction<void(const TObjectPtr<UDisplayClusterConfigurationViewport>&)>& Function) const
+{
+	if (Cluster)
+	{
+		for (const TTuple<FString, TObjectPtr<UDisplayClusterConfigurationClusterNode>>& ClusterNode : Cluster->Nodes)
+		{
+			if (ClusterNode.Value)
+			{
+				for (const TTuple<FString, TObjectPtr<UDisplayClusterConfigurationViewport>>& Viewport : ClusterNode.Value->Viewports)
+				{
+					if (Viewport.Value)
+					{
+						Function(Viewport.Value);
+					}
+				}
+			}
+		}
+	}
+}
+
 bool UDisplayClusterConfigurationData::AssignPostprocess(const FString& NodeId, const FString& PostprocessId, const FString& Type, TMap<FString, FString> Parameters, int32 Order)
 {
 	if (Cluster && Cluster->Nodes.Contains(NodeId))
diff --git a/Engine/Plugins/Runtime/nDisplay/Source/DisplayClusterConfiguration/Public/DisplayClusterConfigurationTypes.h b/Engine/Plugins/Runtime/nDisplay/Source/DisplayClusterConfiguration/Public/DisplayClusterConfigurationTypes.h
--- a/Engine/Plugins/Runtime/nDisplay/Source/DisplayClusterConfiguration/Public/DisplayClusterConfigurationTypes.h
+++ b/Engine/Plugins/Runtime/nDisplay/Source/DisplayClusterConfiguration/Public/DisplayClusterConfigurationTypes.h
@@ -516,6 +516,9 @@ public:
 	// Viewports API
 	UFUNCTION(BlueprintPure, Category = "NDisplay|Configuration")
 	UDisplayClusterConfigurationViewport* GetViewport(const FString& NodeId, const FString& ViewportId) const;
+
+	/** Execute a function on each viewport in each cluster node */
+	void ForEachViewport(const TFunction<void(const TObjectPtr<UDisplayClusterConfigurationViewport>&)>& Function) const;
 	
 	/**
 	* Update\Create node postprocess
diff --git a/Engine/Plugins/Runtime/nDisplay/Source/DisplayClusterConfiguration/Public/DisplayClusterConfigurationTypes_Viewport.h b/Engine/Plugins/Runtime/nDisplay/Source/DisplayClusterConfiguration/Public/DisplayClusterConfigurationTypes_Viewport.h
--- a/Engine/Plugins/Runtime/nDisplay/Source/DisplayClusterConfiguration/Public/DisplayClusterConfigurationTypes_Viewport.h
+++ b/Engine/Plugins/Runtime/nDisplay/Source/DisplayClusterConfiguration/Public/DisplayClusterConfigurationTypes_Viewport.h
@@ -22,6 +22,7 @@
 
 #include "DisplayClusterConfigurationTypes_Viewport.generated.h"
 
+class UDisplayClusterDisplayDeviceBaseComponent;
 struct FDisplayClusterConfigurationICVFX_StageSettings;
 
 USTRUCT(Blueprintable)
@@ -219,6 +220,10 @@ public:
 	FDisplayClusterConfigurationProjection ProjectionPolicy;
 
 #if WITH_EDITORONLY_DATA
+	/** Select a display device to use during preview. An empty string will use the default display device */
+	UPROPERTY(EditAnywhere, Category = "Preview", meta = (DisplayName = "Display Device"))
+	FString DisplayDeviceName;
+	
 	UPROPERTY(EditAnywhere, BlueprintReadWrite, Category = "Preview", meta = (DisplayName = "Preview Frustum"))
 	bool bAllowPreviewFrustumRendering = false;
 #endif
diff --git a/Engine/Plugins/Runtime/nDisplay/Source/DisplayClusterConfigurator/Private/DisplayClusterConfiguratorBlueprintEditor.cpp b/Engine/Plugins/Runtime/nDisplay/Source/DisplayClusterConfigurator/Private/DisplayClusterConfiguratorBlueprintEditor.cpp
--- a/Engine/Plugins/Runtime/nDisplay/Source/DisplayClusterConfigurator/Private/DisplayClusterConfiguratorBlueprintEditor.cpp
+++ b/Engine/Plugins/Runtime/nDisplay/Source/DisplayClusterConfigurator/Private/DisplayClusterConfiguratorBlueprintEditor.cpp
@@ -778,12 +778,12 @@ void FDisplayClusterConfiguratorBlueprintEditor::OnRenameVariable(UBlueprint* Bl
 
 	// Check the configuration's viewports for any matching references to the component variable being renamed, and update those references
 	// to the new variable name
-	if (UDisplayClusterConfigurationData* Config = GetConfig())
+	if (const UDisplayClusterConfigurationData* Config = GetConfig())
 	{
-		for (TPair<FString, UDisplayClusterConfigurationClusterNode*> ClusterNodePair : Config->Cluster->Nodes)
+		for (const TPair<FString, TObjectPtr<UDisplayClusterConfigurationClusterNode>>& ClusterNodePair : Config->Cluster->Nodes)
 		{
 			UDisplayClusterConfigurationClusterNode* ClusterNode = ClusterNodePair.Value;
-			for (TPair<FString, UDisplayClusterConfigurationViewport*> ViewportPair : ClusterNode->Viewports)
+			for (const TPair<FString, TObjectPtr<UDisplayClusterConfigurationViewport>>& ViewportPair : ClusterNode->Viewports)
 			{
 				UDisplayClusterConfigurationViewport* Viewport = ViewportPair.Value;
 
@@ -819,6 +819,11 @@ void FDisplayClusterConfiguratorBlueprintEditor::OnRenameVariable(UBlueprint* Bl
 				{
 					DisplayClusterConfiguratorPropertyUtils::SetPropertyHandleValue(Viewport, GET_MEMBER_NAME_CHECKED(UDisplayClusterConfigurationViewport, Camera), NewVariableName.ToString());
 				}
+
+				if (Viewport->DisplayDeviceName == OldVariableName.ToString())
+				{
+					DisplayClusterConfiguratorPropertyUtils::SetPropertyHandleValue(Viewport, GET_MEMBER_NAME_CHECKED(UDisplayClusterConfigurationViewport, DisplayDeviceName), NewVariableName.ToString());
+				}
 			}
 		}
 	}
diff --git a/Engine/Plugins/Runtime/nDisplay/Source/DisplayClusterConfigurator/Private/DisplayClusterConfiguratorCompiler.cpp b/Engine/Plugins/Runtime/nDisplay/Source/DisplayClusterConfigurator/Private/DisplayClusterConfiguratorCompiler.cpp
--- a/Engine/Plugins/Runtime/nDisplay/Source/DisplayClusterConfigurator/Private/DisplayClusterConfiguratorCompiler.cpp
+++ b/Engine/Plugins/Runtime/nDisplay/Source/DisplayClusterConfigurator/Private/DisplayClusterConfiguratorCompiler.cpp
@@ -9,6 +9,7 @@
 #include "Blueprints/DisplayClusterBlueprint.h"
 #include "Blueprints/DisplayClusterBlueprintGeneratedClass.h"
 #include "Components/DisplayClusterCameraComponent.h"
+#include "Components/DisplayDevice/DisplayClusterDisplayDeviceBaseComponent.h"
 #include "Camera/CameraComponent.h"
 
 #include "Engine/SCS_Node.h"
@@ -261,7 +262,7 @@ void FDisplayClusterConfiguratorKismetCompilerContext::ValidateConfiguration()
 	
 	for (const TPair<FString, TObjectPtr<UDisplayClusterConfigurationClusterNode>>& ClusterNode : BlueprintData->Cluster->Nodes)
 	{
-		// Validate vieports
+		// Validate viewports
 		if (ClusterNode.Value->Viewports.Num() > 0)
 		{
 			for (const TPair<FString, TObjectPtr<UDisplayClusterConfigurationViewport>>& Viewport : ClusterNode.Value->Viewports)
@@ -280,6 +281,44 @@ void FDisplayClusterConfiguratorKismetCompilerContext::ValidateConfiguration()
 				else
 				{
 					ViewportsByName.Add(ViewportName, Viewport.Value);
+
+					// Check display device
+					if (!Viewport.Value->DisplayDeviceName.IsEmpty())
+					{
+						// Only check if bp class has been generated and component owner instantiated. Otherwise native
+						// components can't be easily found since the CDO has already been cleared at this point.
+						if (UBlueprintGeneratedClass* GenClass = DCBlueprint->GetGeneratedClass())
+						{
+							if (AActor* ActorOwner = GenClass->SimpleConstructionScript->GetComponentEditorActorInstance())
+							{
+								bool bDisplayDeviceFound = false;
+								// blueprint component
+								if (const USCS_Node* SCSNode = GenClass->SimpleConstructionScript->FindSCSNode(*Viewport.Value->DisplayDeviceName))
+								{
+									if (Cast<UDisplayClusterDisplayDeviceBaseComponent>(SCSNode->GetActualComponentTemplate(GenClass)))
+									{
+										bDisplayDeviceFound = true;
+									}
+								}
+
+								// native component
+								if (!bDisplayDeviceFound
+									&& FindObjectFast<UDisplayClusterDisplayDeviceBaseComponent>(ActorOwner, *Viewport.Value->DisplayDeviceName))
+								{
+									bDisplayDeviceFound = true;
+								}
+
+								if (!bDisplayDeviceFound)
+								{
+									MessageLog.Error(
+										*LOCTEXT("DisplayDeviceNotFound", "Could not find display device component @@ for viewport @@.").ToString(),
+										*Viewport.Value->DisplayDeviceName,
+										Viewport.Value
+									);
+								}
+							}
+						}
+					}
 				}
 				
 				if (Viewport.Value->ProjectionPolicy.Type.IsEmpty())
diff --git a/Engine/Plugins/Runtime/nDisplay/Source/DisplayClusterConfigurator/Private/Views/Details/Cluster/DisplayClusterConfiguratorViewportDetailsCustomization.cpp b/Engine/Plugins/Runtime/nDisplay/Source/DisplayClusterConfigurator/Private/Views/Details/Cluster/DisplayClusterConfiguratorViewportDetailsCustomization.cpp
--- a/Engine/Plugins/Runtime/nDisplay/Source/DisplayClusterConfigurator/Private/Views/Details/Cluster/DisplayClusterConfiguratorViewportDetailsCustomization.cpp
+++ b/Engine/Plugins/Runtime/nDisplay/Source/DisplayClusterConfigurator/Private/Views/Details/Cluster/DisplayClusterConfiguratorViewportDetailsCustomization.cpp
@@ -8,6 +8,8 @@
 #include "DisplayClusterProjectionStrings.h"
 #include "DisplayClusterConfiguratorLog.h"
 #include "Components/DisplayClusterCameraComponent.h"
+#include "Components/DisplayDevice/DisplayClusterDisplayDeviceBaseComponent.h"
+#include "Views/Details/Widgets/SDisplayClusterConfiguratorComponentPicker.h"
 #include "Views/Details/Widgets/SDisplayClusterConfigurationSearchableComboBox.h"
 
 #include "DetailLayoutBuilder.h"
@@ -49,28 +51,53 @@ void FDisplayClusterConfiguratorViewportDetailsCustomization::CustomizeDetails(I
 		IDetailCategoryBuilder& Category = InLayoutBuilder.EditCategory(TEXT("In Camera VFX"));
 		Category.SetCategoryVisibility(false);
 	}
-	
-	CameraHandle = InLayoutBuilder.GetProperty(GET_MEMBER_NAME_CHECKED(UDisplayClusterConfigurationViewport, Camera));
-	check(CameraHandle->IsValidHandle());
 
-	if (ConfigurationViewportPtr->ProjectionPolicy.Type.Equals(DisplayClusterProjectionStrings::projection::Camera, ESearchCase::IgnoreCase))
 	{
-		CameraHandle->MarkHiddenByCustomization();
+		CameraHandle = InLayoutBuilder.GetProperty(GET_MEMBER_NAME_CHECKED(UDisplayClusterConfigurationViewport, Camera));
+		check(CameraHandle->IsValidHandle());
+
+		if (ConfigurationViewportPtr->ProjectionPolicy.Type.Equals(DisplayClusterProjectionStrings::projection::Camera, ESearchCase::IgnoreCase))
+		{
+			CameraHandle->MarkHiddenByCustomization();
+		}
+		else
+		{
+			ResetCameraOptions();
+
+			if (IDetailPropertyRow* CameraPropertyRow = InLayoutBuilder.EditDefaultProperty(CameraHandle))
+			{
+				CameraPropertyRow->CustomWidget()
+					.NameContent()
+					[
+						CameraHandle->CreatePropertyNameWidget()
+					]
+					.ValueContent()
+					[
+						CreateCustomCameraWidget()
+					];
+			}
+		}
 	}
-	else
+	
 	{
-		ResetCameraOptions();
-
-		if (IDetailPropertyRow* CameraPropertyRow = InLayoutBuilder.EditDefaultProperty(CameraHandle))
+		TSharedPtr<IPropertyHandle> DisplayDeviceHandle = InLayoutBuilder.GetProperty(GET_MEMBER_NAME_CHECKED(UDisplayClusterConfigurationViewport, DisplayDeviceName));
+		check(DisplayDeviceHandle->IsValidHandle());
+		
+		if (IDetailPropertyRow* DisplayDevicePropertyRow = InLayoutBuilder.EditDefaultProperty(DisplayDeviceHandle))
 		{
-			CameraPropertyRow->CustomWidget()
+			DisplayDevicePropertyRow->CustomWidget()
 				.NameContent()
 				[
-					CameraHandle->CreatePropertyNameWidget()
+					DisplayDeviceHandle->CreatePropertyNameWidget()
 				]
 				.ValueContent()
 				[
-					CreateCustomCameraWidget()
+					SNew(SDisplayClusterConfiguratorComponentPicker,
+						UDisplayClusterDisplayDeviceBaseComponent::StaticClass(),
+						GetRootActor(),
+						DisplayDeviceHandle)
+						.DefaultOptionText(NSLOCTEXT("nDisplayDevicePicker", "DefaultText", "Default"))
+						.DefaultOptionValue(TOptional<FString>(""))
 				];
 		}
 	}
diff --git a/Engine/Plugins/Runtime/nDisplay/Source/DisplayClusterConfigurator/Private/Views/Details/DisplayClusterRootActorDetailsCustomization.cpp b/Engine/Plugins/Runtime/nDisplay/Source/DisplayClusterConfigurator/Private/Views/Details/DisplayClusterRootActorDetailsCustomization.cpp
--- a/Engine/Plugins/Runtime/nDisplay/Source/DisplayClusterConfigurator/Private/Views/Details/DisplayClusterRootActorDetailsCustomization.cpp
+++ b/Engine/Plugins/Runtime/nDisplay/Source/DisplayClusterConfigurator/Private/Views/Details/DisplayClusterRootActorDetailsCustomization.cpp
@@ -2,6 +2,9 @@
 
 #include "DisplayClusterRootActorDetailsCustomization.h"
 
+#include "Widgets/SDisplayClusterConfiguratorComponentPicker.h"
+
+#include "Components/DisplayDevice/DisplayClusterDisplayDeviceBaseComponent.h"
 #include "DisplayClusterRootActor.h"
 #include "DisplayClusterConfigurationStrings.h"
 #include "DisplayClusterConfigurationTypes.h"
@@ -155,6 +158,27 @@ void FDisplayClusterRootActorDetailsCustomization::CustomizeDetails(IDetailLayou
 
 	// Update the selected item in the NodeId combo box to match the current value on the root actor
 	UpdateNodeIdSelection();
+
+	// Default Display Device component selection
+	TSharedPtr<IPropertyHandle> DisplayDeviceHandle = InLayoutBuilder.GetProperty(GET_MEMBER_NAME_CHECKED(ADisplayClusterRootActor, DefaultDisplayDeviceName));
+	if (DisplayDeviceHandle.IsValid() && DisplayDeviceHandle->IsValidHandle())
+	{
+		if (IDetailPropertyRow* DisplayDevicePropertyRow = InLayoutBuilder.EditDefaultProperty(DisplayDeviceHandle))
+		{
+			DisplayDevicePropertyRow->CustomWidget()
+				.NameContent()
+				[
+					DisplayDeviceHandle->CreatePropertyNameWidget()
+				]
+				.ValueContent()
+				[
+					SNew(SDisplayClusterConfiguratorComponentPicker,
+						UDisplayClusterDisplayDeviceBaseComponent::StaticClass(),
+						GetRootActor(),
+						DisplayDeviceHandle)
+				];
+		}
+	}
 }
 
 TSharedRef<SWidget> FDisplayClusterRootActorDetailsCustomization::CreateCustomNodeIdWidget()
diff --git a/Engine/Plugins/Runtime/nDisplay/Source/DisplayClusterConfigurator/Private/Views/Details/Widgets/SDisplayClusterConfigurationSearchableComboBox.cpp b/Engine/Plugins/Runtime/nDisplay/Source/DisplayClusterConfigurator/Private/Views/Details/Widgets/SDisplayClusterConfigurationSearchableComboBox.cpp
--- a/Engine/Plugins/Runtime/nDisplay/Source/DisplayClusterConfigurator/Private/Views/Details/Widgets/SDisplayClusterConfigurationSearchableComboBox.cpp
+++ b/Engine/Plugins/Runtime/nDisplay/Source/DisplayClusterConfigurator/Private/Views/Details/Widgets/SDisplayClusterConfigurationSearchableComboBox.cpp
@@ -1,4 +1,4 @@
-// Copyright Epic Games, Inc. All Rights Reservekd.
+// Copyright Epic Games, Inc. All Rights Reserved.
 
 #include "Views/Details/Widgets/SDisplayClusterConfigurationSearchableComboBox.h"
 
diff --git a/Engine/Plugins/Runtime/nDisplay/Source/DisplayClusterConfigurator/Private/Views/Details/Widgets/SDisplayClusterConfigurationSearchableComboBox.h b/Engine/Plugins/Runtime/nDisplay/Source/DisplayClusterConfigurator/Private/Views/Details/Widgets/SDisplayClusterConfigurationSearchableComboBox.h
--- a/Engine/Plugins/Runtime/nDisplay/Source/DisplayClusterConfigurator/Private/Views/Details/Widgets/SDisplayClusterConfigurationSearchableComboBox.h
+++ b/Engine/Plugins/Runtime/nDisplay/Source/DisplayClusterConfigurator/Private/Views/Details/Widgets/SDisplayClusterConfigurationSearchableComboBox.h
@@ -1,4 +1,4 @@
-// Copyright Epic Games, Inc. All Rights Reservekd.
+// Copyright Epic Games, Inc. All Rights Reserved.
 
 #pragma once
 
diff --git a/Engine/Plugins/Runtime/nDisplay/Source/DisplayClusterConfigurator/Private/Views/Details/Widgets/SDisplayClusterConfiguratorComponentPicker.cpp b/Engine/Plugins/Runtime/nDisplay/Source/DisplayClusterConfigurator/Private/Views/Details/Widgets/SDisplayClusterConfiguratorComponentPicker.cpp
new file mode 100644
--- /dev/null
+++ b/Engine/Plugins/Runtime/nDisplay/Source/DisplayClusterConfigurator/Private/Views/Details/Widgets/SDisplayClusterConfiguratorComponentPicker.cpp
@@ -0,0 +1,139 @@
+// Copyright Epic Games, Inc. All Rights Reserved.
+
+#include "SDisplayClusterConfiguratorComponentPicker.h"
+
+#include "SDisplayClusterConfigurationSearchableComboBox.h"
+
+#include "DetailLayoutBuilder.h"
+#include "GameFramework/Actor.h"
+
+void SDisplayClusterConfiguratorComponentPicker::Construct(
+	const FArguments& InArgs,
+	const TSubclassOf<UActorComponent>& InComponentClass,
+	const TWeakObjectPtr<AActor>& InOwningActor,
+	const TSharedPtr<IPropertyHandle>& InPropertyHandle)
+{
+	ComponentClass = InComponentClass;
+	PropertyHandle = InPropertyHandle;
+	OwningActor = InOwningActor;
+
+	DefaultOptionText = InArgs._DefaultOptionText;
+	DefaultOptionValue = InArgs._DefaultOptionValue;
+
+	RefreshComponentOptions();
+
+	ChildSlot
+	[
+		CreateComboBoxWidget()
+	];
+}
+
+TSharedRef<SWidget> SDisplayClusterConfiguratorComponentPicker::CreateComboBoxWidget()
+{
+	if (ComponentComboBox.IsValid())
+	{
+		return ComponentComboBox.ToSharedRef();
+	}
+
+	return SAssignNew(ComponentComboBox, SDisplayClusterConfigurationSearchableComboBox)
+		.OptionsSource(&AvailableComponentOptions)
+		.OnGenerateWidget(this, &SDisplayClusterConfiguratorComponentPicker::MakeOptionComboWidget)
+		.OnSelectionChanged(this, &SDisplayClusterConfiguratorComponentPicker::OnComponentSelected)
+		.ContentPadding(2)
+		.MaxListHeight(200.0f)
+		.Content()
+		[
+			SNew(STextBlock)
+			.Text(this, &SDisplayClusterConfiguratorComponentPicker::GetSelectedComponentText)
+			.Font(IDetailLayoutBuilder::GetDetailFont())
+		];
+}
+
+TSharedRef<SWidget> SDisplayClusterConfiguratorComponentPicker::MakeOptionComboWidget(TSharedPtr<FString> InItem)
+{
+	return SNew(STextBlock)
+		.Text(FText::FromString(*InItem))
+		.Font(IDetailLayoutBuilder::GetDetailFont());
+}
+
+FText SDisplayClusterConfiguratorComponentPicker::GetSelectedComponentText() const
+{
+	FString SelectedOption;
+	if (PropertyHandle.IsValid())
+	{
+		if (FPropertyAccess::Result::MultipleValues == PropertyHandle->GetValue(SelectedOption))
+		{
+			SelectedOption = TEXT("Multiple Values");
+		}
+
+		if (PropertyHandle->GetPropertyClass()->IsChildOf(FNameProperty::StaticClass()))
+		{
+			if (FName(SelectedOption) == NAME_None)
+			{
+				SelectedOption.Empty();
+			}
+		}
+	}
+
+	if (SelectedOption.IsEmpty())
+	{
+		return DefaultOptionText;
+	}
+
+	return FText::FromString(SelectedOption);
+}
+
+void SDisplayClusterConfiguratorComponentPicker::OnComponentSelected(TSharedPtr<FString> InComponentName,
+                                                                           ESelectInfo::Type SelectInfo)
+{
+	if (InComponentName.IsValid())
+	{
+		if (*InComponentName == DefaultOptionText.ToString()
+			&& DefaultOptionValue.IsSet())
+		{
+			const FString& StrValue = DefaultOptionValue.GetValue();
+			PropertyHandle->SetValue(*StrValue);
+		}
+		else
+		{
+			PropertyHandle->SetValue(*InComponentName.Get());
+		}
+		
+		// Reset available options
+		RefreshComponentOptions();
+		ComponentComboBox->ResetOptionsSource(&AvailableComponentOptions);
+		ComponentComboBox->SetIsOpen(false);
+	}
+}
+
+void SDisplayClusterConfiguratorComponentPicker::RefreshComponentOptions()
+{
+	AvailableComponentOptions.Reset();
+
+	FString DefaultOption = DefaultOptionText.ToString();
+	if (!DefaultOption.IsEmpty())
+	{
+		AvailableComponentOptions.Add(MakeShared<FString>(DefaultOption));
+	}
+	
+	if (OwningActor.IsValid())
+	{
+		TArray<UActorComponent*> AvailableComponents;
+		OwningActor->GetComponents(ComponentClass, AvailableComponents);
+		for (const UActorComponent* Component : AvailableComponents)
+		{
+			const FString ComponentName = Component->GetName();
+			if (ComponentName != DefaultOption)
+			{
+				AvailableComponentOptions.Add(MakeShared<FString>(ComponentName));
+			}
+		}
+
+		// Component order not guaranteed, sort for consistency.
+		AvailableComponentOptions.Sort([](const TSharedPtr<FString>& A, const TSharedPtr<FString>& B)
+		{
+			// Default sort isn't compatible with TSharedPtr<FString>.
+			return *A < *B;
+		});
+	}
+}
diff --git a/Engine/Plugins/Runtime/nDisplay/Source/DisplayClusterConfigurator/Private/Views/Details/Widgets/SDisplayClusterConfiguratorComponentPicker.h b/Engine/Plugins/Runtime/nDisplay/Source/DisplayClusterConfigurator/Private/Views/Details/Widgets/SDisplayClusterConfiguratorComponentPicker.h
new file mode 100644
--- /dev/null
+++ b/Engine/Plugins/Runtime/nDisplay/Source/DisplayClusterConfigurator/Private/Views/Details/Widgets/SDisplayClusterConfiguratorComponentPicker.h
@@ -0,0 +1,83 @@
+// Copyright Epic Games, Inc. All Rights Reserved.
+
+#pragma once
+
+#include "CoreMinimal.h"
+#include "SlateFwd.h"
+#include "Templates/SubclassOf.h"
+#include "Widgets/SCompoundWidget.h"
+
+class IPropertyHandle;
+class SDisplayClusterConfigurationSearchableComboBox;
+class UActorComponent;
+
+/**
+ * A searchable text combo box for components which allows reading/writing a string value to a property handle.
+ * Works with CDOs.
+ */
+class SDisplayClusterConfiguratorComponentPicker
+	: public SCompoundWidget
+{
+public:
+	SLATE_BEGIN_ARGS(SDisplayClusterConfiguratorComponentPicker)
+	{}
+	SLATE_ARGUMENT(FText, DefaultOptionText)
+	SLATE_ARGUMENT(TOptional<FString>, DefaultOptionValue)
+	SLATE_END_ARGS()
+	
+	/**
+	 * Construct the widget from a declaration
+	 *
+	 * @param InArgs   Declaration from which to construct the widget
+	 * @param InComponentClass The component class and all children this picker should display
+	 * @param InOwningActor The actor containing the components to display
+	 * @param InPropertyHandle The property handle to read and write values to
+	 */
+	void Construct(const FArguments& InArgs, const TSubclassOf<UActorComponent>& InComponentClass, const TWeakObjectPtr<AActor>& InOwningActor,
+		const TSharedPtr<IPropertyHandle>& InPropertyHandle);
+
+private:
+	/** Creates a combo box widget to represent components */
+	TSharedRef<SWidget> CreateComboBoxWidget();
+
+	/**
+	 * Creates a text block widget to use to display the specified item in the dropdown menu
+	 * @param InItem - The string item to make the text block for
+	 */
+	TSharedRef<SWidget> MakeOptionComboWidget(TSharedPtr<FString> InItem);
+
+	/** Gets the text to display for the currently selected component */
+	FText GetSelectedComponentText() const;
+	
+	/**
+	 * Raised when a component is selected from the dropdown menu
+	 * @param InComponentName - The component name that was selected
+	 * @param SelectInfo - Flag to indicate through what interface the selection was made
+	 */
+	void OnComponentSelected(TSharedPtr<FString> InComponentName, ESelectInfo::Type SelectInfo);
+
+	/** Rebuilds the list of components to show in the dropdown menu */
+	void RefreshComponentOptions();
+	
+private:
+	/** The list of components to display in the dropdown menu */
+	TArray<TSharedPtr<FString>>	AvailableComponentOptions;
+	
+	/** The property handle to read/set the component value to */
+	TSharedPtr<IPropertyHandle> PropertyHandle;
+	
+	/** The combo box that is being displayed in the details panel for the component property */
+	TSharedPtr<SDisplayClusterConfigurationSearchableComboBox> ComponentComboBox;
+
+	/** The actor containing the components */
+	TWeakObjectPtr<AActor> OwningActor;
+
+	/** The component type for this widget */
+	TSubclassOf<UActorComponent> ComponentClass;
+
+	/** The default option to use */
+	FText DefaultOptionText;
+
+	/** The value to apply when selecting the default option */
+	TOptional<FString> DefaultOptionValue;
+};
diff --git a/Engine/Plugins/Runtime/nDisplay/Source/DisplayClusterConfigurator/Private/Views/SCSEditor/SDisplayClusterConfiguratorComponentCombo.cpp b/Engine/Plugins/Runtime/nDisplay/Source/DisplayClusterConfigurator/Private/Views/SCSEditor/SDisplayClusterConfiguratorComponentCombo.cpp
--- a/Engine/Plugins/Runtime/nDisplay/Source/DisplayClusterConfigurator/Private/Views/SCSEditor/SDisplayClusterConfiguratorComponentCombo.cpp
+++ b/Engine/Plugins/Runtime/nDisplay/Source/DisplayClusterConfigurator/Private/Views/SCSEditor/SDisplayClusterConfiguratorComponentCombo.cpp
@@ -6,6 +6,7 @@
 #include "Components/DisplayClusterICVFXCameraComponent.h"
 #include "Components/DisplayClusterScreenComponent.h"
 #include "Components/DisplayClusterXformComponent.h"
+#include "Components/DisplayDevice/DisplayClusterDisplayDeviceBaseComponent.h"
 
 #include "Widgets/Layout/SSpacer.h"
 #include "Widgets/Images/SImage.h"
@@ -459,6 +460,23 @@ void SDisplayClusterConfiguratorComponentClassCombo::GenerateComponentClassList(
 	AddHeader("nDisplay ICVFX Components");
 	AddDCComp(UDisplayClusterICVFXCameraComponent::StaticClass());
 
+	AddHeader("nDisplay Display Device Components");
+	{
+		// Locate all derived display device components. Only handles native classes. If we allow display devices
+		// to be extended in BP then we'll need to handle unloaded assets.
+		
+		TArray<UClass*> DisplayDeviceClasses;
+		GetDerivedClasses(UDisplayClusterDisplayDeviceBaseComponent::StaticClass(), DisplayDeviceClasses);
+
+		for (UClass* Class : DisplayDeviceClasses)
+		{
+			if (!Class->HasAnyClassFlags(CLASS_Abstract | CLASS_Deprecated | CLASS_NewerVersionExists))
+			{
+				AddDCComp(Class);
+			}
+		}
+	}
+	
 	int32 CompIterationIdx = 0;
 	for (const TSharedPtr<FComponentClassComboEntry>& Comp : *ComponentClassListPtr)
 	{
diff --git a/Engine/Plugins/Runtime/nDisplay/Source/DisplayClusterLightCardEditor/Private/Viewport/DisplayClusterLightCardEditorViewportClient.cpp b/Engine/Plugins/Runtime/nDisplay/Source/DisplayClusterLightCardEditor/Private/Viewport/DisplayClusterLightCardEditorViewportClient.cpp
--- a/Engine/Plugins/Runtime/nDisplay/Source/DisplayClusterLightCardEditor/Private/Viewport/DisplayClusterLightCardEditorViewportClient.cpp
+++ b/Engine/Plugins/Runtime/nDisplay/Source/DisplayClusterLightCardEditor/Private/Viewport/DisplayClusterLightCardEditorViewportClient.cpp
@@ -197,6 +197,7 @@ void FDisplayClusterLightCardEditorViewportClient::Tick(float DeltaSeconds)
 
 				if (PreviewComp && LevelInstancePreviewComp)
 				{
+					PreviewComp->SetUseDisplayDevice(false);
 					PreviewComp->SetOverrideTexture(LevelInstancePreviewComp->GetRenderTargetTexturePostProcess());
 				}
 			}
