nDisplay: Fix game build.
CIS issue 339715

diff --git a/Engine/Plugins/Runtime/nDisplay/Source/DisplayCluster/Private/Game/EngineClasses/Scene/DisplayDevice/DisplayClusterDisplayDeviceUtils.cpp b/Engine/Plugins/Runtime/nDisplay/Source/DisplayCluster/Private/Game/EngineClasses/Scene/DisplayDevice/DisplayClusterDisplayDeviceUtils.cpp
--- a/Engine/Plugins/Runtime/nDisplay/Source/DisplayCluster/Private/Game/EngineClasses/Scene/DisplayDevice/DisplayClusterDisplayDeviceUtils.cpp
+++ b/Engine/Plugins/Runtime/nDisplay/Source/DisplayCluster/Private/Game/EngineClasses/Scene/DisplayDevice/DisplayClusterDisplayDeviceUtils.cpp
@@ -7,6 +7,8 @@
 #include "DisplayClusterRootActor.h"
 #include "Misc/DisplayClusterLog.h"
 
+#if WITH_EDITOR
+
 UDisplayClusterDisplayDeviceBaseComponent* UE::DisplayClusterDisplayDeviceUtils::FindAndSyncDisplayDeviceFromViewport(
 	UDisplayClusterConfigurationViewport* InViewport)
 {
@@ -89,3 +91,5 @@ UDisplayClusterDisplayDeviceBaseComponent* UE::DisplayClusterDisplayDeviceUtils:
 
 	return nullptr;
 }
+
+#endif
diff --git a/Engine/Plugins/Runtime/nDisplay/Source/DisplayCluster/Private/Game/EngineClasses/Scene/DisplayDevice/DisplayClusterDisplayDeviceUtils.h b/Engine/Plugins/Runtime/nDisplay/Source/DisplayCluster/Private/Game/EngineClasses/Scene/DisplayDevice/DisplayClusterDisplayDeviceUtils.h
--- a/Engine/Plugins/Runtime/nDisplay/Source/DisplayCluster/Private/Game/EngineClasses/Scene/DisplayDevice/DisplayClusterDisplayDeviceUtils.h
+++ b/Engine/Plugins/Runtime/nDisplay/Source/DisplayCluster/Private/Game/EngineClasses/Scene/DisplayDevice/DisplayClusterDisplayDeviceUtils.h
@@ -7,6 +7,7 @@ class UDisplayClusterDisplayDeviceBaseComponent;
 
 namespace UE::DisplayClusterDisplayDeviceUtils
 {
+#if WITH_EDITOR
 	/**
 	 * Attempt to find the display device given a viewport.
 	 * If the display device is not found on an instance, the CDO is checked in case the display device was renamed and the
@@ -16,4 +17,5 @@ namespace UE::DisplayClusterDisplayDeviceUtils
 	 * @return The display device this viewport uses or nullptr.
 	 */
 	UDisplayClusterDisplayDeviceBaseComponent* FindAndSyncDisplayDeviceFromViewport(UDisplayClusterConfigurationViewport* InViewport);
+#endif
 }
