Fix missing symbol on linux caused by marking inline in the cpp a function that is exposed via the header.
#horde 339732
#rnx

diff --git a/Engine/Plugins/Runtime/nDisplay/Source/DisplayCluster/Private/Game/EngineClasses/Scene/DisplayDevice/DisplayClusterDisplayDeviceComponentBase.cpp b/Engine/Plugins/Runtime/nDisplay/Source/DisplayCluster/Private/Game/EngineClasses/Scene/DisplayDevice/DisplayClusterDisplayDeviceComponentBase.cpp
--- a/Engine/Plugins/Runtime/nDisplay/Source/DisplayCluster/Private/Game/EngineClasses/Scene/DisplayDevice/DisplayClusterDisplayDeviceComponentBase.cpp
+++ b/Engine/Plugins/Runtime/nDisplay/Source/DisplayCluster/Private/Game/EngineClasses/Scene/DisplayDevice/DisplayClusterDisplayDeviceComponentBase.cpp
@@ -8,7 +8,7 @@
 #include "Materials/MaterialInstanceDynamic.h"
 #include "UObject/ConstructorHelpers.h"
 
-inline UDisplayClusterDisplayDeviceBaseComponent::UDisplayClusterDisplayDeviceBaseComponent()
+UDisplayClusterDisplayDeviceBaseComponent::UDisplayClusterDisplayDeviceBaseComponent()
 {
 #if WITH_EDITOR
 	static ConstructorHelpers::FObjectFinder<UMaterial> PreviewMaterialObj(TEXT("/nDisplay/Materials/Preview/M_DisplayDevicePreview"));
