nDisplay: Fixed VIOSO math

diff --git a/Engine/Plugins/Runtime/nDisplay/Source/DisplayClusterProjection/Private/Policy/VIOSO/Windows/DisplayClusterProjectionVIOSOWarper.cpp b/Engine/Plugins/Runtime/nDisplay/Source/DisplayClusterProjection/Private/Policy/VIOSO/Windows/DisplayClusterProjectionVIOSOWarper.cpp
--- a/Engine/Plugins/Runtime/nDisplay/Source/DisplayClusterProjection/Private/Policy/VIOSO/Windows/DisplayClusterProjectionVIOSOWarper.cpp
+++ b/Engine/Plugins/Runtime/nDisplay/Source/DisplayClusterProjection/Private/Policy/VIOSO/Windows/DisplayClusterProjectionVIOSOWarper.cpp
@@ -34,14 +34,6 @@ static FAutoConsoleVariableRef CVarDisplayClusterProjectionVIOSOPolicyEnableCust
 	ECVF_RenderThreadSafe
 );
 
-int32 GDisplayClusterProjectionVIOSOPolicyOverrideViewLocation = 0;
-static FAutoConsoleVariableRef CVarDisplayClusterProjectionVIOSOPolicyOverrideViewLocation(
-	TEXT("nDisplay.render.projection.VIOSO.OverrideViewLocation"),
-	GDisplayClusterProjectionVIOSOPolicyOverrideViewLocation,
-	TEXT("Override viewer eye location with values from VIOSO (0 - off).\n"),
-	ECVF_RenderThreadSafe
-);
-
 int32 GDisplayClusterProjectionVIOSOPolicyFrustumFitToViewport = 0;
 static FAutoConsoleVariableRef CVarDisplayClusterProjectionVIOSOPolicyFrustumFitToViewport(
 	TEXT("nDisplay.render.projection.VIOSO.FitFrustumToViewport"),
@@ -58,11 +50,60 @@ static FAutoConsoleVariableRef CVarDisplayClusterProjectionVIOSOPolicySymmetricF
 	ECVF_RenderThreadSafe
 );
 
-int32 GDisplayClusterProjectionVIOSOPolicyEnableViewDirection = 0;
-static FAutoConsoleVariableRef CVarDisplayClusterProjectionVIOSOPolicyEnableViewDirection(
-	TEXT("nDisplay.render.projection.VIOSO.EnableViewDirection"),
-	GDisplayClusterProjectionVIOSOPolicyEnableViewDirection,
-	TEXT("Enable view direction for VIOSO (0 - off).\n"),
+int32 GDisplayClusterProjectionVIOSOPolicyPitchInverse = 0;
+static FAutoConsoleVariableRef CVarDisplayClusterProjectionVIOSOPolicyPitchInverse(
+	TEXT("nDisplay.render.projection.VIOSO.Pitch.Inverse"),
+	GDisplayClusterProjectionVIOSOPolicyPitchInverse,
+	TEXT("Inverse Pitch (default 0).\n"),
+	ECVF_RenderThreadSafe
+);
+
+int32 GDisplayClusterProjectionVIOSOPolicyYawInverse = 1;
+static FAutoConsoleVariableRef CVarDisplayClusterProjectionVIOSOPolicyYawInverse(
+	TEXT("nDisplay.render.projection.VIOSO.Yaw.Inverse"),
+	GDisplayClusterProjectionVIOSOPolicyYawInverse,
+	TEXT("Inverse Yaw (default 1).\n"),
+	ECVF_RenderThreadSafe
+);
+
+int32 GDisplayClusterProjectionVIOSOPolicyRollInverse = 0;
+static FAutoConsoleVariableRef CVarDisplayClusterProjectionVIOSOPolicyRollInverse(
+	TEXT("nDisplay.render.projection.VIOSO.Roll.Inverse"),
+	GDisplayClusterProjectionVIOSOPolicyRollInverse,
+	TEXT("Inverse Roll (default 0).\n"),
+	ECVF_RenderThreadSafe
+);
+
+int32 GDisplayClusterProjectionVIOSOPolicyPitchAxis = 0;
+static FAutoConsoleVariableRef CVarDisplayClusterProjectionVIOSOPolicyPitchAxis(
+	TEXT("nDisplay.render.projection.VIOSO.Pitch.Axis"),
+	GDisplayClusterProjectionVIOSOPolicyPitchAxis,
+	TEXT("The source axis for Pitch (Default 0).\n")
+	TEXT("0 - X Axis\n")
+	TEXT("1 - Y Axis\n")
+	TEXT("2 - Z Axis\n"),
+	ECVF_RenderThreadSafe
+);
+
+int32 GDisplayClusterProjectionVIOSOPolicyYawAxis = 1;
+static FAutoConsoleVariableRef CVarDisplayClusterProjectionVIOSOPolicyYawAxis(
+	TEXT("nDisplay.render.projection.VIOSO.Yaw.Axis"),
+	GDisplayClusterProjectionVIOSOPolicyYawAxis,
+	TEXT("The source axis for Yaw (Default 1).\n")
+	TEXT("0 - X Axis\n")
+	TEXT("1 - Y Axis\n")
+	TEXT("2 - Z Axis\n"),
+	ECVF_RenderThreadSafe
+);
+
+int32 GDisplayClusterProjectionVIOSOPolicyRollAxis = 2;
+static FAutoConsoleVariableRef CVarDisplayClusterProjectionVIOSOPolicyRollAxis(
+	TEXT("nDisplay.render.projection.VIOSO.Roll.Axis"),
+	GDisplayClusterProjectionVIOSOPolicyRollAxis,
+	TEXT("The source axis for Roll (Default 2).\n")
+	TEXT("0 - X Axis\n")
+	TEXT("1 - Y Axis\n")
+	TEXT("2 - Z Axis\n"),
 	ECVF_RenderThreadSafe
 );
 
@@ -79,34 +120,47 @@ namespace UE::DisplayClusterProjection::VIOSOWarper
 
 	static FTransform ViosoToUETransform(FromViosoConventionMatrix);
 
-	static inline FVector3f ToVioso(const FVector& InPos, const float VIOSOGeometryScale, const float WorldToMeters)
+	static inline FVector3f ToVioso(const FVector& InPos, const float WorldToMeters)
 	{
-		float ScaleIn = (VIOSOGeometryScale / WorldToMeters);
+		float ScaleIn = (1.f / WorldToMeters);
 		const FVector Pos = ViosoToUETransform.InverseTransformPosition(InPos) * ScaleIn;
 
 		return FVector3f(Pos.X, Pos.Y, Pos.Z);
 	}
 
-	static inline FVector FromVioso(const FVector3f& InPos, const float VIOSOGeometryScale, const float WorldToMeters)
+	static inline FVector FromVioso(const FVector3f& InPos, const float WorldToMeters)
 	{
-		float ScaleOut = (WorldToMeters / VIOSOGeometryScale);
+		float ScaleOut = WorldToMeters;
 		FVector UELocation = ViosoToUETransform.TransformPosition(FVector(InPos)) * ScaleOut;
 
 		return UELocation;
 	}
 
-	static inline FRotator FromViosoEuler(const FVector3f& InEuler)
+	static inline constexpr EAxis::Type GetViosoAxis(const int32 InValue)
 	{
-		const FVector InEulerDegree = FMath::RadiansToDegrees(FVector(InEuler));
+		switch (InValue)
+		{
+		case 1:
+			return EAxis::Y;
+		case 2:
+			return EAxis::Z;
+
+		default:
+			break;
+		}
 
-		return FRotator(InEulerDegree.Y, -InEulerDegree.X, InEulerDegree.Z);
+		return EAxis::X;
 	}
 
-	static inline FVector3f ToViosoEuler(const FRotator& InRotator)
+	static inline FRotator FromViosoEuler(const FVector3f& InEuler)
 	{
-		const FVector InEulerRad = FMath::DegreesToRadians(FVector(-InRotator.Yaw, InRotator.Pitch, InRotator.Roll));
+		const FVector InEulerDegree = FMath::RadiansToDegrees(FVector(InEuler));
 
-		return FVector3f(InEulerRad);
+		float Pitch = InEulerDegree.GetComponentForAxis(GetViosoAxis(GDisplayClusterProjectionVIOSOPolicyPitchAxis)) * (GDisplayClusterProjectionVIOSOPolicyPitchInverse ? -1 : 0);
+		float Yaw   = InEulerDegree.GetComponentForAxis(GetViosoAxis(GDisplayClusterProjectionVIOSOPolicyYawAxis))   * (GDisplayClusterProjectionVIOSOPolicyYawInverse ? -1 : 0);
+		float Roll  = InEulerDegree.GetComponentForAxis(GetViosoAxis(GDisplayClusterProjectionVIOSOPolicyRollAxis))  * (GDisplayClusterProjectionVIOSOPolicyRollInverse ? -1 : 0);
+
+		return FRotator(Pitch, Yaw, Roll);
 	}
 };
 using namespace UE::DisplayClusterProjection;
@@ -117,8 +171,8 @@ using namespace UE::DisplayClusterProjection;
 bool FDisplayClusterProjectionVIOSOWarper::CalculateViewProjection(IDisplayClusterViewport* InViewport, const uint32 InContextNum, FVector& InOutViewLocation, FRotator& InOutViewRotation, FMatrix& OutProjMatrix, const float WorldToMeters, const float NCP, const float FCP)
 {
 	// Convert to vioso space:
-	FVector3f InOutEye = VIOSOWarper::ToVioso(InOutViewLocation, VIOSOGeometryScale, WorldToMeters);
-	FVector3f InOutRot = GDisplayClusterProjectionVIOSOPolicyEnableViewDirection ? VIOSOWarper::ToViosoEuler(InOutViewRotation) : FVector3f::ZeroVector;
+	FVector3f InOutEye = VIOSOWarper::ToVioso(InOutViewLocation, WorldToMeters);
+	FVector3f InOutRot = FVector3f::ZeroVector;
 
 	FVector3f OutPos = InOutEye;
 	FVector3f OutDir = InOutRot;
@@ -131,14 +185,6 @@ bool FDisplayClusterProjectionVIOSOWarper::CalculateViewProjection(IDisplayClust
 	// VWB_getPosDirClip(VWB_Warper* pWarper, VWB_float* pEye, VWB_float* pRot, VWB_float* pPos, VWB_float* pDir, VWB_float* pClip, bool symmetric, VWB_float aspect);
 	if (VWB_ERROR_NONE == VIOSOLibrary->VWB_getPosDirClip(pWarper, &InOutEye.X, &InOutRot.X, &OutPos.X, &OutDir.X, &VWB_ViewClip[0], bSymmetric, ViewportAspectRatio))
 	{
-		// Convert to UE coordinate system
-		if (GDisplayClusterProjectionVIOSOPolicyOverrideViewLocation)
-		{
-			// VIOSO returns incorrect values in OutPos.
-			const FVector ViewLocationFromVIOSO = VIOSOWarper::FromVioso(OutPos, VIOSOGeometryScale, WorldToMeters);
-			InOutViewLocation = ViewLocationFromVIOSO;
-		}
-
 		InOutViewRotation = VIOSOWarper::FromViosoEuler(OutDir);
 
 		OutProjMatrix = GetProjMatrix(InViewport, InContextNum);
@@ -182,7 +228,7 @@ bool FDisplayClusterProjectionVIOSOWarper::ExportGeometry(FDisplayClusterProject
 			{
 				VWB_WarpBlendVertex& SrcVertex = PreviewMesh.vtx[VertexIndex];
 
-				const FVector Pts = VIOSOWarper::FromVioso(FVector3f(SrcVertex.pos[0], SrcVertex.pos[1], SrcVertex.pos[2]), VIOSOGeometryScale, WorldToMeters);
+				const FVector Pts = VIOSOWarper::FromVioso(FVector3f(SrcVertex.pos[0], SrcVertex.pos[1], SrcVertex.pos[2]), WorldToMeters);
 
 				OutMeshData.Vertices[VertexIndex] = Pts;
 				OutMeshData.UV[VertexIndex] = FVector2D(SrcVertex.uv[0], SrcVertex.uv[1]);
@@ -283,16 +329,13 @@ bool FDisplayClusterProjectionVIOSOWarper::Initialize(void* pDxDevice, const FVe
 			/// as the gamma is already calculated inside blend map, or to fine-tune, defaults to 1 (no change)
 			pWarper->gamma = VIOSOConfigData.Gamma;
 
-			// Default internal vioso scaling in meters
-			VIOSOGeometryScale = VIOSOConfigData.UnitsInMeter / 1000.f;
-
 			// the transformation matrix to go from VIOSO coordinates to IG coordinates, defaults to indentity
 			// note VIOSO maps are always right-handed, to use with a left-handed world like DirectX, invert the z!
 			// UE is always left-handed
-			float Scale = VIOSOConfigData.UnitsInMeter;
+			const float Scale = VIOSOConfigData.UnitsInMeter;
 			const FMatrix BaseMatrix(
 				FPlane(Scale, 0.f,    0.f,   0.f),
-				FPlane(0.f,   Scale,  1.f,   0.f),
+				FPlane(0.f,   Scale,  0.f,   0.f),
 				FPlane(0.f,   0.f,   -Scale, 0.f),
 				FPlane(0.f,   0.f,    0.f,   1.f));
 			
diff --git a/Engine/Plugins/Runtime/nDisplay/Source/DisplayClusterProjection/Private/Policy/VIOSO/Windows/DisplayClusterProjectionVIOSOWarper.h b/Engine/Plugins/Runtime/nDisplay/Source/DisplayClusterProjection/Private/Policy/VIOSO/Windows/DisplayClusterProjectionVIOSOWarper.h
--- a/Engine/Plugins/Runtime/nDisplay/Source/DisplayClusterProjection/Private/Policy/VIOSO/Windows/DisplayClusterProjectionVIOSOWarper.h
+++ b/Engine/Plugins/Runtime/nDisplay/Source/DisplayClusterProjection/Private/Policy/VIOSO/Windows/DisplayClusterProjectionVIOSOWarper.h
@@ -116,8 +116,5 @@ private:
 
 	// Custom clipping planes
 	FVector2D DefaultClippingPlanes;
-
-	// Vioso geometry scale
-	float VIOSOGeometryScale = 1.f;
 };
 #endif
