nDisplay: Rearranged the order of rendering passes for ICVFX: "ChromakeyCameraIntersect" will render before "LightCardOver".

diff --git a/Engine/Plugins/Runtime/nDisplay/Source/DisplayClusterShaders/Private/Shaders/DisplayClusterShadersWarpblend_ICVFX.cpp b/Engine/Plugins/Runtime/nDisplay/Source/DisplayClusterShaders/Private/Shaders/DisplayClusterShadersWarpblend_ICVFX.cpp
--- a/Engine/Plugins/Runtime/nDisplay/Source/DisplayClusterShaders/Private/Shaders/DisplayClusterShadersWarpblend_ICVFX.cpp
+++ b/Engine/Plugins/Runtime/nDisplay/Source/DisplayClusterShaders/Private/Shaders/DisplayClusterShadersWarpblend_ICVFX.cpp
@@ -1134,6 +1134,35 @@ public:
 		return true;
 	}
 
+	inline bool SetRenderPass_LightCardOver()
+	{
+		if (EnumHasAnyFlags(RenderState, EIcvfxPassRenderState::LightCardOver | EIcvfxPassRenderState::UVLightCardOver))
+		{
+			// Render LC over cameras
+			ActiveCameraIndex = INDEX_NONE;
+			RenderPass = EIcvfxPassRenderPass::LightCardOver;
+			
+			return true;
+		}
+
+		return false;
+	}
+
+	inline bool SetRenderPass_InnerFrustumChromakeyOverlapIterator()
+	{
+		// Chromakey overlap pass
+		if (EnumHasAnyFlags(RenderState, EIcvfxPassRenderState::EnableInnerFrustumChromakeyOverlap))
+		{
+			// Render camera overlaps
+			ActiveCameraIndex = 1;
+			RenderPass = EIcvfxPassRenderPass::InnerFrustumChromakeyOverlapIterator;
+			
+			return true;
+		}
+
+		return false;
+	}
+
 	// Render Full Icvfx Warp&Blend
 	bool DoRenderPass(FRHICommandListImmediate& RHICmdList)
 	{
@@ -1164,48 +1193,30 @@ public:
 			if (!ICVFXParameters.IsCameraUsed(++ActiveCameraIndex))
 			{
 				// All cameras have been rendered.
-
-				if (EnumHasAnyFlags(RenderState, EIcvfxPassRenderState::LightCardOver | EIcvfxPassRenderState::UVLightCardOver))
+				if (!SetRenderPass_InnerFrustumChromakeyOverlapIterator() && !SetRenderPass_LightCardOver())
 				{
-					// Render LC over cameras
-					ActiveCameraIndex = INDEX_NONE;
-					RenderPass = EIcvfxPassRenderPass::LightCardOver;
-					break;
-				}
-				else if (EnumHasAnyFlags(RenderState, EIcvfxPassRenderState::EnableInnerFrustumChromakeyOverlap))
-				{
-					// Render camera overlaps
-					ActiveCameraIndex = 1;
-					RenderPass = EIcvfxPassRenderPass::InnerFrustumChromakeyOverlapIterator;
-					break;
+					// Render completed
+					RenderPass = EIcvfxPassRenderPass::None;
+					return false;
 				}
-
-				// Render completed
-				RenderPass = EIcvfxPassRenderPass::None;
-				return false;
 			}
 			break;
 
 		case EIcvfxPassRenderPass::LightCardOver:
-			
-			if (EnumHasAnyFlags(RenderState, EIcvfxPassRenderState::EnableInnerFrustumChromakeyOverlap))
-			{
-				// Render camera overlaps
-				ActiveCameraIndex = 1;
-				RenderPass = EIcvfxPassRenderPass::InnerFrustumChromakeyOverlapIterator;
-				break;
-			}
-
 			// Render completed
 			RenderPass = EIcvfxPassRenderPass::None;
+
 			return false;
 
 		case EIcvfxPassRenderPass::InnerFrustumChromakeyOverlapIterator:
 			if (!ICVFXParameters.IsCameraUsed(++ActiveCameraIndex))
 			{
-				// Render completed
-				RenderPass = EIcvfxPassRenderPass::None;
-				return false;
+				if (!SetRenderPass_LightCardOver())
+				{
+					// Render completed
+					RenderPass = EIcvfxPassRenderPass::None;
+					return false;
+				}
 			}
 			break;
 
