nDisplay: Fixed issue in the ICVFX window where the transform gizmo would render unusably small in orthographic projection. Cause of the issue was a change in the engine that changed the meaning of a property in the view matrices structure for orthographic projections, so this fix applies a complimentary fix to the gizmo's size calculations to account for the change.

diff --git a/Engine/Plugins/Runtime/nDisplay/Source/DisplayClusterLightCardEditor/Private/Viewport/DisplayClusterLightCardEditorWidget.cpp b/Engine/Plugins/Runtime/nDisplay/Source/DisplayClusterLightCardEditor/Private/Viewport/DisplayClusterLightCardEditorWidget.cpp
--- a/Engine/Plugins/Runtime/nDisplay/Source/DisplayClusterLightCardEditor/Private/Viewport/DisplayClusterLightCardEditorWidget.cpp
+++ b/Engine/Plugins/Runtime/nDisplay/Source/DisplayClusterLightCardEditor/Private/Viewport/DisplayClusterLightCardEditorWidget.cpp
@@ -173,7 +173,10 @@ float FDisplayClusterLightCardEditorWidget::GetLengthScreenScalar(const FSceneVi
 	const float ResolutionScale = View->UnconstrainedViewRect.Size().X / (DPIScale * ReferenceResolution);
 
 	const float DistanceFromView = FMath::Max(FVector::Dist(Origin, View->ViewMatrices.GetViewOrigin()), 1.f);
-	const float ProjectionScale = bIsOrthographic ? 1.0 / View->ViewMatrices.GetScreenScale() : DistanceFromView / View->ViewMatrices.GetScreenScale();
+	const FMatrix& ProjectionMatrix = View->ViewMatrices.GetProjectionMatrix();
+	const float ProjectionScale = bIsOrthographic ? 
+		1.0 / (View->ViewMatrices.GetScreenScale() * FMath::Abs(ProjectionMatrix.M[0][0])) :
+		DistanceFromView / View->ViewMatrices.GetScreenScale();
 
 	const float FinalScalar = DPIScale * ResolutionScale * ProjectionScale;
 	return FinalScalar;
