nDisplay: added missing cvar value to control SecondaryScreenPercentage

diff --git a/Engine/Plugins/Runtime/nDisplay/Source/DisplayCluster/Private/Game/EngineClasses/Basics/DisplayClusterViewportClient.cpp b/Engine/Plugins/Runtime/nDisplay/Source/DisplayCluster/Private/Game/EngineClasses/Basics/DisplayClusterViewportClient.cpp
--- a/Engine/Plugins/Runtime/nDisplay/Source/DisplayCluster/Private/Game/EngineClasses/Basics/DisplayClusterViewportClient.cpp
+++ b/Engine/Plugins/Runtime/nDisplay/Source/DisplayCluster/Private/Game/EngineClasses/Basics/DisplayClusterViewportClient.cpp
@@ -719,7 +719,7 @@ void UDisplayClusterViewportClient::Draw(FViewport* InViewport, FCanvas* SceneCa
 				}
 
 				// Set up secondary resolution fraction for the view family.
-				if (!bStereoRendering && ViewFamily.SupportsScreenPercentage())
+				if (ViewFamily.SupportsScreenPercentage())
 				{
 					float CustomSecondaryScreenPercentage = IConsoleManager::Get().FindConsoleVariable(TEXT("r.SecondaryScreenPercentage.GameViewport"), false)->GetFloat();
 					if (CustomSecondaryScreenPercentage > 0.0)
@@ -778,7 +778,7 @@ void UDisplayClusterViewportClient::Draw(FViewport* InViewport, FCanvas* SceneCa
 					float CustomBufferRatio = DCViewFamily.CustomBufferRatio;
 
 					float GlobalResolutionFraction = 1.0f;
-					float SecondaryScreenPercentage = 1.0f;
+					float SecondaryScreenPercentage = ViewFamily.SecondaryViewFraction;
 
 					if (ViewFamily.EngineShowFlags.ScreenPercentage)
 					{
