nDisplay: fixed dx11 quadbuffer stereo rendering

diff --git a/Engine/Plugins/Runtime/nDisplay/Source/DisplayCluster/Private/Render/Device/DisplayClusterDeviceBase.cpp b/Engine/Plugins/Runtime/nDisplay/Source/DisplayCluster/Private/Render/Device/DisplayClusterDeviceBase.cpp
--- a/Engine/Plugins/Runtime/nDisplay/Source/DisplayCluster/Private/Render/Device/DisplayClusterDeviceBase.cpp
+++ b/Engine/Plugins/Runtime/nDisplay/Source/DisplayCluster/Private/Render/Device/DisplayClusterDeviceBase.cpp
@@ -370,13 +370,6 @@ void FDisplayClusterDeviceBase::FinalizeNewFrame()
 
 	// reset viewport manager ptr on game thread
 	ViewportManagerWeakPtr.Reset();
-
-	// Reset viewport manager proxy on render thread
-	ENQUEUE_RENDER_COMMAND(DisplayClusterDevice_ResetViewportManagerProxy)(
-		[DCRenderDevice = this](FRHICommandListImmediate& RHICmdList)
-		{
-			DCRenderDevice->ViewportManagerProxyWeakPtr.Reset();
-		});
 }
 
 DECLARE_GPU_STAT_NAMED(nDisplay_Device_RenderTexture, TEXT("nDisplay RenderDevice::RenderTexture"));
@@ -421,6 +414,9 @@ void FDisplayClusterDeviceBase::RenderTexture_RenderThread(FRHICommandListImmedi
 			ClearRenderTarget(RHICmdList, SrcTexture);
 		}
 	}
+
+	// Reset viewport manager proxy on render thread
+	ViewportManagerProxyWeakPtr.Reset();
 }
 
 //////////////////////////////////////////////////////////////////////////////////////////////
diff --git a/Engine/Plugins/Runtime/nDisplay/Source/DisplayCluster/Private/Render/Device/DisplayClusterDeviceBase.h b/Engine/Plugins/Runtime/nDisplay/Source/DisplayCluster/Private/Render/Device/DisplayClusterDeviceBase.h
--- a/Engine/Plugins/Runtime/nDisplay/Source/DisplayCluster/Private/Render/Device/DisplayClusterDeviceBase.h
+++ b/Engine/Plugins/Runtime/nDisplay/Source/DisplayCluster/Private/Render/Device/DisplayClusterDeviceBase.h
@@ -102,7 +102,7 @@ private:
 	TWeakPtr<IDisplayClusterViewportManager, ESPMode::ThreadSafe> ViewportManagerWeakPtr;
 
 	// Pointer to the DC ViewportManagerProxy from the active DCRA that is currently being used for rendering.
-	TWeakPtr<FDisplayClusterViewportManagerProxy, ESPMode::ThreadSafe> ViewportManagerProxyWeakPtr;
+	mutable TWeakPtr<FDisplayClusterViewportManagerProxy, ESPMode::ThreadSafe> ViewportManagerProxyWeakPtr;
 
 	EDisplayClusterRenderFrameMode RenderFrameMode = EDisplayClusterRenderFrameMode::Mono;
 	int32 DesiredNumberOfViews = 0;
