[nDisplay] Fix for media settings backward compatibility path
- fixed media output migration
- added missing media input migration
note: previewed in 26704092 (https://p4-swarm.epicgames.net/reviews/26787042)

diff --git a/Engine/Plugins/Runtime/nDisplay/Source/DisplayClusterConfiguration/Private/DisplayClusterConfigurationTypes.cpp b/Engine/Plugins/Runtime/nDisplay/Source/DisplayClusterConfiguration/Private/DisplayClusterConfigurationTypes.cpp
--- a/Engine/Plugins/Runtime/nDisplay/Source/DisplayClusterConfiguration/Private/DisplayClusterConfigurationTypes.cpp
+++ b/Engine/Plugins/Runtime/nDisplay/Source/DisplayClusterConfiguration/Private/DisplayClusterConfigurationTypes.cpp
@@ -200,13 +200,17 @@ void UDisplayClusterConfigurationViewport::PostLoad()
 
 #if WITH_EDITOR
 	PRAGMA_DISABLE_DEPRECATION_WARNINGS
-		if (IsValid(RenderSettings.Media.MediaOutput_DEPRECATED))
+		if (!IsValid(RenderSettings.Media.MediaInput.MediaSource) && IsValid(RenderSettings.Media.MediaSource))
 		{
-			RenderSettings.Media.MediaOutputs.Add(
-				{ RenderSettings.Media.MediaOutput_DEPRECATED , RenderSettings.Media.OutputSyncPolicy_DEPRECATED });
+			RenderSettings.Media.MediaInput.MediaSource = RenderSettings.Media.MediaSource;
+			RenderSettings.Media.MediaSource = nullptr;
+		}
 
-			RenderSettings.Media.MediaOutput_DEPRECATED = nullptr;
-			RenderSettings.Media.OutputSyncPolicy_DEPRECATED = nullptr;
+		if (RenderSettings.Media.MediaOutputs.IsEmpty() && IsValid(RenderSettings.Media.MediaOutput))
+		{
+			RenderSettings.Media.MediaOutputs.Add({ RenderSettings.Media.MediaOutput , RenderSettings.Media.OutputSyncPolicy });
+			RenderSettings.Media.MediaOutput = nullptr;
+			RenderSettings.Media.OutputSyncPolicy = nullptr;
 		}
 	PRAGMA_ENABLE_DEPRECATION_WARNINGS
 #endif // WITH_EDITOR
@@ -309,13 +313,17 @@ void UDisplayClusterConfigurationClusterNode::PostLoad()
 
 #if WITH_EDITOR
 	PRAGMA_DISABLE_DEPRECATION_WARNINGS
-		if (IsValid(Media.MediaOutput_DEPRECATED))
+		if (!IsValid(Media.MediaInput.MediaSource) && IsValid(Media.MediaSource))
 		{
-			Media.MediaOutputs.Add(
-				{ Media.MediaOutput_DEPRECATED , Media.OutputSyncPolicy_DEPRECATED });
+			Media.MediaInput.MediaSource = Media.MediaSource;
+			Media.MediaSource = nullptr;
+		}
 
-			Media.MediaOutput_DEPRECATED = nullptr;
-			Media.OutputSyncPolicy_DEPRECATED = nullptr;
+		if (Media.MediaOutputs.IsEmpty() && IsValid(Media.MediaOutput))
+		{
+			Media.MediaOutputs.Add({ Media.MediaOutput, Media.OutputSyncPolicy });
+			Media.MediaOutput = nullptr;
+			Media.OutputSyncPolicy = nullptr;
 		}
 	PRAGMA_ENABLE_DEPRECATION_WARNINGS
 #endif // WITH_EDITOR
diff --git a/Engine/Plugins/Runtime/nDisplay/Source/DisplayClusterConfiguration/Public/DisplayClusterConfigurationTypes_Media.h b/Engine/Plugins/Runtime/nDisplay/Source/DisplayClusterConfiguration/Public/DisplayClusterConfigurationTypes_Media.h
--- a/Engine/Plugins/Runtime/nDisplay/Source/DisplayClusterConfiguration/Public/DisplayClusterConfigurationTypes_Media.h
+++ b/Engine/Plugins/Runtime/nDisplay/Source/DisplayClusterConfiguration/Public/DisplayClusterConfigurationTypes_Media.h
@@ -79,21 +79,24 @@ public:
 	/** Returns true if at least one media output assigned */
 	bool IsMediaOutputAssigned() const;
 
-protected:
 #if WITH_EDITORONLY_DATA
 
+protected:
 	UPROPERTY(meta = (DeprecatedProperty, DeprecationMessage = "This property has been deprecated"))
 	FString MediaSharingNode_DEPRECATED;
 
 public:
+	UE_DEPRECATED(5.3, "This property has been deprecated. Please refer new MediaInput property.")
 	UPROPERTY(meta = (DeprecatedProperty, DeprecationMessage = "This property has been deprecated. Please refer new MediaInput property."))
-	TObjectPtr<UMediaSource> MediaSource_DEPRECATED;
+	TObjectPtr<UMediaSource> MediaSource;
 
+	UE_DEPRECATED(5.3, "This property has been deprecated. Please refer new MediaOutputs property.")
 	UPROPERTY(meta = (DeprecatedProperty, DeprecationMessage = "This property has been deprecated. Please refer new MediaOutputs property."))
-	TObjectPtr<UMediaOutput> MediaOutput_DEPRECATED;
+	TObjectPtr<UMediaOutput> MediaOutput;
 
-	UPROPERTY(meta = (DeprecatedProperty, DeprecationMessage = "This property has been deprecated. Please refer new MediaOutputs property."))
-	TObjectPtr<UDisplayClusterMediaOutputSynchronizationPolicy> OutputSyncPolicy_DEPRECATED;
+	UE_DEPRECATED(5.3, "This property has been deprecated.")
+	UPROPERTY(meta = (DeprecatedProperty, DeprecationMessage = "This property has been deprecated."))
+	TObjectPtr<UDisplayClusterMediaOutputSynchronizationPolicy> OutputSyncPolicy;
 
 #endif // WITH_EDITORONLY_DATA
 };
diff --git a/Engine/Plugins/Runtime/nDisplay/Source/DisplayClusterConfiguration/Public/DisplayClusterConfigurationTypes_Viewport.h b/Engine/Plugins/Runtime/nDisplay/Source/DisplayClusterConfiguration/Public/DisplayClusterConfigurationTypes_Viewport.h
--- a/Engine/Plugins/Runtime/nDisplay/Source/DisplayClusterConfiguration/Public/DisplayClusterConfigurationTypes_Viewport.h
+++ b/Engine/Plugins/Runtime/nDisplay/Source/DisplayClusterConfiguration/Public/DisplayClusterConfigurationTypes_Viewport.h
@@ -165,7 +165,7 @@ public:
 
 public:
 	//~ Begin UObject interface
-	void PostLoad() override;
+	virtual void PostLoad() override;
 
 #if WITH_EDITOR
 	virtual void PreEditChange(FEditPropertyChain& PropertyAboutToChange) override;
