[nDisplay] Fixed default networking settings

diff --git a/Engine/Plugins/Runtime/nDisplay/Source/DisplayClusterConfiguration/Private/DisplayClusterConfigurationTypes.cpp b/Engine/Plugins/Runtime/nDisplay/Source/DisplayClusterConfiguration/Private/DisplayClusterConfigurationTypes.cpp
--- a/Engine/Plugins/Runtime/nDisplay/Source/DisplayClusterConfiguration/Private/DisplayClusterConfigurationTypes.cpp
+++ b/Engine/Plugins/Runtime/nDisplay/Source/DisplayClusterConfiguration/Private/DisplayClusterConfigurationTypes.cpp
@@ -417,16 +417,6 @@ FDisplayClusterConfigurationClusterSync::FDisplayClusterConfigurationClusterSync
 	InputSyncPolicy.Type  = cluster::input_sync::InputSyncPolicyReplicatePrimary;
 }
 
-FDisplayClusterConfigurationNetworkSettings::FDisplayClusterConfigurationNetworkSettings()
-	: ConnectRetriesAmount     (300)   // ...
-	, ConnectRetryDelay        (1000)  // 5 minutes unless all nodes up
-	, GameStartBarrierTimeout  (1000 * 3600 * 5) // 5 hours unless ready to start rendering
-	, FrameStartBarrierTimeout (1000 * 60 * 30)  // 30 minutes for the barrier
-	, FrameEndBarrierTimeout   (1000 * 60 * 30)  // 30 minutes for the barrier
-	, RenderSyncBarrierTimeout (1000 * 60 * 30)  // 30 minutes for the barrier
-{
-}
-
 void UDisplayClusterConfigurationViewport::GetReferencedMeshNames(TArray<FString>& OutMeshNames) const
 {
 	// Collect all mesh references from projection policies
diff --git a/Engine/Plugins/Runtime/nDisplay/Source/DisplayClusterConfiguration/Private/Formats/JSON500/DisplayClusterConfigurationJsonParser_500.cpp b/Engine/Plugins/Runtime/nDisplay/Source/DisplayClusterConfiguration/Private/Formats/JSON500/DisplayClusterConfigurationJsonParser_500.cpp
--- a/Engine/Plugins/Runtime/nDisplay/Source/DisplayClusterConfiguration/Private/Formats/JSON500/DisplayClusterConfigurationJsonParser_500.cpp
+++ b/Engine/Plugins/Runtime/nDisplay/Source/DisplayClusterConfiguration/Private/Formats/JSON500/DisplayClusterConfigurationJsonParser_500.cpp
@@ -188,22 +188,22 @@ namespace JSON500
 			// Network
 			{
 				Config->Cluster->Network.ConnectRetriesAmount = DisplayClusterHelpers::map::template ExtractValueFromString(
-					CfgJson.Cluster.Network, FString(DisplayClusterConfigurationStrings::config::cluster::network::NetConnectRetriesAmount), (uint32)15);
+					CfgJson.Cluster.Network, FString(DisplayClusterConfigurationStrings::config::cluster::network::NetConnectRetriesAmount), static_cast<uint32>(Config->Cluster->Network.ConnectRetriesAmount));
 
 				Config->Cluster->Network.ConnectRetryDelay = DisplayClusterHelpers::map::template ExtractValueFromString(
-					CfgJson.Cluster.Network, FString(DisplayClusterConfigurationStrings::config::cluster::network::NetConnectRetryDelay), (uint32)1000);
+					CfgJson.Cluster.Network, FString(DisplayClusterConfigurationStrings::config::cluster::network::NetConnectRetryDelay), static_cast<uint32>(Config->Cluster->Network.ConnectRetryDelay));
 
 				Config->Cluster->Network.GameStartBarrierTimeout = DisplayClusterHelpers::map::template ExtractValueFromString(
-					CfgJson.Cluster.Network, FString(DisplayClusterConfigurationStrings::config::cluster::network::NetGameStartBarrierTimeout), (uint32)30000);
+					CfgJson.Cluster.Network, FString(DisplayClusterConfigurationStrings::config::cluster::network::NetGameStartBarrierTimeout), static_cast<uint32>(Config->Cluster->Network.GameStartBarrierTimeout));
 
 				Config->Cluster->Network.FrameStartBarrierTimeout = DisplayClusterHelpers::map::template ExtractValueFromString(
-					CfgJson.Cluster.Network, FString(DisplayClusterConfigurationStrings::config::cluster::network::NetFrameStartBarrierTimeout), (uint32)5000);
+					CfgJson.Cluster.Network, FString(DisplayClusterConfigurationStrings::config::cluster::network::NetFrameStartBarrierTimeout), static_cast<uint32>(Config->Cluster->Network.FrameStartBarrierTimeout));
 
 				Config->Cluster->Network.FrameEndBarrierTimeout = DisplayClusterHelpers::map::template ExtractValueFromString(
-					CfgJson.Cluster.Network, FString(DisplayClusterConfigurationStrings::config::cluster::network::NetFrameEndBarrierTimeout), (uint32)5000);
+					CfgJson.Cluster.Network, FString(DisplayClusterConfigurationStrings::config::cluster::network::NetFrameEndBarrierTimeout), static_cast<uint32>(Config->Cluster->Network.FrameEndBarrierTimeout));
 
 				Config->Cluster->Network.RenderSyncBarrierTimeout = DisplayClusterHelpers::map::template ExtractValueFromString(
-					CfgJson.Cluster.Network, FString(DisplayClusterConfigurationStrings::config::cluster::network::NetRenderSyncBarrierTimeout), (uint32)5000);
+					CfgJson.Cluster.Network, FString(DisplayClusterConfigurationStrings::config::cluster::network::NetRenderSyncBarrierTimeout), static_cast<uint32>(Config->Cluster->Network.RenderSyncBarrierTimeout));
 			}
 
 			// Failover
diff --git a/Engine/Plugins/Runtime/nDisplay/Source/DisplayClusterConfiguration/Public/DisplayClusterConfigurationTypes.h b/Engine/Plugins/Runtime/nDisplay/Source/DisplayClusterConfiguration/Public/DisplayClusterConfigurationTypes.h
--- a/Engine/Plugins/Runtime/nDisplay/Source/DisplayClusterConfiguration/Public/DisplayClusterConfigurationTypes.h
+++ b/Engine/Plugins/Runtime/nDisplay/Source/DisplayClusterConfiguration/Public/DisplayClusterConfigurationTypes.h
@@ -233,33 +233,30 @@ struct DISPLAYCLUSTERCONFIGURATION_API FDisplayClusterConfigurationNetworkSettin
 {
 	GENERATED_BODY()
 
-public:
-	FDisplayClusterConfigurationNetworkSettings();
-
 public:
 	/** Advanced: amount of times nDisplay tries to reconnect before dropping */
-	UPROPERTY(EditAnywhere, BlueprintReadWrite, Category = NDisplay, meta = (ClampMin = "1", ClampMax = "99", UIMin = "1", UIMax = "99"))
-	int32 ConnectRetriesAmount;
+	UPROPERTY(EditAnywhere, BlueprintReadWrite, Category = NDisplay, meta = (ClampMin = "1", ClampMax = "500", UIMin = "1", UIMax = "500"))
+	int32 ConnectRetriesAmount = 300;
 
 	/** Advanced: delay in between connection retries */
 	UPROPERTY(EditAnywhere, BlueprintReadWrite, Category = NDisplay, meta = (ClampMin = "1", ClampMax = "5000", UIMin = "1", UIMax = "5000"))
-	int32 ConnectRetryDelay;
+	int32 ConnectRetryDelay = 1000; // 1 second between attempts
 
 	/** Advanced: timeout for Game Thread Barrier */
 	UPROPERTY(EditAnywhere, BlueprintReadWrite, Category = NDisplay, meta = (ClampMin = "5000", UIMin = "5000"))
-	int32 GameStartBarrierTimeout;
+	int32 GameStartBarrierTimeout = 1000 * 3600 * 5; // 5 hours
 
 	/** Advanced: timeout value for Start Frame Barrier */
 	UPROPERTY(EditAnywhere, BlueprintReadWrite, Category = NDisplay, meta = (ClampMin = "1", UIMin = "1"))
-	int32 FrameStartBarrierTimeout;
+	int32 FrameStartBarrierTimeout = 1000 * 60 * 30; // 30 minutes
 
 	/** Advanced: timeout value for End Frame Barrier */
 	UPROPERTY(EditAnywhere, BlueprintReadWrite, Category = NDisplay, meta = (ClampMin = "1", UIMin = "1"))
-	int32 FrameEndBarrierTimeout;
+	int32 FrameEndBarrierTimeout = 1000 * 60 * 30; // 30 minutes
 
 	/** Advanced: timeout value for Render Sync */
 	UPROPERTY(EditAnywhere, BlueprintReadWrite, Category = NDisplay, meta = (ClampMin = "1", UIMin = "1"))
-	int32 RenderSyncBarrierTimeout;
+	int32 RenderSyncBarrierTimeout = 1000 * 60 * 30; // 30 minutes
 };
 
 USTRUCT(Blueprintable)
