nDisplay: Removed the stage geometry normal map visualization features, as they were generally not exposed or not useful, and were causing DDC caching issues when loading up a stage on a node due to the use of a UTexture2D to store a visualizable version of the normal map. Visualization of the rendered stage geometry will be addressed in a more user friendly way in a future task

diff --git a/Engine/Plugins/Runtime/nDisplay/Source/DisplayCluster/Private/Game/EngineClasses/Scene/DisplayClusterStageGeometryComponent.cpp b/Engine/Plugins/Runtime/nDisplay/Source/DisplayCluster/Private/Game/EngineClasses/Scene/DisplayClusterStageGeometryComponent.cpp
--- a/Engine/Plugins/Runtime/nDisplay/Source/DisplayCluster/Private/Game/EngineClasses/Scene/DisplayClusterStageGeometryComponent.cpp
+++ b/Engine/Plugins/Runtime/nDisplay/Source/DisplayCluster/Private/Game/EngineClasses/Scene/DisplayClusterStageGeometryComponent.cpp
@@ -10,8 +10,6 @@
 #include "TextureResource.h"
 #include "Components/DisplayClusterScreenComponent.h"
 #include "Engine/TextureRenderTarget2D.h"
-#include "Brushes/SlateImageBrush.h"
-#include "Widgets/Images/SImage.h"
 
 const uint32 UDisplayClusterStageGeometryComponent::GeometryMapSize = 512;
 const float UDisplayClusterStageGeometryComponent::GeometryMapFOV = 2.0f * FMath::RadiansToDegrees(FMath::Atan(0.55 * PI));
@@ -34,65 +32,6 @@ void UDisplayClusterStageGeometryComponent::TickComponent(float DeltaTime, ELeve
 	}
 }
 
-#if WITH_EDITOR
-namespace UE::DisplayClusterStageGeometryComponent
-{
-	class STextureImage : public SCompoundWidget
-	{
-	public:
-		SLATE_BEGIN_ARGS(STextureImage) {}
-			SLATE_ARGUMENT(UTexture2D*, Texture)
-		SLATE_END_ARGS()
-
-		void Construct(const FArguments& InArgs)
-		{
-			if (InArgs._Texture)
-			{
-				FTextureResource* Resource = InArgs._Texture->GetResource();
-				if (Resource)
-				{
-					FVector2D Size(Resource->GetSizeX(), Resource->GetSizeY());
-					TextureBrush = MakeShared<FSlateImageBrush>(InArgs._Texture,
-						Size,
-						FSlateColor(FLinearColor(1, 1, 1, 1)),
-						ESlateBrushTileType::NoTile,
-						ESlateBrushImageType::Linear);
-				}
-			}
-
-			if (!TextureBrush.IsValid())
-			{
-				TextureBrush = MakeShared<FSlateBrush>();
-			}
-
-			ChildSlot
-			[
-				SNew(SImage)
-				.Image(TextureBrush.Get())
-			];
-		}
-
-	private:
-		TSharedPtr<FSlateBrush> TextureBrush;
-	};
-}
-
-TSharedPtr<SWidget> UDisplayClusterStageGeometryComponent::GetCustomEditorPreviewWidget()
-{
-	return SNew(SVerticalBox)
-		+ SVerticalBox::Slot()
-		[
-			SNew(UE::DisplayClusterStageGeometryComponent::STextureImage)
-			.Texture(NorthGeometryMap.Texture)
-		]
-		+ SVerticalBox::Slot()
-		[
-			SNew(UE::DisplayClusterStageGeometryComponent::STextureImage)
-			.Texture(SouthGeometryMap.Texture)
-		];
-}
-#endif
-
 void UDisplayClusterStageGeometryComponent::Invalidate(bool bForceImmediateRedraw)
 {
 	if (bForceImmediateRedraw)
@@ -202,27 +141,12 @@ UTextureRenderTarget2D* UDisplayClusterStageGeometryComponent::CreateRenderTarge
 	return RenderTarget;
 }
 
-UTexture2D* UDisplayClusterStageGeometryComponent::CreateTexture2D()
-{
-#if WITH_EDITOR
-	UTexture2D* Texture2D = NewObject<UTexture2D>(this);
-	Texture2D->Source.Init(GeometryMapSize, GeometryMapSize, 1, 1, TSF_BGRA8);
-
-	return Texture2D;
-#else
-	return nullptr;
-#endif
-}
-
 void UDisplayClusterStageGeometryComponent::RedrawGeometryMap()
 {
 	if (!bGeometryMapLoaded)
 	{
 		NorthGeometryMap.RenderTarget = CreateRenderTarget();
-		NorthGeometryMap.Texture = CreateTexture2D();
-
 		SouthGeometryMap.RenderTarget = CreateRenderTarget();
-		SouthGeometryMap.Texture = CreateTexture2D();
 
 		bGeometryMapLoaded = true;
 	}
@@ -311,26 +235,6 @@ void UDisplayClusterStageGeometryComponent::GenerateGeometryMap(bool bIsNorthMap
 			Canvas.Flush_GameThread();
 
 			TexResource->ReadFloat16Pixels(GeometryMap.GeometryData);
-
-#if WITH_EDITOR
-			if (GeometryMap.Texture)
-			{
-				uint8* MipData = GeometryMap.Texture->Source.LockMip(0);
-				for (int32 Index = 0; Index < GeometryMap.GeometryData.Num(); ++Index)
-				{
-					uint8* DestPtr = &MipData[Index * sizeof(FColor)];
-					const FColor SourceColor = GeometryMap.GeometryData[Index].GetFloats().ToFColor(false);
-
-					*DestPtr++ = SourceColor.B;
-					*DestPtr++ = SourceColor.G;
-					*DestPtr++ = SourceColor.R;
-					*DestPtr++ = 0xFF;
-				}
-				GeometryMap.Texture->Source.UnlockMip(0);
-
-				GeometryMap.Texture->PostEditChange();
-			}
-#endif
 		}
 	}
 }
diff --git a/Engine/Plugins/Runtime/nDisplay/Source/DisplayCluster/Public/Components/DisplayClusterStageGeometryComponent.h b/Engine/Plugins/Runtime/nDisplay/Source/DisplayCluster/Public/Components/DisplayClusterStageGeometryComponent.h
--- a/Engine/Plugins/Runtime/nDisplay/Source/DisplayCluster/Public/Components/DisplayClusterStageGeometryComponent.h
+++ b/Engine/Plugins/Runtime/nDisplay/Source/DisplayCluster/Public/Components/DisplayClusterStageGeometryComponent.h
@@ -23,10 +23,6 @@ struct FDisplayClusterStageGeometryMap
 	UPROPERTY()
 	TObjectPtr<UTextureRenderTarget2D> RenderTarget;
 
-	/** A displayable texture of the geometry map */
-	UPROPERTY()
-	TObjectPtr<UTexture2D> Texture;
-
 	/** The raw texture data of the geometry map, which stores the normal vector in the RGB components and the depth in the A component */
 	TArray<FFloat16Color> GeometryData;
 
@@ -52,11 +48,6 @@ public:
 
 	//~ Begin UObject interface
 	virtual void TickComponent(float DeltaTime, ELevelTick TickType, FActorComponentTickFunction* ThisTickFunction) override;
-
-#if WITH_EDITOR
-	virtual bool GetEditorPreviewInfo(float DeltaTime, FMinimalViewInfo& ViewOut) override { return true; }
-	virtual TSharedPtr<SWidget> GetCustomEditorPreviewWidget() override;
-#endif
 	//~ End UObject interface
 
 	/**
@@ -87,19 +78,10 @@ public:
 	 */
 	bool MorphProceduralMesh(UProceduralMeshComponent* InProceduralMeshComponent);
 
-	/** Gets the geometry map texture for the northern hemisphere of the stage */
-	UTexture2D* GetNorthGeometryMapTexture() const { return NorthGeometryMap.Texture; }
-
-	/** Gets the geometry map texture for the southern hemisphere of the stage */
-	UTexture2D* GetSouthGeometryMapTexture() const { return SouthGeometryMap.Texture; }
-
 private:
 	/** Creates a new render target that can be used to render the stage's geometry map */
 	UTextureRenderTarget2D* CreateRenderTarget();
 
-	/** Creates a new texture that can be used to display the stage's geometry map */
-	UTexture2D* CreateTexture2D();
-
 	/** Redraws the stage's geometry map */
 	void RedrawGeometryMap();
 
diff --git a/Engine/Plugins/Runtime/nDisplay/Source/DisplayClusterLightCardEditor/Private/Viewport/DisplayClusterLightCardEditorViewportClient.cpp b/Engine/Plugins/Runtime/nDisplay/Source/DisplayClusterLightCardEditor/Private/Viewport/DisplayClusterLightCardEditorViewportClient.cpp
--- a/Engine/Plugins/Runtime/nDisplay/Source/DisplayClusterLightCardEditor/Private/Viewport/DisplayClusterLightCardEditorViewportClient.cpp
+++ b/Engine/Plugins/Runtime/nDisplay/Source/DisplayClusterLightCardEditor/Private/Viewport/DisplayClusterLightCardEditorViewportClient.cpp
@@ -73,13 +73,6 @@ static TAutoConsoleVariable<bool> CVarICVFXPanelAutoPan(
 	true,
 	TEXT("When true, the stage view will automatically tilt and pan on supported map projections as you drag objects near the edges of the screen."));
 
-int32 GICVFXPanelDisplayNormalMapVisualization = 0;
-static FAutoConsoleVariableRef CVarICVFXPanelDisplayNormalMapVisualization(
-	TEXT("nDisplay.panel.DisplayNormalMapVisualization"),
-	GICVFXPanelDisplayNormalMapVisualization,
-	TEXT("Displays the normal map visualization of the stage's geometry map, which is used to determine physical stage geometry"),
-	ECVF_RenderThreadSafe
-);
 //////////////////////////////////////////////////////////////////////////
 // FDisplayClusterLightCardEditorViewportClient
 
@@ -383,20 +376,6 @@ void FDisplayClusterLightCardEditorViewportClient::Draw(FViewport* InViewport, F
 		DrawCanvas(*Viewport, *View, *Canvas);
 	}
 
-	if (bDisplayNormalMapVisualization || GICVFXPanelDisplayNormalMapVisualization)
-	{
-		auto DrawNormalMap = [Canvas, this](bool bShowNorthMap, FVector2D Position)
-		{
-			if (const UTexture* NormalMapTexture = ProjectionHelper->GetNormalMapTexture(bShowNorthMap))
-			{
-				Canvas->DrawTile(Position.X, Position.Y, 512, 512, 0, 0, 1, 1, FLinearColor::White, NormalMapTexture->GetResource());
-			}
-		};
-
-		DrawNormalMap(true, FVector2D(0.0f, 0.0f));
-		DrawNormalMap(false, FVector2D(0.0f, 512.0f));
-	}
-
 	// Remove temporary debug lines.
 	// Possibly a hack. Lines may get added without the scene being rendered etc.
 	if (World->LineBatcher != NULL && (World->LineBatcher->BatchedLines.Num() || World->LineBatcher->BatchedPoints.Num() || World->LineBatcher->BatchedMeshes.Num() ) )
diff --git a/Engine/Plugins/Runtime/nDisplay/Source/DisplayClusterLightCardEditor/Private/Viewport/DisplayClusterLightCardEditorViewportClient.h b/Engine/Plugins/Runtime/nDisplay/Source/DisplayClusterLightCardEditor/Private/Viewport/DisplayClusterLightCardEditorViewportClient.h
--- a/Engine/Plugins/Runtime/nDisplay/Source/DisplayClusterLightCardEditor/Private/Viewport/DisplayClusterLightCardEditorViewportClient.h
+++ b/Engine/Plugins/Runtime/nDisplay/Source/DisplayClusterLightCardEditor/Private/Viewport/DisplayClusterLightCardEditorViewportClient.h
@@ -398,9 +398,6 @@ private:
 	/** The increment to change the FOV by when using the scroll wheel */
 	float FOVScrollIncrement = 5.0f;
 
-	/** Indicates if the normal map should be displayed to the screen */
-	bool bDisplayNormalMapVisualization = false;
-
 	/** Current input mode */
 	EInputMode InputMode = EInputMode::Idle;
 
diff --git a/Engine/Plugins/Runtime/nDisplay/Source/DisplayClusterScenePreview/Private/DisplayClusterLightCardEditorHelper.cpp b/Engine/Plugins/Runtime/nDisplay/Source/DisplayClusterScenePreview/Private/DisplayClusterLightCardEditorHelper.cpp
--- a/Engine/Plugins/Runtime/nDisplay/Source/DisplayClusterScenePreview/Private/DisplayClusterLightCardEditorHelper.cpp
+++ b/Engine/Plugins/Runtime/nDisplay/Source/DisplayClusterScenePreview/Private/DisplayClusterLightCardEditorHelper.cpp
@@ -240,12 +240,8 @@ void FDisplayClusterLightCardEditorHelper::SetLevelInstanceRootActor(ADisplayClu
 
 const UTexture2D* FDisplayClusterLightCardEditorHelper::GetNormalMapTexture(bool bShowNorthMap)
 {
-	if (ADisplayClusterRootActor* RootActor = UpdateRootActor())
-	{
-		UTexture2D* NormalMapTexture = bShowNorthMap ? RootActor->GetStageGeometryComponent()->GetNorthGeometryMapTexture() : RootActor->GetStageGeometryComponent()->GetSouthGeometryMapTexture();
-		return NormalMapTexture;
-	}
-
+	// TODO: Exposure to the raw stage geometry normal maps was removed since they are generally undecipherable, and an upcoming task (UE-153862) will add a
+	// way to visualize the normal maps in a way that is far more user friendly
 	return nullptr;
 }
 
diff --git a/Engine/Plugins/Runtime/nDisplay/Source/DisplayClusterScenePreview/Public/DisplayClusterLightCardEditorHelper.h b/Engine/Plugins/Runtime/nDisplay/Source/DisplayClusterScenePreview/Public/DisplayClusterLightCardEditorHelper.h
--- a/Engine/Plugins/Runtime/nDisplay/Source/DisplayClusterScenePreview/Public/DisplayClusterLightCardEditorHelper.h
+++ b/Engine/Plugins/Runtime/nDisplay/Source/DisplayClusterScenePreview/Public/DisplayClusterLightCardEditorHelper.h
@@ -113,6 +113,7 @@ public:
 	 *
 	 * @param bShowNorthMap If true, return the texture for the north normal map. Otherwise, return the texture for the south normal map.
 	 */
+	UE_DEPRECATED(5.3, "The raw stage geoemetry normal maps are no longer exposed by the stage geometry component")
 	DISPLAYCLUSTERSCENEPREVIEW_API const UTexture2D* GetNormalMapTexture(bool bShowNorthMap);
 
 	/**
diff --git a/Engine/Plugins/VirtualProduction/EpicStageApp/Source/EpicStageApp/Private/StageAppRouteHandler.cpp b/Engine/Plugins/VirtualProduction/EpicStageApp/Source/EpicStageApp/Private/StageAppRouteHandler.cpp
--- a/Engine/Plugins/VirtualProduction/EpicStageApp/Source/EpicStageApp/Private/StageAppRouteHandler.cpp
+++ b/Engine/Plugins/VirtualProduction/EpicStageApp/Source/EpicStageApp/Private/StageAppRouteHandler.cpp
@@ -32,12 +32,6 @@
 
 #define LOCTEXT_NAMESPACE "StageAppRouteHandler"
 
-static TAutoConsoleVariable<bool> CVarStageAppDebugLightCardHelperNormals(
-	TEXT("StageApp.DebugLightCardHelperNormals"),
-	false,
-	TEXT("Whether to overlay preview renders with a debug display of the normal maps for the lightcard helper.")
-);
-
 static TAutoConsoleVariable<float> CVarStageAppDragTimeoutCheckInterval(
 	TEXT("StageApp.DragTimeoutCheckInterval"),
 	2.f,
@@ -558,30 +552,6 @@ void FStageAppRouteHandler::HandleWebSocketNDisplayPreviewRender(const FRemoteCo
 				return;
 			}
 
-			if (CVarStageAppDebugLightCardHelperNormals.GetValueOnGameThread())
-			{
-				// Draw the normal texture in the top-left corner if available
-				if (PerRendererData)
-				{
-					FDisplayClusterLightCardEditorHelper& Helper = PerRendererData->GetLightCardHelper();
-					const FGameTime Time = FGameTime::GetTimeSinceAppStart();
-					const ERHIFeatureLevel::Type FeatureLevel = PerRendererData->GetRootActor()->GetWorld()->Scene->GetFeatureLevel();
-					const FIntPoint HalfResolution = Resolution / 2;
-
-					auto DrawNormalMap = [&](const FIntPoint& Position, bool bShowNorthMap) {
-						if (const UTexture2D* NorthMapTexture = Helper.GetNormalMapTexture(bShowNorthMap))
-						{
-							FCanvas Canvas(RenderTarget, nullptr, Time, FeatureLevel);
-							Canvas.DrawTile(Position.X, Position.Y, HalfResolution.X, HalfResolution.Y, 0, 0, 1, 1, FColor::White, NorthMapTexture->GetResource());
-							Canvas.Flush_GameThread();
-						}
-					};
-
-					DrawNormalMap(FIntPoint(0, 0), true);
-					DrawNormalMap(FIntPoint(HalfResolution.X, 0), false);
-				}
-			}
-
 			// Read image data
 			TArray<FColor> PixelData;
 			RenderTarget->ReadPixels(PixelData);
