[nDisplay] Fix nDisplay screens not being able to be resized properly.

diff --git a/Engine/Plugins/Runtime/nDisplay/Source/DisplayCluster/Private/Game/EngineClasses/Scene/DisplayClusterScreenComponent.cpp b/Engine/Plugins/Runtime/nDisplay/Source/DisplayCluster/Private/Game/EngineClasses/Scene/DisplayClusterScreenComponent.cpp
--- a/Engine/Plugins/Runtime/nDisplay/Source/DisplayCluster/Private/Game/EngineClasses/Scene/DisplayClusterScreenComponent.cpp
+++ b/Engine/Plugins/Runtime/nDisplay/Source/DisplayCluster/Private/Game/EngineClasses/Scene/DisplayClusterScreenComponent.cpp
@@ -9,6 +9,7 @@
 #include "Materials/MaterialInterface.h"
 #include "Materials/Material.h"
 #include "UObject/ConstructorHelpers.h"
+#include "UObject/StructOnScope.h"
 
 
 UDisplayClusterScreenComponent::UDisplayClusterScreenComponent(const FObjectInitializer& ObjectInitializer)
@@ -41,14 +42,14 @@ UDisplayClusterScreenComponent::UDisplayClusterScreenComponent(const FObjectInit
 
 FVector2D UDisplayClusterScreenComponent::GetScreenSize() const
 {
-	const FVector ComponentScale = GetComponentScale();
+	const FVector ComponentScale = GetRelativeScale3D();
 	const FVector2D ComponentScale2D(ComponentScale.Y, ComponentScale.Z);
 	return ComponentScale2D;
 }
 
 void UDisplayClusterScreenComponent::SetScreenSize(const FVector2D& InSize)
 {
-	SetWorldScale3D(FVector(1.f, InSize.X, InSize.Y));
+	SetRelativeScale3D(FVector(1.f, InSize.X, InSize.Y));
 
 #if WITH_EDITOR
 	Size = InSize;
@@ -84,7 +85,7 @@ void UDisplayClusterScreenComponent::PostEditChangeProperty(FPropertyChangedEven
 		{
 			SetScreenSize(Size);
 		}
-		else if (PropertyChangedEvent.MemberProperty->GetFName() == TEXT("RelativeScale3D"))
+		else if (PropertyChangedEvent.MemberProperty->GetFName() == GetRelativeScale3DPropertyName())
 		{
 			UpdateScreenSizeFromScale();
 		}
@@ -95,6 +96,31 @@ void UDisplayClusterScreenComponent::PostEditChangeProperty(FPropertyChangedEven
 
 void UDisplayClusterScreenComponent::UpdateScreenSizeFromScale()
 {
-	Size = GetScreenSize();
+	// The Size property just reflects the 2D [Y,Z] relative transform scale and should always be in sync.
+	const FVector2D ScreenSize = GetScreenSize();
+	const TSharedPtr<FStructOnScope> ScreenSizeStructHighPrecision = MakeShared<FStructOnScope>(TBaseStructure<FVector2D>::Get(),
+		(uint8*)&ScreenSize);
+
+	const FStructProperty* SizeProperty = FindFieldChecked<FStructProperty>(UDisplayClusterScreenComponent::StaticClass(),
+		GET_MEMBER_NAME_CHECKED(UDisplayClusterScreenComponent, Size));
+	uint8* SizePropAddr = SizeProperty->ContainerPtrToValuePtr<uint8>(this);
+
+	// HACK: We need to purposely lose precision when storing the size in order to match what the UI would show if entering
+	// it directly in the details panel. If we don't then it can be easy to lose the ability to propagate changes from
+	// the CDO to instances when entering in high precision numbers such as using the slider.
+	//
+	// What happens is PropertyNode.cpp won't propagate a change if the instance property value isn't identical
+	// to the previous CDO value. This is a deep comparison comparing the actual values of the property (high precision).
+	// But the previous value is exported as text which loses precision, then is imported into a temporary property for comparison.
+	// The deep comparison is done against the live property (this property) which has the true double values, but the temporary
+	// property only has the imported text low precision values, so the comparison will fail and the value will not propagate.
+	
+	// TODO: Instead of using a separate FVector2D for size, we should probably just use the correct scale property handles
+	// through customization. Doing so may require public editor API changes since the transform handles are all customized by default
+	// and fairly locked down. We would also need to still support the locked aspect ration option.
+	
+	FString ScreenSizeStringLowPrecision;
+	SizeProperty->ExportText_Direct(ScreenSizeStringLowPrecision, ScreenSizeStructHighPrecision->GetStructMemory(), ScreenSizeStructHighPrecision->GetStructMemory(), this, PPF_None);
+	SizeProperty->ImportText_Direct(*ScreenSizeStringLowPrecision, SizePropAddr, this, PPF_SerializedAsImportText);
 }
 #endif
diff --git a/Engine/Plugins/Runtime/nDisplay/Source/DisplayClusterConfigurator/Private/Views/Details/Components/DisplayClusterConfiguratorScreenComponentDetailsCustomization.cpp b/Engine/Plugins/Runtime/nDisplay/Source/DisplayClusterConfigurator/Private/Views/Details/Components/DisplayClusterConfiguratorScreenComponentDetailsCustomization.cpp
--- a/Engine/Plugins/Runtime/nDisplay/Source/DisplayClusterConfigurator/Private/Views/Details/Components/DisplayClusterConfiguratorScreenComponentDetailsCustomization.cpp
+++ b/Engine/Plugins/Runtime/nDisplay/Source/DisplayClusterConfigurator/Private/Views/Details/Components/DisplayClusterConfiguratorScreenComponentDetailsCustomization.cpp
@@ -154,8 +154,10 @@ void FDisplayClusterConfiguratorScreenDetailsCustomization::GetAspectRatioAndSet
 	if (UDisplayClusterScreenComponent* Archetype = Cast<UDisplayClusterScreenComponent>(ScreenComponentPtr->GetArchetype()))
 	{
 		// Set the DEFAULT value here, that way user can always reset to default for the current preset.
-		Archetype->Modify();
-		Archetype->SetScreenSize(Preset.Size);
+		if (Archetype->GetScreenSize() != Preset.Size)
+		{
+			Archetype->SetScreenSize(Preset.Size);
+		}
 	}
 
 	if (OutAspectRatio)
