[nDisplay] Fix light card actor rendering with incorrect settings after being restored from a level snapshot.

diff --git a/Engine/Plugins/Runtime/nDisplay/Source/DisplayCluster/Private/Game/EngineClasses/Scene/DisplayClusterLightCardActor.cpp b/Engine/Plugins/Runtime/nDisplay/Source/DisplayCluster/Private/Game/EngineClasses/Scene/DisplayClusterLightCardActor.cpp
--- a/Engine/Plugins/Runtime/nDisplay/Source/DisplayCluster/Private/Game/EngineClasses/Scene/DisplayClusterLightCardActor.cpp
+++ b/Engine/Plugins/Runtime/nDisplay/Source/DisplayCluster/Private/Game/EngineClasses/Scene/DisplayClusterLightCardActor.cpp
@@ -147,12 +147,35 @@ void ADisplayClusterLightCardActor::OnConstruction(const FTransform& Transform)
 
 	UpdatePolygonTexture();
 
-	UMaterialInterface* Material = LightCardComponent->GetMaterial(0);
-
-	if (Material && !Material->IsA<UMaterialInstanceDynamic>())
+	if (UMaterialInterface* Material = LightCardComponent->GetMaterial(0))
 	{
-		UMaterialInstanceDynamic* LightCardMatInstance = UMaterialInstanceDynamic::Create(Material, LightCardComponent, TEXT("LightCardMID"));
-		LightCardComponent->SetMaterial(0, LightCardMatInstance);
+		if (UMaterialInstanceDynamic* MaterialInstanceDynamic = Cast<UMaterialInstanceDynamic>(Material))
+		{
+			// Clear param values and always call UpdateLightCardMaterialInstance(). Otherwise a snapshots restore
+			// may result in a desync between actor settings and MID parameter values.
+
+			// When level snapshots restores an actor the CS can run multiple times (spawn, posteditchangeproperty, etc).
+			// Our mesh component is a default sub object, so the "new" copy of it exists at this point and we set the MID.
+			// But after the actor is finished constructing, snapshots deserializes components reattaching them to the actor.
+			// What looks like is happening is the MID instance is somehow preserved, but parameter values are lost/outdated.
+			// This causes a problem in that we already set the parameter value so the value on the render thread is correct,
+			// but the value we have on the game thread is not. IE, on construction "UseMask" is true initially and the
+			// MID is updated correctly, but after snapshots loads our properties "UseMask" might be "false". Then,
+			// when snapshots deserializes the components the "true" value on the game thread incorrectly becomes "false".
+			// When we attempt to update the MID on tick to the correct value of "false" it thinks it already has
+			// the "false" value so it doesn't update the parameter render thread which is currently "true".
+			
+			// As long as we always call UpdateLightCardMaterialInstance() we can avoid the problem, but clearing the
+			// parameters out is an extra precaution to ensure we have the correct values.
+			
+			MaterialInstanceDynamic->ClearParameterValues();
+		}
+		else
+		{
+			// Initial MID creation
+			UMaterialInstanceDynamic* LightCardMatInstance = UMaterialInstanceDynamic::Create(Material, LightCardComponent, TEXT("LightCardMID"));
+			LightCardComponent->SetMaterial(0, LightCardMatInstance);
+		}
 
 		UpdateLightCardMaterialInstance();
 	}
