nDisplay:
- DCRA PreviewMesh can now be used to standalone and+or packaged builds.
- Added a new function to the WarpBlend interface to get or create a warp mesh for previewing.
- Bug fixes for the 'MPCDI' projection policy
   * The geometry profile type for PFM was missing
   * PFM cache ignored scale multiplier

diff --git a/Engine/Plugins/Runtime/nDisplay/Source/DisplayCluster/Private/Game/EngineClasses/Scene/DisplayClusterPreviewComponent.cpp b/Engine/Plugins/Runtime/nDisplay/Source/DisplayCluster/Private/Game/EngineClasses/Scene/DisplayClusterPreviewComponent.cpp
--- a/Engine/Plugins/Runtime/nDisplay/Source/DisplayCluster/Private/Game/EngineClasses/Scene/DisplayClusterPreviewComponent.cpp
+++ b/Engine/Plugins/Runtime/nDisplay/Source/DisplayCluster/Private/Game/EngineClasses/Scene/DisplayClusterPreviewComponent.cpp
@@ -281,6 +281,12 @@ bool UDisplayClusterPreviewComponent::UpdatePreviewMesh()
 
 void UDisplayClusterPreviewComponent::ReleasePreviewMesh()
 {
+	if (!bIsRootActorPreviewMesh && PreviewMesh)
+	{
+		PreviewMesh->UnregisterComponent();
+		PreviewMesh->DestroyComponent();
+	}
+
 	// Forget old mesh with material
 	PreviewMesh = nullptr;
 	CurrentMeshMaterial = nullptr;
diff --git a/Engine/Plugins/Runtime/nDisplay/Source/DisplayCluster/Private/Render/Containers/DisplayClusterRender_MeshComponent.cpp b/Engine/Plugins/Runtime/nDisplay/Source/DisplayCluster/Private/Render/Containers/DisplayClusterRender_MeshComponent.cpp
--- a/Engine/Plugins/Runtime/nDisplay/Source/DisplayCluster/Private/Render/Containers/DisplayClusterRender_MeshComponent.cpp
+++ b/Engine/Plugins/Runtime/nDisplay/Source/DisplayCluster/Private/Render/Containers/DisplayClusterRender_MeshComponent.cpp
@@ -130,6 +130,23 @@ IDisplayClusterRender_MeshComponentProxy* FDisplayClusterRender_MeshComponent::G
 	return MeshComponentProxy;
 }
 
+UMeshComponent* FDisplayClusterRender_MeshComponent::GetMeshComponent() const
+{
+	switch (GetGeometrySource())
+	{
+	case EDisplayClusterRender_MeshComponentGeometrySource::StaticMeshComponentRef:
+		return GetStaticMeshComponent();
+
+	case EDisplayClusterRender_MeshComponentGeometrySource::ProceduralMeshComponentRef:
+		return GetProceduralMeshComponent();
+
+	default:
+		break;
+	}
+
+	return nullptr;
+}
+
 UStaticMeshComponent* FDisplayClusterRender_MeshComponent::GetStaticMeshComponent() const
 {
 	check(IsInGameThread());
diff --git a/Engine/Plugins/Runtime/nDisplay/Source/DisplayCluster/Private/Render/Containers/DisplayClusterRender_MeshComponent.h b/Engine/Plugins/Runtime/nDisplay/Source/DisplayCluster/Private/Render/Containers/DisplayClusterRender_MeshComponent.h
--- a/Engine/Plugins/Runtime/nDisplay/Source/DisplayCluster/Private/Render/Containers/DisplayClusterRender_MeshComponent.h
+++ b/Engine/Plugins/Runtime/nDisplay/Source/DisplayCluster/Private/Render/Containers/DisplayClusterRender_MeshComponent.h
@@ -30,6 +30,7 @@ public:
 	virtual void ReleaseProxyGeometry() override;
 
 	virtual USceneComponent* GetOriginComponent() const override;
+	virtual UMeshComponent* GetMeshComponent() const override;
 	virtual UStaticMeshComponent* GetStaticMeshComponent() const override;
 	virtual const UStaticMesh* GetStaticMesh() const override;
 	virtual UProceduralMeshComponent* GetProceduralMeshComponent() const override;
diff --git a/Engine/Plugins/Runtime/nDisplay/Source/DisplayCluster/Public/Render/Containers/IDisplayClusterRender_MeshComponent.h b/Engine/Plugins/Runtime/nDisplay/Source/DisplayCluster/Public/Render/Containers/IDisplayClusterRender_MeshComponent.h
--- a/Engine/Plugins/Runtime/nDisplay/Source/DisplayCluster/Public/Render/Containers/IDisplayClusterRender_MeshComponent.h
+++ b/Engine/Plugins/Runtime/nDisplay/Source/DisplayCluster/Public/Render/Containers/IDisplayClusterRender_MeshComponent.h
@@ -7,6 +7,7 @@
 #include "Render/Containers/DisplayClusterRender_MeshGeometry.h"
 
 class USceneComponent;
+class UMeshComponent;
 class UStaticMesh;
 class UStaticMeshComponent;
 class UProceduralMeshComponent;
@@ -98,6 +99,13 @@ public:
 	*/
 	virtual USceneComponent* GetOriginComponent() const = 0;
 
+	/**
+	* Get referenced Mesh component object
+	*
+	* @return - pointer to the Mesh component
+	*/
+	virtual UMeshComponent* GetMeshComponent() const = 0;
+
 	/**
 	* Get referenced StaticMesh component object
 	*
diff --git a/Engine/Plugins/Runtime/nDisplay/Source/DisplayCluster/Public/Render/Projection/IDisplayClusterProjectionPolicy.h b/Engine/Plugins/Runtime/nDisplay/Source/DisplayCluster/Public/Render/Projection/IDisplayClusterProjectionPolicy.h
--- a/Engine/Plugins/Runtime/nDisplay/Source/DisplayCluster/Public/Render/Projection/IDisplayClusterProjectionPolicy.h
+++ b/Engine/Plugins/Runtime/nDisplay/Source/DisplayCluster/Public/Render/Projection/IDisplayClusterProjectionPolicy.h
@@ -52,7 +52,7 @@ public:
 	 * This component is used to convert from the local DCRA space to the local projection policy space,
 	 * which contains the calibrated geometry data used for the warp.
 	 */
-	virtual const USceneComponent* const GetOriginComponent() const
+	virtual USceneComponent* const GetOriginComponent() const
 	{
 		return nullptr;
 	}
@@ -284,7 +284,6 @@ public:
 		return false;
 	}
 
-#if WITH_EDITOR
 	/**
 	* Ask projection policy instance if it has any mesh based preview
 	*
@@ -329,5 +328,4 @@ public:
 	{
 		return nullptr;
 	}
-#endif
 };
diff --git a/Engine/Plugins/Runtime/nDisplay/Source/DisplayClusterProjection/Private/Policy/DisplayClusterProjectionPolicyBase.h b/Engine/Plugins/Runtime/nDisplay/Source/DisplayClusterProjection/Private/Policy/DisplayClusterProjectionPolicyBase.h
--- a/Engine/Plugins/Runtime/nDisplay/Source/DisplayClusterProjection/Private/Policy/DisplayClusterProjectionPolicyBase.h
+++ b/Engine/Plugins/Runtime/nDisplay/Source/DisplayClusterProjection/Private/Policy/DisplayClusterProjectionPolicyBase.h
@@ -25,24 +25,14 @@ public:
 		return PolicyInstanceId;
 	}
 
-	const USceneComponent* const GetOriginComp() const
-	{
-		return PolicyOriginComponentRef.GetOrFindSceneComponent();
-	}
-
-	USceneComponent* GetOriginComp()
-	{
-		return PolicyOriginComponentRef.GetOrFindSceneComponent();
-	}
-
 	void SetOriginComp(USceneComponent* OriginComp)
 	{
 		PolicyOriginComponentRef.SetSceneComponent(OriginComp);
 	}
 
-	virtual const USceneComponent* const GetOriginComponent() const override
+	virtual USceneComponent* const GetOriginComponent() const override
 	{
-		return GetOriginComp();
+		return PolicyOriginComponentRef.GetOrFindSceneComponent();
 	}
 
 	virtual const TMap<FString, FString>& GetParameters() const override
diff --git a/Engine/Plugins/Runtime/nDisplay/Source/DisplayClusterProjection/Private/Policy/Domeprojection/DisplayClusterProjectionDomeprojectionPolicyBase.cpp b/Engine/Plugins/Runtime/nDisplay/Source/DisplayClusterProjection/Private/Policy/Domeprojection/DisplayClusterProjectionDomeprojectionPolicyBase.cpp
--- a/Engine/Plugins/Runtime/nDisplay/Source/DisplayClusterProjection/Private/Policy/Domeprojection/DisplayClusterProjectionDomeprojectionPolicyBase.cpp
+++ b/Engine/Plugins/Runtime/nDisplay/Source/DisplayClusterProjection/Private/Policy/Domeprojection/DisplayClusterProjectionDomeprojectionPolicyBase.cpp
@@ -114,7 +114,7 @@ bool FDisplayClusterProjectionDomeprojectionPolicyBase::CalculateView(IDisplayCl
 	const float FCP = GDisplayClusterProjectionDomeprojectionPolicyEnableNearClippingPlane ? InFCP : 1;
 
 	// Get origin component
-	const USceneComponent* const OriginComp = GetOriginComp();
+	const USceneComponent* const OriginComp = GetOriginComponent();
 	check(OriginComp);
 
 	// Get world-origin matrix
diff --git a/Engine/Plugins/Runtime/nDisplay/Source/DisplayClusterProjection/Private/Policy/EasyBlend/DisplayClusterProjectionEasyBlendPolicyBase.cpp b/Engine/Plugins/Runtime/nDisplay/Source/DisplayClusterProjection/Private/Policy/EasyBlend/DisplayClusterProjectionEasyBlendPolicyBase.cpp
--- a/Engine/Plugins/Runtime/nDisplay/Source/DisplayClusterProjection/Private/Policy/EasyBlend/DisplayClusterProjectionEasyBlendPolicyBase.cpp
+++ b/Engine/Plugins/Runtime/nDisplay/Source/DisplayClusterProjection/Private/Policy/EasyBlend/DisplayClusterProjectionEasyBlendPolicyBase.cpp
@@ -135,7 +135,7 @@ bool FDisplayClusterProjectionEasyBlendPolicyBase::CalculateView(IDisplayCluster
 	const float WorldScale = WorldToMeters / 100.f;
 
 	// Get origin component
-	const USceneComponent* const OriginComp = GetOriginComp();
+	const USceneComponent* const OriginComp = GetOriginComponent();
 	check(OriginComp);
 
 	// Get world-origin matrix
diff --git a/Engine/Plugins/Runtime/nDisplay/Source/DisplayClusterProjection/Private/Policy/MPCDI/DisplayClusterProjectionMPCDIPolicy.cpp b/Engine/Plugins/Runtime/nDisplay/Source/DisplayClusterProjection/Private/Policy/MPCDI/DisplayClusterProjectionMPCDIPolicy.cpp
--- a/Engine/Plugins/Runtime/nDisplay/Source/DisplayClusterProjection/Private/Policy/MPCDI/DisplayClusterProjectionMPCDIPolicy.cpp
+++ b/Engine/Plugins/Runtime/nDisplay/Source/DisplayClusterProjection/Private/Policy/MPCDI/DisplayClusterProjectionMPCDIPolicy.cpp
@@ -174,10 +174,6 @@ void FDisplayClusterProjectionMPCDIPolicy::HandleEndScene(IDisplayClusterViewpor
 	check(IsInGameThread());
 
 	ImplRelease();
-
-#if WITH_EDITOR
-	ReleasePreviewMeshComponent();
-#endif
 }
 
 void FDisplayClusterProjectionMPCDIPolicy::ImplRelease()
@@ -242,7 +238,7 @@ bool FDisplayClusterProjectionMPCDIPolicy::CalculateView(IDisplayClusterViewport
 	const float WorldScale = WorldToMeters / 100.f;
 
 	// Get view location in local space
-	const USceneComponent* const OriginComp = GetOriginComp();
+	const USceneComponent* const OriginComp = GetOriginComponent();
 
 	// Initialize frustum
 	TSharedPtr<FDisplayClusterWarpEye, ESPMode::ThreadSafe> WarpEye = MakeShared<FDisplayClusterWarpEye, ESPMode::ThreadSafe>(InViewport->ToSharedPtr(), InContextNum);
@@ -448,92 +444,11 @@ void FDisplayClusterProjectionMPCDIPolicy::ApplyWarpBlend_RenderThread(FRHIComma
 	}
 }
 
-#if WITH_EDITOR
-
-#include "ProceduralMeshComponent.h"
-
-void FDisplayClusterProjectionMPCDIPolicy::ReleasePreviewMeshComponent()
-{
-	USceneComponent* PreviewMeshComp = PreviewMeshComponentRef.GetOrFindSceneComponent();
-	if (PreviewMeshComp != nullptr)
-	{
-		PreviewMeshComp->UnregisterComponent();
-		PreviewMeshComp->DestroyComponent();
-	}
-
-	PreviewMeshComponentRef.ResetSceneComponent();
-}
-
 UMeshComponent* FDisplayClusterProjectionMPCDIPolicy::GetOrCreatePreviewMeshComponent(IDisplayClusterViewport* InViewport, bool& bOutIsRootActorComponent)
 {
-	check(IsInGameThread());
-
-	if (WarpBlendInterface.IsValid() == false || !bIsPreviewMeshEnabled)
-	{
-		return nullptr;
-	}
-
-	// MPCDI can use a static mesh component instead of a warpmap (2D)
-	if (WarpBlendInterface.IsValid())
-	{
-		if (UMeshComponent* StaticMeshComponentRef = WarpBlendInterface->GetStaticMeshComponent())
-		{
-			bOutIsRootActorComponent = true;
-
-			return StaticMeshComponentRef;
-		}
-	}
-
-	// used created mesh component
-	bOutIsRootActorComponent = false;
-
-	USceneComponent* OriginComp = GetOriginComp();
-
-	// Return Exist mesh component
-	USceneComponent* PreviewMeshComp = PreviewMeshComponentRef.GetOrFindSceneComponent();
-	if (PreviewMeshComp != nullptr)
-	{
-		UProceduralMeshComponent* PreviewMesh = Cast<UProceduralMeshComponent>(PreviewMeshComp);
-		if (PreviewMesh != nullptr)
-		{
-			// update attachment to parent
-			PreviewMesh->AttachToComponent(OriginComp, FAttachmentTransformRules(EAttachmentRule::KeepRelative, false));
-			return PreviewMesh;
-		}
-	}
-
-	// Downscale preview mesh dimension to max limit
-	const uint32 PreviewGeometryDimLimit = 128;
-
-	// Create new WarpMesh component
-	FDisplayClusterWarpGeometryOBJ MeshData;
-	if (WarpBlendInterface->ExportWarpMapGeometry(MeshData, PreviewGeometryDimLimit))
-	{
-		const FString CompName = FString::Printf(TEXT("MPCDI_%s_impl"), *GetId());
-
-		// Creta new object
-		UProceduralMeshComponent* MeshComp = NewObject<UProceduralMeshComponent>(OriginComp, FName(*CompName), EObjectFlags::RF_DuplicateTransient | RF_Transient | RF_TextExportTransient);
-		if (MeshComp)
-		{
-			MeshComp->RegisterComponent();
-			MeshComp->AttachToComponent(OriginComp, FAttachmentTransformRules(EAttachmentRule::KeepRelative, false));
-			MeshComp->CreateMeshSection(0, MeshData.Vertices, MeshData.Triangles, MeshData.Normal, MeshData.UV, TArray<FColor>(), TArray<FProcMeshTangent>(), false);
-			MeshComp->SetIsVisualizationComponent(true);
-
-			// Because of "nDisplay.render.show.visualizationcomponents" we need extra flag to exclude this geometry from render
-			MeshComp->SetHiddenInGame(true);
-
-			// Store reference to mesh component
-			PreviewMeshComponentRef.SetSceneComponent(MeshComp);
-			return MeshComp;
-		}
-	}
-
-	return nullptr;
+	return (bIsPreviewMeshEnabled && WarpBlendInterface.IsValid()) ? WarpBlendInterface->GetOrCreateMeshComponent(InViewport, bOutIsRootActorComponent) : nullptr;
 }
 
-#endif
-
 bool FDisplayClusterProjectionMPCDIPolicy::CreateWarpBlendFromConfig(IDisplayClusterViewport* InViewport)
 {
 	check(IsInGameThread());
diff --git a/Engine/Plugins/Runtime/nDisplay/Source/DisplayClusterProjection/Private/Policy/MPCDI/DisplayClusterProjectionMPCDIPolicy.h b/Engine/Plugins/Runtime/nDisplay/Source/DisplayClusterProjection/Private/Policy/MPCDI/DisplayClusterProjectionMPCDIPolicy.h
--- a/Engine/Plugins/Runtime/nDisplay/Source/DisplayClusterProjection/Private/Policy/MPCDI/DisplayClusterProjectionMPCDIPolicy.h
+++ b/Engine/Plugins/Runtime/nDisplay/Source/DisplayClusterProjection/Private/Policy/MPCDI/DisplayClusterProjectionMPCDIPolicy.h
@@ -78,20 +78,11 @@ private:
 	bool bInvalidConfiguration = false;
 	bool bIsPreviewMeshEnabled = false;
 
-#if WITH_EDITOR
 protected:
-	//////////////////////////////////////////////////////////////////////////////////////////////
-	// IDisplayClusterProjectionPolicyPreview
-	//////////////////////////////////////////////////////////////////////////////////////////////
 	virtual bool HasPreviewMesh() override
 	{
 		return true;
 	}
-	virtual class UMeshComponent* GetOrCreatePreviewMeshComponent(IDisplayClusterViewport* InViewport, bool& bOutIsRootActorComponent) override;
-
-	void ReleasePreviewMeshComponent();
 
-private:
-	FDisplayClusterSceneComponentRef PreviewMeshComponentRef;
-#endif
+	virtual class UMeshComponent* GetOrCreatePreviewMeshComponent(IDisplayClusterViewport* InViewport, bool& bOutIsRootActorComponent) override;
 };
diff --git a/Engine/Plugins/Runtime/nDisplay/Source/DisplayClusterProjection/Private/Policy/Mesh/DisplayClusterProjectionMeshPolicy.cpp b/Engine/Plugins/Runtime/nDisplay/Source/DisplayClusterProjection/Private/Policy/Mesh/DisplayClusterProjectionMeshPolicy.cpp
--- a/Engine/Plugins/Runtime/nDisplay/Source/DisplayClusterProjection/Private/Policy/Mesh/DisplayClusterProjectionMeshPolicy.cpp
+++ b/Engine/Plugins/Runtime/nDisplay/Source/DisplayClusterProjection/Private/Policy/Mesh/DisplayClusterProjectionMeshPolicy.cpp
@@ -68,7 +68,7 @@ bool FDisplayClusterProjectionMeshPolicy::CreateWarpMeshInterface(IDisplayCluste
 		{
 			FDisplayClusterWarpInitializer_StaticMesh CreateParameters;
 
-			CreateParameters.OriginComponent = GetOriginComp();
+			CreateParameters.OriginComponent = GetOriginComponent();
 
 			CreateParameters.StaticMeshComponent = WarpCfg.StaticMeshComponent;
 			CreateParameters.StaticMeshComponentLODIndex = WarpCfg.StaticMeshComponentLODIndex;
@@ -82,7 +82,7 @@ bool FDisplayClusterProjectionMeshPolicy::CreateWarpMeshInterface(IDisplayCluste
 		{
 			FDisplayClusterWarpInitializer_ProceduralMesh CreateParameters;
 
-			CreateParameters.OriginComponent = GetOriginComp();
+			CreateParameters.OriginComponent = GetOriginComponent();
 
 			CreateParameters.ProceduralMeshComponent = WarpCfg.ProceduralMeshComponent;
 			CreateParameters.ProceduralMeshComponentSectionIndex = WarpCfg.ProceduralMeshComponentSectionIndex;
@@ -251,18 +251,3 @@ bool FDisplayClusterProjectionMeshPolicy::GetWarpMeshConfiguration(IDisplayClust
 
 	return true;
 }
-
-#if WITH_EDITOR
-UMeshComponent* FDisplayClusterProjectionMeshPolicy::GetOrCreatePreviewMeshComponent(IDisplayClusterViewport* InViewport, bool& bOutIsRootActorComponent)
-{
-	if (WarpBlendInterface.IsValid())
-		{
-		bOutIsRootActorComponent = true;
-
-		return WarpBlendInterface->GetStaticMeshComponent();
-	}
-
-	return nullptr;
-}
-#endif
-
diff --git a/Engine/Plugins/Runtime/nDisplay/Source/DisplayClusterProjection/Private/Policy/Mesh/DisplayClusterProjectionMeshPolicy.h b/Engine/Plugins/Runtime/nDisplay/Source/DisplayClusterProjection/Private/Policy/Mesh/DisplayClusterProjectionMeshPolicy.h
--- a/Engine/Plugins/Runtime/nDisplay/Source/DisplayClusterProjection/Private/Policy/Mesh/DisplayClusterProjectionMeshPolicy.h
+++ b/Engine/Plugins/Runtime/nDisplay/Source/DisplayClusterProjection/Private/Policy/Mesh/DisplayClusterProjectionMeshPolicy.h
@@ -25,15 +25,6 @@ public:
 	virtual const FString& GetType() const override;
 	virtual bool HandleStartScene(IDisplayClusterViewport* InViewport) override;
 
-#if WITH_EDITOR
-	virtual bool HasPreviewMesh() override
-	{
-		return true;
-	}
-
-	virtual UMeshComponent* GetOrCreatePreviewMeshComponent(IDisplayClusterViewport* InViewport, bool& bOutIsRootActorComponent) override;
-#endif
-
 public:
 	/** Parse the config data for a mesh id and try to retrieve it from the root actor. */
 	bool CreateWarpMeshInterface(IDisplayClusterViewport* InViewport);
diff --git a/Engine/Plugins/Runtime/nDisplay/Source/DisplayClusterProjection/Private/Policy/Simple/DisplayClusterProjectionSimplePolicy.cpp b/Engine/Plugins/Runtime/nDisplay/Source/DisplayClusterProjection/Private/Policy/Simple/DisplayClusterProjectionSimplePolicy.cpp
--- a/Engine/Plugins/Runtime/nDisplay/Source/DisplayClusterProjection/Private/Policy/Simple/DisplayClusterProjectionSimplePolicy.cpp
+++ b/Engine/Plugins/Runtime/nDisplay/Source/DisplayClusterProjection/Private/Policy/Simple/DisplayClusterProjectionSimplePolicy.cpp
@@ -214,14 +214,10 @@ void FDisplayClusterProjectionSimplePolicy::ReleaseMeshData(IDisplayClusterViewp
 {
 }
 
-#if WITH_EDITOR
 UMeshComponent* FDisplayClusterProjectionSimplePolicy::GetOrCreatePreviewMeshComponent(IDisplayClusterViewport* InViewport, bool& bOutIsRootActorComponent)
 {
 	bOutIsRootActorComponent = true;
-
-	USceneComponent* SceneComponent = ScreenCompRef.GetOrFindSceneComponent();
-
-	if (SceneComponent)
+	if (USceneComponent* SceneComponent = ScreenCompRef.GetOrFindSceneComponent())
 	{
 		UDisplayClusterScreenComponent* ScreenComp = StaticCast<UDisplayClusterScreenComponent*>(SceneComponent);
 		return (ScreenComp == nullptr) ? nullptr : ScreenComp;
@@ -229,4 +225,3 @@ UMeshComponent* FDisplayClusterProjectionSimplePolicy::GetOrCreatePreviewMeshCom
 
 	return nullptr;
 }
-#endif
diff --git a/Engine/Plugins/Runtime/nDisplay/Source/DisplayClusterProjection/Private/Policy/Simple/DisplayClusterProjectionSimplePolicy.h b/Engine/Plugins/Runtime/nDisplay/Source/DisplayClusterProjection/Private/Policy/Simple/DisplayClusterProjectionSimplePolicy.h
--- a/Engine/Plugins/Runtime/nDisplay/Source/DisplayClusterProjection/Private/Policy/Simple/DisplayClusterProjectionSimplePolicy.h
+++ b/Engine/Plugins/Runtime/nDisplay/Source/DisplayClusterProjection/Private/Policy/Simple/DisplayClusterProjectionSimplePolicy.h
@@ -29,14 +29,12 @@ public:
 	virtual bool CalculateView(IDisplayClusterViewport* InViewport, const uint32 InContextNum, FVector& InOutViewLocation, FRotator& InOutViewRotation, const FVector& ViewOffset, const float WorldToMeters, const float NCP, const float FCP) override;
 	virtual bool GetProjectionMatrix(IDisplayClusterViewport* InViewport, const uint32 InContextNum, FMatrix& OutPrjMatrix) override;
 
-#if WITH_EDITOR
 	virtual bool HasPreviewMesh() override
 	{
 		return true;
 	}
 
 	virtual UMeshComponent* GetOrCreatePreviewMeshComponent(IDisplayClusterViewport* InViewport, bool& bOutIsRootActorComponent) override;
-#endif
 
 protected:
 	bool InitializeMeshData(IDisplayClusterViewport* InViewport);
diff --git a/Engine/Plugins/Runtime/nDisplay/Source/DisplayClusterProjection/Private/Policy/VIOSO/Windows/DisplayClusterProjectionVIOSOPolicy.cpp b/Engine/Plugins/Runtime/nDisplay/Source/DisplayClusterProjection/Private/Policy/VIOSO/Windows/DisplayClusterProjectionVIOSOPolicy.cpp
--- a/Engine/Plugins/Runtime/nDisplay/Source/DisplayClusterProjection/Private/Policy/VIOSO/Windows/DisplayClusterProjectionVIOSOPolicy.cpp
+++ b/Engine/Plugins/Runtime/nDisplay/Source/DisplayClusterProjection/Private/Policy/VIOSO/Windows/DisplayClusterProjectionVIOSOPolicy.cpp
@@ -18,6 +18,8 @@
 #include "Render/Viewport/IDisplayClusterViewport.h"
 #include "Render/Viewport/IDisplayClusterViewportProxy.h"
 
+#include "ProceduralMeshComponent.h"
+
 //////////////////////////////////////////////////////////////////////////////////////////////
 // FDisplayClusterProjectionVIOSOPolicy
 //////////////////////////////////////////////////////////////////////////////////////////////
@@ -76,10 +78,6 @@ void FDisplayClusterProjectionVIOSOPolicy::HandleEndScene(IDisplayClusterViewpor
 	check(IsInGameThread());
 
 	ImplRelease();
-
-#if WITH_EDITOR
-	ReleasePreviewMeshComponent();
-#endif
 }
 
 void FDisplayClusterProjectionVIOSOPolicy::ImplRelease()
@@ -101,7 +99,7 @@ bool FDisplayClusterProjectionVIOSOPolicy::CalculateView(IDisplayClusterViewport
 	}
 
 	// Get view location in local space
-	const USceneComponent* const OriginComp = GetOriginComp();
+	const USceneComponent* const OriginComp = GetOriginComponent();
 	const FTransform& World2LocalTransform = (OriginComp ? OriginComp->GetComponentTransform() : FTransform::Identity);
 
 	// Get our base camera location and view offset in local space (MPCDI space)
@@ -212,21 +210,6 @@ bool FDisplayClusterProjectionVIOSOPolicy::ImplApplyWarpBlend_RenderThread(FRHIC
 	return true;
 }
 
-#if WITH_EDITOR
-#include "ProceduralMeshComponent.h"
-
-void FDisplayClusterProjectionVIOSOPolicy::ReleasePreviewMeshComponent()
-{
-	USceneComponent* PreviewMeshComp = PreviewMeshComponentRef.GetOrFindSceneComponent();
-	if (PreviewMeshComp != nullptr)
-	{
-		PreviewMeshComp->UnregisterComponent();
-		PreviewMeshComp->DestroyComponent();
-	}
-
-	PreviewMeshComponentRef.ResetSceneComponent();
-}
-
 UMeshComponent* FDisplayClusterProjectionVIOSOPolicy::GetOrCreatePreviewMeshComponent(IDisplayClusterViewport* InViewport, bool& bOutIsRootActorComponent)
 {
 	check(IsInGameThread());
@@ -236,27 +219,13 @@ UMeshComponent* FDisplayClusterProjectionVIOSOPolicy::GetOrCreatePreviewMeshComp
 		return nullptr;
 	}
 
-	// used created mesh component
+	// Create a new DCRA mesh component
 	bOutIsRootActorComponent = false;
 
-	USceneComponent* OriginComp = GetOriginComp();
-
-	// Return Exist mesh component
-	USceneComponent* PreviewMeshComp = PreviewMeshComponentRef.GetOrFindSceneComponent();
-	if (PreviewMeshComp != nullptr)
-	{
-		UProceduralMeshComponent* PreviewMesh = Cast<UProceduralMeshComponent>(PreviewMeshComp);
-		if (PreviewMesh != nullptr)
-		{
-			// update attachment to parent
-			PreviewMesh->AttachToComponent(OriginComp, FAttachmentTransformRules(EAttachmentRule::KeepRelative, false));
-			return PreviewMesh;
-		}
-	}
-
 	// Get geometry data
+	USceneComponent* OriginComp = GetOriginComponent();
 	TSharedPtr<FDisplayClusterProjectionVIOSOGeometryExportData, ESPMode::ThreadSafe> GeometryExportData = FDisplayClusterProjectionVIOSOGeometryExportData::Create(VIOSOLibrary, ViosoConfigData);
-	if (GeometryExportData.IsValid())
+	if (OriginComp && GeometryExportData.IsValid())
 	{
 		// Create new WarpMesh component
 		const FString CompName = FString::Printf(TEXT("VIOSO_%s_impl"), *GetId());
@@ -268,17 +237,17 @@ UMeshComponent* FDisplayClusterProjectionVIOSOPolicy::GetOrCreatePreviewMeshComp
 			MeshComp->RegisterComponent();
 			MeshComp->AttachToComponent(OriginComp, FAttachmentTransformRules(EAttachmentRule::KeepRelative, false));
 			MeshComp->CreateMeshSection(0, GeometryExportData->Vertices, GeometryExportData->Triangles, GeometryExportData->Normal, GeometryExportData->UV, TArray<FColor>(), TArray<FProcMeshTangent>(), false);
+
+#if WITH_EDITOR
 			MeshComp->SetIsVisualizationComponent(true);
+#endif
 
 			// Because of "nDisplay.render.show.visualizationcomponents" we need extra flag to exclude this geometry from render
 			MeshComp->SetHiddenInGame(true);
 
-			// Store reference to mesh component
-			PreviewMeshComponentRef.SetSceneComponent(MeshComp);
 			return MeshComp;
 		}
 	}
 
 	return nullptr;
 }
-#endif
diff --git a/Engine/Plugins/Runtime/nDisplay/Source/DisplayClusterProjection/Private/Policy/VIOSO/Windows/DisplayClusterProjectionVIOSOPolicy.h b/Engine/Plugins/Runtime/nDisplay/Source/DisplayClusterProjection/Private/Policy/VIOSO/Windows/DisplayClusterProjectionVIOSOPolicy.h
--- a/Engine/Plugins/Runtime/nDisplay/Source/DisplayClusterProjection/Private/Policy/VIOSO/Windows/DisplayClusterProjectionVIOSOPolicy.h
+++ b/Engine/Plugins/Runtime/nDisplay/Source/DisplayClusterProjection/Private/Policy/VIOSO/Windows/DisplayClusterProjectionVIOSOPolicy.h
@@ -55,8 +55,6 @@ protected:
 
 	FCriticalSection DllAccessCS;
 
-private:
-#if WITH_EDITOR
 protected:
 	//////////////////////////////////////////////////////////////////////////////////////////////
 	// IDisplayClusterProjectionPolicyPreview
@@ -65,11 +63,6 @@ protected:
 	{
 		return true;
 	}
-	virtual class UMeshComponent* GetOrCreatePreviewMeshComponent(IDisplayClusterViewport* InViewport, bool& bOutIsRootActorComponent) override;
-
-	void ReleasePreviewMeshComponent();
 
-private:
-	FDisplayClusterSceneComponentRef PreviewMeshComponentRef;
-#endif
+	virtual class UMeshComponent* GetOrCreatePreviewMeshComponent(IDisplayClusterViewport* InViewport, bool& bOutIsRootActorComponent) override;
 };
diff --git a/Engine/Plugins/Runtime/nDisplay/Source/DisplayClusterWarp/Private/WarpBlend/DisplayClusterWarpBlend.cpp b/Engine/Plugins/Runtime/nDisplay/Source/DisplayClusterWarp/Private/WarpBlend/DisplayClusterWarpBlend.cpp
--- a/Engine/Plugins/Runtime/nDisplay/Source/DisplayClusterWarp/Private/WarpBlend/DisplayClusterWarpBlend.cpp
+++ b/Engine/Plugins/Runtime/nDisplay/Source/DisplayClusterWarp/Private/WarpBlend/DisplayClusterWarpBlend.cpp
@@ -4,10 +4,17 @@
 #include "WarpBlend/Math/DisplayClusterWarpBlendMath_Frustum.h"
 
 #include "Render/Containers/IDisplayClusterRender_MeshComponent.h"
+#include "Render/Projection/IDisplayClusterProjectionPolicy.h"
+#include "Render/Viewport/IDisplayClusterViewport.h"
+
+#include "Blueprints/DisplayClusterWarpGeometry.h"
 
 #include "HAL/IConsoleManager.h"
+#include "UObject/UObjectGlobals.h"
+
 #include "ProceduralMeshComponent.h"
 #include "Components/StaticMeshComponent.h"
+#include "ProceduralMeshComponent.h"
 
 // Setup frustum projection cache
 static TAutoConsoleVariable<int32> CVarMPCDIFrustumCacheDepth(
@@ -129,22 +136,59 @@ const FDisplayClusterWarpData& FDisplayClusterWarpBlend::GetWarpData(const uint3
 
 }
 
-UMeshComponent* FDisplayClusterWarpBlend::GetStaticMeshComponent() const
+UMeshComponent* FDisplayClusterWarpBlend::GetOrCreateMeshComponent(IDisplayClusterViewport* InViewport, bool& bExistingComponent) const
 {
-	const TSharedPtr<IDisplayClusterRender_MeshComponent, ESPMode::ThreadSafe>& MeshComponent = GeometryContext.GeometryProxy.MeshComponent;
-	if (MeshComponent.IsValid())
+	const TSharedPtr<IDisplayClusterProjectionPolicy, ESPMode::ThreadSafe> ProjectionPolicy = InViewport ? InViewport->GetProjectionPolicy() : nullptr;
+
+	switch(GetWarpGeometryType())
 	{
-		switch (MeshComponent->GetGeometrySource())
+	case EDisplayClusterWarpGeometryType::WarpMesh:
+	case EDisplayClusterWarpGeometryType::WarpProceduralMesh:
+		// use the existing DCRA component
+		bExistingComponent = true;
+
+		return GeometryContext.GeometryProxy.MeshComponent.IsValid() ? GeometryContext.GeometryProxy.MeshComponent->GetMeshComponent() : nullptr;
+
+	case EDisplayClusterWarpGeometryType::WarpMap:
+		if (ProjectionPolicy.IsValid())
 		{
-		case EDisplayClusterRender_MeshComponentGeometrySource::StaticMeshComponentRef:
-			return MeshComponent->GetStaticMeshComponent();
+			// create a new mesh component
+			bExistingComponent = false;
+
+			// Downscale preview mesh dimension to max limit
+			const uint32 PreviewGeometryDimLimit = 128;
+
+			USceneComponent* OriginComp = ProjectionPolicy->GetOriginComponent();
+
+			// Create new WarpMesh component
+			FDisplayClusterWarpGeometryOBJ MeshData;
+			if (OriginComp && ExportWarpMapGeometry(MeshData, PreviewGeometryDimLimit))
+			{
+				const FString CompName = FString::Printf(TEXT("ExportPFM_%s"), *ProjectionPolicy->GetId());
+
+				// Creta new object
+				UProceduralMeshComponent* MeshComp = NewObject<UProceduralMeshComponent>(OriginComp, FName(*CompName), EObjectFlags::RF_DuplicateTransient | RF_Transient | RF_TextExportTransient);
+				if (MeshComp)
+				{
+					MeshComp->RegisterComponent();
+					MeshComp->AttachToComponent(OriginComp, FAttachmentTransformRules(EAttachmentRule::KeepRelative, false));
+					MeshComp->CreateMeshSection(0, MeshData.Vertices, MeshData.Triangles, MeshData.Normal, MeshData.UV, TArray<FColor>(), TArray<FProcMeshTangent>(), false);
+#if WITH_EDITOR
+					MeshComp->SetIsVisualizationComponent(true);
+#endif
 
-		case EDisplayClusterRender_MeshComponentGeometrySource::ProceduralMeshComponentRef:
-			return MeshComponent->GetProceduralMeshComponent();
+					// Because of "nDisplay.render.show.visualizationcomponents" we need extra flag to exclude this geometry from render
+					MeshComp->SetHiddenInGame(true);
 
-		default:
-			break;
+					return MeshComp;
+				}
+			}
 		}
+
+	break;
+
+	default:
+		break;
 	}
 
 	return nullptr;
@@ -152,8 +196,5 @@ UMeshComponent* FDisplayClusterWarpBlend::GetStaticMeshComponent() const
 
 bool FDisplayClusterWarpBlend::ExportWarpMapGeometry(FDisplayClusterWarpGeometryOBJ& OutMeshData, uint32 InMaxDimension) const
 {
-#if WITH_EDITOR
 	return FDisplayClusterWarpBlendExporter_WarpMap::ExportWarpMap(GeometryContext, OutMeshData, InMaxDimension);
-#endif
-	return false;
 }
diff --git a/Engine/Plugins/Runtime/nDisplay/Source/DisplayClusterWarp/Private/WarpBlend/DisplayClusterWarpBlend.h b/Engine/Plugins/Runtime/nDisplay/Source/DisplayClusterWarp/Private/WarpBlend/DisplayClusterWarpBlend.h
--- a/Engine/Plugins/Runtime/nDisplay/Source/DisplayClusterWarp/Private/WarpBlend/DisplayClusterWarpBlend.h
+++ b/Engine/Plugins/Runtime/nDisplay/Source/DisplayClusterWarp/Private/WarpBlend/DisplayClusterWarpBlend.h
@@ -82,7 +82,7 @@ public:
 	virtual FDisplayClusterWarpData& GetWarpData(const uint32 ContextNum) override;
 	virtual const FDisplayClusterWarpData& GetWarpData(const uint32 ContextNum) const override;
 
-	virtual UMeshComponent* GetStaticMeshComponent() const override;
+	virtual UMeshComponent* GetOrCreateMeshComponent(IDisplayClusterViewport* InViewport, bool& bExistingComponent) const override;
 
 	//~!IDisplayClusterWarpBlend
 
diff --git a/Engine/Plugins/Runtime/nDisplay/Source/DisplayClusterWarp/Private/WarpBlend/Loader/DisplayClusterWarpBlendLoader_MPCDI.cpp b/Engine/Plugins/Runtime/nDisplay/Source/DisplayClusterWarp/Private/WarpBlend/Loader/DisplayClusterWarpBlendLoader_MPCDI.cpp
--- a/Engine/Plugins/Runtime/nDisplay/Source/DisplayClusterWarp/Private/WarpBlend/Loader/DisplayClusterWarpBlendLoader_MPCDI.cpp
+++ b/Engine/Plugins/Runtime/nDisplay/Source/DisplayClusterWarp/Private/WarpBlend/Loader/DisplayClusterWarpBlendLoader_MPCDI.cpp
@@ -45,6 +45,7 @@ TSharedPtr<FDisplayClusterWarpBlend, ESPMode::ThreadSafe> FDisplayClusterWarpBle
 
 	// Filling all internal data manually:
 	RegionLoader.AlphaMapGammaEmbedded = InConstructParameters.AlphaMapEmbeddedAlpha;
+	RegionLoader.MPCDIAttributes = InConstructParameters.MPCDIAttributes;
 
 	RegionLoader.WarpMap = RegionLoader.CreateTextureFromPFMFile(InConstructParameters.PFMFileName, InConstructParameters.PFMScale, InConstructParameters.bIsUnrealGameSpace);
 	if (RegionLoader.WarpMap.IsValid())
@@ -59,9 +60,6 @@ TSharedPtr<FDisplayClusterWarpBlend, ESPMode::ThreadSafe> FDisplayClusterWarpBle
 			RegionLoader.BetaMap = RegionLoader.CreateTextureFromFile(InConstructParameters.BetaMapFileName);
 		}
 
-		// Copy MPCDI attributes
-		RegionLoader.MPCDIAttributes = InConstructParameters.MPCDIAttributes;
-
 		// After all internal data are filled in, the warpblend interface is created and initialized.
 		return RegionLoader.CreateWarpBlendInterface();
 	}
diff --git a/Engine/Plugins/Runtime/nDisplay/Source/DisplayClusterWarp/Private/WarpBlend/Loader/DisplayClusterWarpBlendLoader_MPCDIRegionLoader.cpp b/Engine/Plugins/Runtime/nDisplay/Source/DisplayClusterWarp/Private/WarpBlend/Loader/DisplayClusterWarpBlendLoader_MPCDIRegionLoader.cpp
--- a/Engine/Plugins/Runtime/nDisplay/Source/DisplayClusterWarp/Private/WarpBlend/Loader/DisplayClusterWarpBlendLoader_MPCDIRegionLoader.cpp
+++ b/Engine/Plugins/Runtime/nDisplay/Source/DisplayClusterWarp/Private/WarpBlend/Loader/DisplayClusterWarpBlendLoader_MPCDIRegionLoader.cpp
@@ -66,6 +66,7 @@ TSharedPtr<FDisplayClusterWarpBlend, ESPMode::ThreadSafe> FDisplayClusterWarpBle
 	TSharedPtr<FDisplayClusterWarpBlend, ESPMode::ThreadSafe> WarpBlend = MakeShared<FDisplayClusterWarpBlend, ESPMode::ThreadSafe>();
 
 	WarpBlend->GeometryContext.GeometryProxy.GeometryType = EDisplayClusterWarpGeometryType::WarpMap;
+	WarpBlend->GeometryContext.GeometryProxy.MPCDIAttributes = MPCDIAttributes;
 
 	// Frustum geometry source
 	switch (MPCDIAttributes.ProfileType)
@@ -81,12 +82,8 @@ TSharedPtr<FDisplayClusterWarpBlend, ESPMode::ThreadSafe> FDisplayClusterWarpBle
 		break;
 	}
 
-	WarpBlend->GeometryContext.GeometryProxy.MPCDIAttributes = MPCDIAttributes;
-
 	FDisplayClusterWarpBlend_GeometryProxy& Proxy = WarpBlend->GeometryContext.GeometryProxy;
-
 	Proxy.WarpMapTexture = WarpMap;
-
 	Proxy.AlphaMapEmbeddedGamma = AlphaMapGammaEmbedded;
 	Proxy.AlphaMapTexture = AlphaMap;
 	Proxy.BetaMapTexture = BetaMap;
@@ -157,17 +154,19 @@ TSharedPtr<IDisplayClusterRender_Texture, ESPMode::ThreadSafe> FDisplayClusterWa
 TSharedPtr<IDisplayClusterRender_Texture, ESPMode::ThreadSafe> FDisplayClusterWarpBlendMPCDIRegionLoader::CreateTextureFromPFMFile(const FString& InPFMFile, float PFMScale, bool bIsUnrealGameSpace)
 {
 	// new resource, initialize once
-	const FString PFMFileFullPath = FDisplayClusterWarpBlendMPCDIFileLoader::GetFullPathToFile(InPFMFile);
-	if (PFMFileFullPath.IsEmpty())
+	const FString FullPath2PFMFile = FDisplayClusterWarpBlendMPCDIFileLoader::GetFullPathToFile(InPFMFile);
+	if (FullPath2PFMFile.IsEmpty())
 	{
 		return nullptr;
 	}
 
-	TSharedPtr<IDisplayClusterRender_Texture, ESPMode::ThreadSafe> Texture = GetOrCreateCachedTextureImpl(PFMFileFullPath);
+	const FString UniqueName = FString::Printf(TEXT("%s=%f"), *FMD5::HashAnsiString(*FullPath2PFMFile.ToLower()), PFMScale);
+
+	TSharedPtr<IDisplayClusterRender_Texture, ESPMode::ThreadSafe> Texture = GetOrCreateCachedTextureImpl(UniqueName);
 	if (Texture && !Texture->IsEnabled())
 	{
 		bool bResult = false;
-		const std::string FileName = TCHAR_TO_ANSI(*PFMFileFullPath);
+		const std::string FileName = TCHAR_TO_ANSI(*FullPath2PFMFile);
 		mpcdi::PFM* PFMData;
 		mpcdi::MPCDI_Error res = mpcdi::PfmIO::Read(FileName, PFMData);
 		if (mpcdi::MPCDI_SUCCESS == res && PFMData)
diff --git a/Engine/Plugins/Runtime/nDisplay/Source/DisplayClusterWarp/Private/WarpBlend/Loader/DisplayClusterWarpBlendLoader_WarpMap.cpp b/Engine/Plugins/Runtime/nDisplay/Source/DisplayClusterWarp/Private/WarpBlend/Loader/DisplayClusterWarpBlendLoader_WarpMap.cpp
--- a/Engine/Plugins/Runtime/nDisplay/Source/DisplayClusterWarp/Private/WarpBlend/Loader/DisplayClusterWarpBlendLoader_WarpMap.cpp
+++ b/Engine/Plugins/Runtime/nDisplay/Source/DisplayClusterWarp/Private/WarpBlend/Loader/DisplayClusterWarpBlendLoader_WarpMap.cpp
@@ -173,8 +173,12 @@ bool FDisplayClusterWarpBlendLoader::LoadFromPoint(const EDisplayClusterWarpProf
 
 bool FDisplayClusterWarpBlendLoader::LoadFromPFM(const EDisplayClusterWarpProfileType InProfileType, mpcdi::PFM* SourcePFM, float PFMScale, bool bIsUnrealGameSpace)
 {
-	check(SourcePFM);
+	if (InProfileType == EDisplayClusterWarpProfileType::Invalid)
+	{
+		return false;
+	}
 
+	check(SourcePFM);
 	if (!InitializeWarpDataImpl(SourcePFM->GetSizeX(), SourcePFM->GetSizeY()))
 	{
 		return false;
diff --git a/Engine/Plugins/Runtime/nDisplay/Source/DisplayClusterWarp/Public/IDisplayClusterWarpBlend.h b/Engine/Plugins/Runtime/nDisplay/Source/DisplayClusterWarp/Public/IDisplayClusterWarpBlend.h
--- a/Engine/Plugins/Runtime/nDisplay/Source/DisplayClusterWarp/Public/IDisplayClusterWarpBlend.h
+++ b/Engine/Plugins/Runtime/nDisplay/Source/DisplayClusterWarp/Public/IDisplayClusterWarpBlend.h
@@ -6,6 +6,8 @@
 #include "Containers/DisplayClusterWarpEye.h"
 
 class UMeshComponent;
+class IDisplayClusterViewport;
+
 /**
  * WarpBlend interface for MPCDI and mesh projection policies
  */
@@ -93,6 +95,12 @@ public:
 	/** Get internal warp context data. */
 	virtual const FDisplayClusterWarpData& GetWarpData(const uint32 ContextNum) const = 0;
 
-	/** Get frustum mesh component. */
-	virtual UMeshComponent* GetStaticMeshComponent() const = 0;
+	/** Get a mesh component with the geometry used for the warp blend.
+	* 
+	* @param InViewport - The viewport with the projection policy to which the instance belongs
+	* @param bExistingComponent - true if the component already exists in the DCRA and does not need to be deleted externally.
+	* 
+	* @return - ptr to the mesh component with the geometry, or nullptr in case of failure
+	*/
+	virtual UMeshComponent* GetOrCreateMeshComponent(IDisplayClusterViewport* InViewport, bool& bExistingComponent) const = 0;
 };
