Fixing compile error

diff --git a/Engine/Plugins/Runtime/nDisplay/Source/DisplayClusterLightCardEditor/Private/DisplayClusterLightCardEditor.cpp b/Engine/Plugins/Runtime/nDisplay/Source/DisplayClusterLightCardEditor/Private/DisplayClusterLightCardEditor.cpp
--- a/Engine/Plugins/Runtime/nDisplay/Source/DisplayClusterLightCardEditor/Private/DisplayClusterLightCardEditor.cpp
+++ b/Engine/Plugins/Runtime/nDisplay/Source/DisplayClusterLightCardEditor/Private/DisplayClusterLightCardEditor.cpp
@@ -145,7 +145,7 @@ TArray<AActor*> FDisplayClusterLightCardEditor::FindAllManagedActors() const
 			RootActorLightCardActors.Append(reinterpret_cast<TSet<ADisplayClusterLightCardActor*>&>(ChromakeyCardActors));
 		}
 		
-		for (const TObjectPtr<ADisplayClusterLightCardActor>& LightCardActor : RootActorLightCardActors)
+		for (ADisplayClusterLightCardActor* LightCardActor : RootActorLightCardActors)
 		{
 			ManagedActors.Add(LightCardActor);
 		}
