nDisplay: fixed replication startup freezes by wrong connection sequence

diff --git a/Engine/Plugins/Runtime/nDisplay/Source/DisplayCluster/Private/Cluster/DisplayClusterNetDriverHelper.cpp b/Engine/Plugins/Runtime/nDisplay/Source/DisplayCluster/Private/Cluster/DisplayClusterNetDriverHelper.cpp
--- a/Engine/Plugins/Runtime/nDisplay/Source/DisplayCluster/Private/Cluster/DisplayClusterNetDriverHelper.cpp
+++ b/Engine/Plugins/Runtime/nDisplay/Source/DisplayCluster/Private/Cluster/DisplayClusterNetDriverHelper.cpp
@@ -1,6 +1,7 @@
 // Copyright Epic Games, Inc. All Rights Reserved.
 
 #include "Cluster/DisplayClusterNetDriverHelper.h"
+#include "Config/IDisplayClusterConfigManager.h"
 #include "DisplayClusterConfigurationTypes.h"
 
 #include "Misc/DisplayClusterGlobals.h"
@@ -8,6 +9,66 @@
 
 #include "Network/Service/ClusterEventsBinary/DisplayClusterClusterEventsBinaryClient.h"
 
+int32 FDisplayClusterNetDriverHelper::GetNumberOfClusterNodes() const
+{
+	if (GDisplayCluster == nullptr)
+	{
+		return 0;
+	}
+
+	IDisplayClusterConfigManager* DisplayClusterConfigManager = GDisplayCluster->GetConfigMgr();
+	if (DisplayClusterConfigManager == nullptr)
+	{
+		return 0;
+	}
+
+	UDisplayClusterConfigurationData* DisplayClusterConfigurationData = DisplayClusterConfigManager->GetConfig();
+	if (DisplayClusterConfigurationData == nullptr)
+	{
+		return 0;
+	}
+
+	if (!IsValid(DisplayClusterConfigurationData->Cluster))
+	{
+		return 0;
+	}
+
+	return DisplayClusterConfigurationData->Cluster->Nodes.Num();
+}
+
+FString FDisplayClusterNetDriverHelper::GetPrimaryNodeAddress()
+{
+	FString NoAddress;
+
+	if (GDisplayCluster == nullptr)
+	{
+		return NoAddress;
+	}
+
+	IDisplayClusterConfigManager* DisplayClusterConfigManager = GDisplayCluster->GetConfigMgr();
+
+	if (DisplayClusterConfigManager == nullptr)
+	{
+		return NoAddress;
+	}
+
+	UDisplayClusterConfigurationData* DisplayClusterConfigurationData = DisplayClusterConfigManager->GetConfig();
+
+	if (DisplayClusterConfigurationData == nullptr)
+	{
+		return NoAddress;
+	}
+
+	if (!IsValid(DisplayClusterConfigurationData->Cluster))
+	{
+		return NoAddress;
+	}
+
+	FString PrimaryNodeId = DisplayClusterConfigurationData->Cluster->PrimaryNode.Id;
+
+	return DisplayClusterConfigurationData->Cluster->GetNode(PrimaryNodeId)->Host;
+}
+
 bool FDisplayClusterNetDriverHelper::RegisterClusterEventsBinaryClient(uint32 ClusterId, const FString& ClientAddress, uint16 ClientPort)
 {	
 	TSharedPtr<FDisplayClusterClusterEventsBinaryClient>* FoundClient = ClusterEventsBinaryClients.Find(ClusterId);
@@ -51,11 +112,12 @@ bool FDisplayClusterNetDriverHelper::HasClient(uint32 ClusterId)
 	return ClusterEventsBinaryClients.Contains(ClusterId);
 }
 
-bool FDisplayClusterNetDriverHelper::GetRequiredArguments(const FURL& URL, const TCHAR*& OutClusterId, const TCHAR*& OutPrimaryNodeId, const TCHAR*& OutPrimaryNodePort)
+bool FDisplayClusterNetDriverHelper::GetRequiredArguments(const FURL& URL, const TCHAR*& OutClusterId, const TCHAR*& OutPrimaryNodeId, const TCHAR*& OutPrimaryNodePort, const TCHAR*& OutClusterNodesNum)
 {
 	if (!URL.HasOption(DisplayClusterStrings::uri_args::ClusterId) ||
 		!URL.HasOption(DisplayClusterStrings::uri_args::PrimaryNodeId) ||
-		!URL.HasOption(DisplayClusterStrings::uri_args::PrimaryNodePort))
+		!URL.HasOption(DisplayClusterStrings::uri_args::PrimaryNodePort) ||
+		!URL.HasOption(DisplayClusterStrings::uri_args::ClusterNodesNum))
 	{
 		return false;
 	}
@@ -63,6 +125,7 @@ bool FDisplayClusterNetDriverHelper::GetRequiredArguments(const FURL& URL, const
 	OutClusterId = URL.GetOption(DisplayClusterStrings::uri_args::ClusterId, nullptr);
 	OutPrimaryNodeId = URL.GetOption(DisplayClusterStrings::uri_args::PrimaryNodeId, nullptr);
 	OutPrimaryNodePort = URL.GetOption(DisplayClusterStrings::uri_args::PrimaryNodePort, nullptr);
+	OutClusterNodesNum = URL.GetOption(DisplayClusterStrings::uri_args::ClusterNodesNum, nullptr);
 
 	return true;
 }
diff --git a/Engine/Plugins/Runtime/nDisplay/Source/DisplayCluster/Public/Cluster/DisplayClusterNetDriverHelper.h b/Engine/Plugins/Runtime/nDisplay/Source/DisplayCluster/Public/Cluster/DisplayClusterNetDriverHelper.h
--- a/Engine/Plugins/Runtime/nDisplay/Source/DisplayCluster/Public/Cluster/DisplayClusterNetDriverHelper.h
+++ b/Engine/Plugins/Runtime/nDisplay/Source/DisplayCluster/Public/Cluster/DisplayClusterNetDriverHelper.h
@@ -20,6 +20,12 @@ struct DISPLAYCLUSTER_API FDisplayClusterNetDriverHelper
 	GENERATED_BODY()
 
 public:
+	/** Returns number of nodes in current cluster */
+	int32 GetNumberOfClusterNodes() const;
+
+	/** Returns primary node IP address of current cluster */
+	FString GetPrimaryNodeAddress();
+
 	/** Creates and stores nDisplay cluster events client using cluster id and primary node parameters
 	 *  @param СlusterId unique identifier of cluster
 	 *  @param ClientAddress primary node IP address
@@ -40,7 +46,7 @@ public:
 	/** Checks if URL has all required arguments for DisplayClusterNetDriver
 	 *  @param URL Url to check
 	 */
-	static bool GetRequiredArguments(const FURL& URL, const TCHAR*& OutClusterId, const TCHAR*& OutPrimaryNodeId, const TCHAR*& OutPrimaryNodePort);
+	static bool GetRequiredArguments(const FURL& URL, const TCHAR*& OutClusterId, const TCHAR*& OutPrimaryNodeId, const TCHAR*& OutPrimaryNodePort, const TCHAR*& OutClusterNodesNum);
 
 	/** Sends command to nDisplay cluster specified by id
 	 *  @param СlusterId identifier of cluster
diff --git a/Engine/Plugins/Runtime/nDisplay/Source/DisplayCluster/Public/Misc/DisplayClusterStrings.h b/Engine/Plugins/Runtime/nDisplay/Source/DisplayCluster/Public/Misc/DisplayClusterStrings.h
--- a/Engine/Plugins/Runtime/nDisplay/Source/DisplayCluster/Public/Misc/DisplayClusterStrings.h
+++ b/Engine/Plugins/Runtime/nDisplay/Source/DisplayCluster/Public/Misc/DisplayClusterStrings.h
@@ -65,5 +65,6 @@ namespace DisplayClusterStrings
 		static constexpr const TCHAR* ClusterId = TEXT("ClusterId=");
 		static constexpr const TCHAR* PrimaryNodeId = TEXT("PrimaryNodeId=");
 		static constexpr const TCHAR* PrimaryNodePort = TEXT("PrimaryNodePort=");
+		static constexpr const TCHAR* ClusterNodesNum = TEXT("ClusterNodesCount=");
 	}
 };
diff --git a/Engine/Plugins/Runtime/nDisplay/Source/DisplayClusterReplication/Private/DisplayClusterNetConnection.cpp b/Engine/Plugins/Runtime/nDisplay/Source/DisplayClusterReplication/Private/DisplayClusterNetConnection.cpp
--- a/Engine/Plugins/Runtime/nDisplay/Source/DisplayClusterReplication/Private/DisplayClusterNetConnection.cpp
+++ b/Engine/Plugins/Runtime/nDisplay/Source/DisplayClusterReplication/Private/DisplayClusterNetConnection.cpp
@@ -33,6 +33,7 @@ UDisplayClusterNetConnection::UDisplayClusterNetConnection(const FObjectInitiali
 	Super(ObjectInitializer),
 	ClientId(0),
 	ClusterId(0),
+	ClusterNodesNum(0),
 	NodePort(0),
 	bNodeIsPrimary(false),
 	bIsClusterConnection(false),
@@ -66,8 +67,9 @@ void UDisplayClusterNetConnection::SetClientLoginState(const EClientLoginState::
 		const TCHAR* ClusterIdString = nullptr;
 		const TCHAR* PrimaryNodeId = nullptr;
 		const TCHAR* PrimaryNodePortString = nullptr;
+		const TCHAR* ClusterNodesNumString = nullptr;
 
-		if (!FDisplayClusterNetDriverHelper::GetRequiredArguments(ClientURL, ClusterIdString, PrimaryNodeId, PrimaryNodePortString))
+		if (!FDisplayClusterNetDriverHelper::GetRequiredArguments(ClientURL, ClusterIdString, PrimaryNodeId, PrimaryNodePortString, ClusterNodesNumString))
 		{
 			UE_LOG(LogDisplayClusterNetDriver, Verbose, TEXT("Cluster connection Join has failed: URI arguments are invalid"));
 			return;
@@ -79,6 +81,8 @@ void UDisplayClusterNetConnection::SetClientLoginState(const EClientLoginState::
 		NodeAddress = GetRemoteAddr()->ToString(false);
 		NodePort = FCString::Atoi(PrimaryNodePortString);
 
+		ClusterNodesNum = FCString::Atoi(ClusterNodesNumString);
+
 		bIsClusterConnection = true;
 
 		if (PrimaryNodeId == NodeName)
@@ -115,14 +119,14 @@ void UDisplayClusterNetConnection::ReceivedPacket(FBitReader& Reader, bool bIsRe
 		Super::ReceivedPacket(Reader, bIsReinjectedPacket, false);
 		InPackets.Add(InPacketId, Reader);
 
-		UE_LOG(LogDisplayClusterNetDriver, Verbose, TEXT("Packet %i received and accumulated"), InPacketId);
+		UE_LOG(LogDisplayClusterNetDriver, VeryVerbose, TEXT("Packet %i received and accumulated"), InPacketId);
 	}
 	else
 	{
 		// This is server or client that is not ready to replication
 		Super::ReceivedPacket(Reader, bIsReinjectedPacket, bDispatchPacket);
 
-		UE_LOG(LogDisplayClusterNetDriver, Verbose, TEXT("Packet %i received and processed"), InPacketId);
+		UE_LOG(LogDisplayClusterNetDriver, VeryVerbose, TEXT("Packet %i received and processed"), InPacketId);
 	}
 }
 
@@ -138,7 +142,7 @@ void UDisplayClusterNetConnection::ProcessPacket(int32 PacketId)
 
 	if (!InPackets.Contains(PacketId))
 	{
-		UE_LOG(LogDisplayClusterNetDriver, Verbose, TEXT("Packet %i skipped!"), PacketId);
+		UE_LOG(LogDisplayClusterNetDriver, VeryVerbose, TEXT("Packet %i skipped!"), PacketId);
 
 		return;
 	}
@@ -153,7 +157,7 @@ void UDisplayClusterNetConnection::ProcessPacket(int32 PacketId)
 			break;
 		}
 
-		UE_LOG(LogDisplayClusterNetDriver, Verbose, TEXT("Packet %i processed"), CurrentPacketId);
+		UE_LOG(LogDisplayClusterNetDriver, VeryVerbose, TEXT("Packet %i processed"), CurrentPacketId);
 
 		bool bSkipAck = false;
 		bool bHasBunchErrors = false;
diff --git a/Engine/Plugins/Runtime/nDisplay/Source/DisplayClusterReplication/Private/DisplayClusterNetDriver.cpp b/Engine/Plugins/Runtime/nDisplay/Source/DisplayClusterReplication/Private/DisplayClusterNetDriver.cpp
--- a/Engine/Plugins/Runtime/nDisplay/Source/DisplayClusterReplication/Private/DisplayClusterNetDriver.cpp
+++ b/Engine/Plugins/Runtime/nDisplay/Source/DisplayClusterReplication/Private/DisplayClusterNetDriver.cpp
@@ -39,6 +39,7 @@
 
 UDisplayClusterNetDriver::UDisplayClusterNetDriver(const FObjectInitializer& ObjectInitializer) :
 	Super(ObjectInitializer),
+	bClusterHasConnected(false),
 	bLastBunchWasAcked(false),
 	bConnectionViewersAreReady(false)
 {
@@ -90,8 +91,9 @@ bool UDisplayClusterNetDriver::InitListen(FNetworkNotify* InNotify, FURL& Listen
 		const TCHAR* ClusterIdString = nullptr;
 		const TCHAR* NodeId = nullptr;
 		const TCHAR* NodePortString = nullptr;
+		const TCHAR* ClusterNodesNumString = nullptr;
 
-		if (!FDisplayClusterNetDriverHelper::GetRequiredArguments(ListenURL, ClusterIdString, NodeId, NodePortString))
+		if (!FDisplayClusterNetDriverHelper::GetRequiredArguments(ListenURL, ClusterIdString, NodeId, NodePortString, ClusterNodesNumString))
 		{
 			UE_LOG(LogDisplayClusterNetDriver, Warning, TEXT("InitListen:: Can't get required URL arguments"));
 
@@ -123,52 +125,56 @@ void UDisplayClusterNetDriver::TickDispatch(float DeltaTime)
 		return;
 	}
 
-	// Check nodes that are ready for synchronous packets processing
-	for (UDisplayClusterNetConnection* NodeConnectionIt : NodeConnections)
+	// Register binary client when ALL nodes connected
+	if (!bClusterHasConnected)
 	{
-		// NodeConnection is already in sync mode - skip
-		if (NodeConnectionIt->bSynchronousMode)
-		{
-			continue;
-		}
+		UE_LOG(LogDisplayClusterNetDriver, VeryVerbose, TEXT("Primary Node connecitons = [%u]"), PrimaryNodeConnections.Num());
 
-		// NodeConnection has not joined server - skip
-		if ((NodeConnectionIt->ClientLoginState != EClientLoginState::Type::ReceivedJoin) || (NodeConnectionIt->OwningActor == nullptr))
-		{
-			continue;
-		}
-		
-		// Event listener for cluster is not registered - skip
-		if (!ClusterNetworkDriverHelper->HasClient(NodeConnectionIt->ClusterId))
+		for (UDisplayClusterNetConnection* PrimaryNodeConnectionIt : PrimaryNodeConnections)
 		{
-			continue;
-		}
+			TArray<UDisplayClusterNetConnection*>* ClusterNodeConnectionsFound = ClusterConnections.Find(PrimaryNodeConnectionIt->ClusterId);
 
-		// Notify primary node that current node connection is ready
-		FDisplayClusterClusterEventBinary NetworkDriverSyncEvent;
-		GenerateClusterCommandsEvent(NetworkDriverSyncEvent, NodeSyncEvent);
+			if (ClusterNodeConnectionsFound == nullptr)
+			{
+				continue;
+			}
 
-		const bool bEventWasSent = ClusterNetworkDriverHelper->SendCommandToCluster(NodeConnectionIt->ClusterId, NetworkDriverSyncEvent);
+			TArray<UDisplayClusterNetConnection*>& ClusterNodeConnections = *ClusterNodeConnectionsFound;
 
-		// If event was sent to primary cluster node
-		if (bEventWasSent)
-		{
-			// Mark all node connections related to current ClusterId as ready
-			for (UDisplayClusterNetConnection* ClusterNodeConnectionIt : NodeConnections)
+			if (PrimaryNodeConnectionIt->ClusterNodesNum != ClusterNodeConnections.Num())
 			{
-				if (NodeConnectionIt->bSynchronousMode)
-				{
-					continue;
-				}
+				UE_LOG(LogDisplayClusterNetDriver, VeryVerbose, TEXT("PrimaryNodeConnectionIt->ClusterNodesNum = [%u], ClusterNodeConnections.Num() = [%u]"), PrimaryNodeConnectionIt->ClusterNodesNum, ClusterNodeConnections.Num());
+				continue;
+			}
+
+			if (ClusterNetworkDriverHelper->RegisterClusterEventsBinaryClient(PrimaryNodeConnectionIt->ClusterId, PrimaryNodeConnectionIt->NodeAddress, PrimaryNodeConnectionIt->NodePort))
+			{
+				UE_LOG(LogDisplayClusterNetDriver, VeryVerbose, TEXT("Registered primary node for cluster [%u] node [%s]"), PrimaryNodeConnectionIt->ClusterId, *PrimaryNodeConnectionIt->NodeName);
+			}
+
+			// Notify primary node that current node connection is ready
+			FDisplayClusterClusterEventBinary NetworkDriverSyncEvent;
+			GenerateClusterCommandsEvent(NetworkDriverSyncEvent, NodeSyncEvent);
 
-				if (ClusterNodeConnectionIt->ClusterId == NodeConnectionIt->ClusterId)
+			// send command 
+			const bool bEventWasSent = ClusterNetworkDriverHelper->SendCommandToCluster(PrimaryNodeConnectionIt->ClusterId, NetworkDriverSyncEvent);
+
+			// If event was sent to primary cluster node
+			if (bEventWasSent)
+			{
+				bClusterHasConnected = true;
+				// Mark all node connections related to current ClusterId as ready
+				for (UDisplayClusterNetConnection* ClusterNodeConnectionIt : NodeConnections)
 				{
-					NodeConnectionIt->bSynchronousMode = true;
-					SyncConnections.Add(NodeConnectionIt);
+					if (ClusterNodeConnectionIt->ClusterId == PrimaryNodeConnectionIt->ClusterId)
+					{
+						ClusterNodeConnectionIt->bSynchronousMode = true;
+						SyncConnections.Add(ClusterNodeConnectionIt);
+					}
 				}
-			}
 
-			UE_LOG(LogDisplayClusterNetDriver, Verbose, TEXT("Setting synchronous mode in cluster [%u]"), NodeConnectionIt->ClusterId);
+				UE_LOG(LogDisplayClusterNetDriver, VeryVerbose, TEXT("Setting synchronous mode in cluster [%u]"), PrimaryNodeConnectionIt->ClusterId);
+			}
 		}
 	}
 }
@@ -180,12 +186,10 @@ void UDisplayClusterNetDriver::AddNodeConnection(UDisplayClusterNetConnection* N
 	if (NetConnection->bNodeIsPrimary)
 	{
 		PrimaryNodeConnections.Add(NetConnection);
-
-		if (ClusterNetworkDriverHelper->RegisterClusterEventsBinaryClient(NetConnection->ClusterId, NetConnection->NodeAddress, NetConnection->NodePort))
-		{
-			UE_LOG(LogDisplayClusterNetDriver, Verbose, TEXT("Registered primary node for cluster [%u] node [%s]"), NetConnection->ClusterId, *NetConnection->NodeName);
-		}
 	}
+
+	TArray<UDisplayClusterNetConnection*>& Connections = ClusterConnections.FindOrAdd(NetConnection->ClusterId);
+	Connections.Add(NetConnection);
 }
 
 void UDisplayClusterNetDriver::RemoveNodeConnection(UDisplayClusterNetConnection* NetConnection)
@@ -194,14 +198,32 @@ void UDisplayClusterNetDriver::RemoveNodeConnection(UDisplayClusterNetConnection
 	{
 		if (ClusterNetworkDriverHelper->RemoveClusterEventsBinaryClient(NetConnection->ClusterId))
 		{
-			UE_LOG(LogDisplayClusterNetDriver, Verbose, TEXT("Removed primary node for cluster [%u] node [%s]"), NetConnection->ClusterId, *NetConnection->NodeName);
+			UE_LOG(LogDisplayClusterNetDriver, VeryVerbose, TEXT("Removed primary node for cluster [%u] node [%s]"), NetConnection->ClusterId, *NetConnection->NodeName);
 		}
 
 		PrimaryNodeConnections.Remove(NetConnection);
 	}
 
-	SyncConnections.Remove(NetConnection);
-	NodeConnections.Remove(NetConnection);
+
+	TArray<UDisplayClusterNetConnection*>* Connections = ClusterConnections.Find(NetConnection->ClusterId);
+
+	if (Connections != nullptr)
+	{
+		for (UDisplayClusterNetConnection* NodeConnectionIt : *Connections)
+		{
+			SyncConnections.Remove(NodeConnectionIt);
+			NodeConnections.Remove(NodeConnectionIt);
+		}
+		Connections->Remove(NetConnection);
+	}
+
+	// Reset packet queues
+	if (SyncConnections.IsEmpty())
+	{
+		LastAckedPacket.Reset();
+		OutPacketsQueues.Reset();
+		bClusterHasConnected = false;
+	}
 }
 
 void UDisplayClusterNetDriver::HandleEvent(FDisplayClusterClusterEventBinary const& InEvent)
@@ -227,13 +249,13 @@ void UDisplayClusterNetDriver::HandleEvent(FDisplayClusterClusterEventBinary con
 
 	if (InEvent.EventId == NodeSyncEvent)
 	{
-		UE_LOG(LogDisplayClusterNetDriver, Verbose, TEXT("[StartSynchronousMode] command received"));
+		UE_LOG(LogDisplayClusterNetDriver, VeryVerbose, TEXT("[StartSynchronousMode] command received"));
 
 		// Client received command to start synchrosonus packets processing
 		ClusterServerConnection->bSynchronousMode = true;
 		ClusterServerConnection->ClientId = TextKeyUtil::HashString(ClusterServerConnection->Challenge);
 
-		UE_LOG(LogDisplayClusterNetDriver, Verbose, TEXT("Client node self-introducing: %u"), ClusterServerConnection->ClientId);
+		UE_LOG(LogDisplayClusterNetDriver, VeryVerbose, TEXT("Client node self-introducing: %u"), ClusterServerConnection->ClientId);
 	}
 
 	if (InEvent.EventId == PacketSyncEvent)
@@ -324,7 +346,7 @@ void UDisplayClusterNetDriver::TickFlush(float DeltaSeconds)
 		// Driver waits until last sent bunch of any sync connection will be acked
 		bLastBunchWasAcked |= (LastOutPacketId - ConnectionAckId > 1);
 
-		UE_LOG(LogDisplayClusterNetDriver, Verbose, TEXT("Client %u; Packet %i sent, %i acked"), SyncConnectionIt->ClientId, ConnectionAckId, LastAckedPacketIdRef);
+		UE_LOG(LogDisplayClusterNetDriver, VeryVerbose, TEXT("Client %u; Packet %i sent, %i acked"), SyncConnectionIt->ClientId, ConnectionAckId, LastAckedPacketIdRef);
 	}
 
 	// Find slice of packets which are ready for processing on clients
@@ -375,12 +397,15 @@ void UDisplayClusterNetDriver::TickFlush(float DeltaSeconds)
 
 			// Slice queue by removing all packets which are behind ReadyForProcessingPacket
 			TDeque<int32>& ConnectionPacketQueue = OutPacketsQueues[UniqueID];
-			for (int j = 0; j < PacketIdx; j++)
+			if (!ConnectionPacketQueue.IsEmpty())
 			{
-				ConnectionPacketQueue.PopFirst();
+				for (int j = 0; j < PacketIdx; j++)
+				{
+					ConnectionPacketQueue.PopFirst();
+				}
 			}
 
-			UE_LOG(LogDisplayClusterNetDriver, Verbose, TEXT("Client %u; [ProcessPacket %i] command sent"), SyncConnection->ClientId, ReadyForProcessingPacket);
+			UE_LOG(LogDisplayClusterNetDriver, VeryVerbose, TEXT("Client %u; [ProcessPacket %i] command sent"), SyncConnection->ClientId, ReadyForProcessingPacket);
 		}
 
 		// Generate cluster event
diff --git a/Engine/Plugins/Runtime/nDisplay/Source/DisplayClusterReplication/Public/DisplayClusterNetConnection.h b/Engine/Plugins/Runtime/nDisplay/Source/DisplayClusterReplication/Public/DisplayClusterNetConnection.h
--- a/Engine/Plugins/Runtime/nDisplay/Source/DisplayClusterReplication/Public/DisplayClusterNetConnection.h
+++ b/Engine/Plugins/Runtime/nDisplay/Source/DisplayClusterReplication/Public/DisplayClusterNetConnection.h
@@ -38,6 +38,9 @@ public:
 	/** Cluster unique identifier, HashString from node config file path */
 	uint32 ClusterId;
 
+	/** Cluster nodes number */
+	uint32 ClusterNodesNum;
+
 	/** Cluster node port for binary cluster events */
 	uint16 NodePort;
 
diff --git a/Engine/Plugins/Runtime/nDisplay/Source/DisplayClusterReplication/Public/DisplayClusterNetDriver.h b/Engine/Plugins/Runtime/nDisplay/Source/DisplayClusterReplication/Public/DisplayClusterNetDriver.h
--- a/Engine/Plugins/Runtime/nDisplay/Source/DisplayClusterReplication/Public/DisplayClusterNetDriver.h
+++ b/Engine/Plugins/Runtime/nDisplay/Source/DisplayClusterReplication/Public/DisplayClusterNetDriver.h
@@ -99,6 +99,9 @@ protected:
 	// nDisplay Primary node connections
 	TSet<UDisplayClusterNetConnection*> PrimaryNodeConnections;
 
+	// nDisplay Node connections which corresponds to specific cluster
+	TMap<uint32, TArray<UDisplayClusterNetConnection*>> ClusterConnections;
+
 	// nDisplay node connections that participate in synchronous replication
 	TSet<UDisplayClusterNetConnection*> SyncConnections;
 
@@ -114,6 +117,9 @@ protected:
 	// Binary event listener
 	FOnClusterEventBinaryListener EventBinaryListener;
 
+	// True if cluster has connected
+	bool bClusterHasConnected;
+
 	// Used to control state of acked packets queue
 	bool bLastBunchWasAcked;
 
diff --git a/Engine/Plugins/VirtualProduction/Switchboard/Source/Switchboard/switchboard/devices/ndisplay/plugin_ndisplay.py b/Engine/Plugins/VirtualProduction/Switchboard/Source/Switchboard/switchboard/devices/ndisplay/plugin_ndisplay.py
--- a/Engine/Plugins/VirtualProduction/Switchboard/Source/Switchboard/switchboard/devices/ndisplay/plugin_ndisplay.py
+++ b/Engine/Plugins/VirtualProduction/Switchboard/Source/Switchboard/switchboard/devices/ndisplay/plugin_ndisplay.py
@@ -711,6 +711,12 @@ class DevicenDisplay(DeviceUnreal):
         if self.nodeconfig and 'port_ce_binary' in self.nodeconfig:
             primary_node_binary_event_port = self.nodeconfig['port_ce_binary']
 
+        cluster_nodes_count = 0
+        for device in self.active_unreal_devices:
+            is_device_connected = device.status == DeviceStatus.CLOSED or device.status == DeviceStatus.OPEN
+            if device.device_type == "nDisplay" and is_device_connected:
+                cluster_nodes_count = cluster_nodes_count + 1
+
         def get_common_args():
             """ 
             get common args for client and server 
@@ -718,7 +724,8 @@ class DevicenDisplay(DeviceUnreal):
             return '?'.join([
                     f'ClusterId={os.getpid()}',
                     f'PrimaryNodeId={primary_device_name}',
-                    f'PrimaryNodePort={primary_node_binary_event_port}'
+                    f'PrimaryNodePort={primary_node_binary_event_port}',
+                    f'ClusterNodesCount={cluster_nodes_count}'
                 ])
 
         def get_client_args(server_address):
