nDisplay: Fixed preview rendering for 'mesh' projectiuon policy.

diff --git a/Engine/Plugins/Runtime/nDisplay/Source/DisplayClusterProjection/Private/Policy/MPCDI/DisplayClusterProjectionMPCDIPolicy.h b/Engine/Plugins/Runtime/nDisplay/Source/DisplayClusterProjection/Private/Policy/MPCDI/DisplayClusterProjectionMPCDIPolicy.h
--- a/Engine/Plugins/Runtime/nDisplay/Source/DisplayClusterProjection/Private/Policy/MPCDI/DisplayClusterProjectionMPCDIPolicy.h
+++ b/Engine/Plugins/Runtime/nDisplay/Source/DisplayClusterProjection/Private/Policy/MPCDI/DisplayClusterProjectionMPCDIPolicy.h
@@ -74,7 +74,7 @@ protected:
 	TArray<FDisplayClusterWarpContext> WarpBlendContexts;
 	TArray<FDisplayClusterWarpContext> WarpBlendContexts_Proxy;
 
-private:
+protected:
 	bool bInvalidConfiguration = false;
 	bool bIsPreviewMeshEnabled = false;
 
diff --git a/Engine/Plugins/Runtime/nDisplay/Source/DisplayClusterProjection/Private/Policy/Mesh/DisplayClusterProjectionMeshPolicy.cpp b/Engine/Plugins/Runtime/nDisplay/Source/DisplayClusterProjection/Private/Policy/Mesh/DisplayClusterProjectionMeshPolicy.cpp
--- a/Engine/Plugins/Runtime/nDisplay/Source/DisplayClusterProjection/Private/Policy/Mesh/DisplayClusterProjectionMeshPolicy.cpp
+++ b/Engine/Plugins/Runtime/nDisplay/Source/DisplayClusterProjection/Private/Policy/Mesh/DisplayClusterProjectionMeshPolicy.cpp
@@ -59,6 +59,9 @@ bool FDisplayClusterProjectionMeshPolicy::CreateWarpMeshInterface(IDisplayCluste
 			return false;
 		}
 
+		// Mesh poolicy always shows a preview on an existing mesh
+		bIsPreviewMeshEnabled = true;
+
 		// The mesh always uses DCRootActor as its origin because all geometry is in UE space.
 		InitializeOriginComponent(InViewport, TEXT(""));
 
