[nDisplay] Fixed performance drop in MP+SharedMemMedia when 'DCRA::bEnableInnerFrustums' is unchecked

diff --git a/Engine/Plugins/Runtime/nDisplay/Source/DisplayCluster/Private/Game/EngineClasses/Scene/DisplayClusterRootActor.cpp b/Engine/Plugins/Runtime/nDisplay/Source/DisplayCluster/Private/Game/EngineClasses/Scene/DisplayClusterRootActor.cpp
--- a/Engine/Plugins/Runtime/nDisplay/Source/DisplayCluster/Private/Game/EngineClasses/Scene/DisplayClusterRootActor.cpp
+++ b/Engine/Plugins/Runtime/nDisplay/Source/DisplayCluster/Private/Game/EngineClasses/Scene/DisplayClusterRootActor.cpp
@@ -442,13 +442,16 @@ UDisplayClusterConfigurationData* ADisplayClusterRootActor::GetConfigData() cons
 
 bool ADisplayClusterRootActor::IsInnerFrustumEnabled(const FString& InnerFrustumID) const
 {
+	// Common condition for any camera/frustum
+	const bool bClusterGlobalInnersEnabled = (CurrentConfigData ? CurrentConfigData->StageSettings.bEnableInnerFrustums : true);
+
 	// add more GUI rules here
 	// Inner Frustum Enabled
 	//  Camera_1  [ ]
 	//  Camera_2  [X]
 	//  Camera_3  [X]
 
-	return true;
+	return bClusterGlobalInnersEnabled; // && bOtherCondition(s)
 }
 
 int ADisplayClusterRootActor::GetInnerFrustumPriority(const FString& InnerFrustumID) const
