Added UE::RenderCommandPipe::FSyncScope to wrap render commands that invoke BeginScenePrimitiveRendering.

diff --git a/Engine/Plugins/Runtime/RenderTrace/Source/Private/RenderTrace.cpp b/Engine/Plugins/Runtime/RenderTrace/Source/Private/RenderTrace.cpp
--- a/Engine/Plugins/Runtime/RenderTrace/Source/Private/RenderTrace.cpp
+++ b/Engine/Plugins/Runtime/RenderTrace/Source/Private/RenderTrace.cpp
@@ -643,6 +643,8 @@ void FRenderTraceQueue::Tick(float DeltaTime)
 
 	if (RequestsToUpdate.Num() > 0)
 	{
+		UE::RenderCommandPipe::FSyncScope SyncScope;
+
 		ENQUEUE_RENDER_COMMAND(FRenderTraceUpdaterTick)(
 			[InRequestsToUpdate=MoveTemp(RequestsToUpdate)](FRHICommandListImmediate& RHICmdList)
 			{
diff --git a/Engine/Plugins/Runtime/nDisplay/Source/DisplayCluster/Private/Render/Viewport/LightCard/DisplayClusterViewportLightCardManagerProxy.cpp b/Engine/Plugins/Runtime/nDisplay/Source/DisplayCluster/Private/Render/Viewport/LightCard/DisplayClusterViewportLightCardManagerProxy.cpp
--- a/Engine/Plugins/Runtime/nDisplay/Source/DisplayCluster/Private/Render/Viewport/LightCard/DisplayClusterViewportLightCardManagerProxy.cpp
+++ b/Engine/Plugins/Runtime/nDisplay/Source/DisplayCluster/Private/Render/Viewport/LightCard/DisplayClusterViewportLightCardManagerProxy.cpp
@@ -42,6 +42,8 @@ void FDisplayClusterViewportLightCardManagerProxy::ReleaseUVLightCardResource()
 
 void FDisplayClusterViewportLightCardManagerProxy::RenderUVLightCard(FSceneInterface* InScene, const FDisplayClusterShaderParameters_UVLightCards& InParameters) const
 {
+	UE::RenderCommandPipe::FSyncScope SyncScope;
+
 	ENQUEUE_RENDER_COMMAND(DisplayClusterViewportLightCardManagerProxy_RenderUVLightCard)(
 		[InProxyData = SharedThis(this), Scene = InScene, Parameters = InParameters](FRHICommandListImmediate& RHICmdList)
 		{
diff --git a/Engine/Plugins/Runtime/nDisplay/Source/DisplayClusterLightCardEditorShaders/Private/DisplayClusterMeshProjectionRenderer.cpp b/Engine/Plugins/Runtime/nDisplay/Source/DisplayClusterLightCardEditorShaders/Private/DisplayClusterMeshProjectionRenderer.cpp
--- a/Engine/Plugins/Runtime/nDisplay/Source/DisplayClusterLightCardEditorShaders/Private/DisplayClusterMeshProjectionRenderer.cpp
+++ b/Engine/Plugins/Runtime/nDisplay/Source/DisplayClusterLightCardEditorShaders/Private/DisplayClusterMeshProjectionRenderer.cpp
@@ -906,6 +906,8 @@ void FDisplayClusterMeshProjectionRenderer::Render(FCanvas* Canvas, FSceneInterf
 	const bool bIsHitTesting = Canvas->IsHitTesting();
 	FHitProxyConsumer* HitProxyConsumer = Canvas->GetHitProxyConsumer();
 
+	UE::RenderCommandPipe::FSyncScope SyncScope;
+
 	ENQUEUE_RENDER_COMMAND(FDrawProjectedMeshes)(
 		[RenderTarget, Scene, RenderSettings, bIsHitTesting, HitProxyConsumer, this](FRHICommandListImmediate& RHICmdList)
 		{
diff --git a/Engine/Source/Runtime/Landscape/Private/LandscapeGrass.cpp b/Engine/Source/Runtime/Landscape/Private/LandscapeGrass.cpp
--- a/Engine/Source/Runtime/Landscape/Private/LandscapeGrass.cpp
+++ b/Engine/Source/Runtime/Landscape/Private/LandscapeGrass.cpp
@@ -1315,6 +1315,8 @@ public:
 		RenderTargetTexture->InitCustomFormat(TargetSize.X, TargetSize.Y, PF_B8G8R8A8, bForceLinearGamma);
 		RenderTargetResource = RenderTargetTexture->GameThread_GetRenderTargetResource()->GetTextureRenderTarget2DResource();
 
+		UE::RenderCommandPipe::FSyncScope SyncScope;
+
 		// render
 		FLandscapeGrassWeightExporter_RenderThread* Exporter = this;
 		ENQUEUE_RENDER_COMMAND(FDrawSceneCommand)(
diff --git a/Engine/Source/Runtime/Landscape/Private/LandscapePhysicalMaterial.cpp b/Engine/Source/Runtime/Landscape/Private/LandscapePhysicalMaterial.cpp
--- a/Engine/Source/Runtime/Landscape/Private/LandscapePhysicalMaterial.cpp
+++ b/Engine/Source/Runtime/Landscape/Private/LandscapePhysicalMaterial.cpp
@@ -612,6 +612,8 @@ void FLandscapePhysicalMaterialRenderTask::UpdateInternal(bool bInFlush)
 		FLandscapePhysicalMaterialRenderTaskImpl* Task = &GTaskPool.Pool[PoolHandle];
 		check(Task->LandscapeComponent != nullptr);
 
+		UE::RenderCommandPipe::FSyncScope SyncScope;
+
 		// If the material gets recompiled, for example, the component's scene proxy could be changed before the RT update runs, so we need to make sure we're getting the scene proxy on the game thread :
 		const FLandscapeComponentSceneProxy* SceneProxy = static_cast<const FLandscapeComponentSceneProxy*>(Task->LandscapeComponent->SceneProxy);
 		ENQUEUE_RENDER_COMMAND(FLandscapePhysicalMaterialFlush)(
