nDisplay: Fixed bugs related to ICVFX InCamera soft edges:
- fixed ICVFX shader GetFeatherAlpha() - the alpha of the InCamera inner area should always be set to 1 (parameter 'feather')
- fixed function GetCameraSoftEdge():
  * custom frustum made changes to soft edges even if it was not enabled
  * The softedge 'z' and 'w' values are used for other purposes and cannot be changed by custom frustum.

diff --git a/Engine/Plugins/Runtime/nDisplay/Shaders/Private/MPCDIOverlayShaders.usf b/Engine/Plugins/Runtime/nDisplay/Shaders/Private/MPCDIOverlayShaders.usf
--- a/Engine/Plugins/Runtime/nDisplay/Shaders/Private/MPCDIOverlayShaders.usf
+++ b/Engine/Plugins/Runtime/nDisplay/Shaders/Private/MPCDIOverlayShaders.usf
@@ -71,16 +71,24 @@ float GetFeatherAlpha(float2 EdgeOffset, float Feather)
 {
 	float Alpha = 1;
 
-	Alpha *= smoothstep(0, Feather, EdgeOffset.x); // horizontal component
-	Alpha *= smoothstep(0, Feather, EdgeOffset.y); // vertical component
+	// do not change the alpha within an in-camera frame
+	if (EdgeOffset.x < 1)
+	{
+		Alpha *= smoothstep(0, Feather, EdgeOffset.x); // horizontal component
+	}
+	
+	// do not change the alpha within an in-camera frame
+	if (EdgeOffset.y < 1)
+	{
+		Alpha *= smoothstep(0, Feather, EdgeOffset.y); // vertical component
+	}
 
 	return Alpha;
 }
 
 float GetInnerCameraSoftEdgeAlpha(float2 CameraUV, float4 InInnerCameraSoftEdge)
 {
-
-	float2 EdgeOffset = float2(1,1);
+	float2 EdgeOffset = float2(1, 1);
 
 	// InInnerCameraSoftEdge.x values are used to pass left/right soft edge values
 	if(CameraUV.x < InInnerCameraSoftEdge.x) // left 
diff --git a/Engine/Plugins/Runtime/nDisplay/Source/DisplayClusterConfiguration/Private/DisplayClusterConfigurationTypes_ICVFX.cpp b/Engine/Plugins/Runtime/nDisplay/Source/DisplayClusterConfiguration/Private/DisplayClusterConfigurationTypes_ICVFX.cpp
--- a/Engine/Plugins/Runtime/nDisplay/Source/DisplayClusterConfiguration/Private/DisplayClusterConfigurationTypes_ICVFX.cpp
+++ b/Engine/Plugins/Runtime/nDisplay/Source/DisplayClusterConfiguration/Private/DisplayClusterConfigurationTypes_ICVFX.cpp
@@ -272,35 +272,36 @@ FVector4 FDisplayClusterConfigurationICVFX_CameraSettings::GetCameraSoftEdge(con
 	// Z for new parameter Feather
 	ResultSoftEdge.Z = SoftEdge.Feather;
 
-	const float CameraBufferRatio = GetCameraBufferRatio(InStageSettings);
+	// Custom frustum made changes to soft edges
+	if (CustomFrustum.bEnable)
+	{
+		// default - percents
+		const float ConvertToPercent = 0.01f;
 
-	// default - percents
-	const float ConvertToPercent = 0.01f;
+		float Left = ClampPercent(CustomFrustum.Left * ConvertToPercent);
+		float Right = ClampPercent(CustomFrustum.Right * ConvertToPercent);
+		float Top = ClampPercent(CustomFrustum.Top * ConvertToPercent);
+		float Bottom = ClampPercent(CustomFrustum.Bottom * ConvertToPercent);
 
-	float Left = ClampPercent(CustomFrustum.Left * ConvertToPercent);
-	float Right = ClampPercent(CustomFrustum.Right * ConvertToPercent);
-	float Top = ClampPercent(CustomFrustum.Top * ConvertToPercent);
-	float Bottom = ClampPercent(CustomFrustum.Bottom * ConvertToPercent);
+		if (CustomFrustum.Mode == EDisplayClusterConfigurationViewportCustomFrustumMode::Pixels)
+		{
+			const float CameraBufferRatio = GetCameraBufferRatio(InStageSettings);
+			const FIntPoint FrameSize = GetCameraFrameSize(InStageSettings);
 
-	if (CustomFrustum.Mode == EDisplayClusterConfigurationViewportCustomFrustumMode::Pixels)
-	{
-		const FIntPoint FrameSize = GetCameraFrameSize(InStageSettings);
+			const float  FrameWidth = FrameSize.X * CameraBufferRatio;
+			const float FrameHeight = FrameSize.Y * CameraBufferRatio;
 
-		const float  FrameWidth = FrameSize.X * CameraBufferRatio;
-		const float FrameHeight = FrameSize.Y * CameraBufferRatio;
+			Left = ClampPercent(CustomFrustum.Left / FrameWidth);
+			Right = ClampPercent(CustomFrustum.Right / FrameWidth);
+			Top = ClampPercent(CustomFrustum.Top / FrameHeight);
+			Bottom = ClampPercent(CustomFrustum.Bottom / FrameHeight);
+		}
 
-		Left = ClampPercent(CustomFrustum.Left / FrameWidth);
-		Right = ClampPercent(CustomFrustum.Right / FrameWidth);
-		Top = ClampPercent(CustomFrustum.Top / FrameHeight);
-		Bottom = ClampPercent(CustomFrustum.Bottom / FrameHeight);
+		// recalculate soft edge related offsets based on frustum
+		ResultSoftEdge.X /= (1 + Left + Right);
+		ResultSoftEdge.Y /= (1 + Top + Bottom);
 	}
 
-	// recalculate soft edge related offsets based on frustum
-	ResultSoftEdge.X /= (1 + Left + Right);
-	ResultSoftEdge.Y /= (1 + Top + Bottom);
-	ResultSoftEdge.Z /= (1 + Left + Right);
-	ResultSoftEdge.W /= (1 + Top + Bottom);
-
 	return ResultSoftEdge;
 }
 
