nDisplay: Fixed shader compilation issue where the MAX_INNER_CAMERAS_AMOUNT define wasn't being set on all MPCDI shaders, causing compiliation issues.

diff --git a/Engine/Plugins/Runtime/nDisplay/Shaders/Private/MPCDIOverlayShaders.usf b/Engine/Plugins/Runtime/nDisplay/Shaders/Private/MPCDIOverlayShaders.usf
--- a/Engine/Plugins/Runtime/nDisplay/Shaders/Private/MPCDIOverlayShaders.usf
+++ b/Engine/Plugins/Runtime/nDisplay/Shaders/Private/MPCDIOverlayShaders.usf
@@ -176,7 +176,7 @@ float4 GetInnerCameraColor(float4 WarpedUV, OutputVS IN)
 			CameraColor.w = GetInnerCameraSoftEdgeAlpha(CameraUV, InnerCameraSoftEdge);
 
 			// **** disable incamera render inside overlapped areas
-			#if !INNER_CAMERA_OVERLAP
+			#if ENABLE_INNER_CAMERA_OVERLAP && !INNER_CAMERA_OVERLAP
 			if(OverlappedInnerCamerasNum > 0 && CameraColor.w < 1)
 			{
 				for(uint OverlappedCameraIndex = 0; OverlappedCameraIndex < OverlappedInnerCamerasNum; OverlappedCameraIndex++)
@@ -226,7 +226,7 @@ float4 GetInnerCameraColor(float4 WarpedUV, OutputVS IN)
 				#endif
 
 				// ********* When camera overlap we render chromakey **********
-				#if INNER_CAMERA_OVERLAP
+				#if ENABLE_INNER_CAMERA_OVERLAP && INNER_CAMERA_OVERLAP
 					if(CameraColor.w >= 0)
 					{
 						OverlayChromakeyColor.w = 0;
diff --git a/Engine/Plugins/Runtime/nDisplay/Shaders/Private/MPCDIUtils.ush b/Engine/Plugins/Runtime/nDisplay/Shaders/Private/MPCDIUtils.ush
--- a/Engine/Plugins/Runtime/nDisplay/Shaders/Private/MPCDIUtils.ush
+++ b/Engine/Plugins/Runtime/nDisplay/Shaders/Private/MPCDIUtils.ush
@@ -42,9 +42,15 @@ float4x4 InnerCameraProjectionMatrix;
 // The total number of InCameras, which are rendered below.
 int OverlappedInnerCamerasNum;
 
+#ifdef MAX_INNER_CAMERAS_AMOUNT
+#define ENABLE_INNER_CAMERA_OVERLAP 1
+
 // Matrices and SoftEdges for InCameras, which are rendered below.
 float4x4 OverlappedInnerCamerasProjectionMatrices[MAX_INNER_CAMERAS_AMOUNT];
 float4 OverlappedInnerCameraSoftEdges[MAX_INNER_CAMERAS_AMOUNT];
+#else
+#define ENABLE_INNER_CAMERA_OVERLAP 0
+#endif
 
 float AlphaEmbeddedGamma;
 
