nDisplay: Fixed RTT cleanup bug in OutputRemap.

diff --git a/Engine/Plugins/Runtime/nDisplay/Source/DisplayCluster/Private/Render/Viewport/Postprocess/DisplayClusterViewportPostProcessOutputRemap.cpp b/Engine/Plugins/Runtime/nDisplay/Source/DisplayCluster/Private/Render/Viewport/Postprocess/DisplayClusterViewportPostProcessOutputRemap.cpp
--- a/Engine/Plugins/Runtime/nDisplay/Source/DisplayCluster/Private/Render/Viewport/Postprocess/DisplayClusterViewportPostProcessOutputRemap.cpp
+++ b/Engine/Plugins/Runtime/nDisplay/Source/DisplayCluster/Private/Render/Viewport/Postprocess/DisplayClusterViewportPostProcessOutputRemap.cpp
@@ -25,19 +25,6 @@
 #include "Engine/StaticMesh.h"
 #include "ProceduralMeshComponent.h"
 
-#include "ClearQuad.h"
-
-// Enable/Disable ClearTexture for OutputRemap RTT
-static TAutoConsoleVariable<int32> CVarClearOutputRemapFrameTextureEnabled(
-	TEXT("nDisplay.render.output_remap.ClearTextureEnabled"),
-	1,
-	TEXT("Enables OutputRemap RTT clearing before postprocessing.\n")
-	TEXT("0 : disabled\n")
-	TEXT("1 : enabled\n")
-	,
-	ECVF_RenderThreadSafe
-);
-
 //////////////////////////////////////////////////////////////////////////////////////////////
 // FDisplayClusterViewportPostProcessOutputRemap
 //////////////////////////////////////////////////////////////////////////////////////////////
@@ -253,21 +240,11 @@ void FDisplayClusterViewportPostProcessOutputRemap::PerformPostProcessFrame_Rend
 		const IDisplayClusterRender_MeshComponentProxy* MeshProxy = OutputRemapMesh->GetMeshComponentProxy_RenderThread();
 		if (MeshProxy!=nullptr && MeshProxy->IsEnabled_RenderThread())
 		{
-			const bool bClearOutputRemapFrameTextureEnabled = CVarClearOutputRemapFrameTextureEnabled.GetValueOnRenderThread() != 0;
-
 			for (int32 Index = 0; Index < InFrameTargets->Num(); Index++)
 			{
 				FRHITexture2D* InOutTexture = (*InFrameTargets)[Index];
 				FRHITexture2D* TempTargetableTexture = (*InAdditionalFrameTargets)[Index];
 
-				if (bClearOutputRemapFrameTextureEnabled)
-				{
-					// Do clear output frame texture before whole frame remap
-					const FIntPoint Size = TempTargetableTexture->GetSizeXY();
-					RHICmdList.SetViewport(0, 0, 0.0f, Size.X, Size.Y, 1.0f);
-					DrawClearQuad(RHICmdList, FLinearColor::Black);
-				}
-
 				if (ShadersAPI.RenderPostprocess_OutputRemap(RHICmdList, InOutTexture, TempTargetableTexture, *MeshProxy))
 				{
 					TransitionAndCopyTexture(RHICmdList, TempTargetableTexture, InOutTexture, {});
diff --git a/Engine/Plugins/Runtime/nDisplay/Source/DisplayClusterShaders/Private/Shaders/DisplayClusterShadersPostprocess_OutputRemap.cpp b/Engine/Plugins/Runtime/nDisplay/Source/DisplayClusterShaders/Private/Shaders/DisplayClusterShadersPostprocess_OutputRemap.cpp
--- a/Engine/Plugins/Runtime/nDisplay/Source/DisplayClusterShaders/Private/Shaders/DisplayClusterShadersPostprocess_OutputRemap.cpp
+++ b/Engine/Plugins/Runtime/nDisplay/Source/DisplayClusterShaders/Private/Shaders/DisplayClusterShadersPostprocess_OutputRemap.cpp
@@ -49,6 +49,17 @@ static TAutoConsoleVariable<int32> CVarOutputRemapShaderType(
 	,ECVF_RenderThreadSafe
 );
 
+// Enable/Disable ClearTexture for OutputRemap RTT
+static TAutoConsoleVariable<int32> CVarClearOutputRemapClearRTT(
+	TEXT("nDisplay.render.output_remap.clear_rtt"),
+	1,
+	TEXT("Enables OutputRemap RTT clearing before postprocessing.\n")
+	TEXT("0 : disabled\n")
+	TEXT("1 : enabled\n")
+	,
+	ECVF_RenderThreadSafe
+);
+
 BEGIN_SHADER_PARAMETER_STRUCT(FOutputRemapVertexShaderParameters, )
 END_SHADER_PARAMETER_STRUCT()
 
@@ -137,6 +148,8 @@ bool FDisplayClusterShadersPostprocess_OutputRemap::RenderPostprocess_OutputRema
 
 	bool bResult = false;
 
+	const bool bClearTextureEnabled = CVarClearOutputRemapClearRTT.GetValueOnRenderThread() != 0;
+
 	SCOPED_GPU_STAT(RHICmdList, nDisplay_PostProcess_OutputRemap);
 	SCOPED_DRAW_EVENT(RHICmdList, nDisplay_PostProcess_OutputRemap);
 
@@ -148,6 +161,12 @@ bool FDisplayClusterShadersPostprocess_OutputRemap::RenderPostprocess_OutputRema
 		const FIntPoint TargetSizeXY = InRenderTargetableDestTexture->GetSizeXY();
 		RHICmdList.SetViewport(0, 0, 0.0f, TargetSizeXY.X, TargetSizeXY.Y, 1.0f);
 
+		if (bClearTextureEnabled)
+		{
+			// Clean the texture of the output frame before remapping the whole frame
+			DrawClearQuad(RHICmdList, FLinearColor::Black);
+		}
+
 		const ERHIFeatureLevel::Type RenderFeatureLevel = GMaxRHIFeatureLevel;
 		const auto GlobalShaderMap = GetGlobalShaderMap(RenderFeatureLevel);
 
