More ComponentID to GetPrimitiveSceneId() fixed

diff --git a/Engine/Plugins/Runtime/nDisplay/Source/DisplayCluster/Private/Game/EngineClasses/Scene/DisplayClusterRootActor.cpp b/Engine/Plugins/Runtime/nDisplay/Source/DisplayCluster/Private/Game/EngineClasses/Scene/DisplayClusterRootActor.cpp
--- a/Engine/Plugins/Runtime/nDisplay/Source/DisplayCluster/Private/Game/EngineClasses/Scene/DisplayClusterRootActor.cpp
+++ b/Engine/Plugins/Runtime/nDisplay/Source/DisplayCluster/Private/Game/EngineClasses/Scene/DisplayClusterRootActor.cpp
@@ -486,7 +486,7 @@ void ImplCollectChildHiddenComponents(TSet<FPrimitiveComponentId>& OutPrimitives
 				UPrimitiveComponent* PrimComp = Cast<UPrimitiveComponent>(ChildIt);
 				if (PrimComp)
 				{
-					OutPrimitives.Add(PrimComp->ComponentId);
+					OutPrimitives.Add(PrimComp->GetPrimitiveSceneId());
 				}
 			}
 		}
@@ -514,7 +514,7 @@ void ADisplayClusterRootActor::GetTypedPrimitives(TSet<FPrimitiveComponentId>& O
 						UPrimitiveComponent* PrimComp = Cast<UPrimitiveComponent>(CompIt);
 						if (PrimComp)
 						{
-							OutPrimitives.Add(PrimComp->ComponentId);
+							OutPrimitives.Add(PrimComp->GetPrimitiveSceneId());
 						}
 
 						if (bCollectChildrenVisualizationComponent)
@@ -530,7 +530,7 @@ void ADisplayClusterRootActor::GetTypedPrimitives(TSet<FPrimitiveComponentId>& O
 				UPrimitiveComponent* PrimComp = Cast<UPrimitiveComponent>(CompIt);
 				if (PrimComp)
 				{
-					OutPrimitives.Add(PrimComp->ComponentId);
+					OutPrimitives.Add(PrimComp->GetPrimitiveSceneId());
 				}
 
 				if (bCollectChildrenVisualizationComponent)
@@ -585,7 +585,7 @@ bool ADisplayClusterRootActor::GetHiddenInGamePrimitives(TSet<FPrimitiveComponen
 			{
 				if(UPrimitiveComponent* PrimComp = Cast<UPrimitiveComponent>(PreviewMesh))
 				{
-					OutPrimitives.Add(PrimComp->ComponentId);
+					OutPrimitives.Add(PrimComp->GetPrimitiveSceneId());
 				}
 			}
 		}
@@ -600,7 +600,7 @@ bool ADisplayClusterRootActor::GetHiddenInGamePrimitives(TSet<FPrimitiveComponen
 		{
 			if (CompIt->bHiddenInGame)
 			{
-				OutPrimitives.Add(CompIt->ComponentId);
+				OutPrimitives.Add(CompIt->GetPrimitiveSceneId());
 			}
 
 			ImplCollectChildHiddenComponents(OutPrimitives, CompIt);
@@ -684,7 +684,7 @@ bool ADisplayClusterRootActor::GetHiddenInGamePrimitives(TSet<FPrimitiveComponen
 #endif
 								)
 							{
-								PrimitiveComponentsArray[Index].Add(PrimComp->ComponentId);
+								PrimitiveComponentsArray[Index].Add(PrimComp->GetPrimitiveSceneId());
 							}
 
 							ImplCollectChildHiddenComponents(PrimitiveComponentsArray[Index], PrimComp);
diff --git a/Engine/Plugins/Runtime/nDisplay/Source/DisplayCluster/Private/Render/Viewport/Configuration/DisplayClusterViewportConfigurationHelpers_Visibility.cpp b/Engine/Plugins/Runtime/nDisplay/Source/DisplayCluster/Private/Render/Viewport/Configuration/DisplayClusterViewportConfigurationHelpers_Visibility.cpp
--- a/Engine/Plugins/Runtime/nDisplay/Source/DisplayCluster/Private/Render/Viewport/Configuration/DisplayClusterViewportConfigurationHelpers_Visibility.cpp
+++ b/Engine/Plugins/Runtime/nDisplay/Source/DisplayCluster/Private/Render/Viewport/Configuration/DisplayClusterViewportConfigurationHelpers_Visibility.cpp
@@ -31,7 +31,7 @@ static void ImplCollectVisibility(ADisplayClusterRootActor& RootActor, const FDi
 				{
 					if (const UPrimitiveComponent* PrimComp = Cast<UPrimitiveComponent>(Component))
 					{
-						OutAdditionalComponentsList.Add(PrimComp->ComponentId);
+						OutAdditionalComponentsList.Add(PrimComp->GetPrimitiveSceneId());
 					}
 				}
 			}
diff --git a/Engine/Plugins/Runtime/nDisplay/Source/DisplayCluster/Private/Render/Viewport/DisplayClusterViewport_VisibilitySettings.cpp b/Engine/Plugins/Runtime/nDisplay/Source/DisplayCluster/Private/Render/Viewport/DisplayClusterViewport_VisibilitySettings.cpp
--- a/Engine/Plugins/Runtime/nDisplay/Source/DisplayCluster/Private/Render/Viewport/DisplayClusterViewport_VisibilitySettings.cpp
+++ b/Engine/Plugins/Runtime/nDisplay/Source/DisplayCluster/Private/Render/Viewport/DisplayClusterViewport_VisibilitySettings.cpp
@@ -38,9 +38,9 @@ static void GetPrimitiveComponentsFromLayers(UWorld* World, const TArray<FName>&
 					{
 						if (const UPrimitiveComponent* PrimComp = Cast<UPrimitiveComponent>(Component))
 						{
-							if (!InExcludePrimitivesList || !InExcludePrimitivesList->Contains(PrimComp->ComponentId))
+							if (!InExcludePrimitivesList || !InExcludePrimitivesList->Contains(PrimComp->GetPrimitiveSceneId()))
 							{
-								OutPrimitives.Add(PrimComp->ComponentId);
+								OutPrimitives.Add(PrimComp->GetPrimitiveSceneId());
 							}
 						}
 					}
