nDisplay: Added a toolbar button to import an MPCDI config into an nDisplay config blueprint, which will add all necessary cluster nodes and viewports, and configure them correctly, as well as adding screen components to match if the MPCDI config is a 2D config.

diff --git a/Engine/Plugins/Runtime/nDisplay/Source/DisplayClusterConfigurator/DisplayClusterConfigurator.Build.cs b/Engine/Plugins/Runtime/nDisplay/Source/DisplayClusterConfigurator/DisplayClusterConfigurator.Build.cs
--- a/Engine/Plugins/Runtime/nDisplay/Source/DisplayClusterConfigurator/DisplayClusterConfigurator.Build.cs
+++ b/Engine/Plugins/Runtime/nDisplay/Source/DisplayClusterConfigurator/DisplayClusterConfigurator.Build.cs
@@ -13,6 +13,7 @@ public class DisplayClusterConfigurator : ModuleRules
 				"DisplayCluster",
 				"DisplayClusterConfiguration",
 				"DisplayClusterProjection",
+				"DisplayClusterWarp",
 
 				"AdvancedPreviewScene",
 				"ApplicationCore",
@@ -36,6 +37,7 @@ public class DisplayClusterConfigurator : ModuleRules
 				"MainFrame",
 				"MediaAssets",
 				"MessageLog",
+				"Networking",
 				"PinnedCommandList",
 				"Projects",
 				"PropertyEditor",
diff --git a/Engine/Plugins/Runtime/nDisplay/Source/DisplayClusterConfigurator/Private/ClusterConfiguration/ViewModels/DisplayClusterConfiguratorClusterNodeViewModel.cpp b/Engine/Plugins/Runtime/nDisplay/Source/DisplayClusterConfigurator/Private/ClusterConfiguration/ViewModels/DisplayClusterConfiguratorClusterNodeViewModel.cpp
--- a/Engine/Plugins/Runtime/nDisplay/Source/DisplayClusterConfigurator/Private/ClusterConfiguration/ViewModels/DisplayClusterConfiguratorClusterNodeViewModel.cpp
+++ b/Engine/Plugins/Runtime/nDisplay/Source/DisplayClusterConfigurator/Private/ClusterConfiguration/ViewModels/DisplayClusterConfiguratorClusterNodeViewModel.cpp
@@ -10,6 +10,7 @@ FDisplayClusterConfiguratorClusterNodeViewModel::FDisplayClusterConfiguratorClus
 	ClusterNodePtr = ClusterNode;
 
 	INIT_PROPERTY_HANDLE(UDisplayClusterConfigurationClusterNode, ClusterNode, WindowRect);
+	INIT_PROPERTY_HANDLE(UDisplayClusterConfigurationClusterNode, ClusterNode, Host);
 }
 
 void FDisplayClusterConfiguratorClusterNodeViewModel::SetWindowRect(const FDisplayClusterConfigurationRectangle& NewWindowRect)
@@ -65,4 +66,25 @@ void FDisplayClusterConfiguratorClusterNodeViewModel::SetWindowRect(const FDispl
 	{
 		ClusterNode->MarkPackageDirty();
 	}
-}
\ No newline at end of file
+}
+
+void FDisplayClusterConfiguratorClusterNodeViewModel::SetHost(const FString& NewHost)
+{
+	UDisplayClusterConfigurationClusterNode* ClusterNode = ClusterNodePtr.Get();
+	check(ClusterNode);
+	ClusterNode->Modify();
+
+	bool bClusterNodeChanged = false;
+
+	if (ClusterNode->Host != NewHost)
+	{
+		bClusterNodeChanged = true;
+
+		HostHandle->SetValue(NewHost);
+	}
+
+	if (bClusterNodeChanged)
+	{
+		ClusterNode->MarkPackageDirty();
+	}
+}
diff --git a/Engine/Plugins/Runtime/nDisplay/Source/DisplayClusterConfigurator/Private/ClusterConfiguration/ViewModels/DisplayClusterConfiguratorClusterNodeViewModel.h b/Engine/Plugins/Runtime/nDisplay/Source/DisplayClusterConfigurator/Private/ClusterConfiguration/ViewModels/DisplayClusterConfiguratorClusterNodeViewModel.h
--- a/Engine/Plugins/Runtime/nDisplay/Source/DisplayClusterConfigurator/Private/ClusterConfiguration/ViewModels/DisplayClusterConfiguratorClusterNodeViewModel.h
+++ b/Engine/Plugins/Runtime/nDisplay/Source/DisplayClusterConfigurator/Private/ClusterConfiguration/ViewModels/DisplayClusterConfiguratorClusterNodeViewModel.h
@@ -16,9 +16,11 @@ public:
 	FDisplayClusterConfiguratorClusterNodeViewModel(UDisplayClusterConfigurationClusterNode* ClusterNode);
 
 	void SetWindowRect(const FDisplayClusterConfigurationRectangle& NewWindowRect);
+	void SetHost(const FString& NewHost);
 
 private:
 	TWeakObjectPtr<UDisplayClusterConfigurationClusterNode> ClusterNodePtr;
 
 	PROPERTY_HANDLE(WindowRect);
+	PROPERTY_HANDLE(Host);
 };
\ No newline at end of file
diff --git a/Engine/Plugins/Runtime/nDisplay/Source/DisplayClusterConfigurator/Private/ClusterConfiguration/ViewModels/DisplayClusterConfiguratorViewportViewModel.cpp b/Engine/Plugins/Runtime/nDisplay/Source/DisplayClusterConfigurator/Private/ClusterConfiguration/ViewModels/DisplayClusterConfiguratorViewportViewModel.cpp
--- a/Engine/Plugins/Runtime/nDisplay/Source/DisplayClusterConfigurator/Private/ClusterConfiguration/ViewModels/DisplayClusterConfiguratorViewportViewModel.cpp
+++ b/Engine/Plugins/Runtime/nDisplay/Source/DisplayClusterConfigurator/Private/ClusterConfiguration/ViewModels/DisplayClusterConfiguratorViewportViewModel.cpp
@@ -9,10 +9,24 @@ FDisplayClusterConfiguratorViewportViewModel::FDisplayClusterConfiguratorViewpor
 {
 	ViewportPtr = Viewport;
 
+	INIT_PROPERTY_HANDLE(UDisplayClusterConfigurationViewport, Viewport, Camera);
 	INIT_PROPERTY_HANDLE(UDisplayClusterConfigurationViewport, Viewport, Region);
 	INIT_PROPERTY_HANDLE(UDisplayClusterConfigurationViewport, Viewport, ViewportRemap);
 }
 
+void FDisplayClusterConfiguratorViewportViewModel::SetCamera(const FString& NewCamera)
+{
+	UDisplayClusterConfigurationViewport* Viewport = ViewportPtr.Get();
+	check(Viewport);
+	Viewport->Modify();
+
+	if (NewCamera != Viewport->Camera)
+	{
+		CameraHandle->SetValue(NewCamera);
+		Viewport->MarkPackageDirty();
+	}
+}
+
 void FDisplayClusterConfiguratorViewportViewModel::SetRegion(const FDisplayClusterConfigurationRectangle& NewRegion)
 {
 	UDisplayClusterConfigurationViewport* Viewport = ViewportPtr.Get();
diff --git a/Engine/Plugins/Runtime/nDisplay/Source/DisplayClusterConfigurator/Private/ClusterConfiguration/ViewModels/DisplayClusterConfiguratorViewportViewModel.h b/Engine/Plugins/Runtime/nDisplay/Source/DisplayClusterConfigurator/Private/ClusterConfiguration/ViewModels/DisplayClusterConfiguratorViewportViewModel.h
--- a/Engine/Plugins/Runtime/nDisplay/Source/DisplayClusterConfigurator/Private/ClusterConfiguration/ViewModels/DisplayClusterConfiguratorViewportViewModel.h
+++ b/Engine/Plugins/Runtime/nDisplay/Source/DisplayClusterConfigurator/Private/ClusterConfiguration/ViewModels/DisplayClusterConfiguratorViewportViewModel.h
@@ -18,6 +18,9 @@ public:
 	/** Creates a new view model for the specified viewport configuration object */
 	FDisplayClusterConfiguratorViewportViewModel(UDisplayClusterConfigurationViewport* Viewport);
 
+	/** Sets the camera property of the viewport, propagating the change through property handles and marking the viewport's package as dirty  */
+	void SetCamera(const FString& NewCamera);
+
 	/** Sets the region of the viewport,  propagating the change through property handles and marking the viewport's package as dirty */
 	void SetRegion(const FDisplayClusterConfigurationRectangle& NewRegion);
 
@@ -28,6 +31,7 @@ private:
 	/** A pointer to the viewport configuration object the view model encapsulates */
 	TWeakObjectPtr<UDisplayClusterConfigurationViewport> ViewportPtr;
 
+	PROPERTY_HANDLE(Camera);
 	PROPERTY_HANDLE(Region);
 	PROPERTY_HANDLE(ViewportRemap);
 };
\ No newline at end of file
diff --git a/Engine/Plugins/Runtime/nDisplay/Source/DisplayClusterConfigurator/Private/DisplayClusterConfiguratorBlueprintEditor.cpp b/Engine/Plugins/Runtime/nDisplay/Source/DisplayClusterConfigurator/Private/DisplayClusterConfiguratorBlueprintEditor.cpp
--- a/Engine/Plugins/Runtime/nDisplay/Source/DisplayClusterConfigurator/Private/DisplayClusterConfiguratorBlueprintEditor.cpp
+++ b/Engine/Plugins/Runtime/nDisplay/Source/DisplayClusterConfigurator/Private/DisplayClusterConfiguratorBlueprintEditor.cpp
@@ -67,7 +67,9 @@
 #include "Kismet2/DebuggerCommands.h"
 #include "Subsystems/PanelExtensionSubsystem.h"
 #include "ScopedTransaction.h"
-
+#include "Blueprints/DisplayClusterBlueprintGeneratedClass.h"
+#include "Misc/DisplayClusterHelpers.h"
+#include "MPCDI/DisplayClusterConfiguratorMPCDIImporter.h"
 
 #define LOCTEXT_NAMESPACE "DisplayClusterConfiguratorBlueprintEditor"
 
@@ -909,6 +911,8 @@ void FDisplayClusterConfiguratorBlueprintEditor::BindCommands()
 		FCanExecuteAction(),
 		FIsActionChecked::CreateSP(this, &FDisplayClusterConfiguratorBlueprintEditor::IsExportOnSaveSet)
 	);
+
+	ToolkitCommands->MapAction(Commands.ImportMPCDI, FExecuteAction::CreateSP(this, &FDisplayClusterConfiguratorBlueprintEditor::ImportMPCDI_Clicked));
 }
 
 void FDisplayClusterConfiguratorBlueprintEditor::CreateWidgets()
@@ -1075,6 +1079,61 @@ void FDisplayClusterConfiguratorBlueprintEditor::ToggleExportOnSaveSetting()
 	Settings->SaveConfig();
 }
 
+void FDisplayClusterConfiguratorBlueprintEditor::ImportMPCDI_Clicked()
+{
+	const FString MPCDIFileDescription = LOCTEXT("MPCDIFileDescription", "MPCDI file").ToString();
+	const FString MPCDIFileExtension = TEXT("*.mpcdi");
+	const FString FileTypes = FString::Printf(TEXT("%s (%s)|%s"), *MPCDIFileDescription, *MPCDIFileExtension, *MPCDIFileExtension);
+
+	TArray<FString> OpenFileNames;
+	IDesktopPlatform* DesktopPlatform = FDesktopPlatformModule::Get();
+	bool bFileSelected = false;
+	int32 FilterIndex = -1;
+
+	// Open file dialog
+	if (DesktopPlatform)
+	{
+		const void* ParentWindowHandle = FSlateApplication::Get().FindBestParentWindowHandleForDialogs(nullptr);
+
+		bFileSelected = DesktopPlatform->OpenFileDialog(
+			ParentWindowHandle,
+			LOCTEXT("ImportMPCDIDialogTitle", "Import MPCDI").ToString(),
+			FEditorDirectories::Get().GetLastDirectory(ELastDirectory::GENERIC_IMPORT),
+			TEXT(""),
+			FileTypes,
+			EFileDialogFlags::None,
+			OpenFileNames,
+			FilterIndex
+		);
+	}
+
+	// Load file
+	if (bFileSelected)
+	{
+		if (OpenFileNames.Num() > 0)
+		{
+			const FString& FileName = OpenFileNames[0];
+			FEditorDirectories::Get().SetLastDirectory(ELastDirectory::GENERIC_IMPORT, FileName);
+
+			UDisplayClusterBlueprint* Blueprint = LoadedBlueprint.Get();
+
+			FDisplayClusterConfiguratorMPCDIImporterParams ImportParams {};
+			if (FDisplayClusterConfiguratorMPCDIImporter::ImportMPCDIIntoBlueprint(FileName, Blueprint, ImportParams))
+			{
+				GetEditorData()->MarkPackageDirty();
+				ClusterChanged();
+
+				// Since the blueprint component transforms are edited directly by the importer, we must
+				// fully destroy and recreate the preview actor, as otherwise the engine will not propagate the
+				// new values to the preview since there is now a difference between the preview and blueprint values
+				DestroyPreview();
+				UpdateSubobjectPreview(true);
+				SubobjectEditor->UpdateTree();
+			}
+		}
+	}
+}
+
 TStatId FDisplayClusterConfiguratorBlueprintEditor::GetStatId() const
 {
 	RETURN_QUICK_DECLARE_CYCLE_STAT(FDisplayClusterConfiguratorBlueprintEditor, STATGROUP_Tickables);
diff --git a/Engine/Plugins/Runtime/nDisplay/Source/DisplayClusterConfigurator/Private/DisplayClusterConfiguratorBlueprintEditor.h b/Engine/Plugins/Runtime/nDisplay/Source/DisplayClusterConfigurator/Private/DisplayClusterConfiguratorBlueprintEditor.h
--- a/Engine/Plugins/Runtime/nDisplay/Source/DisplayClusterConfigurator/Private/DisplayClusterConfiguratorBlueprintEditor.h
+++ b/Engine/Plugins/Runtime/nDisplay/Source/DisplayClusterConfigurator/Private/DisplayClusterConfiguratorBlueprintEditor.h
@@ -214,6 +214,8 @@ private:
 
 	bool IsExportOnSaveSet() const;
 	void ToggleExportOnSaveSetting();
+
+	void ImportMPCDI_Clicked();
 	//~ End UI command handlers
 
 	void OnReadOnlyChanged(bool bReadOnly);
diff --git a/Engine/Plugins/Runtime/nDisplay/Source/DisplayClusterConfigurator/Private/DisplayClusterConfiguratorCommands.cpp b/Engine/Plugins/Runtime/nDisplay/Source/DisplayClusterConfigurator/Private/DisplayClusterConfiguratorCommands.cpp
--- a/Engine/Plugins/Runtime/nDisplay/Source/DisplayClusterConfigurator/Private/DisplayClusterConfiguratorCommands.cpp
+++ b/Engine/Plugins/Runtime/nDisplay/Source/DisplayClusterConfigurator/Private/DisplayClusterConfiguratorCommands.cpp
@@ -11,6 +11,7 @@ void FDisplayClusterConfiguratorCommands::RegisterCommands()
 	UI_COMMAND(Export, "Export", "Export to an nDisplay config. Requires a primary cluster node set. This does not export all components and is meant for launching the cluster from the command line", EUserInterfaceActionType::Button, FInputChord());
 	UI_COMMAND(EditConfig, "EditConfig", "Edit config file with text editor", EUserInterfaceActionType::Button, FInputChord());
 	UI_COMMAND(ExportConfigOnSave, "Export on Save", "Export to nDisplay config automatically on save. Requires a config previously exported", EUserInterfaceActionType::ToggleButton, FInputChord());
+	UI_COMMAND(ImportMPCDI, "Import MPCDI", "Import an MPCDI configuration", EUserInterfaceActionType::Button, FInputChord());
 
 	UI_COMMAND(AddNewClusterNode, "Add New Cluster Node", "Adds a new cluster node to the cluster config", EUserInterfaceActionType::Button, FInputChord());
 	UI_COMMAND(AddNewViewport, "Add New Viewport", "Adds a new viewport to the cluster node", EUserInterfaceActionType::Button, FInputChord());
diff --git a/Engine/Plugins/Runtime/nDisplay/Source/DisplayClusterConfigurator/Private/DisplayClusterConfiguratorCommands.h b/Engine/Plugins/Runtime/nDisplay/Source/DisplayClusterConfigurator/Private/DisplayClusterConfiguratorCommands.h
--- a/Engine/Plugins/Runtime/nDisplay/Source/DisplayClusterConfigurator/Private/DisplayClusterConfiguratorCommands.h
+++ b/Engine/Plugins/Runtime/nDisplay/Source/DisplayClusterConfigurator/Private/DisplayClusterConfiguratorCommands.h
@@ -23,6 +23,7 @@ public:
 	TSharedPtr<FUICommandInfo> Export;
 	TSharedPtr<FUICommandInfo> EditConfig;
 	TSharedPtr<FUICommandInfo> ExportConfigOnSave;
+	TSharedPtr<FUICommandInfo> ImportMPCDI;
 
 	// Cluster Configuration commands
 	TSharedPtr<FUICommandInfo> AddNewClusterNode;
diff --git a/Engine/Plugins/Runtime/nDisplay/Source/DisplayClusterConfigurator/Private/MPCDI/DisplayClusterConfiguratorMPCDIImporter.cpp b/Engine/Plugins/Runtime/nDisplay/Source/DisplayClusterConfigurator/Private/MPCDI/DisplayClusterConfiguratorMPCDIImporter.cpp
new file mode 100644
--- /dev/null
+++ b/Engine/Plugins/Runtime/nDisplay/Source/DisplayClusterConfigurator/Private/MPCDI/DisplayClusterConfiguratorMPCDIImporter.cpp
@@ -0,0 +1,295 @@
+// Copyright Epic Games, Inc. All Rights Reserved.
+
+#include "DisplayClusterConfiguratorMPCDIImporter.h"
+
+#include "DisplayClusterProjectionStrings.h"
+#include "IDisplayClusterWarp.h"
+#include "Blueprints/DisplayClusterBlueprint.h"
+#include "Blueprints/DisplayClusterBlueprintGeneratedClass.h"
+#include "ClusterConfiguration/DisplayClusterConfiguratorClusterUtils.h"
+#include "ClusterConfiguration/ViewModels/DisplayClusterConfiguratorClusterNodeViewModel.h"
+#include "ClusterConfiguration/ViewModels/DisplayClusterConfiguratorProjectionPolicyViewModel.h"
+#include "ClusterConfiguration/ViewModels/DisplayClusterConfiguratorViewportViewModel.h"
+#include "Components/DisplayClusterCameraComponent.h"
+#include "Components/DisplayClusterScreenComponent.h"
+#include "Engine/SCS_Node.h"
+#include "Engine/SimpleConstructionScript.h"
+#include "Misc/DisplayClusterHelpers.h"
+
+bool FDisplayClusterConfiguratorMPCDIImporter::ImportMPCDIIntoBlueprint(const FString& InFilePath, UDisplayClusterBlueprint* InBlueprint, const FDisplayClusterConfiguratorMPCDIImporterParams& InParams)
+{
+	const FString MPCDIFileFullPath = DisplayClusterHelpers::filesystem::GetFullPathForConfigResource(InFilePath);
+
+	TMap<FString, TMap<FString, FDisplayClusterWarpMPCDIAttributes>> MPCDIFile;
+	IDisplayClusterWarp::Get().ReadMPCDFileStructure(InFilePath, MPCDIFile);
+
+
+	FName ViewOriginComponentName = !InParams.ViewOriginComponentName.IsNone() ? InParams.ViewOriginComponentName : TEXT("DefaultViewPoint");
+
+	USCS_Node* ViewOriginNode = nullptr;
+	UDisplayClusterCameraComponent* ViewOriginComponent = Cast<UDisplayClusterCameraComponent>(FindBlueprintSceneComponent(InBlueprint, ViewOriginComponentName, &ViewOriginNode));
+
+	FIPv4Address CurrentIPAddress = InParams.HostStartingIPAddress;
+	for (const TPair<FString, TMap<FString, FDisplayClusterWarpMPCDIAttributes>>& Buffer : MPCDIFile)
+	{
+		for (const TPair<FString, FDisplayClusterWarpMPCDIAttributes>& Region : Buffer.Value)
+		{
+			if (InParams.bCreateStageGeometryComponents)
+			{
+				// For MPCDI 2D, we create Display Cluster screen components to match the MPCDI buffer and region configuration
+				if (Region.Value.ProfileType == EDisplayClusterWarpProfileType::warp_2D)
+				{
+					bool bFoundExistingScreen = false;
+					USCS_Node* ScreenNode = FindOrCreateScreenNodeForRegion(InBlueprint, Region.Key, bFoundExistingScreen);
+					if (ScreenNode)
+					{
+						UDisplayClusterScreenComponent* ScreenComponent = CastChecked<UDisplayClusterScreenComponent>(ScreenNode->GetActualComponentTemplate(InBlueprint->GetGeneratedClass()));
+						ConfigureScreenFromRegion(ScreenComponent, ViewOriginComponent, Region.Value, InParams);
+
+						// If a new screen had to be created, it needs to be parented to a valid parent component
+						if (!bFoundExistingScreen)
+						{
+							USCS_Node* ParentNode = nullptr;
+							USceneComponent* ParentComponent = nullptr;
+							if (InParams.ParentComponentName.IsNone())
+							{
+								ParentComponent = InBlueprint->SimpleConstructionScript->GetSceneRootComponentTemplate(false, &ParentNode);
+							}
+							else
+							{
+								ParentComponent = FindBlueprintSceneComponent(InBlueprint, InParams.ParentComponentName, &ParentNode);
+							}
+
+							if (ParentNode && ParentNode->GetSCS() == InBlueprint->SimpleConstructionScript)
+							{
+								ParentNode->AddChildNode(ScreenNode);
+							}
+							else
+							{
+								ScreenNode->SetParent(ParentComponent);
+								InBlueprint->SimpleConstructionScript->AddNode(ScreenNode);
+							}
+						}
+					}
+				}
+			}
+
+			bool bFoundExistingViewport = false;
+			UDisplayClusterConfigurationViewport* Viewport = FindOrCreateViewportForRegion(InBlueprint, Region.Key, bFoundExistingViewport);
+			if (Viewport)
+			{
+				ConfigureViewportFromRegion(Viewport, Region.Value, InParams);
+
+				if (!bFoundExistingViewport)
+				{
+					FDisplayClusterConfiguratorProjectionPolicyViewModel ProjectionPolicyViewModel(Viewport);
+
+					ProjectionPolicyViewModel.SetPolicyType(DisplayClusterProjectionStrings::projection::MPCDI);
+
+					ProjectionPolicyViewModel.SetParameterValue(DisplayClusterProjectionStrings::cfg::mpcdi::MPCDITypeKey, DisplayClusterProjectionStrings::cfg::mpcdi::TypeMPCDI);
+					ProjectionPolicyViewModel.SetParameterValue(DisplayClusterProjectionStrings::cfg::mpcdi::File, MPCDIFileFullPath);
+					ProjectionPolicyViewModel.SetParameterValue(DisplayClusterProjectionStrings::cfg::mpcdi::Buffer, Buffer.Key);
+					ProjectionPolicyViewModel.SetParameterValue(DisplayClusterProjectionStrings::cfg::mpcdi::Region, Region.Key);
+					ProjectionPolicyViewModel.SetParameterValue(DisplayClusterProjectionStrings::cfg::mpcdi::EnablePreview, TEXT("true"));
+
+					if (Region.Value.ProfileType == EDisplayClusterWarpProfileType::warp_2D)
+					{
+						ProjectionPolicyViewModel.SetParameterValue(DisplayClusterProjectionStrings::cfg::mpcdi::MPCDIType, TEXT("2d"));
+
+						if (InParams.bCreateStageGeometryComponents)
+						{
+							ProjectionPolicyViewModel.SetParameterValue(DisplayClusterProjectionStrings::cfg::mpcdi::Component, GetScreenNameForRegion(Region.Key));
+						}
+					}
+
+					if (!InParams.ParentComponentName.IsNone())
+					{
+						ProjectionPolicyViewModel.SetParameterValue(DisplayClusterProjectionStrings::cfg::mpcdi::Origin, InParams.ParentComponentName.ToString());
+					}
+
+					if (UDisplayClusterConfigurationClusterNode* ClusterNode = Cast<UDisplayClusterConfigurationClusterNode>(Viewport->GetOuter()))
+					{
+						FDisplayClusterConfiguratorClusterNodeViewModel ClusterNodeViewModel(ClusterNode);
+						ClusterNodeViewModel.SetHost(CurrentIPAddress.ToString());
+					}
+				}
+
+				UDisplayClusterConfigurationCluster* Cluster = InBlueprint->GetOrLoadConfig()->Cluster;
+				const bool bHasValidPrimaryNode = Cluster->Nodes.Contains(Cluster->PrimaryNode.Id);
+				if (!bHasValidPrimaryNode)
+				{
+					Cluster->Modify();
+					Cluster->PrimaryNode.Id = GetClusterNodeNameForRegion(Region.Key);
+				}
+
+				if (InParams.bIncrementHostIPAddress)
+				{
+					++CurrentIPAddress.Value;
+				}
+			}
+		}
+	}
+
+	return true;
+}
+
+USCS_Node* FDisplayClusterConfiguratorMPCDIImporter::FindOrCreateScreenNodeForRegion(UDisplayClusterBlueprint* InBlueprint, const FString& RegionId, bool& bOutFoundExistingScreen)
+{
+	bOutFoundExistingScreen = false;
+
+	const FString ScreenName = GetScreenNameForRegion(RegionId);
+	if (USCS_Node* ExistingScreenNode = InBlueprint->SimpleConstructionScript->FindSCSNode(*ScreenName))
+	{
+		bOutFoundExistingScreen = true;
+		return ExistingScreenNode;
+	}
+
+	return InBlueprint->SimpleConstructionScript->CreateNode(UDisplayClusterScreenComponent::StaticClass(), *ScreenName);
+}
+
+void FDisplayClusterConfiguratorMPCDIImporter::ConfigureScreenFromRegion(
+	UDisplayClusterScreenComponent* InScreenComponent,
+	UDisplayClusterCameraComponent* InViewOriginComponent,
+	const FDisplayClusterWarpMPCDIAttributes& InAttributes,
+	const FDisplayClusterConfiguratorMPCDIImporterParams& InParams)
+{
+	const FVector2D BufferActualSize(InAttributes.Buffer.Resolution.X * InParams.BufferToWorldScale, InAttributes.Buffer.Resolution.Y * InParams.BufferToWorldScale);
+	const FVector2D RegionActualPosition((InAttributes.Region.Pos.X - 0.5) * BufferActualSize.X, (InAttributes.Region.Pos.Y - 0.5) * BufferActualSize.Y);
+	const FVector2D RegionActualSize(InAttributes.Region.Size.X * BufferActualSize.X, InAttributes.Region.Size.Y * BufferActualSize.Y);
+
+	FVector ScreenPosition = FVector(InParams.BufferToWorldDistance, RegionActualPosition.X + RegionActualSize.X * 0.5, RegionActualPosition.Y + RegionActualSize.Y * 0.5);
+
+	if (InViewOriginComponent)
+	{
+		// We want to position the screen component to be in front of the view origin it is assigned to, so we must compute
+		// the transforms from view origin to root and from root to screen component to correct position
+
+		// Blueprint components don't have world transforms, so compute the view origin to root transform manually
+		FTransform ViewOriginTransform = FTransform::Identity;
+		for (const USceneComponent* Comp = InViewOriginComponent; Comp != nullptr; Comp = Comp->GetAttachParent())
+		{
+			ViewOriginTransform *= Comp->GetRelativeTransform();
+		}
+
+		// The screen may also be parented to a non-root component, so also compute its root transform
+		FTransform ScreenTransform = FTransform::Identity;
+		for (const USceneComponent* Comp = InScreenComponent->GetAttachParent(); Comp != nullptr; Comp = Comp->GetAttachParent())
+		{
+			ScreenTransform *= Comp->GetRelativeTransform();
+		}
+
+		FTransform ViewOriginToScreen = ViewOriginTransform * ScreenTransform.Inverse();
+		ScreenPosition = ViewOriginToScreen.TransformPositionNoScale(ScreenPosition);
+	}
+
+	InScreenComponent->SetRelativeLocation(ScreenPosition);
+	InScreenComponent->SetScreenSize(RegionActualSize);
+}
+
+UDisplayClusterConfigurationViewport* FDisplayClusterConfiguratorMPCDIImporter::FindOrCreateViewportForRegion(UDisplayClusterBlueprint* InBlueprint, const FString& RegionId, bool& bOutFoundExistingViewport)
+{
+	UDisplayClusterConfigurationViewport* Viewport = nullptr;
+	bOutFoundExistingViewport = false;
+
+	const FString ViewportName = GetViewportNameForRegion(RegionId);
+	UDisplayClusterConfigurationCluster* Cluster = InBlueprint->GetOrLoadConfig()->Cluster;
+	for (const TPair<FString, TObjectPtr<UDisplayClusterConfigurationClusterNode>>& NodePair : Cluster->Nodes)
+	{
+		if (NodePair.Value->Viewports.Contains(ViewportName))
+		{
+			Viewport = NodePair.Value->Viewports[ViewportName];
+			bOutFoundExistingViewport = true;
+			break;
+		}
+	}
+
+	if (!Viewport)
+	{
+		const FString NodeName = GetClusterNodeNameForRegion(RegionId);
+		UDisplayClusterConfigurationClusterNode* ClusterNodeTemplate = NewObject<UDisplayClusterConfigurationClusterNode>(InBlueprint, NAME_None, RF_Transactional | RF_ArchetypeObject | RF_Public);
+		UDisplayClusterConfigurationClusterNode* NewClusterNode = FDisplayClusterConfiguratorClusterUtils::AddClusterNodeToCluster(ClusterNodeTemplate, Cluster, NodeName);
+
+		UDisplayClusterConfigurationViewport* ViewportTemplate = NewObject<UDisplayClusterConfigurationViewport>(InBlueprint, NAME_None, RF_Transactional | RF_ArchetypeObject | RF_Public);
+		Viewport = FDisplayClusterConfiguratorClusterUtils::AddViewportToClusterNode(ViewportTemplate, NewClusterNode, ViewportName);
+	}
+
+	return Viewport;
+}
+
+void FDisplayClusterConfiguratorMPCDIImporter::ConfigureViewportFromRegion(UDisplayClusterConfigurationViewport* InViewport, const FDisplayClusterWarpMPCDIAttributes& InAttributes, const FDisplayClusterConfiguratorMPCDIImporterParams& InParams)
+{
+	if (UDisplayClusterConfigurationClusterNode* ClusterNode = Cast<UDisplayClusterConfigurationClusterNode>(InViewport->GetOuter()))
+	{
+		FDisplayClusterConfiguratorClusterNodeViewModel ClusterNodeViewModel(ClusterNode);
+		ClusterNodeViewModel.SetWindowRect(FDisplayClusterConfigurationRectangle(0, 0, InAttributes.Region.Resolution.X, InAttributes.Region.Resolution.Y));
+	}
+
+	FDisplayClusterConfiguratorViewportViewModel ViewportViewModel(InViewport);
+	ViewportViewModel.SetRegion(FDisplayClusterConfigurationRectangle(0, 0, InAttributes.Region.Resolution.X, InAttributes.Region.Resolution.Y));
+
+	if (InParams.ViewOriginComponentName != NAME_None)
+	{
+		ViewportViewModel.SetCamera(InParams.ViewOriginComponentName.ToString());
+	}
+}
+
+USceneComponent* FDisplayClusterConfiguratorMPCDIImporter::FindBlueprintSceneComponent(UDisplayClusterBlueprint* InBlueprint, const FName& ComponentName, USCS_Node** OutComponentNode)
+{
+	*OutComponentNode = nullptr;
+
+	UClass* GeneratedClass = InBlueprint->GetGeneratedClass();
+	UClass* ParentClass = InBlueprint->ParentClass;
+
+	AActor* CDO = Cast<AActor>(GeneratedClass->GetDefaultObject(false));
+	if (CDO == nullptr && ParentClass != nullptr)
+	{
+		CDO = Cast<AActor>(ParentClass->GetDefaultObject(false));
+	}
+
+	// First, check to see if there exists a native component on the CDO that matches the specified name
+	USceneComponent* FoundComponent = nullptr;
+	if (CDO)
+	{
+		for (UActorComponent* Component : CDO->GetComponents())
+		{
+			if (USceneComponent* SceneComponent = Cast<USceneComponent>(Component))
+			{
+				if (SceneComponent->GetFName() == ComponentName)
+				{
+					FoundComponent = SceneComponent;
+					break;
+				}
+			}
+		}
+	}
+
+	// If a native component was not found, check the SCS to see if one exists there
+	if (!FoundComponent)
+	{
+		if (USCS_Node* FoundNode = InBlueprint->SimpleConstructionScript->FindSCSNode(ComponentName))
+		{
+			if (FoundNode->ComponentTemplate && FoundNode->ComponentTemplate->IsA<USceneComponent>())
+			{
+				*OutComponentNode = FoundNode;
+				FoundComponent = Cast<USceneComponent>(FoundNode->ComponentTemplate);
+			}
+		}
+	}
+
+	return FoundComponent;
+}
+
+FString FDisplayClusterConfiguratorMPCDIImporter::GetScreenNameForRegion(const FString& RegionId)
+{
+	return FString::Printf(TEXT("%s_Screen"), *RegionId);
+}
+
+FString FDisplayClusterConfiguratorMPCDIImporter::GetClusterNodeNameForRegion(const FString& RegionId)
+{
+	return FString::Printf(TEXT("%s_Node"), *RegionId);
+}
+
+FString FDisplayClusterConfiguratorMPCDIImporter::GetViewportNameForRegion(const FString& RegionId)
+{
+	return FString::Printf(TEXT("%s_Viewport"), *RegionId);
+}
diff --git a/Engine/Plugins/Runtime/nDisplay/Source/DisplayClusterConfigurator/Private/MPCDI/DisplayClusterConfiguratorMPCDIImporter.h b/Engine/Plugins/Runtime/nDisplay/Source/DisplayClusterConfigurator/Private/MPCDI/DisplayClusterConfiguratorMPCDIImporter.h
new file mode 100644
--- /dev/null
+++ b/Engine/Plugins/Runtime/nDisplay/Source/DisplayClusterConfigurator/Private/MPCDI/DisplayClusterConfiguratorMPCDIImporter.h
@@ -0,0 +1,109 @@
+// Copyright Epic Games, Inc. All Rights Reserved.
+
+#pragma once
+
+#include "CoreMinimal.h"
+#include "Interfaces/IPv4/IPv4Address.h"
+
+class UDisplayClusterCameraComponent;
+struct FDisplayClusterWarpMPCDIAttributes;
+class USceneComponent;
+class USCS_Node;
+class UDisplayClusterConfigurationViewport;
+class UDisplayClusterBlueprint;
+class UDisplayClusterScreenComponent;
+
+/** USed to pass parameters into the MPCDI importer */
+struct FDisplayClusterConfiguratorMPCDIImporterParams
+{
+	/** The amount to scale an MPCDI buffer by when converting from pixels to world coordinates */
+	float BufferToWorldScale = 0.1;
+
+	/** The distance from the origin to position an MPCDI buffer in world coordinates */
+	float BufferToWorldDistance = 100;
+
+	/** The name of the component to make the parent of the MPCDI geometry */
+	FName ParentComponentName = NAME_None;
+
+	/** The name of the view origin component to link to the MPCDI geometry. If none is supplied, the default view origin is used */
+	FName ViewOriginComponentName = NAME_None;
+
+	/** The IP address to configure the first cluster node with */
+	FIPv4Address HostStartingIPAddress = FIPv4Address::InternalLoopback;
+
+	/** Indicates that the IP address will be incremented for each MPCDI region a cluster node is created for */
+	bool bIncrementHostIPAddress = true;
+
+	/** Indicates if stage geometry components are created to match the MPCDI geometry */
+	bool bCreateStageGeometryComponents = true;
+};
+
+/** Helper class that can import an MPCDI configuration into an nDisplay config blueprint */
+class FDisplayClusterConfiguratorMPCDIImporter
+{
+public:
+	/**
+	 * Imports the specified MPCDI file into the specified nDisplay config blueprint, configurating the blueprint to match the MPCDI file configuration
+	 * @param InFilePath - The file path to the MPCDI file to import
+	 * @param InBlueprint - The blueprint to write the MPCDI configuration to
+	 * @param InParams - The parameters to use when importing the MPCDI config
+	 **/
+	static bool ImportMPCDIIntoBlueprint(const FString& InFilePath, UDisplayClusterBlueprint* InBlueprint, const FDisplayClusterConfiguratorMPCDIImporterParams& InParams);
+
+private:
+	/**
+	 * Attempts to find an existing SCS screen node for the specified region, and if one is not found, creates a new node
+	 * 
+	 * @param InBlueprint - The blueprint to find the node on
+	 * @param RegionId - The ID of the region to find an existing node for
+	 * @param bOutFoundExistingScreen - An out param indicating if an existing node was found or not
+	 **/
+	static USCS_Node* FindOrCreateScreenNodeForRegion(UDisplayClusterBlueprint* InBlueprint, const FString& RegionId, bool& bOutFoundExistingScreen);
+
+	/** 
+	 * Configures a screen component to match the configuration of an MPCDI region
+	 * 
+	 * @param InScreenComponent - The component to configure
+	 * @param InViewOriginComponent - The view origin component to configure the screens from
+	 * @param InAttributes - The attributes of the region to configure
+	 * @param InParams - The import parameters 
+	 **/
+	static void ConfigureScreenFromRegion(UDisplayClusterScreenComponent* InScreenComponent, UDisplayClusterCameraComponent* InViewOriginComponent, const FDisplayClusterWarpMPCDIAttributes& InAttributes, const FDisplayClusterConfiguratorMPCDIImporterParams& InParams);
+
+	/**
+	 * Attempts to find an existing viewport configuration for the specified region, and if one is not found, creates a new viewport configuration
+	 * 
+	 * @param InBlueprint - The blueprint to find the viewport on
+	 * @param RegionId - The ID of the region to find an existing viewport for
+	 * @param bOutFoundExistingViewport - An out param indicating if an existing viewport was found or not
+	 **/
+	static UDisplayClusterConfigurationViewport* FindOrCreateViewportForRegion(UDisplayClusterBlueprint* InBlueprint, const FString& RegionId, bool& bOutFoundExistingViewport);
+
+	/** 
+	 * Configures a viewport configuration to match the configuration of an MPCDI region
+	 * 
+	 * @param InViewport - The viewport configuration to configure
+	 * @param InAttributes - The attributes of the region to configure
+	 * @param InParams - The import parameters 
+	 **/
+	static void ConfigureViewportFromRegion(UDisplayClusterConfigurationViewport* InViewport, const FDisplayClusterWarpMPCDIAttributes& InAttributes, const FDisplayClusterConfiguratorMPCDIImporterParams& InParams);
+
+	/**
+	 * Attempts to find a specified scene component in the blueprint's component hierarchy
+	 * 
+	 * @param InBlueprint - The blueprint to search
+	 * @param ComponentName - The name of the component to search for
+	 * @param OutComponentNode - The corresponding SCS node of the found component if the component was a non-native component
+	 * @return The found scene component, or nullptr if none was found
+	 **/
+	static USceneComponent* FindBlueprintSceneComponent(UDisplayClusterBlueprint* InBlueprint, const FName& ComponentName, USCS_Node** OutComponentNode);
+
+	/** Creates a screen component name for a region using its ID */
+	static FString GetScreenNameForRegion(const FString& RegionId);
+
+	/** Creates a cluster node name for a region using its ID */
+	static FString GetClusterNodeNameForRegion(const FString& RegionId);
+
+	/** Creates a viewport name for a region using its ID */
+	static FString GetViewportNameForRegion(const FString& RegionId);
+};
\ No newline at end of file
diff --git a/Engine/Plugins/Runtime/nDisplay/Source/DisplayClusterConfigurator/Private/Views/DisplayClusterConfiguratorToolbar.cpp b/Engine/Plugins/Runtime/nDisplay/Source/DisplayClusterConfigurator/Private/Views/DisplayClusterConfiguratorToolbar.cpp
--- a/Engine/Plugins/Runtime/nDisplay/Source/DisplayClusterConfigurator/Private/Views/DisplayClusterConfiguratorToolbar.cpp
+++ b/Engine/Plugins/Runtime/nDisplay/Source/DisplayClusterConfigurator/Private/Views/DisplayClusterConfiguratorToolbar.cpp
@@ -77,6 +77,7 @@ void FDisplayClusterConfiguratorToolbar::FillModesToolbar(FToolBarBuilder& Toolb
 		/*bSimpleComboBox =*/true
 	);
 
+	ToolbarBuilder.AddToolBarButton(Commands.ImportMPCDI);
 	ToolbarBuilder.EndSection();
 
 	/*
