nDisplay: DC SVE to display frame stats, for now just frame number and timecode. Use DC.Stats.Frame 1 to enable on nodes.

diff --git a/Engine/Plugins/Runtime/nDisplay/Shaders/Private/FrameStats.usf b/Engine/Plugins/Runtime/nDisplay/Shaders/Private/FrameStats.usf
new file mode 100644
--- /dev/null
+++ b/Engine/Plugins/Runtime/nDisplay/Shaders/Private/FrameStats.usf
@@ -0,0 +1,98 @@
+// Copyright Epic Games, Inc. All Rights Reserved.
+
+#include "/Engine/Public/Platform.ush"
+#include "/Engine/Private/MiniFontCommon.ush"
+
+uint FrameCount;
+uint Timecode;
+static const uint Semicolon = 58;
+static const float4 FontColor = float4(0.1f, 0.9f, 0.1f, 0.0f);
+
+// Function taken from MiniFontCommon.ush, but with alpha support
+void PrintCharacterA(int2 PixelPos, inout float4 OutColor, float4 FontColor, inout int2 LeftTop, int CharacterID)
+{
+	uint2 Rel = (uint2)(PixelPos - LeftTop);
+
+	// Background: black character as shadow
+	FLATTEN if (Rel.x < (NATIVE_CHARACTER_RES + 1) && Rel.y < (NATIVE_CHARACTER_RES + 1))
+	{
+		OutColor = lerp(OutColor, float4(0.0f, 0.0f, 0.0f, 0.0f), SampleMiniFont(CharacterID, Rel - uint2(1, 1)));
+	}
+	// Foreground: colored character
+	FLATTEN if (Rel.x < (NATIVE_CHARACTER_RES) && Rel.y < (NATIVE_CHARACTER_RES))
+	{
+		OutColor = lerp(OutColor, FontColor, SampleMiniFont(CharacterID, Rel));
+	}
+
+	LeftTop.x += NATIVE_CHARACTER_RES;
+}
+
+void PrintTwoDigitUintA(int2 PixelPos, inout float4 OutColor, float4 FontColor, inout int2 Cursor, uint Number)
+{
+	float FloatNumber = float(Number) + 0.05f;
+
+	PrintCharacterA(PixelPos, OutColor, FontColor, Cursor, ExtractDigitFromFloat(FloatNumber, 10));
+	PrintCharacterA(PixelPos, OutColor, FontColor, Cursor, ExtractDigitFromFloat(FloatNumber, 1));
+}
+
+void PrintLargeUintA(int2 PixelPos, inout float4 OutColor, float4 FontColor, inout int2 Cursor, uint Number)
+{
+	float FloatNumber = float(Number) + 0.05f;
+
+	// before period
+	PrintCharacterA(PixelPos, OutColor, FontColor, Cursor, ExtractDigitFromFloat(FloatNumber, 1000000));
+	PrintCharacterA(PixelPos, OutColor, FontColor, Cursor, ExtractDigitFromFloat(FloatNumber, 100000));
+	PrintCharacterA(PixelPos, OutColor, FontColor, Cursor, ExtractDigitFromFloat(FloatNumber, 10000));
+	PrintCharacterA(PixelPos, OutColor, FontColor, Cursor, ExtractDigitFromFloat(FloatNumber, 1000));
+	PrintCharacterA(PixelPos, OutColor, FontColor, Cursor, ExtractDigitFromFloat(FloatNumber, 100));
+	PrintCharacterA(PixelPos, OutColor, FontColor, Cursor, ExtractDigitFromFloat(FloatNumber, 10));
+	PrintCharacterA(PixelPos, OutColor, FontColor, Cursor, ExtractDigitFromFloat(FloatNumber, 1));
+}
+
+void MainPS(in noperspective float4 InUVAndScreenPos : TEXCOORD0,
+	float4 SvPosition : SV_POSITION,
+	out float4 OutColor : SV_Target0)
+{
+	const uint HoursTC   = (Timecode >> 24) & 0xFF;
+	const uint MinutesTC = (Timecode >> 16) & 0xFF;
+	const uint SecondsTC = (Timecode >> 8)  & 0xFF;
+	const uint FramesTC  = (Timecode >> 0)  & 0xFF;
+	const int2 PixelPos = int2(SvPosition.xy);
+	
+	// Inverted-alpha to match the translucency of scene color
+	OutColor = float4(0.0f, 0.0f, 0.0f, 1.0f);
+	
+	int2 Cursor = 1;
+	const int StartX = Cursor.x;
+
+	PrintTwoDigitUintA(PixelPos, OutColor, FontColor, Cursor, HoursTC);
+	PrintCharacterA(PixelPos, OutColor, FontColor, Cursor, Semicolon);
+	PrintTwoDigitUintA(PixelPos, OutColor, FontColor, Cursor, MinutesTC);
+	PrintCharacterA(PixelPos, OutColor, FontColor, Cursor, Semicolon);
+	PrintTwoDigitUintA(PixelPos, OutColor, FontColor, Cursor, SecondsTC);
+	PrintCharacterA(PixelPos, OutColor, FontColor, Cursor, Semicolon);
+	PrintTwoDigitUintA(PixelPos, OutColor, FontColor, Cursor, FramesTC);
+
+	Cursor.x = StartX;
+	Cursor.y += 12;
+	
+	PrintLargeUintA(PixelPos, OutColor, FontColor, Cursor, FrameCount);
+}
+
+Texture2D		InputTexture;
+SamplerState	InputTextureSampler;
+Texture2D		FrameStatsTexture;
+SamplerState	FrameStatsTextureSampler;
+
+void MainOutputPS(in noperspective float4 InUVAndScreenPos : TEXCOORD0,
+	float4 SvPosition : SV_POSITION,
+	out float4 OutColor : SV_Target0)
+{
+	uint2 PhysicalDimension;
+	FrameStatsTexture.GetDimensions(PhysicalDimension.x, PhysicalDimension.y);
+	
+	float4 SceneColor = InputTexture.Sample(InputTextureSampler, InUVAndScreenPos.xy);
+	float4 FrameStats = FrameStatsTexture.Sample(FrameStatsTextureSampler, 0.5f * SvPosition.xy/float2(PhysicalDimension));
+
+	OutColor = FrameStats + SceneColor * FrameStats.a;
+}
diff --git a/Engine/Plugins/Runtime/nDisplay/Source/DisplayCluster/Private/Render/Viewport/DisplayClusterViewportFrameStatsViewExtension.cpp b/Engine/Plugins/Runtime/nDisplay/Source/DisplayCluster/Private/Render/Viewport/DisplayClusterViewportFrameStatsViewExtension.cpp
new file mode 100644
--- /dev/null
+++ b/Engine/Plugins/Runtime/nDisplay/Source/DisplayCluster/Private/Render/Viewport/DisplayClusterViewportFrameStatsViewExtension.cpp
@@ -0,0 +1,182 @@
+// Copyright Epic Games, Inc. All Rights Reserved.
+
+#include "DisplayClusterViewportFrameStatsViewExtension.h"
+#include "DisplayClusterSceneViewExtensions.h"
+
+#include "Render/Viewport/DisplayClusterViewportManager.h"
+#include "Render/Viewport/DisplayClusterViewport.h"
+#include "Render/Viewport/DisplayClusterViewportManagerProxy.h"
+
+#include "GlobalShader.h"
+#include "Misc/App.h"
+#include "Misc/Timecode.h"
+#include "SceneRendering.h"
+#include "SceneView.h"
+#include "ScreenPass.h"
+#include "ShaderParameters.h"
+#include "ShaderParameterStruct.h"
+#include "SystemTextures.h"
+#include "PostProcess/PostProcessTestImage.h"
+#include "PostProcess/PostProcessMaterialInputs.h"
+
+
+int32 GDisplayClusterShowFrameStats = 0;
+static FAutoConsoleVariableRef CVarDisplayClusterShowFrameStats(
+	TEXT("DC.Stats.Frame"),
+	GDisplayClusterShowFrameStats,
+	TEXT("Show frame stats (timecode, frame number) in a repeating pattern."),
+	ECVF_RenderThreadSafe
+);
+
+BEGIN_SHADER_PARAMETER_STRUCT(FDisplayClusterFrameStatsShaderParameters, )
+	SHADER_PARAMETER_TEXTURE(Texture2D, MiniFontTexture)
+	SHADER_PARAMETER(uint32, FrameCount)
+	SHADER_PARAMETER(uint32, Timecode)
+	RENDER_TARGET_BINDING_SLOTS()
+END_SHADER_PARAMETER_STRUCT()
+
+class FDisplayClusterFrameStatsShader : public FGlobalShader
+{
+public:
+	DECLARE_GLOBAL_SHADER(FDisplayClusterFrameStatsShader);
+	SHADER_USE_PARAMETER_STRUCT(FDisplayClusterFrameStatsShader, FGlobalShader);
+
+	using FParameters = FDisplayClusterFrameStatsShaderParameters;
+};
+
+IMPLEMENT_SHADER_TYPE(, FDisplayClusterFrameStatsShader, TEXT("/Plugin/nDisplay/Private/FrameStats.usf"), TEXT("MainPS"), SF_Pixel);
+
+
+BEGIN_SHADER_PARAMETER_STRUCT(FDisplayClusterFrameStatsOutputShaderParameters, )
+	SHADER_PARAMETER_RDG_TEXTURE(Texture2D, InputTexture)
+	SHADER_PARAMETER_SAMPLER(SamplerState, InputTextureSampler)
+	SHADER_PARAMETER_RDG_TEXTURE(Texture2D, FrameStatsTexture)
+	SHADER_PARAMETER_SAMPLER(SamplerState, FrameStatsTextureSampler)
+	RENDER_TARGET_BINDING_SLOTS()
+END_SHADER_PARAMETER_STRUCT()
+
+class FDisplayClusterFrameStatsOutputShader : public FGlobalShader
+{
+public:
+	DECLARE_GLOBAL_SHADER(FDisplayClusterFrameStatsOutputShader);
+	SHADER_USE_PARAMETER_STRUCT(FDisplayClusterFrameStatsOutputShader, FGlobalShader);
+
+	using FParameters = FDisplayClusterFrameStatsOutputShaderParameters;
+};
+
+IMPLEMENT_SHADER_TYPE(, FDisplayClusterFrameStatsOutputShader, TEXT("/Plugin/nDisplay/Private/FrameStats.usf"), TEXT("MainOutputPS"), SF_Pixel);
+
+namespace UE::DisplayClusterViewExtension
+{
+	static FRHITexture* GetMiniFontTexture()
+	{
+		return GSystemTextures.AsciiTexture ? GSystemTextures.AsciiTexture->GetRHI() : GSystemTextures.WhiteDummy->GetRHI();
+	}
+}
+
+
+///////////////////////////////////////////////////////////////////////////////////////
+// FDisplayClusterViewportFrameStatsViewExtension
+///////////////////////////////////////////////////////////////////////////////////////
+FDisplayClusterViewportFrameStatsViewExtension::FDisplayClusterViewportFrameStatsViewExtension(const FAutoRegister& AutoRegister, const FDisplayClusterViewportManager* InViewportManager)
+	: FSceneViewExtensionBase(AutoRegister)
+	, ViewportManagerWeakPtr(InViewportManager->AsShared())
+{ }
+
+FDisplayClusterViewportFrameStatsViewExtension::~FDisplayClusterViewportFrameStatsViewExtension()
+{
+	ViewportManagerWeakPtr.Reset();
+}
+
+void FDisplayClusterViewportFrameStatsViewExtension::SubscribeToPostProcessingPass(EPostProcessingPass PassId, FAfterPassCallbackDelegateArray& InOutPassCallbacks, bool bIsPassEnabled)
+{
+	if (PassId == EPostProcessingPass::Tonemap)
+	{
+		InOutPassCallbacks.Add(FAfterPassCallbackDelegate::CreateRaw(this, &FDisplayClusterViewportFrameStatsViewExtension::PostProcessPassAfterTonemap_RenderThread));
+	}
+}
+
+void FDisplayClusterViewportFrameStatsViewExtension::BeginRenderViewFamily(FSceneViewFamily& InViewFamily)
+{
+	const FTimecode CurrentTimecode = FApp::GetTimecode();
+
+	EncodedTimecode = (static_cast<uint8>(CurrentTimecode.Hours) << 24u) | (static_cast<uint8>(CurrentTimecode.Minutes) << 16u) | (static_cast<uint8>(CurrentTimecode.Seconds) << 8u) | static_cast<uint8>(CurrentTimecode.Frames);
+	FrameCount = GFrameCounter;
+}
+
+FScreenPassTexture FDisplayClusterViewportFrameStatsViewExtension::PostProcessPassAfterTonemap_RenderThread(FRDGBuilder& GraphBuilder, const FSceneView& View, const FPostProcessMaterialInputs& Inputs)
+{
+	const FViewInfo& ViewInfo = static_cast<const FViewInfo&>(View);
+	const FScreenPassTexture& SceneColor = Inputs.GetInput(EPostProcessMaterialInput::SceneColor);
+	check(SceneColor.IsValid());
+
+	FGlobalShaderMap* ShaderMap = GetGlobalShaderMap(View.GetFeatureLevel());
+	TShaderMapRef<FScreenPassVS> VertexShader(ShaderMap);
+
+	FRDGTextureDesc TextureFrameStatsDesc = FRDGTextureDesc::Create2D(FIntPoint(128, 64), PF_R8G8B8A8, FClearValueBinding::Black, TexCreate_ShaderResource | TexCreate_RenderTargetable);
+	FRDGTexture* TextureFrameStats = GraphBuilder.CreateTexture(TextureFrameStatsDesc, TEXT("DisplayCluster.FrameStatsRT"));
+
+	// Render frame stats to a small render target to minimize shader costs
+	// Note: As a further optimization, this could be renderered once per view family.
+	{
+		const FScreenPassTextureViewport Viewport(TextureFrameStats);
+		FRenderTargetBinding FrameStatsRenderTargetBinding(TextureFrameStats, ERenderTargetLoadAction::EClear);
+
+		TShaderMapRef<FDisplayClusterFrameStatsShader> FrameStatsPixelShader(ShaderMap);
+		FDisplayClusterFrameStatsShaderParameters* Parameters = GraphBuilder.AllocParameters<FDisplayClusterFrameStatsShaderParameters>();
+		Parameters->MiniFontTexture = UE::DisplayClusterViewExtension::GetMiniFontTexture();
+		Parameters->RenderTargets[0] = FrameStatsRenderTargetBinding;
+		Parameters->FrameCount = FrameCount;
+		Parameters->Timecode = EncodedTimecode;
+
+		AddDrawScreenPass(GraphBuilder, RDG_EVENT_NAME("DisplayClusterViewportFrameStatsPass"), ViewInfo, Viewport, Viewport, VertexShader, FrameStatsPixelShader, Parameters);
+	}
+
+	FScreenPassRenderTarget Output = Inputs.OverrideOutput;
+	if (!Output.IsValid())
+	{
+		Output = FScreenPassRenderTarget::CreateFromInput(GraphBuilder, SceneColor, View.GetOverwriteLoadAction(), TEXT("DisplayCluster.FrameStatsOutput"));
+	}
+
+	// Render output composited with frame stats texture
+	{
+		const FScreenPassTextureViewport InputViewport(SceneColor);
+		const FScreenPassTextureViewport OutputViewport(Output);
+
+		TShaderMapRef<FDisplayClusterFrameStatsOutputShader> OutputPixelShader(ShaderMap);
+		FDisplayClusterFrameStatsOutputShaderParameters* Parameters = GraphBuilder.AllocParameters<FDisplayClusterFrameStatsOutputShaderParameters>();
+		Parameters->InputTexture = SceneColor.Texture;
+		Parameters->InputTextureSampler = TStaticSamplerState<>::GetRHI();
+		Parameters->FrameStatsTexture = TextureFrameStats;
+		Parameters->FrameStatsTextureSampler = TStaticSamplerState<SF_Point, AM_Wrap, AM_Wrap>::GetRHI();
+		Parameters->RenderTargets[0] = Output.GetRenderTargetBinding();
+
+		AddDrawScreenPass(GraphBuilder, RDG_EVENT_NAME("DisplayClusterViewportFrameStatsOutputPass"), ViewInfo, OutputViewport, InputViewport, VertexShader, OutputPixelShader, Parameters);
+	}
+
+	return Output;
+}
+
+bool FDisplayClusterViewportFrameStatsViewExtension::IsActiveThisFrame_Internal(const FSceneViewExtensionContext& Context) const
+{
+	if (IsActive())
+	{
+		static const FDisplayClusterSceneViewExtensionContext DCViewExtensionContext;
+		if (Context.IsA(MoveTempIfPossible(DCViewExtensionContext)))
+		{
+			const FDisplayClusterSceneViewExtensionContext& DisplayContext = static_cast<const FDisplayClusterSceneViewExtensionContext&>(Context);
+			if (DisplayContext.GetViewportManager() == GetViewportManager())
+			{
+				// Apply only for DC viewports
+				return true;
+			}
+		}
+	}
+
+	return false;
+}
+
+bool FDisplayClusterViewportFrameStatsViewExtension::IsActive() const
+{
+	return GetViewportManager() != nullptr && GDisplayClusterShowFrameStats > 0;
+}
diff --git a/Engine/Plugins/Runtime/nDisplay/Source/DisplayCluster/Private/Render/Viewport/DisplayClusterViewportFrameStatsViewExtension.h b/Engine/Plugins/Runtime/nDisplay/Source/DisplayCluster/Private/Render/Viewport/DisplayClusterViewportFrameStatsViewExtension.h
new file mode 100644
--- /dev/null
+++ b/Engine/Plugins/Runtime/nDisplay/Source/DisplayCluster/Private/Render/Viewport/DisplayClusterViewportFrameStatsViewExtension.h
@@ -0,0 +1,48 @@
+// Copyright Epic Games, Inc. All Rights Reserved.
+
+#pragma once
+#include "SceneViewExtension.h"
+
+class FDisplayClusterViewportManager;
+
+#define DISPLAYCLUSTER_SCENE_DEBUG_VIEW_EXTENSION_PRIORITY 999
+
+/**
+ * DC-specific view extension to display frame stats
+ */
+class FDisplayClusterViewportFrameStatsViewExtension : public FSceneViewExtensionBase
+{
+public:
+	FDisplayClusterViewportFrameStatsViewExtension(const FAutoRegister& AutoRegister, const FDisplayClusterViewportManager* InViewportManager);
+	virtual ~FDisplayClusterViewportFrameStatsViewExtension();
+
+public:
+	//~ Begin ISceneViewExtension interface
+	virtual void SetupViewFamily(FSceneViewFamily& InViewFamily) override { }
+	virtual void SetupView(FSceneViewFamily& InViewFamily, FSceneView& InView) override { }
+	virtual void BeginRenderViewFamily(FSceneViewFamily& InViewFamily) override;
+
+	void SubscribeToPostProcessingPass(EPostProcessingPass PassId, FAfterPassCallbackDelegateArray& InOutPassCallbacks, bool bIsPassEnabled);
+	FScreenPassTexture PostProcessPassAfterTonemap_RenderThread(FRDGBuilder& GraphBuilder, const FSceneView& View, const FPostProcessMaterialInputs& InOutInputs);
+
+	virtual int32 GetPriority() const override { return DISPLAYCLUSTER_SCENE_DEBUG_VIEW_EXTENSION_PRIORITY; }
+
+protected:
+	virtual bool IsActiveThisFrame_Internal(const FSceneViewExtensionContext& Context) const override;
+
+private:
+	/** Get viewport manager ptr. */
+	inline const FDisplayClusterViewportManager* GetViewportManager() const
+	{
+		return ViewportManagerWeakPtr.IsValid() ? ViewportManagerWeakPtr.Pin().Get() : nullptr;
+	}
+
+	/** True, if VE can be used at the moment. */
+	bool IsActive() const;
+
+private:
+	std::atomic<uint32> FrameCount;
+	std::atomic<uint32> EncodedTimecode;
+
+	TWeakPtr<const FDisplayClusterViewportManager, ESPMode::ThreadSafe> ViewportManagerWeakPtr;
+};
diff --git a/Engine/Plugins/Runtime/nDisplay/Source/DisplayCluster/Private/Render/Viewport/DisplayClusterViewportManager.cpp b/Engine/Plugins/Runtime/nDisplay/Source/DisplayCluster/Private/Render/Viewport/DisplayClusterViewportManager.cpp
--- a/Engine/Plugins/Runtime/nDisplay/Source/DisplayCluster/Private/Render/Viewport/DisplayClusterViewportManager.cpp
+++ b/Engine/Plugins/Runtime/nDisplay/Source/DisplayCluster/Private/Render/Viewport/DisplayClusterViewportManager.cpp
@@ -4,6 +4,7 @@
 #include "Render/Viewport/DisplayClusterViewportManagerProxy.h"
 #include "Render/Viewport/DisplayClusterViewportManagerViewExtension.h"
 #include "Render/Viewport/DisplayClusterViewportManagerViewPointExtension.h"
+#include "Render/Viewport/DisplayClusterViewportFrameStatsViewExtension.h"
 
 #include "IDisplayCluster.h"
 #include "Render/IDisplayClusterRenderManager.h"
@@ -621,6 +622,11 @@ bool FDisplayClusterViewportManager::BeginNewFrame(FViewport* InViewport, UWorld
 		ViewportManagerViewPointExtension = FSceneViewExtensions::NewExtension<FDisplayClusterViewportManagerViewPointExtension>(this);
 	}
 
+	if (!FrameStatsViewExtension.IsValid())
+	{
+		FrameStatsViewExtension = FSceneViewExtensions::NewExtension<FDisplayClusterViewportFrameStatsViewExtension>(this);
+	}
+
 	// Before new frame
 	if (PostProcessManager.IsValid())
 	{
diff --git a/Engine/Plugins/Runtime/nDisplay/Source/DisplayCluster/Private/Render/Viewport/DisplayClusterViewportManager.h b/Engine/Plugins/Runtime/nDisplay/Source/DisplayCluster/Private/Render/Viewport/DisplayClusterViewportManager.h
--- a/Engine/Plugins/Runtime/nDisplay/Source/DisplayCluster/Private/Render/Viewport/DisplayClusterViewportManager.h
+++ b/Engine/Plugins/Runtime/nDisplay/Source/DisplayCluster/Private/Render/Viewport/DisplayClusterViewportManager.h
@@ -226,6 +226,9 @@ private:
 	// This VE is used for LocalPlayer::GetViewPoint(). It calls the ISceneViewExtension::SetupViewPoint() function of this VE.
 	TSharedPtr<class FDisplayClusterViewportManagerViewPointExtension, ESPMode::ThreadSafe> ViewportManagerViewPointExtension;
 
+	// This VE is used to display frame information such as timecode and frame number.
+	TSharedPtr<class FDisplayClusterViewportFrameStatsViewExtension, ESPMode::ThreadSafe> FrameStatsViewExtension;
+
 #if WITH_EDITOR
 	/** The handle for FCoreUObjectDelegates::GetPreGarbageCollectDelegate */
 	FDelegateHandle PreGarbageCollectHandle;
