Fixed CineCameraComponent registering for tick in builds where it won't do anything useful.

diff --git a/Engine/Plugins/Runtime/nDisplay/Source/DisplayCluster/Private/Game/EngineClasses/Scene/DisplayClusterICVFXCameraComponent.cpp b/Engine/Plugins/Runtime/nDisplay/Source/DisplayCluster/Private/Game/EngineClasses/Scene/DisplayClusterICVFXCameraComponent.cpp
--- a/Engine/Plugins/Runtime/nDisplay/Source/DisplayCluster/Private/Game/EngineClasses/Scene/DisplayClusterICVFXCameraComponent.cpp
+++ b/Engine/Plugins/Runtime/nDisplay/Source/DisplayCluster/Private/Game/EngineClasses/Scene/DisplayClusterICVFXCameraComponent.cpp
@@ -19,6 +19,10 @@
 #include "DisplayClusterEnums.h"
 #include "Version/DisplayClusterICVFXCameraCustomVersion.h"
 
+UDisplayClusterICVFXCameraComponent::UDisplayClusterICVFXCameraComponent(const FObjectInitializer& ObjectInitializer)
+{
+	PrimaryComponentTick.bCanEverTick = true;
+}
 
 void UDisplayClusterICVFXCameraComponent::Serialize(FArchive& Ar)
 {
diff --git a/Engine/Plugins/Runtime/nDisplay/Source/DisplayCluster/Public/Components/DisplayClusterICVFXCameraComponent.h b/Engine/Plugins/Runtime/nDisplay/Source/DisplayCluster/Public/Components/DisplayClusterICVFXCameraComponent.h
--- a/Engine/Plugins/Runtime/nDisplay/Source/DisplayCluster/Public/Components/DisplayClusterICVFXCameraComponent.h
+++ b/Engine/Plugins/Runtime/nDisplay/Source/DisplayCluster/Public/Components/DisplayClusterICVFXCameraComponent.h
@@ -28,8 +28,6 @@ class DISPLAYCLUSTER_API UDisplayClusterICVFXCameraComponent
 	GENERATED_BODY()
 
 public:
-	UDisplayClusterICVFXCameraComponent(const FObjectInitializer& ObjectInitializer)
-	{ }
 
 	virtual void Serialize(FArchive& Ar) override;
 	virtual void PostLoad() override;
diff --git a/Engine/Source/Runtime/CinematicCamera/Private/CineCameraComponent.cpp b/Engine/Source/Runtime/CinematicCamera/Private/CineCameraComponent.cpp
--- a/Engine/Source/Runtime/CinematicCamera/Private/CineCameraComponent.cpp
+++ b/Engine/Source/Runtime/CinematicCamera/Private/CineCameraComponent.cpp
@@ -34,7 +34,12 @@ UCineCameraComponent::UCineCameraComponent()
 	bTickInEditor = true;
 #endif
 	
+	// NOTE: in UE 5.3 and onward, classes deriving from this component must explicitly register for tick if
+	// overriding the tick function, since this class only ticks productively with the following defines on
+#if WITH_EDITORONLY_DATA || ENABLE_DRAW_DEBUG
 	PrimaryComponentTick.bCanEverTick = true;
+#endif
+
 	bAutoActivate = true;
 
 	bConstrainAspectRatio = true;
