Fixed compile due to missing declaration (excluded from compile locally, sorry).

#rnx

diff --git a/Engine/Plugins/Runtime/nDisplay/Source/DisplayCluster/Public/Components/DisplayClusterICVFXCameraComponent.h b/Engine/Plugins/Runtime/nDisplay/Source/DisplayCluster/Public/Components/DisplayClusterICVFXCameraComponent.h
--- a/Engine/Plugins/Runtime/nDisplay/Source/DisplayCluster/Public/Components/DisplayClusterICVFXCameraComponent.h
+++ b/Engine/Plugins/Runtime/nDisplay/Source/DisplayCluster/Public/Components/DisplayClusterICVFXCameraComponent.h
@@ -28,6 +28,7 @@ class DISPLAYCLUSTER_API UDisplayClusterICVFXCameraComponent
 	GENERATED_BODY()
 
 public:
+	UDisplayClusterICVFXCameraComponent(const FObjectInitializer& ObjectInitializer);
 
 	virtual void Serialize(FArchive& Ar) override;
 	virtual void PostLoad() override;
