nDisplay: Fixed issues with light card colors and blending caused by inconsistent gamma encoding for 3D and UV light cards with and without OCIO. Now, all light cards are gamma encoded with gamma 2.2 when rendered, and will be decoded if needed during the final composite pass when blending occurs.

diff --git a/Engine/Plugins/Runtime/nDisplay/Shaders/Private/MPCDIOverlayShaders.usf b/Engine/Plugins/Runtime/nDisplay/Shaders/Private/MPCDIOverlayShaders.usf
--- a/Engine/Plugins/Runtime/nDisplay/Shaders/Private/MPCDIOverlayShaders.usf
+++ b/Engine/Plugins/Runtime/nDisplay/Shaders/Private/MPCDIOverlayShaders.usf
@@ -37,15 +37,16 @@ OutputPS Passthrough_PS(OutputVS IN)
 ////////////////////////////////////////////////////////////////////////////////
 float4 GetLightCardTextureColor(float2 InUV)
 {
-	return LightCardTexture.Sample(LightCardSampler, InUV);
+	float4 OutColor = LightCardTexture.Sample(LightCardSampler, InUV);
+	return float4(pow(OutColor.xyz, LightCardGamma), OutColor.w);
 }
 
 float4 GetUVLightCardTextureColor(float2 InUV)
 {
-	float4 OutColor =UVLightCardTexture.Sample(UVLightCardSampler, InUV);
+	float4 OutColor = UVLightCardTexture.Sample(UVLightCardSampler, InUV);
 
 	// UVLightCard alpha not inverted
-	return float4(OutColor.xyz, 1 - OutColor.w);
+	return float4(pow(OutColor.xyz, LightCardGamma), 1 - OutColor.w);
 }
 
 float3 ComposeOverlayColor(float4 InColor, float4 InOverlayColor)
diff --git a/Engine/Plugins/Runtime/nDisplay/Shaders/Private/MPCDIUtils.ush b/Engine/Plugins/Runtime/nDisplay/Shaders/Private/MPCDIUtils.ush
--- a/Engine/Plugins/Runtime/nDisplay/Shaders/Private/MPCDIUtils.ush
+++ b/Engine/Plugins/Runtime/nDisplay/Shaders/Private/MPCDIUtils.ush
@@ -52,6 +52,7 @@ float4 OverlappedInnerCameraSoftEdges[MAX_INNER_CAMERAS_AMOUNT];
 #define ENABLE_INNER_CAMERA_OVERLAP 0
 #endif
 
+float LightCardGamma;
 float AlphaEmbeddedGamma;
 
 float4 InnerCameraSoftEdge;
diff --git a/Engine/Plugins/Runtime/nDisplay/Shaders/Private/UVLightCardGammaCorrection.usf b/Engine/Plugins/Runtime/nDisplay/Shaders/Private/UVLightCardGammaCorrection.usf
new file mode 100644
--- /dev/null
+++ b/Engine/Plugins/Runtime/nDisplay/Shaders/Private/UVLightCardGammaCorrection.usf
@@ -0,0 +1,14 @@
+// Copyright Epic Games, Inc. All Rights Reserved.
+
+#include "/Engine/Private/Common.ush"
+
+Texture2D InputTexture;
+SamplerState InputSampler;
+float InverseGamma = 1.0;
+
+float4 GammaCorrectionPS(noperspective float4 UVAndScreenPos : TEXCOORD0) : SV_Target0
+{
+	float2 UV = UVAndScreenPos.xy;
+	float4 TextureColor = Texture2DSample(InputTexture, InputSampler, UV);
+	return float4(pow(TextureColor.rgb, InverseGamma), TextureColor.a);
+}
\ No newline at end of file
diff --git a/Engine/Plugins/Runtime/nDisplay/Source/DisplayCluster/Private/Render/Viewport/Configuration/DisplayClusterViewportConfigurationHelpers_ICVFX.cpp b/Engine/Plugins/Runtime/nDisplay/Source/DisplayCluster/Private/Render/Viewport/Configuration/DisplayClusterViewportConfigurationHelpers_ICVFX.cpp
--- a/Engine/Plugins/Runtime/nDisplay/Source/DisplayCluster/Private/Render/Viewport/Configuration/DisplayClusterViewportConfigurationHelpers_ICVFX.cpp
+++ b/Engine/Plugins/Runtime/nDisplay/Source/DisplayCluster/Private/Render/Viewport/Configuration/DisplayClusterViewportConfigurationHelpers_ICVFX.cpp
@@ -807,6 +807,13 @@ void FDisplayClusterViewportConfigurationHelpers_ICVFX::UpdateLightcardViewportS
 	{
 		BaseViewport.RenderSettingsICVFX.ICVFX.UVLightCard.ViewportId = DstViewport.GetId();
 	}
+
+	if (EnumHasAnyFlags(DstViewport.RenderSettingsICVFX.RuntimeFlags, EDisplayClusterViewportRuntimeICVFXFlags::Lightcard | EDisplayClusterViewportRuntimeICVFXFlags::UVLightcard))
+	{
+		// Set the light card gamma used to linearize the light card textures before blending during final composite
+		// The OCIO pass will have already linearized the light card renders, so set the gamma to 1 if using OCIO on the light cards
+		BaseViewport.RenderSettingsICVFX.ICVFX.LightCardGamma = DstViewport.OpenColorIO.IsValid() ? 1.0 : 2.2;
+	}
 }
 
 void FDisplayClusterViewportConfigurationHelpers_ICVFX::UpdateCameraCustomFrustum(
diff --git a/Engine/Plugins/Runtime/nDisplay/Source/DisplayCluster/Private/Render/Viewport/DisplayClusterViewportManager.cpp b/Engine/Plugins/Runtime/nDisplay/Source/DisplayCluster/Private/Render/Viewport/DisplayClusterViewportManager.cpp
--- a/Engine/Plugins/Runtime/nDisplay/Source/DisplayCluster/Private/Render/Viewport/DisplayClusterViewportManager.cpp
+++ b/Engine/Plugins/Runtime/nDisplay/Source/DisplayCluster/Private/Render/Viewport/DisplayClusterViewportManager.cpp
@@ -692,9 +692,10 @@ bool FDisplayClusterViewportManager::BeginNewFrame(FViewport* InViewport, UWorld
 				ContextIt.RenderThreadData.EngineDisplayGamma = ViewportDisplayGamma;
 
 
-				if (DefaultEngineShowFlags)
+				if (DefaultEngineShowFlags && !EnumHasAnyFlags(ViewportIt->RenderSettingsICVFX.RuntimeFlags, EDisplayClusterViewportRuntimeICVFXFlags::UVLightcard))
 				{
 					// Todo: DefaultEngineShowFlags is not the correct value, since each viewport uses its own Engine flags.
+					// UV Light cards should not be set to the default engine flags, since they are rendered using a custom render pass that does not utilize the flags
 					ContextIt.RenderThreadData.EngineShowFlags = *DefaultEngineShowFlags;
 				}
 			}
@@ -1182,6 +1183,26 @@ TSharedPtr<IDisplayClusterProjectionPolicy, ESPMode::ThreadSafe> FDisplayCluster
 	return nullptr;
 }
 
+TArray<TSharedPtr<IDisplayClusterViewport, ESPMode::ThreadSafe>> FDisplayClusterViewportManager::GetViewportsForWarpPolicy(const TSharedPtr<IDisplayClusterWarpPolicy>& InWarpPolicy) const
+{
+	TArray<TSharedPtr<IDisplayClusterViewport, ESPMode::ThreadSafe>> WarpPolicyViewports;
+
+	for (const TSharedPtr<FDisplayClusterViewport, ESPMode::ThreadSafe>& Viewport : ImplGetEntireClusterViewports())
+	{
+		bool bViewportValid = Viewport.IsValid();
+		bool bProjectionPolicyValid = Viewport->GetProjectionPolicy().IsValid();
+		IDisplayClusterWarpPolicy* WarpPolicy = Viewport->GetProjectionPolicy()->GetWarpPolicy();
+		bool bIsSameWarpPolicy = WarpPolicy == InWarpPolicy.Get();
+
+		if (Viewport.IsValid() && Viewport->GetProjectionPolicy().IsValid() && Viewport->GetProjectionPolicy()->GetWarpPolicy() == InWarpPolicy.Get())
+		{
+			WarpPolicyViewports.Add(Viewport);
+		}
+	}
+
+	return WarpPolicyViewports;
+}
+
 void FDisplayClusterViewportManager::MarkComponentGeometryDirty(const FName InComponentName)
 {
 	check(IsInGameThread());
diff --git a/Engine/Plugins/Runtime/nDisplay/Source/DisplayCluster/Private/Render/Viewport/LightCard/DisplayClusterViewportLightCardManager.cpp b/Engine/Plugins/Runtime/nDisplay/Source/DisplayCluster/Private/Render/Viewport/LightCard/DisplayClusterViewportLightCardManager.cpp
--- a/Engine/Plugins/Runtime/nDisplay/Source/DisplayCluster/Private/Render/Viewport/LightCard/DisplayClusterViewportLightCardManager.cpp
+++ b/Engine/Plugins/Runtime/nDisplay/Source/DisplayCluster/Private/Render/Viewport/LightCard/DisplayClusterViewportLightCardManager.cpp
@@ -182,6 +182,7 @@ void FDisplayClusterViewportLightCardManager::RenderUVLightCard()
 			FDisplayClusterShaderParameters_UVLightCards UVLightCardParameters;
 			UVLightCardParameters.ProjectionPlaneSize = ADisplayClusterLightCardActor::UVPlaneDefaultSize;
 			UVLightCardParameters.bRenderFinalColor = ViewportManager->ShouldRenderFinalColor();
+			UVLightCardParameters.LightCardGamma = 2.2;
 
 			// Store any components that were invisible but forced to be visible so they can be set back to invisible after the render
 			TArray<UPrimitiveComponent*> ComponentsToUnload;
diff --git a/Engine/Plugins/Runtime/nDisplay/Source/DisplayClusterShaders/Private/Shaders/DisplayClusterShadersPreprocess_UVLightCards.cpp b/Engine/Plugins/Runtime/nDisplay/Source/DisplayClusterShaders/Private/Shaders/DisplayClusterShadersPreprocess_UVLightCards.cpp
--- a/Engine/Plugins/Runtime/nDisplay/Source/DisplayClusterShaders/Private/Shaders/DisplayClusterShadersPreprocess_UVLightCards.cpp
+++ b/Engine/Plugins/Runtime/nDisplay/Source/DisplayClusterShaders/Private/Shaders/DisplayClusterShadersPreprocess_UVLightCards.cpp
@@ -17,6 +17,7 @@
 #include "SceneInterface.h"
 #include "SceneRendererInterface.h"
 #include "SceneTexturesConfig.h"
+#include "ScreenPass.h"
 #include "ShaderParameters/DisplayClusterShaderParameters_UVLightCards.h"
 #include "UnrealClient.h"
 
@@ -172,6 +173,21 @@ BEGIN_SHADER_PARAMETER_STRUCT(FUVLightCardPassParameters, )
 	RENDER_TARGET_BINDING_SLOTS()
 END_SHADER_PARAMETER_STRUCT()
 
+class FGammaCorrectionPS : public FGlobalShader
+{
+	DECLARE_SHADER_TYPE(FGammaCorrectionPS, Global);
+	SHADER_USE_PARAMETER_STRUCT(FGammaCorrectionPS, FGlobalShader);
+
+	BEGIN_SHADER_PARAMETER_STRUCT(FParameters, )
+		SHADER_PARAMETER_RDG_TEXTURE(Texture2D, InputTexture)
+		SHADER_PARAMETER_SAMPLER(SamplerState, InputSampler)
+		SHADER_PARAMETER(float, InverseGamma)
+		RENDER_TARGET_BINDING_SLOTS()
+		END_SHADER_PARAMETER_STRUCT()
+};
+
+IMPLEMENT_GLOBAL_SHADER(FGammaCorrectionPS, "/Plugin/nDisplay/Private/UVLightCardGammaCorrection.usf", "GammaCorrectionPS", SF_Pixel);
+
 DECLARE_GPU_STAT_NAMED(nDisplay_UVLightCards_Render, TEXT("nDisplay UVLightCards::Render"));
 
 bool FDisplayClusterShadersPreprocess_UVLightCards::RenderPreprocess_UVLightCards(FRHICommandListImmediate& RHICmdList, FSceneInterface* InScene, FRenderTarget* InRenderTarget, const FDisplayClusterShaderParameters_UVLightCards& InParameters)
@@ -277,11 +293,16 @@ bool FDisplayClusterShadersPreprocess_UVLightCards::RenderPreprocess_UVLightCard
 	FRDGTextureRef OutputTexture = GraphBuilder.RegisterExternalTexture(CreateRenderTarget(InRenderTarget->GetRenderTargetTexture(), TEXT("UVLightCardRenderTarget")));
 	FRenderTargetBinding OutputRenderTargetBinding(OutputTexture, ERenderTargetLoadAction::EClear);
 
+	// Intermediate render target for the raw color output of the light cards
+	FRDGTextureRef ColorTexture = GraphBuilder.CreateTexture(OutputTexture->Desc, TEXT("DisplayClusterUVLightCards.ColorTexture"));
+	FRenderTargetBinding ColorRenderTargetBinding(ColorTexture, ERenderTargetLoadAction::EClear);
+
+	// First, render all the UV light cards to a plane which represents the UV plane for the stage
 	FUVLightCardPassParameters* PassParameters = GraphBuilder.AllocParameters<FUVLightCardPassParameters>();
 	PassParameters->View = View->ViewUniformBuffer;
 	PassParameters->Scene = GetSceneUniformBufferRef(GraphBuilder, *View);
 	PassParameters->InstanceCulling = FInstanceCullingContext::CreateDummyInstanceCullingUniformBuffer(GraphBuilder);
-	PassParameters->RenderTargets[0] = OutputRenderTargetBinding;
+	PassParameters->RenderTargets[0] = ColorRenderTargetBinding;
 
 	GraphBuilder.AddPass(RDG_EVENT_NAME("DisplayClusterUVLightCards::Render"),
 		PassParameters,
@@ -307,6 +328,27 @@ bool FDisplayClusterShadersPreprocess_UVLightCards::RenderPreprocess_UVLightCard
 			});
 		});
 
+	// Now gamma correct the UV light card render
+	FGammaCorrectionPS::FParameters* GammaCorrectionParameters = GraphBuilder.AllocParameters<FGammaCorrectionPS::FParameters>();
+	GammaCorrectionParameters->InputTexture = ColorTexture;
+	GammaCorrectionParameters->InputSampler = TStaticSamplerState<>::GetRHI();
+	GammaCorrectionParameters->InverseGamma = 1.0 / InParameters.LightCardGamma;
+	GammaCorrectionParameters->RenderTargets[0] = OutputRenderTargetBinding;
+
+	FGlobalShaderMap* GlobalShaderMap = GetGlobalShaderMap(View->FeatureLevel);
+	TShaderMapRef<FScreenPassVS> ScreenPassVS(GlobalShaderMap);
+	TShaderMapRef<FGammaCorrectionPS> GammaCorrectionPS(GlobalShaderMap);
+
+	AddDrawScreenPass(
+		GraphBuilder,
+		RDG_EVENT_NAME("DisplayClusterUVLightCards::GammaCorrection"),
+		FScreenPassViewInfo(*View),
+		FScreenPassTextureViewport(OutputTexture),
+		FScreenPassTextureViewport(ColorTexture),
+		ScreenPassVS,
+		GammaCorrectionPS,
+		GammaCorrectionParameters);
+
 	GraphBuilder.Execute();
 
 	return true;
diff --git a/Engine/Plugins/Runtime/nDisplay/Source/DisplayClusterShaders/Private/Shaders/DisplayClusterShadersWarpblend_ICVFX.cpp b/Engine/Plugins/Runtime/nDisplay/Source/DisplayClusterShaders/Private/Shaders/DisplayClusterShadersWarpblend_ICVFX.cpp
--- a/Engine/Plugins/Runtime/nDisplay/Source/DisplayClusterShaders/Private/Shaders/DisplayClusterShadersWarpblend_ICVFX.cpp
+++ b/Engine/Plugins/Runtime/nDisplay/Source/DisplayClusterShaders/Private/Shaders/DisplayClusterShadersWarpblend_ICVFX.cpp
@@ -296,6 +296,7 @@ BEGIN_SHADER_PARAMETER_STRUCT(FIcvfxPixelShaderParameters, )
 	SHADER_PARAMETER(FMatrix44f, OverlayProjectionMatrix)
 	SHADER_PARAMETER(FMatrix44f, InnerCameraProjectionMatrix)
 
+	SHADER_PARAMETER(float, LightCardGamma)
 	SHADER_PARAMETER(float, AlphaEmbeddedGamma)
 
 	SHADER_PARAMETER(FVector4f, InnerCameraSoftEdge)
@@ -623,6 +624,7 @@ public:
 		{
 			RenderPassData.PSParameters.LightCardSampler = TStaticSamplerState<SF_Trilinear, AM_Clamp, AM_Clamp, AM_Clamp>::GetRHI();
 			RenderPassData.PSParameters.LightCardTexture = ICVFXParameters.LightCard.Texture;
+			RenderPassData.PSParameters.LightCardGamma = ICVFXParameters.LightCardGamma;
 
 			return true;
 		}
@@ -670,6 +672,7 @@ public:
 		{
 			RenderPassData.PSParameters.UVLightCardSampler = TStaticSamplerState<SF_Trilinear, AM_Clamp, AM_Clamp, AM_Clamp>::GetRHI();
 			RenderPassData.PSParameters.UVLightCardTexture = ICVFXParameters.UVLightCard.Texture;
+			RenderPassData.PSParameters.LightCardGamma = ICVFXParameters.LightCardGamma;
 
 			return true;
 		}
diff --git a/Engine/Plugins/Runtime/nDisplay/Source/DisplayClusterShaders/Public/ShaderParameters/DisplayClusterShaderParameters_ICVFX.h b/Engine/Plugins/Runtime/nDisplay/Source/DisplayClusterShaders/Public/ShaderParameters/DisplayClusterShaderParameters_ICVFX.h
--- a/Engine/Plugins/Runtime/nDisplay/Source/DisplayClusterShaders/Public/ShaderParameters/DisplayClusterShaderParameters_ICVFX.h
+++ b/Engine/Plugins/Runtime/nDisplay/Source/DisplayClusterShaders/Public/ShaderParameters/DisplayClusterShaderParameters_ICVFX.h
@@ -105,6 +105,8 @@ public:
 
 		UVLightCard.Reset();
 		LightCard.Reset();
+
+		LightCardGamma = 2.2;
 	}
 
 	// Implement copy ref and arrays
@@ -119,6 +121,7 @@ public:
 		UVLightCard = InParameters.UVLightCard;
 		LightCard   = InParameters.LightCard;
 		LightCardMode = InParameters.LightCardMode;
+		LightCardGamma = InParameters.LightCardGamma;
 	}
 
 	inline void SortCamerasRenderOrder()
@@ -332,4 +335,7 @@ public:
 	FDisplayClusterShaderParametersICVFX_ViewportResource    UVLightCard;
 	FDisplayClusterShaderParametersICVFX_ViewportResource    LightCard;
 	EDisplayClusterShaderParametersICVFX_LightCardRenderMode LightCardMode = EDisplayClusterShaderParametersICVFX_LightCardRenderMode::Under;
+
+	// The gamma that the light card renders have been encoded with, used to linearize during final composite
+	float LightCardGamma = 2.2;
 };
diff --git a/Engine/Plugins/Runtime/nDisplay/Source/DisplayClusterShaders/Public/ShaderParameters/DisplayClusterShaderParameters_UVLightCards.h b/Engine/Plugins/Runtime/nDisplay/Source/DisplayClusterShaders/Public/ShaderParameters/DisplayClusterShaderParameters_UVLightCards.h
--- a/Engine/Plugins/Runtime/nDisplay/Source/DisplayClusterShaders/Public/ShaderParameters/DisplayClusterShaderParameters_UVLightCards.h
+++ b/Engine/Plugins/Runtime/nDisplay/Source/DisplayClusterShaders/Public/ShaderParameters/DisplayClusterShaderParameters_UVLightCards.h
@@ -15,6 +15,9 @@ struct FDisplayClusterShaderParameters_UVLightCards
 	/** The size of the plane in world units to project the UV space onto */
 	float ProjectionPlaneSize;
 
+	/** The gamma to apply to the light cards when gamma correcting them */
+	float LightCardGamma = 1.0;
+
 	/** Indicates whether to render the final color of the light cards or not */
 	bool bRenderFinalColor;
 };
