nDisplay: Fix bad commit from CL 27350219

diff --git a/Engine/Plugins/Runtime/nDisplay/Source/DisplayCluster/Private/Render/Viewport/DisplayClusterViewportManager.cpp b/Engine/Plugins/Runtime/nDisplay/Source/DisplayCluster/Private/Render/Viewport/DisplayClusterViewportManager.cpp
--- a/Engine/Plugins/Runtime/nDisplay/Source/DisplayCluster/Private/Render/Viewport/DisplayClusterViewportManager.cpp
+++ b/Engine/Plugins/Runtime/nDisplay/Source/DisplayCluster/Private/Render/Viewport/DisplayClusterViewportManager.cpp
@@ -1183,26 +1183,6 @@ TSharedPtr<IDisplayClusterProjectionPolicy, ESPMode::ThreadSafe> FDisplayCluster
 	return nullptr;
 }
 
-TArray<TSharedPtr<IDisplayClusterViewport, ESPMode::ThreadSafe>> FDisplayClusterViewportManager::GetViewportsForWarpPolicy(const TSharedPtr<IDisplayClusterWarpPolicy>& InWarpPolicy) const
-{
-	TArray<TSharedPtr<IDisplayClusterViewport, ESPMode::ThreadSafe>> WarpPolicyViewports;
-
-	for (const TSharedPtr<FDisplayClusterViewport, ESPMode::ThreadSafe>& Viewport : ImplGetEntireClusterViewports())
-	{
-		bool bViewportValid = Viewport.IsValid();
-		bool bProjectionPolicyValid = Viewport->GetProjectionPolicy().IsValid();
-		IDisplayClusterWarpPolicy* WarpPolicy = Viewport->GetProjectionPolicy()->GetWarpPolicy();
-		bool bIsSameWarpPolicy = WarpPolicy == InWarpPolicy.Get();
-
-		if (Viewport.IsValid() && Viewport->GetProjectionPolicy().IsValid() && Viewport->GetProjectionPolicy()->GetWarpPolicy() == InWarpPolicy.Get())
-		{
-			WarpPolicyViewports.Add(Viewport);
-		}
-	}
-
-	return WarpPolicyViewports;
-}
-
 void FDisplayClusterViewportManager::MarkComponentGeometryDirty(const FName InComponentName)
 {
 	check(IsInGameThread());
