nDisplay: Add an optional render pass to display device components. The default implementation performs an OCIO pass.
Fix default display devices not being cached.

diff --git a/Engine/Plugins/Runtime/nDisplay/Source/DisplayCluster/Private/Game/EngineClasses/Scene/DisplayClusterPreviewComponent.cpp b/Engine/Plugins/Runtime/nDisplay/Source/DisplayCluster/Private/Game/EngineClasses/Scene/DisplayClusterPreviewComponent.cpp
--- a/Engine/Plugins/Runtime/nDisplay/Source/DisplayCluster/Private/Game/EngineClasses/Scene/DisplayClusterPreviewComponent.cpp
+++ b/Engine/Plugins/Runtime/nDisplay/Source/DisplayCluster/Private/Game/EngineClasses/Scene/DisplayClusterPreviewComponent.cpp
@@ -106,13 +106,15 @@ UDisplayClusterDisplayDeviceBaseComponent* UDisplayClusterPreviewComponent::GetD
 	{
 		UDisplayClusterDisplayDeviceBaseComponent* CachedComponent =
 			Cast<UDisplayClusterDisplayDeviceBaseComponent>(CachedDisplayDevice.GetComponent(GetOwner()));
-		if (CachedComponent && CachedComponent->GetName() == ViewportConfig->DisplayDeviceName)
+		if (CachedComponent
+			&& (CachedComponent->GetName() == ViewportConfig->DisplayDeviceName
+				|| (ViewportConfig->DisplayDeviceName.IsEmpty() && bIsUsingDefaultDisplayDevice && CachedComponent == RootActor->GetDefaultDisplayDevice())))
 		{
 			DeviceBaseComponent = CachedComponent;
 		}
 		else
 		{
-			DeviceBaseComponent = UE::DisplayClusterDisplayDeviceUtils::FindAndSyncDisplayDeviceFromViewport(ViewportConfig);
+			DeviceBaseComponent = UE::DisplayClusterDisplayDeviceUtils::FindAndSyncDisplayDeviceFromViewport(ViewportConfig, &bIsUsingDefaultDisplayDevice);
 		}
 	}
 
@@ -221,7 +223,7 @@ bool UDisplayClusterPreviewComponent::UpdatePreviewMesh()
 		// It must have a renderable resource (preview render target or override texture), have a valid viewport configured,
 		// and have that viewport either actively being rendered to by the preview renderer OR have an override texture supplied
 		// externally
-		const bool bHasRenderableResource = RenderTarget != nullptr || RenderTargetPostProcess != nullptr || OverrideTexture != nullptr;
+		const bool bHasRenderableResource = RenderTarget != nullptr || RenderTargetPostProcess != nullptr || OverrideTexture != nullptr || DisplayDeviceRenderTarget != nullptr;
 		const bool bIsViewportValid = Viewport != nullptr && Viewport->GetProjectionPolicy().IsValid();
 		const bool bOutputToPreviewMesh = bHasRenderableResource && bIsViewportValid && (Viewport->GetRenderSettings().bEnable || OverrideTexture);
 		if (bOutputToPreviewMesh)
@@ -314,8 +316,7 @@ void UDisplayClusterPreviewComponent::UpdatePreviewMaterial()
 		}
 		else
 		{
-			PreviewMaterialInstance->SetTextureParameterValue(TEXT("Preview"),
-				RootActor && RootActor->bPreviewEnablePostProcess ? RenderTargetPostProcess : RenderTarget);
+			PreviewMaterialInstance->SetTextureParameterValue(TEXT("Preview"), CurrentPreviewRenderTargetPtr);
 		}
 
 		// Allow display device to perform any processing on the material instance.
@@ -346,6 +347,7 @@ void UDisplayClusterPreviewComponent::ReleasePreviewRenderTarget()
 {
 	ReleaseRenderTargetImpl(&RenderTarget);
 	ReleaseRenderTargetImpl(&RenderTargetPostProcess);
+	ReleaseRenderTargetImpl(&DisplayDeviceRenderTarget);
 }
 
 void UDisplayClusterPreviewComponent::UpdatePreviewRenderTarget()
@@ -360,6 +362,33 @@ void UDisplayClusterPreviewComponent::UpdatePreviewRenderTarget()
 		{
 			UpdateRenderTargetImpl(&RenderTarget);
 		}
+
+		UTextureRenderTarget2D* DestinationRenderTarget = nullptr;
+		if (UTextureRenderTarget2D* SourceRenderTarget = RootActor->bPreviewEnablePostProcess
+			                                                 ? RenderTargetPostProcess
+			                                                 : RenderTarget)
+		{
+			DestinationRenderTarget = SourceRenderTarget;
+
+			// Perform any render passes on the render target by the display device
+			if (UDisplayClusterDisplayDeviceBaseComponent* CachedComponent =
+			Cast<UDisplayClusterDisplayDeviceBaseComponent>(CachedDisplayDevice.GetComponent(GetOwner())))
+			{
+				if (CachedComponent->IsRenderPassEnabled() && !RootActor->IsTemplate())
+				{
+					UpdateRenderTargetImpl(&DisplayDeviceRenderTarget);
+					CachedComponent->RenderPass_GameThread(SourceRenderTarget, DisplayDeviceRenderTarget);
+					DestinationRenderTarget = DisplayDeviceRenderTarget;
+				}
+			}
+		}
+
+		if (DestinationRenderTarget != CurrentPreviewRenderTargetPtr)
+		{
+			// Preview render target has changed, update the material
+			CurrentPreviewRenderTargetPtr = DestinationRenderTarget;
+			UpdatePreviewMaterial();
+		}
 	}
 }
 
diff --git a/Engine/Plugins/Runtime/nDisplay/Source/DisplayCluster/Private/Game/EngineClasses/Scene/DisplayClusterRootActorEditor.cpp b/Engine/Plugins/Runtime/nDisplay/Source/DisplayCluster/Private/Game/EngineClasses/Scene/DisplayClusterRootActorEditor.cpp
--- a/Engine/Plugins/Runtime/nDisplay/Source/DisplayCluster/Private/Game/EngineClasses/Scene/DisplayClusterRootActorEditor.cpp
+++ b/Engine/Plugins/Runtime/nDisplay/Source/DisplayCluster/Private/Game/EngineClasses/Scene/DisplayClusterRootActorEditor.cpp
@@ -10,7 +10,7 @@
 #include "Components/DisplayClusterSceneComponentSyncParent.h"
 #include "Components/DisplayClusterScreenComponent.h"
 #include "Components/DisplayClusterStageGeometryComponent.h"
-#include "Components/DisplayDevice/DisplayClusterDisplayDeviceBaseComponent.h"
+#include "Components/DisplayDevice/DisplayClusterDisplayDeviceComponent.h"
 #include "Components/LineBatchComponent.h"
 #include "Components/MeshComponent.h"
 #include "Components/SceneComponent.h"
diff --git a/Engine/Plugins/Runtime/nDisplay/Source/DisplayCluster/Private/Game/EngineClasses/Scene/DisplayDevice/DisplayClusterDisplayDeviceBaseComponent.cpp b/Engine/Plugins/Runtime/nDisplay/Source/DisplayCluster/Private/Game/EngineClasses/Scene/DisplayDevice/DisplayClusterDisplayDeviceBaseComponent.cpp
new file mode 100644
--- /dev/null
+++ b/Engine/Plugins/Runtime/nDisplay/Source/DisplayCluster/Private/Game/EngineClasses/Scene/DisplayDevice/DisplayClusterDisplayDeviceBaseComponent.cpp
@@ -0,0 +1,39 @@
+// Copyright Epic Games, Inc. All Rights Reserved.
+
+#include "Components/DisplayDevice/DisplayClusterDisplayDeviceBaseComponent.h"
+
+#include "DisplayClusterRootActor.h"
+
+#include "Materials/Material.h"
+#include "UObject/ConstructorHelpers.h"
+
+UDisplayClusterDisplayDeviceBaseComponent::UDisplayClusterDisplayDeviceBaseComponent()
+{
+#if WITH_EDITOR
+	static ConstructorHelpers::FObjectFinder<UMaterial> PreviewMaterialObj(TEXT("/nDisplay/Materials/Preview/M_DisplayDevicePreview"));
+	check(PreviewMaterialObj.Object);
+	PreviewMaterial = PreviewMaterialObj.Object;
+
+	static ConstructorHelpers::FObjectFinder<UMaterial> MeshMaterialObj(TEXT("/nDisplay/Materials/Preview/M_DisplayDeviceMesh"));
+	check(MeshMaterialObj.Object);
+	MeshMaterial = MeshMaterialObj.Object;
+#endif
+}
+
+#if WITH_EDITOR
+void UDisplayClusterDisplayDeviceBaseComponent::PostEditChangeProperty(FPropertyChangedEvent& PropertyChangedEvent)
+{
+	Super::PostEditChangeProperty(PropertyChangedEvent);
+
+	const FName PropertyName = PropertyChangedEvent.GetPropertyName();
+	
+	if (PropertyName == GET_MEMBER_NAME_CHECKED(UDisplayClusterDisplayDeviceBaseComponent, PreviewMaterial)
+		|| PropertyName == GET_MEMBER_NAME_CHECKED(UDisplayClusterDisplayDeviceBaseComponent, MeshMaterial))
+	{
+		if (ADisplayClusterRootActor* RootActor = Cast<ADisplayClusterRootActor>(GetOwner()))
+		{
+			RootActor->ResetPreviewComponents_Editor(true);
+		}
+	}
+}
+#endif
diff --git a/Engine/Plugins/Runtime/nDisplay/Source/DisplayCluster/Private/Game/EngineClasses/Scene/DisplayDevice/DisplayClusterDisplayDeviceComponent.cpp b/Engine/Plugins/Runtime/nDisplay/Source/DisplayCluster/Private/Game/EngineClasses/Scene/DisplayDevice/DisplayClusterDisplayDeviceComponent.cpp
new file mode 100644
--- /dev/null
+++ b/Engine/Plugins/Runtime/nDisplay/Source/DisplayCluster/Private/Game/EngineClasses/Scene/DisplayDevice/DisplayClusterDisplayDeviceComponent.cpp
@@ -0,0 +1,31 @@
+// Copyright Epic Games, Inc. All Rights Reserved.
+
+#include "Components/DisplayDevice/DisplayClusterDisplayDeviceComponent.h"
+
+#include "DisplayClusterRootActor.h"
+#include "IDisplayCluster.h"
+
+#include "Engine/TextureRenderTarget2D.h"
+#include "Materials/MaterialInstanceDynamic.h"
+#include "OpenColorIORendering.h"
+
+UDisplayClusterDisplayDeviceComponent::UDisplayClusterDisplayDeviceComponent()
+{
+}
+
+void UDisplayClusterDisplayDeviceComponent::OnUpdatePreviewMaterialInstance(UMaterialInstanceDynamic* InMaterialInstance)
+{
+	Super::OnUpdatePreviewMaterialInstance(InMaterialInstance);
+
+	if (InMaterialInstance)
+	{
+		InMaterialInstance->SetScalarParameterValue(TEXT("Exposure"), Exposure);
+	}
+}
+
+void UDisplayClusterDisplayDeviceComponent::RenderPass_GameThread(UTexture* InSourceTexture, UTextureRenderTarget2D* InRenderTarget)
+{
+	Super::RenderPass_GameThread(InSourceTexture, InRenderTarget);
+
+	FOpenColorIORendering::ApplyColorTransform(GetWorld(), GetColorConversionSettings(), InSourceTexture, InRenderTarget);
+}
diff --git a/Engine/Plugins/Runtime/nDisplay/Source/DisplayCluster/Private/Game/EngineClasses/Scene/DisplayDevice/DisplayClusterDisplayDeviceComponentBase.cpp b/Engine/Plugins/Runtime/nDisplay/Source/DisplayCluster/Private/Game/EngineClasses/Scene/DisplayDevice/DisplayClusterDisplayDeviceComponentBase.cpp
deleted file mode 100644
diff --git a/Engine/Plugins/Runtime/nDisplay/Source/DisplayCluster/Private/Game/EngineClasses/Scene/DisplayDevice/DisplayClusterDisplayDeviceUtils.cpp b/Engine/Plugins/Runtime/nDisplay/Source/DisplayCluster/Private/Game/EngineClasses/Scene/DisplayDevice/DisplayClusterDisplayDeviceUtils.cpp
--- a/Engine/Plugins/Runtime/nDisplay/Source/DisplayCluster/Private/Game/EngineClasses/Scene/DisplayDevice/DisplayClusterDisplayDeviceUtils.cpp
+++ b/Engine/Plugins/Runtime/nDisplay/Source/DisplayCluster/Private/Game/EngineClasses/Scene/DisplayDevice/DisplayClusterDisplayDeviceUtils.cpp
@@ -10,7 +10,7 @@
 #if WITH_EDITOR
 
 UDisplayClusterDisplayDeviceBaseComponent* UE::DisplayClusterDisplayDeviceUtils::FindAndSyncDisplayDeviceFromViewport(
-	UDisplayClusterConfigurationViewport* InViewport)
+	UDisplayClusterConfigurationViewport* InViewport, bool* bOutIsDefaultDevice)
 {
 	if (IsValid(InViewport))
 	{
@@ -19,6 +19,10 @@ UDisplayClusterDisplayDeviceBaseComponent* UE::DisplayClusterDisplayDeviceUtils:
 			UDisplayClusterDisplayDeviceBaseComponent* DeviceBaseComponent = nullptr;
 
 			const bool bUseDefaultDevice = InViewport->DisplayDeviceName.IsEmpty();
+			if (bOutIsDefaultDevice)
+			{
+				*bOutIsDefaultDevice = bUseDefaultDevice;
+			}
 			if (bUseDefaultDevice)
 			{
 				DeviceBaseComponent = RootActor->GetDefaultDisplayDevice();
@@ -68,7 +72,7 @@ UDisplayClusterDisplayDeviceBaseComponent* UE::DisplayClusterDisplayDeviceUtils:
 											InViewport->Modify();
 											InViewport->DisplayDeviceName = ViewportCDO->DisplayDeviceName;
 									
-											return FindAndSyncDisplayDeviceFromViewport(InViewport);
+											return FindAndSyncDisplayDeviceFromViewport(InViewport, bOutIsDefaultDevice);
 										}
 									}
 								}
diff --git a/Engine/Plugins/Runtime/nDisplay/Source/DisplayCluster/Private/Game/EngineClasses/Scene/DisplayDevice/DisplayClusterDisplayDeviceUtils.h b/Engine/Plugins/Runtime/nDisplay/Source/DisplayCluster/Private/Game/EngineClasses/Scene/DisplayDevice/DisplayClusterDisplayDeviceUtils.h
--- a/Engine/Plugins/Runtime/nDisplay/Source/DisplayCluster/Private/Game/EngineClasses/Scene/DisplayDevice/DisplayClusterDisplayDeviceUtils.h
+++ b/Engine/Plugins/Runtime/nDisplay/Source/DisplayCluster/Private/Game/EngineClasses/Scene/DisplayDevice/DisplayClusterDisplayDeviceUtils.h
@@ -14,8 +14,9 @@ namespace UE::DisplayClusterDisplayDeviceUtils
 	 * viewport on the instance will be synced to the new name.
 	 *
 	 * @param InViewport The viewport object which points to a display device and belongs to a root actor.
+	 * @param bOutIsDefaultDevice If the display device returned is the default display device.
 	 * @return The display device this viewport uses or nullptr.
 	 */
-	UDisplayClusterDisplayDeviceBaseComponent* FindAndSyncDisplayDeviceFromViewport(UDisplayClusterConfigurationViewport* InViewport);
+	UDisplayClusterDisplayDeviceBaseComponent* FindAndSyncDisplayDeviceFromViewport(UDisplayClusterConfigurationViewport* InViewport, bool* bOutIsDefaultDevice = nullptr);
 #endif
 }
diff --git a/Engine/Plugins/Runtime/nDisplay/Source/DisplayCluster/Public/Components/DisplayClusterPreviewComponent.h b/Engine/Plugins/Runtime/nDisplay/Source/DisplayCluster/Public/Components/DisplayClusterPreviewComponent.h
--- a/Engine/Plugins/Runtime/nDisplay/Source/DisplayCluster/Public/Components/DisplayClusterPreviewComponent.h
+++ b/Engine/Plugins/Runtime/nDisplay/Source/DisplayCluster/Public/Components/DisplayClusterPreviewComponent.h
@@ -145,6 +145,17 @@ protected:
 	TObjectPtr<UTextureRenderTarget2D> RenderTargetPostProcess;
 	
 private:
+	/**
+	 * The chosen preview render target to use for the preview material. This could be standard, post processed, or
+	 * processed by the display device
+	 */
+	UPROPERTY(Transient)
+	TObjectPtr<UTextureRenderTarget2D> CurrentPreviewRenderTargetPtr;
+
+	/** The preview render target processed by the display device */
+	UPROPERTY(Transient)
+	TObjectPtr<UTextureRenderTarget2D> DisplayDeviceRenderTarget;
+	
 	// Saved mesh policy params
 	UPROPERTY(Transient)
 	FDisplayClusterConfigurationProjection WarpMeshSavedProjectionPolicy;
@@ -162,6 +173,10 @@ private:
 	UPROPERTY(Transient)
 	mutable FComponentReference CachedDisplayDevice;
 
+	/** If the display device being used is default */
+	UPROPERTY(Transient)
+	mutable bool bIsUsingDefaultDisplayDevice = false;
+
 	/** If this preview component should utilize the display device component */
 	bool bUseDisplayDevice = true;
 
diff --git a/Engine/Plugins/Runtime/nDisplay/Source/DisplayCluster/Public/Components/DisplayDevice/DisplayClusterDisplayDeviceBaseComponent.h b/Engine/Plugins/Runtime/nDisplay/Source/DisplayCluster/Public/Components/DisplayDevice/DisplayClusterDisplayDeviceBaseComponent.h
--- a/Engine/Plugins/Runtime/nDisplay/Source/DisplayCluster/Public/Components/DisplayDevice/DisplayClusterDisplayDeviceBaseComponent.h
+++ b/Engine/Plugins/Runtime/nDisplay/Source/DisplayCluster/Public/Components/DisplayDevice/DisplayClusterDisplayDeviceBaseComponent.h
@@ -8,7 +8,13 @@
 
 class UMaterial;
 class UMaterialInstanceDynamic;
+class UTexture;
+class UTextureRenderTarget2D;
 
+/**
+ * Display Device Components can be added to nDisplay root actors and assigned to viewport nodes to allow additional
+ * processing on the preview material.
+ */
 UCLASS(Abstract, ClassGroup = (DisplayCluster))
 class DISPLAYCLUSTER_API UDisplayClusterDisplayDeviceBaseComponent
 	: public USceneComponent
@@ -17,20 +23,34 @@ class DISPLAYCLUSTER_API UDisplayClusterDisplayDeviceBaseComponent
 public:
 	UDisplayClusterDisplayDeviceBaseComponent();
 
+#if WITH_EDITOR
+	virtual void PostEditChangeProperty(FPropertyChangedEvent& PropertyChangedEvent) override;
+#endif
+	
 	/** Retrieve the preview material used. */
 	TObjectPtr<UMaterial> GetPreviewMaterial() const { return PreviewMaterial; }
 
 	/** Retrieve the base mesh material. */
 	TObjectPtr<UMaterial> GetMeshMaterial() const { return MeshMaterial; }
 
+	/** If the display device is configured for additional render passes. */
+	bool IsRenderPassEnabled() const { return bEnableRenderPass; }
+	
 	/** Perform any operations on the preview material instance, such as setting parameter values. */
 	virtual void OnUpdatePreviewMaterialInstance(UMaterialInstanceDynamic* InMaterialInstance) {}
-
-#if WITH_EDITOR
-	virtual void PostEditChangeProperty(FPropertyChangedEvent& PropertyChangedEvent) override;
-#endif
+	
+	/**
+	 * Start a render pass from the game thread.
+	 * @param InSourceTexture The source texture which should be used as input to any render passes.
+	 * @param InRenderTarget The destination render target which should be used as output for render passes.
+	 */
+	virtual void RenderPass_GameThread(UTexture* InSourceTexture, UTextureRenderTarget2D* InRenderTarget) {}
 	
 protected:
+	/** If render passes are enabled. */
+	UPROPERTY(EditAnywhere, Category=RenderPass)
+	bool bEnableRenderPass = false;
+	
 	/** The material the preview components will create a material instance from when rendering the nDisplay preview. */
 	UPROPERTY(EditAnywhere, Category=Material, NoClear)
 	TObjectPtr<UMaterial> PreviewMaterial = nullptr;
@@ -39,20 +59,3 @@ protected:
 	UPROPERTY(EditAnywhere, Category=Material, NoClear)
 	TObjectPtr<UMaterial> MeshMaterial = nullptr;
 };
-
-UCLASS(ClassGroup = (DisplayCluster), meta = (BlueprintSpawnableComponent, DisplayName = "NDisplay Display Device"))
-class DISPLAYCLUSTER_API UDisplayClusterDisplayDeviceComponent
-	: public UDisplayClusterDisplayDeviceBaseComponent
-{
-	GENERATED_BODY()
-
-public:
-	UDisplayClusterDisplayDeviceComponent();
-	
-	virtual void OnUpdatePreviewMaterialInstance(UMaterialInstanceDynamic* InMaterialInstance) override;
-	
-protected:
-	/** Adjust the exposure for the emissive input. */
-	UPROPERTY(EditAnywhere, Category=Material)
-	float Exposure = 0.f;
-};
diff --git a/Engine/Plugins/Runtime/nDisplay/Source/DisplayCluster/Public/Components/DisplayDevice/DisplayClusterDisplayDeviceComponent.h b/Engine/Plugins/Runtime/nDisplay/Source/DisplayCluster/Public/Components/DisplayDevice/DisplayClusterDisplayDeviceComponent.h
new file mode 100644
--- /dev/null
+++ b/Engine/Plugins/Runtime/nDisplay/Source/DisplayCluster/Public/Components/DisplayDevice/DisplayClusterDisplayDeviceComponent.h
@@ -0,0 +1,37 @@
+// Copyright Epic Games, Inc. All Rights Reserved.
+
+#pragma once
+
+#include "DisplayClusterDisplayDeviceBaseComponent.h"
+#include "OpenColorIOColorSpace.h"
+
+#include "DisplayClusterDisplayDeviceComponent.generated.h"
+
+class UMaterial;
+class UMaterialInstanceDynamic;
+class UTexture;
+class UTextureRenderTarget2D;
+
+UCLASS(ClassGroup = (DisplayCluster), meta = (BlueprintSpawnableComponent, DisplayName = "NDisplay Display Device"))
+class DISPLAYCLUSTER_API UDisplayClusterDisplayDeviceComponent
+	: public UDisplayClusterDisplayDeviceBaseComponent
+{
+	GENERATED_BODY()
+
+public:
+	UDisplayClusterDisplayDeviceComponent();
+
+	const FOpenColorIOColorConversionSettings& GetColorConversionSettings() const { return ColorConversionSettings; }
+	
+	virtual void OnUpdatePreviewMaterialInstance(UMaterialInstanceDynamic* InMaterialInstance) override;
+	virtual void RenderPass_GameThread(UTexture* InSourceTexture, UTextureRenderTarget2D* InRenderTarget) override;
+
+protected:
+	/** Adjust the exposure for the emissive input. */
+	UPROPERTY(EditAnywhere, Category=Material)
+	float Exposure = 0.f;
+	
+	/** Color grading settings. */
+	UPROPERTY(EditAnywhere, Category = RenderPass, meta = (DisplayAfter="bEnableRenderPass"))
+	FOpenColorIOColorConversionSettings ColorConversionSettings;
+};
