nDisplay: Exposed the "Exposure compensation" property on post process volumes to the ICVFX color grading drawer's property panel

diff --git a/Engine/Plugins/Runtime/nDisplay/Source/DisplayClusterColorGrading/Private/DataModelGenerators/DisplayClusterColorGradingGenerator_PostProcessVolume.cpp b/Engine/Plugins/Runtime/nDisplay/Source/DisplayClusterColorGrading/Private/DataModelGenerators/DisplayClusterColorGradingGenerator_PostProcessVolume.cpp
--- a/Engine/Plugins/Runtime/nDisplay/Source/DisplayClusterColorGrading/Private/DataModelGenerators/DisplayClusterColorGradingGenerator_PostProcessVolume.cpp
+++ b/Engine/Plugins/Runtime/nDisplay/Source/DisplayClusterColorGrading/Private/DataModelGenerators/DisplayClusterColorGradingGenerator_PostProcessVolume.cpp
@@ -64,6 +64,12 @@ public:
 					GroupCategory.AddProperty(ChildHandle);
 				}
 			}
+			else if (ChildHandle->GetProperty()->GetNameCPP() == GET_MEMBER_NAME_STRING_CHECKED(FPostProcessSettings, AutoExposureBias))
+			{
+				// UE-192122: Also add Exposure compensation property to details view, to match DCRA properties
+				IDetailCategoryBuilder& GroupCategory = DetailBuilder.EditCategory(FName("DetailView_" + GroupName), FText::FromString(GroupName), ECategoryPriority::Important);
+				GroupCategory.AddProperty(ChildHandle);
+			}
 		}
 	}
 };
