nDisplay: Added check for IsAlwaysFlushToWall when moving light cards in ICVFX panel, so that they will not be set flush to wall when Flush to Wall is disabled in their properties.

diff --git a/Engine/Plugins/Runtime/nDisplay/Source/DisplayClusterScenePreview/Private/DisplayClusterLightCardEditorHelper.cpp b/Engine/Plugins/Runtime/nDisplay/Source/DisplayClusterScenePreview/Private/DisplayClusterLightCardEditorHelper.cpp
--- a/Engine/Plugins/Runtime/nDisplay/Source/DisplayClusterScenePreview/Private/DisplayClusterLightCardEditorHelper.cpp
+++ b/Engine/Plugins/Runtime/nDisplay/Source/DisplayClusterScenePreview/Private/DisplayClusterLightCardEditorHelper.cpp
@@ -1196,8 +1196,11 @@ void FDisplayClusterLightCardEditorHelper::InternalMoveActorTo(
 	SetActorCoordinates(Actor, FSphericalCoordinates(InverseRotatedPosition));
 	Actor->UpdateStageActorTransform();
 
-	ADisplayClusterRootActor* RootActor = (Actor->IsProxy() || !LevelInstanceRootActor.IsValid()) ? CachedRootActor.Get() : LevelInstanceRootActor.Get();
-	RootActor->MakeStageActorFlushToWall(Actor.AsActor());
+	if (Actor->IsAlwaysFlushToWall())
+	{
+		ADisplayClusterRootActor* RootActor = (Actor->IsProxy() || !LevelInstanceRootActor.IsValid()) ? CachedRootActor.Get() : LevelInstanceRootActor.Get();
+		RootActor->MakeStageActorFlushToWall(Actor.AsActor());
+	}
 
 #if WITH_EDITOR
 	if (!Actor->IsProxy())
diff --git a/Engine/Plugins/Runtime/nDisplayModularFeatures/Source/DisplayClusterLightCardExtender/Public/StageActor/IDisplayClusterStageActor.h b/Engine/Plugins/Runtime/nDisplayModularFeatures/Source/DisplayClusterLightCardExtender/Public/StageActor/IDisplayClusterStageActor.h
--- a/Engine/Plugins/Runtime/nDisplayModularFeatures/Source/DisplayClusterLightCardExtender/Public/StageActor/IDisplayClusterStageActor.h
+++ b/Engine/Plugins/Runtime/nDisplayModularFeatures/Source/DisplayClusterLightCardExtender/Public/StageActor/IDisplayClusterStageActor.h
@@ -57,7 +57,7 @@ public:
 	virtual FVector2D GetScale() const = 0;
 
 	virtual void SetAlwaysFlushToWall(bool bInAlwaysFlushToWall) { }
-	virtual bool IsAlwaysFlushToWall() const { return false; }
+	virtual bool IsAlwaysFlushToWall() const { return true; }
 
 
 	/** Return property names defined for this stage actor for use in property notifies */
