nDisplay: Speculative crash fix when opening Virtual_Camera level on QAVP on Linux.

diff --git a/Engine/Plugins/Runtime/nDisplay/Source/DisplayClusterConfigurator/Private/DisplayClusterConfiguratorModule.cpp b/Engine/Plugins/Runtime/nDisplay/Source/DisplayClusterConfigurator/Private/DisplayClusterConfiguratorModule.cpp
--- a/Engine/Plugins/Runtime/nDisplay/Source/DisplayClusterConfigurator/Private/DisplayClusterConfiguratorModule.cpp
+++ b/Engine/Plugins/Runtime/nDisplay/Source/DisplayClusterConfigurator/Private/DisplayClusterConfiguratorModule.cpp
@@ -393,12 +393,15 @@ EVisibility FDisplayClusterConfiguratorModule::GetViewportsFrozenWarningVisibili
 {
 	if (GEditor)
 	{
-		for (TActorIterator<ADisplayClusterRootActor> It(GEditor->GetEditorWorldContext().World()); It; ++It)
+		if (const UWorld* World = GEditor->GetEditorWorldContext().World())
 		{
-			const UDisplayClusterConfigurationData* ConfigData = It->GetConfigData();
-			if (ConfigData && ConfigData->StageSettings.bFreezeRenderOuterViewports)
+			for (TActorIterator<ADisplayClusterRootActor> It(World); It; ++It)
 			{
-				return EVisibility::Visible;
+				const UDisplayClusterConfigurationData* ConfigData = It->GetConfigData();
+				if (ConfigData && ConfigData->StageSettings.bFreezeRenderOuterViewports)
+				{
+					return EVisibility::Visible;
+				}
 			}
 		}
 	}
