nDisplay: Added frustum fitting warp policy and in frustum fit view origin component, which will fit a stage's geometry into either the active engine camera or an external cine camera if specified by fitting the projection of the stage geometry onto the camera's projection. In addition, users can preview the frustum fit in editor, which will display a preview version of the stage geometry, scaled appropriately, within the camera's frustum in the level

diff --git a/Engine/Plugins/Runtime/nDisplay/Source/DisplayCluster/Private/Game/EngineClasses/Scene/DisplayClusterPreviewComponent.cpp b/Engine/Plugins/Runtime/nDisplay/Source/DisplayCluster/Private/Game/EngineClasses/Scene/DisplayClusterPreviewComponent.cpp
--- a/Engine/Plugins/Runtime/nDisplay/Source/DisplayCluster/Private/Game/EngineClasses/Scene/DisplayClusterPreviewComponent.cpp
+++ b/Engine/Plugins/Runtime/nDisplay/Source/DisplayCluster/Private/Game/EngineClasses/Scene/DisplayClusterPreviewComponent.cpp
@@ -42,6 +42,10 @@ UDisplayClusterPreviewComponent::UDisplayClusterPreviewComponent(const FObjectIn
 	: Super(ObjectInitializer)
 {
 #if WITH_EDITOR
+	static ConstructorHelpers::FObjectFinder<UMaterial> PreviewMaterialObj(TEXT("/nDisplay/Materials/Preview/M_ProjPolicyPreview"));
+	check(PreviewMaterialObj.Object);
+	MovablePreviewMaterial = PreviewMaterialObj.Object;
+
 	bWantsInitializeComponent = true;
 #endif
 }
@@ -171,6 +175,12 @@ void UDisplayClusterPreviewComponent::SetPreviewMeshMaterial(UMaterial* InMateri
 		{
 			PreviewMaterialInstance = UMaterialInstanceDynamic::Create(InMaterial, this);
 		}
+
+		if (MovablePreviewMesh && MovablePreviewMaterial && MovablePreviewMaterialInstance == nullptr)
+		{
+			MovablePreviewMaterialInstance = UMaterialInstanceDynamic::Create(MovablePreviewMaterial, this);
+		}
+
 		UpdatePreviewMaterial();
 
 		// Set preview material
@@ -178,6 +188,14 @@ void UDisplayClusterPreviewComponent::SetPreviewMeshMaterial(UMaterial* InMateri
 		{
 			PreviewMesh->SetMaterial(0, PreviewMaterialInstance);
 		}
+
+		if (MovablePreviewMaterialInstance)
+		{
+			if (MovablePreviewMesh)
+			{
+				MovablePreviewMesh->SetMaterial(0, MovablePreviewMaterialInstance);
+			}
+		}
 	}
 }
 
@@ -195,7 +213,8 @@ UMaterial* UDisplayClusterPreviewComponent::GetMeshMaterialFromDisplayDevice() c
 
 void UDisplayClusterPreviewComponent::UpdatePreviewMeshReference()
 {
-	if (PreviewMesh && PreviewMesh->GetName().Find(TEXT("TRASH_")) != INDEX_NONE)
+	if ((PreviewMesh && PreviewMesh->GetName().Find(TEXT("TRASH_")) != INDEX_NONE) ||
+		(MovablePreviewMesh && MovablePreviewMesh->GetName().Find(TEXT("TRASH_")) != INDEX_NONE))
 	{
 		// Screen components are regenerated from construction scripts, but preview components are added in dynamically. This preview component may end up
 		// pointing to invalid data on reconstruction.
@@ -246,6 +265,18 @@ bool UDisplayClusterPreviewComponent::UpdatePreviewMesh()
 					// Get new mesh ptr
 					PreviewMesh = Viewport->GetProjectionPolicy()->GetOrCreatePreviewMeshComponent(Viewport, bIsRootActorPreviewMesh);
 
+					if (Viewport->GetProjectionPolicy()->HasPreviewMovableMesh())
+					{
+						// Get new movable mesh ptr
+						MovablePreviewMesh = Viewport->GetProjectionPolicy()->GetOrCreatePreviewMovableMeshComponent(Viewport);
+
+						if (MovablePreviewMesh)
+						{
+							// Make the movable preview mesh invisible by default
+							MovablePreviewMesh->SetVisibility(false);
+						}
+					}
+
 					// Update saved proj policy parameters
 					WarpMeshSavedProjectionPolicy = ViewportConfig->ProjectionPolicy;
 				}
@@ -256,6 +287,11 @@ bool UDisplayClusterPreviewComponent::UpdatePreviewMesh()
 					PreviewMesh->SetCastShadow(false);
 				}
 
+				if (MovablePreviewMesh != nullptr)
+				{
+					MovablePreviewMesh->SetCastShadow(false);
+				}
+
 				UMaterial* PreviewMaterial = GetPreviewMaterialFromDisplayDevice();
 				if (CurrentMeshMaterial != PreviewMaterial)
 				{
@@ -291,6 +327,7 @@ void UDisplayClusterPreviewComponent::ReleasePreviewMesh()
 
 	// Forget old mesh with material
 	PreviewMesh = nullptr;
+	MovablePreviewMesh = nullptr;
 	CurrentMeshMaterial = nullptr;
 }
 
@@ -326,6 +363,21 @@ void UDisplayClusterPreviewComponent::UpdatePreviewMaterial()
 			CachedComponent->OnUpdatePreviewMaterialInstance(PreviewMaterialInstance);
 		}
 	}
+
+	if (MovablePreviewMaterialInstance)
+	{
+		MovablePreviewMaterialInstance->SetScalarParameterValue(TEXT("Opacity"), 1.0);
+
+		if (OverrideTexture)
+		{
+			MovablePreviewMaterialInstance->SetTextureParameterValue(TEXT("Preview"), OverrideTexture);
+		}
+		else
+		{
+			MovablePreviewMaterialInstance->SetTextureParameterValue(TEXT("Preview"),
+				RootActor && RootActor->bPreviewEnablePostProcess ? RenderTargetPostProcess : RenderTarget);
+		}
+	}
 }
 
 void UDisplayClusterPreviewComponent::ReleasePreviewMaterial()
@@ -341,6 +393,16 @@ void UDisplayClusterPreviewComponent::ReleasePreviewMaterial()
 
 		PreviewMaterialInstance = nullptr;
 	}
+
+	if (MovablePreviewMaterialInstance != nullptr)
+	{
+		if (!MovablePreviewMaterialInstance->HasAnyFlags(RF_BeginDestroyed))
+		{
+			MovablePreviewMaterialInstance->ClearParameterValues();
+		}
+
+		MovablePreviewMaterialInstance = nullptr;
+	}
 }
 
 void UDisplayClusterPreviewComponent::ReleasePreviewRenderTarget()
diff --git a/Engine/Plugins/Runtime/nDisplay/Source/DisplayCluster/Private/Render/Viewport/DisplayClusterViewportManager.cpp b/Engine/Plugins/Runtime/nDisplay/Source/DisplayCluster/Private/Render/Viewport/DisplayClusterViewportManager.cpp
--- a/Engine/Plugins/Runtime/nDisplay/Source/DisplayCluster/Private/Render/Viewport/DisplayClusterViewportManager.cpp
+++ b/Engine/Plugins/Runtime/nDisplay/Source/DisplayCluster/Private/Render/Viewport/DisplayClusterViewportManager.cpp
@@ -1183,6 +1183,21 @@ TSharedPtr<IDisplayClusterProjectionPolicy, ESPMode::ThreadSafe> FDisplayCluster
 	return nullptr;
 }
 
+TArray<TSharedPtr<IDisplayClusterViewport, ESPMode::ThreadSafe>> FDisplayClusterViewportManager::GetViewportsForWarpPolicy(const TSharedPtr<IDisplayClusterWarpPolicy>& InWarpPolicy) const
+{
+	TArray<TSharedPtr<IDisplayClusterViewport, ESPMode::ThreadSafe>> WarpPolicyViewports;
+
+	for (const TSharedPtr<FDisplayClusterViewport, ESPMode::ThreadSafe>& Viewport : ImplGetEntireClusterViewports())
+	{
+		if (Viewport.IsValid() && Viewport->GetProjectionPolicy().IsValid() && Viewport->GetProjectionPolicy()->GetWarpPolicy() == InWarpPolicy.Get())
+		{
+			WarpPolicyViewports.Add(Viewport);
+		}
+	}
+
+	return WarpPolicyViewports;
+}
+
 void FDisplayClusterViewportManager::MarkComponentGeometryDirty(const FName InComponentName)
 {
 	check(IsInGameThread());
diff --git a/Engine/Plugins/Runtime/nDisplay/Source/DisplayCluster/Private/Render/Viewport/DisplayClusterViewportManager.h b/Engine/Plugins/Runtime/nDisplay/Source/DisplayCluster/Private/Render/Viewport/DisplayClusterViewportManager.h
--- a/Engine/Plugins/Runtime/nDisplay/Source/DisplayCluster/Private/Render/Viewport/DisplayClusterViewportManager.h
+++ b/Engine/Plugins/Runtime/nDisplay/Source/DisplayCluster/Private/Render/Viewport/DisplayClusterViewportManager.h
@@ -105,6 +105,8 @@ public:
 		return TArrayView<TSharedPtr<IDisplayClusterViewport, ESPMode::ThreadSafe>>((TSharedPtr<IDisplayClusterViewport, ESPMode::ThreadSafe>*)(EntireClusterViewports.GetData()), EntireClusterViewports.Num());
 	}
 
+	virtual TArray<TSharedPtr<IDisplayClusterViewport, ESPMode::ThreadSafe>> GetViewportsForWarpPolicy(const TSharedPtr<IDisplayClusterWarpPolicy>& InWarpPolicy) const override;
+
 	virtual void MarkComponentGeometryDirty(const FName InComponentName = NAME_None) override;
 	virtual void AddReferencedObjects(FReferenceCollector& Collector) override;
 
diff --git a/Engine/Plugins/Runtime/nDisplay/Source/DisplayCluster/Public/Components/DisplayClusterPreviewComponent.h b/Engine/Plugins/Runtime/nDisplay/Source/DisplayCluster/Public/Components/DisplayClusterPreviewComponent.h
--- a/Engine/Plugins/Runtime/nDisplay/Source/DisplayCluster/Public/Components/DisplayClusterPreviewComponent.h
+++ b/Engine/Plugins/Runtime/nDisplay/Source/DisplayCluster/Public/Components/DisplayClusterPreviewComponent.h
@@ -77,6 +77,12 @@ public:
 		return PreviewMesh;
 	}
 
+	UMeshComponent* GetMovablePreviewMesh()
+	{
+		UpdatePreviewMeshReference();
+		return MovablePreviewMesh;
+	}
+
 	UTexture* GetOverrideTexture() const
 	{
 		return OverrideTexture;
@@ -189,13 +195,22 @@ private:
 	/** The current material assigned to the preview mesh */
 	UPROPERTY(Transient)
 	TObjectPtr<UMaterial> CurrentMeshMaterial = nullptr;
-	
+
+	UPROPERTY(Transient)
+	TObjectPtr<UMaterial> MovablePreviewMaterial = nullptr;
+
 	UPROPERTY(Transient)
 	bool bIsRootActorPreviewMesh = false;
 
+	UPROPERTY(Transient)
+	TObjectPtr<UMeshComponent> MovablePreviewMesh = nullptr;
+
 	UPROPERTY(Transient)
 	TObjectPtr<UMaterialInstanceDynamic> PreviewMaterialInstance = nullptr;
 
+	UPROPERTY(Transient)
+	TObjectPtr<UMaterialInstanceDynamic> MovablePreviewMaterialInstance = nullptr;
+
 	UPROPERTY(Transient)
 	TObjectPtr<UTexture> OverrideTexture = nullptr;
 
diff --git a/Engine/Plugins/Runtime/nDisplay/Source/DisplayCluster/Public/Render/Viewport/IDisplayClusterViewportManager.h b/Engine/Plugins/Runtime/nDisplay/Source/DisplayCluster/Public/Render/Viewport/IDisplayClusterViewportManager.h
--- a/Engine/Plugins/Runtime/nDisplay/Source/DisplayCluster/Public/Render/Viewport/IDisplayClusterViewportManager.h
+++ b/Engine/Plugins/Runtime/nDisplay/Source/DisplayCluster/Public/Render/Viewport/IDisplayClusterViewportManager.h
@@ -10,6 +10,7 @@
 
 #include "SceneView.h"
 
+class IDisplayClusterWarpPolicy;
 class UWorld;
 class FViewport;
 class FSceneViewFamilyContext;
@@ -196,6 +197,15 @@ public:
 	*/
 	virtual const TArrayView<TSharedPtr<IDisplayClusterViewport, ESPMode::ThreadSafe>> GetEntireClusterViewports() const = 0;
 
+	/**
+	 * Return all viewports associated with the specified warp policy
+	 * [Game thread func]
+	 *
+	 * @param InWarpPolicy - The warp policy to get viewports for
+	 * @return - A list of viewports associated with the warp policy
+	 */
+	virtual TArray<TSharedPtr<IDisplayClusterViewport, ESPMode::ThreadSafe>> GetViewportsForWarpPolicy(const TSharedPtr<IDisplayClusterWarpPolicy>& InWarpPolicy) const = 0;
+
 	/**
 	* Mark the geometry of the referenced component(s) as dirty (ProceduralMesh, etc)
 	* [Game thread func]
diff --git a/Engine/Plugins/Runtime/nDisplay/Source/DisplayClusterConfigurator/Private/Views/SCSEditor/SDisplayClusterConfiguratorComponentCombo.cpp b/Engine/Plugins/Runtime/nDisplay/Source/DisplayClusterConfigurator/Private/Views/SCSEditor/SDisplayClusterConfiguratorComponentCombo.cpp
--- a/Engine/Plugins/Runtime/nDisplay/Source/DisplayClusterConfigurator/Private/Views/SCSEditor/SDisplayClusterConfiguratorComponentCombo.cpp
+++ b/Engine/Plugins/Runtime/nDisplay/Source/DisplayClusterConfigurator/Private/Views/SCSEditor/SDisplayClusterConfiguratorComponentCombo.cpp
@@ -7,6 +7,7 @@
 #include "Components/DisplayClusterScreenComponent.h"
 #include "Components/DisplayClusterXformComponent.h"
 #include "Components/DisplayDevice/DisplayClusterDisplayDeviceBaseComponent.h"
+#include "Components/DisplayClusterInFrustumFitCameraComponent.h"
 
 #include "Widgets/Layout/SSpacer.h"
 #include "Widgets/Images/SImage.h"
@@ -456,6 +457,7 @@ void SDisplayClusterConfiguratorComponentClassCombo::GenerateComponentClassList(
 	AddDCComp(UDisplayClusterXformComponent::StaticClass());
 	AddDCComp(UDisplayClusterScreenComponent::StaticClass());
 	AddDCComp(UDisplayClusterCameraComponent::StaticClass());
+	AddDCComp(UDisplayClusterInFrustumFitCameraComponent::StaticClass());
 
 	AddHeader("nDisplay ICVFX Components");
 	AddDCComp(UDisplayClusterICVFXCameraComponent::StaticClass());
diff --git a/Engine/Plugins/Runtime/nDisplay/Source/DisplayClusterProjection/Private/Policy/MPCDI/DisplayClusterProjectionMPCDIPolicy.cpp b/Engine/Plugins/Runtime/nDisplay/Source/DisplayClusterProjection/Private/Policy/MPCDI/DisplayClusterProjectionMPCDIPolicy.cpp
--- a/Engine/Plugins/Runtime/nDisplay/Source/DisplayClusterProjection/Private/Policy/MPCDI/DisplayClusterProjectionMPCDIPolicy.cpp
+++ b/Engine/Plugins/Runtime/nDisplay/Source/DisplayClusterProjection/Private/Policy/MPCDI/DisplayClusterProjectionMPCDIPolicy.cpp
@@ -449,6 +449,42 @@ UMeshComponent* FDisplayClusterProjectionMPCDIPolicy::GetOrCreatePreviewMeshComp
 	return (bIsPreviewMeshEnabled && WarpBlendInterface.IsValid()) ? WarpBlendInterface->GetOrCreateMeshComponent(InViewport, bOutIsRootActorComponent) : nullptr;
 }
 
+UMeshComponent* FDisplayClusterProjectionMPCDIPolicy::GetOrCreatePreviewMovableMeshComponent(IDisplayClusterViewport* InViewport)
+{
+#if WITH_EDITOR
+	// If we have already created a preview mesh component before, return that component
+	if (UMeshComponent* MovableMeshComp = Cast<UMeshComponent>(MovablePreviewMeshComponentRef.GetOrFindSceneComponent()))
+	{
+		return MovableMeshComp;
+	}
+	// Create a new mesh component. Best option is to simply create a duplicate of the static preview component by passing a template into NewObject
+	// Attach the movable mesh component to the DCRA's root component, since we don't care about hierarchy here
+	if (ADisplayClusterRootActor* RootActor = InViewport->GetRootActor())
+	{
+		USceneComponent* RootComponent = RootActor->GetRootComponent();
+		bool bIsRootActorComponent = false;
+		if (UMeshComponent* PreviewMeshComponent = GetOrCreatePreviewMeshComponent(InViewport, bIsRootActorComponent))
+		{
+			const FString CompName = FString::Printf(TEXT("MPCDI_%s_movable_impl"), *GetId());
+			const EObjectFlags ObjectFlags = EObjectFlags::RF_DuplicateTransient | RF_Transient | RF_TextExportTransient;
+			UMeshComponent* NewMovableMeshComp = NewObject<UMeshComponent>(RootComponent, PreviewMeshComponent->GetClass(), *CompName, ObjectFlags, PreviewMeshComponent);
+			if (NewMovableMeshComp)
+			{
+				NewMovableMeshComp->RegisterComponent();
+				NewMovableMeshComp->AttachToComponent(RootComponent, FAttachmentTransformRules(EAttachmentRule::KeepRelative, false));
+				NewMovableMeshComp->SetIsVisualizationComponent(true);
+				NewMovableMeshComp->SetHiddenInGame(true);
+				// Store reference to mesh component
+				MovablePreviewMeshComponentRef.SetSceneComponent(NewMovableMeshComp);
+				return NewMovableMeshComp;
+			}
+		}
+	}
+#endif
+
+	return nullptr;
+}
+
 bool FDisplayClusterProjectionMPCDIPolicy::CreateWarpBlendFromConfig(IDisplayClusterViewport* InViewport)
 {
 	check(IsInGameThread());
diff --git a/Engine/Plugins/Runtime/nDisplay/Source/DisplayClusterProjection/Private/Policy/MPCDI/DisplayClusterProjectionMPCDIPolicy.h b/Engine/Plugins/Runtime/nDisplay/Source/DisplayClusterProjection/Private/Policy/MPCDI/DisplayClusterProjectionMPCDIPolicy.h
--- a/Engine/Plugins/Runtime/nDisplay/Source/DisplayClusterProjection/Private/Policy/MPCDI/DisplayClusterProjectionMPCDIPolicy.h
+++ b/Engine/Plugins/Runtime/nDisplay/Source/DisplayClusterProjection/Private/Policy/MPCDI/DisplayClusterProjectionMPCDIPolicy.h
@@ -84,5 +84,11 @@ protected:
 		return true;
 	}
 
+	virtual bool HasPreviewMovableMesh() override { return true; }
+
 	virtual class UMeshComponent* GetOrCreatePreviewMeshComponent(IDisplayClusterViewport* InViewport, bool& bOutIsRootActorComponent) override;
+	virtual class UMeshComponent* GetOrCreatePreviewMovableMeshComponent(IDisplayClusterViewport* InViewport) override;
+
+private:
+	FDisplayClusterSceneComponentRef MovablePreviewMeshComponentRef;
 };
diff --git a/Engine/Plugins/Runtime/nDisplay/Source/DisplayClusterWarp/Private/Components/DisplayClusterInFrustumFitCameraComponent.cpp b/Engine/Plugins/Runtime/nDisplay/Source/DisplayClusterWarp/Private/Components/DisplayClusterInFrustumFitCameraComponent.cpp
new file mode 100644
--- /dev/null
+++ b/Engine/Plugins/Runtime/nDisplay/Source/DisplayClusterWarp/Private/Components/DisplayClusterInFrustumFitCameraComponent.cpp
@@ -0,0 +1,219 @@
+// Copyright Epic Games, Inc. All Rights Reserved.
+
+#include "Components/DisplayClusterInFrustumFitCameraComponent.h"
+#include "Components/BillboardComponent.h"
+
+#include "CineCameraActor.h"
+#include "Camera/CameraComponent.h"
+#include "CineCameraComponent.h"
+
+#include "DisplayClusterWarpLog.h"
+#include "DisplayClusterWarpStrings.h"
+
+#include "IDisplayCluster.h"
+#include "IDisplayClusterWarpBlend.h"
+
+#include "Render/IDisplayClusterRenderManager.h"
+
+#include "Render/Viewport/IDisplayClusterViewportManager.h"
+#include "Render/Viewport/IDisplayClusterViewport.h"
+#include "Render/Projection/IDisplayClusterProjectionPolicy.h"
+#include "Render/Warp/IDisplayClusterWarpPolicyFactory.h"
+#include "Render/Warp/IDisplayClusterWarpPolicy.h"
+
+#include "DisplayClusterRootActor.h"
+#include "Components/DisplayClusterPreviewComponent.h"
+
+namespace UE::DisplayClusterWarp::ViewPointComponent
+{
+	static inline TSharedPtr<IDisplayClusterWarpPolicy, ESPMode::ThreadSafe> ImplCreateWarpPolicy(const FString& InWarpPolicyType, const FString& InWarpPolicyName)
+	{
+		static IDisplayCluster& DisplayClusterSingleton = IDisplayCluster::Get();
+		if (IDisplayClusterRenderManager* const DCRenderManager = DisplayClusterSingleton.GetRenderMgr())
+		{
+			TSharedPtr<IDisplayClusterWarpPolicyFactory> WarpPolicyFactory = DCRenderManager->GetWarpPolicyFactory(InWarpPolicyType);
+			if (WarpPolicyFactory.IsValid())
+			{
+				return WarpPolicyFactory->Create(InWarpPolicyType, InWarpPolicyName);
+			}
+		}
+
+		return nullptr;
+	}
+};
+using namespace UE::DisplayClusterWarp::ViewPointComponent;
+
+//--------------------------------------------------------------------------------
+// UDisplayClusterInFrustumFitCameraComponent
+//--------------------------------------------------------------------------------
+UDisplayClusterInFrustumFitCameraComponent::UDisplayClusterInFrustumFitCameraComponent(const FObjectInitializer& ObjectInitializer)
+	: UDisplayClusterCameraComponent(ObjectInitializer)
+{
+	PrimaryComponentTick.bCanEverTick = true;
+	PrimaryComponentTick.TickGroup = TG_PostUpdateWork;
+	PrimaryComponentTick.bStartWithTickEnabled = true;
+	bAutoActivate = true;
+}
+
+void UDisplayClusterInFrustumFitCameraComponent::TickComponent(float DeltaTime, ELevelTick TickType, FActorComponentTickFunction* ThisTickFunction)
+{
+	// todo: render preview
+
+#if WITH_EDITOR
+	if (WarpPolicy.IsValid())
+	{
+		if (bRefreshPreviewState)
+		{
+			// Set the preview mesh visibility state to match the bShowPreviewFrustumFit flag
+			if (ADisplayClusterRootActor* ParentRootActor = Cast<ADisplayClusterRootActor>(GetOwner()))
+			{
+				if (IDisplayClusterViewportManager* ViewportManager = ParentRootActor->GetViewportManager())
+				{
+					TArray<TSharedPtr<IDisplayClusterViewport, ESPMode::ThreadSafe>> Viewports = ViewportManager->GetViewportsForWarpPolicy(WarpPolicy);
+					for (const TSharedPtr<IDisplayClusterViewport, ESPMode::ThreadSafe>& Viewport : Viewports)
+					{
+						if (UMeshComponent* MovableMeshComponent = Viewport->GetProjectionPolicy()->GetOrCreatePreviewMovableMeshComponent(Viewport.Get()))
+						{
+							MovableMeshComponent->SetVisibility(bShowPreviewFrustumFit);
+						}
+					}
+				}
+			}
+
+			bRefreshPreviewState = false;
+		}
+	}
+#endif
+
+	// Tick warp policy instance
+	if (WarpPolicy.IsValid())
+	{
+		if (ADisplayClusterRootActor* ParentRootActor = Cast<ADisplayClusterRootActor>(GetOwner()))
+		{
+			if (IDisplayClusterViewportManager* ViewportManager = ParentRootActor->GetViewportManager())
+			{
+				WarpPolicy->Tick(ViewportManager, DeltaTime);
+			}
+		}
+	}
+}
+
+void UDisplayClusterInFrustumFitCameraComponent::GetDesiredView(FMinimalViewInfo& InOutViewInfo, float* OutCustomNearClippingPlane)
+{
+	if (IsEnabled())
+	{
+		if (ADisplayClusterRootActor* ParentRootActor = Cast<ADisplayClusterRootActor>(GetOwner()))
+		{
+			if (IDisplayClusterViewport::GetCameraComponentView(GetExternalCameraComponent(), ParentRootActor->GetWorldDeltaSeconds(), bUseCameraPostprocess, InOutViewInfo, OutCustomNearClippingPlane))
+			{
+				// 1. Use external camera for rendering
+				return;
+			}
+
+			if(IDisplayClusterViewportManager* ViewportManager = ParentRootActor->GetViewportManager())
+			{
+				if (IDisplayClusterViewport::GetPlayerCameraView(ViewportManager->GetCurrentWorld(), bUseCameraPostprocess, InOutViewInfo))
+				{
+					// 2. Use active game camera
+					return;
+				}
+			}
+		}
+	}
+
+	// use default logic
+	return UDisplayClusterCameraComponent::GetDesiredView(InOutViewInfo, OutCustomNearClippingPlane);
+}
+
+bool UDisplayClusterInFrustumFitCameraComponent::ShouldUseEntireClusterViewports(IDisplayClusterViewportManager* InViewportManager) const
+{
+	// Only when this component is enabled should viewports be created for the entire cluster that accesses this component.
+	return IsEnabled();
+}
+
+IDisplayClusterWarpPolicy* UDisplayClusterInFrustumFitCameraComponent::GetWarpPolicy(IDisplayClusterViewportManager* InViewportManager)
+{
+	// We can ask for different types of warp policies, depending on the rules of the user settings
+	const FString NewWaprPolicyType = DisplayClusterWarpStrings::warp::InFrustumFit;
+
+	// when returns different type, this will recreate warp policy instance
+	if (WarpPolicy.IsValid() && WarpPolicy->GetType() != NewWaprPolicyType)
+	{
+		WarpPolicy.Reset();
+	}
+
+	if (!WarpPolicy.IsValid())
+	{
+		WarpPolicy = ImplCreateWarpPolicy(NewWaprPolicyType, GetName());
+	}
+
+	return WarpPolicy.Get();
+}
+
+void UDisplayClusterInFrustumFitCameraComponent::OnRegister()
+{
+	Super::OnRegister();
+
+
+#if WITH_EDITOR
+	if (GIsEditor && !IsRunningCommandlet())
+	{
+		if (SpriteComponent)
+		{
+			SpriteComponent->SpriteInfo.Category = TEXT("NDisplayCameraViewOrigin");
+			SpriteComponent->SpriteInfo.DisplayName = NSLOCTEXT("DisplayClusterInFrustumFitCameraComponent", "DisplayClusterInFrustumFitCameraComponentSpriteInfo", "nDisplay InFrustumFit View Origin");
+		}
+	}
+
+	RefreshVisualRepresentation();
+#endif
+
+}
+
+bool UDisplayClusterInFrustumFitCameraComponent::IsEnabled() const
+{
+	return bEnableCameraProjection;
+}
+
+UCameraComponent* UDisplayClusterInFrustumFitCameraComponent::GetExternalCameraComponent() const
+{
+	if (ACineCameraActor* CineCamera = ExternalCameraActor.Get())
+	{
+		return CineCamera->GetCameraComponent();
+	}
+
+	return nullptr;
+}
+
+#if WITH_EDITOR
+void UDisplayClusterInFrustumFitCameraComponent::PostEditChangeProperty(FPropertyChangedEvent& PropertyChangedEvent)
+{
+	Super::PostEditChangeProperty(PropertyChangedEvent);
+
+	if (PropertyChangedEvent.Property && PropertyChangedEvent.Property->GetFName() == GET_MEMBER_NAME_CHECKED(UDisplayClusterInFrustumFitCameraComponent, bShowPreviewFrustumFit))
+	{
+		InvalidatePreviewState();
+	}
+}
+
+bool UDisplayClusterInFrustumFitCameraComponent::GetEditorPreviewInfo(float DeltaTime, FMinimalViewInfo& ViewOut)
+{
+	if (UCameraComponent* CameraComponent = GetExternalCameraComponent())
+	{
+		return CameraComponent->GetEditorPreviewInfo(DeltaTime, ViewOut);
+	}
+
+	return false;
+}
+
+TSharedPtr<SWidget> UDisplayClusterInFrustumFitCameraComponent::GetCustomEditorPreviewWidget()
+{
+	if (UCameraComponent* CameraComponent = GetExternalCameraComponent())
+	{
+		return CameraComponent->GetCustomEditorPreviewWidget();
+	}
+
+	return nullptr;
+}
+
+#endif
diff --git a/Engine/Plugins/Runtime/nDisplay/Source/DisplayClusterWarp/Private/DisplayClusterWarpModule.cpp b/Engine/Plugins/Runtime/nDisplay/Source/DisplayClusterWarp/Private/DisplayClusterWarpModule.cpp
--- a/Engine/Plugins/Runtime/nDisplay/Source/DisplayClusterWarp/Private/DisplayClusterWarpModule.cpp
+++ b/Engine/Plugins/Runtime/nDisplay/Source/DisplayClusterWarp/Private/DisplayClusterWarpModule.cpp
@@ -7,28 +7,67 @@
 
 #include "IDisplayCluster.h"
 #include "Render/IDisplayClusterRenderManager.h"
+#include "Policy/InFrustumFit/DisplayClusterWarpInFrustumFitPolicyFactory.h"
 
+//////////////////////////////////////////////////////////////////////////////////////////////
+// FDisplayClusterWarpModule
+//////////////////////////////////////////////////////////////////////////////////////////////
 FDisplayClusterWarpModule::FDisplayClusterWarpModule()
 {
-	UE_LOG(LogDisplayClusterWarp, Verbose, TEXT("Warp module has been instantiated"));
+	TSharedPtr<IDisplayClusterWarpPolicyFactory> Factory;
+
+	// InFrustumFit warp policy
+	Factory = MakeShared<FDisplayClusterWarpInFrustumFitPolicyFactory>();
+	WarpPolicyFactories.Emplace(DisplayClusterWarpStrings::warp::InFrustumFit, Factory);
+
+	UE_LOG(LogDisplayClusterWarp, Log, TEXT("Warp module has been instantiated"));
 }
 
 FDisplayClusterWarpModule::~FDisplayClusterWarpModule()
 {
-	UE_LOG(LogDisplayClusterWarp, Verbose, TEXT("Warp module has been destroyed"));
+	UE_LOG(LogDisplayClusterWarp, Log, TEXT("Warp module has been destroyed"));
 }
 
-//////////////////////////////////////////////////////////////////////////////////////////////
-// IModuleInterface
-//////////////////////////////////////////////////////////////////////////////////////////////
 void FDisplayClusterWarpModule::StartupModule()
 {
-	UE_LOG(LogDisplayClusterWarp, Verbose, TEXT("Warp module has been started"));
+	UE_LOG(LogDisplayClusterWarp, Log, TEXT("Warp module startup"));
+
+	IDisplayClusterRenderManager* RenderMgr = IDisplayCluster::Get().GetRenderMgr();
+	if (RenderMgr)
+	{
+		for (auto it = WarpPolicyFactories.CreateIterator(); it; ++it)
+		{
+			UE_LOG(LogDisplayClusterWarp, Log, TEXT("Registering <%s> warp policy factory..."), *it->Key);
+
+			if (!RenderMgr->RegisterWarpPolicyFactory(it->Key, it->Value))
+			{
+				UE_LOG(LogDisplayClusterWarp, Warning, TEXT("Couldn't register <%s> warp policy factory"), *it->Key);
+			}
+		}
+	}
+
+	UE_LOG(LogDisplayClusterWarp, Log, TEXT("Warp module has started"));
 }
 
 void FDisplayClusterWarpModule::ShutdownModule()
 {
-	UE_LOG(LogDisplayClusterWarp, Verbose, TEXT("Warp module has been shutdown"));
+	UE_LOG(LogDisplayClusterWarp, Log, TEXT("Warp module shutdown"));
+
+	IDisplayClusterRenderManager* RenderMgr = IDisplayCluster::Get().GetRenderMgr();
+	if (RenderMgr)
+	{
+		for (auto it = WarpPolicyFactories.CreateConstIterator(); it; ++it)
+		{
+			UE_LOG(LogDisplayClusterWarp, Log, TEXT("Un-registering <%s> warp factory..."), *it->Key);
+
+			if (!RenderMgr->UnregisterWarpPolicyFactory(it->Key))
+			{
+				UE_LOG(LogDisplayClusterWarp, Warning, TEXT("An error occurred during un-registering the <%s> warp factory"), *it->Key);
+			}
+		}
+	}
+
+	UE_LOG(LogDisplayClusterWarp, Log, TEXT("Warp module has been shutdown."));
 }
 
 IMPLEMENT_MODULE(FDisplayClusterWarpModule, DisplayClusterWarp);
diff --git a/Engine/Plugins/Runtime/nDisplay/Source/DisplayClusterWarp/Private/DisplayClusterWarpModule.h b/Engine/Plugins/Runtime/nDisplay/Source/DisplayClusterWarp/Private/DisplayClusterWarpModule.h
--- a/Engine/Plugins/Runtime/nDisplay/Source/DisplayClusterWarp/Private/DisplayClusterWarpModule.h
+++ b/Engine/Plugins/Runtime/nDisplay/Source/DisplayClusterWarp/Private/DisplayClusterWarpModule.h
@@ -54,4 +54,8 @@ public:
 	{
 		return FDisplayClusterWarpBlendLoader_MPCDI::ReadMPCDFileStructure(InMPCDIFileName, OutMPCDIFileStructure);
 	}
+
+private:
+	// Available factories
+	TMap<FString, TSharedPtr<IDisplayClusterWarpPolicyFactory>> WarpPolicyFactories;
 };
diff --git a/Engine/Plugins/Runtime/nDisplay/Source/DisplayClusterWarp/Private/Policy/DisplayClusterWarpPolicyBase.cpp b/Engine/Plugins/Runtime/nDisplay/Source/DisplayClusterWarp/Private/Policy/DisplayClusterWarpPolicyBase.cpp
new file mode 100644
--- /dev/null
+++ b/Engine/Plugins/Runtime/nDisplay/Source/DisplayClusterWarp/Private/Policy/DisplayClusterWarpPolicyBase.cpp
@@ -0,0 +1,18 @@
+// Copyright Epic Games, Inc. All Rights Reserved.
+
+#include "Policy/DisplayClusterWarpPolicyBase.h"
+#include "Render/Viewport/IDisplayClusterViewport.h"
+#include "Render/Projection/IDisplayClusterProjectionPolicy.h"
+
+#include "Components/DisplayClusterCameraComponent.h"
+
+#include "Containers/DisplayClusterWarpEye.h"
+
+#include "DisplayClusterWarpLog.h"
+
+//////////////////////////////////////////////////////////////////////////////////////////////
+// FDisplayClusterWarpPolicyBase
+//////////////////////////////////////////////////////////////////////////////////////////////
+FDisplayClusterWarpPolicyBase::FDisplayClusterWarpPolicyBase(const FString& InType, const FString& InWarpPolicyName)
+	: PolicyInstanceId(FString::Printf(TEXT("WarpPolicy_%s_%s"), *InType, *InWarpPolicyName))
+{ }
diff --git a/Engine/Plugins/Runtime/nDisplay/Source/DisplayClusterWarp/Private/Policy/DisplayClusterWarpPolicyBase.h b/Engine/Plugins/Runtime/nDisplay/Source/DisplayClusterWarp/Private/Policy/DisplayClusterWarpPolicyBase.h
new file mode 100644
--- /dev/null
+++ b/Engine/Plugins/Runtime/nDisplay/Source/DisplayClusterWarp/Private/Policy/DisplayClusterWarpPolicyBase.h
@@ -0,0 +1,39 @@
+// Copyright Epic Games, Inc. All Rights Reserved.
+
+#pragma once
+#include "Render/Warp/IDisplayClusterWarpPolicy.h"
+#include "Containers/DisplayClusterWarpEye.h"
+#include "Templates/SharedPointer.h"
+
+class IDisplayClusterWarpBlend;
+
+/**
+ * Base warp policy
+ */
+class FDisplayClusterWarpPolicyBase
+	: public IDisplayClusterWarpPolicy
+	, public TSharedFromThis<FDisplayClusterWarpPolicyBase, ESPMode::ThreadSafe>
+{
+public:
+	FDisplayClusterWarpPolicyBase(const FString& InType, const FString& InWarpPolicyName);
+
+	virtual TSharedPtr<IDisplayClusterWarpPolicy, ESPMode::ThreadSafe> ToSharedPtr() override
+	{
+		return AsShared();
+	}
+
+	virtual TSharedPtr<const IDisplayClusterWarpPolicy, ESPMode::ThreadSafe> ToSharedPtr() const override
+	{
+		return AsShared();
+	}
+
+public:
+	virtual const FString& GetId() const override
+	{
+		return PolicyInstanceId;
+	}
+
+private:
+	// unique name
+	const FString PolicyInstanceId;
+};
diff --git a/Engine/Plugins/Runtime/nDisplay/Source/DisplayClusterWarp/Private/Policy/InFrustumFit/DisplayClusterWarpInFrustumFitPolicy.cpp b/Engine/Plugins/Runtime/nDisplay/Source/DisplayClusterWarp/Private/Policy/InFrustumFit/DisplayClusterWarpInFrustumFitPolicy.cpp
new file mode 100644
--- /dev/null
+++ b/Engine/Plugins/Runtime/nDisplay/Source/DisplayClusterWarp/Private/Policy/InFrustumFit/DisplayClusterWarpInFrustumFitPolicy.cpp
@@ -0,0 +1,394 @@
+// Copyright Epic Games, Inc. All Rights Reserved.
+
+#include "Policy/InFrustumFit/DisplayClusterWarpInFrustumFitPolicy.h"
+
+#include "DisplayClusterRootActor.h"
+#include "Containers/DisplayClusterWarpEye.h"
+
+#include "Render/Viewport/IDisplayClusterViewport.h"
+#include "Render/Projection/IDisplayClusterProjectionPolicy.h"
+
+#include "IDisplayClusterWarpBlend.h"
+
+#include "DisplayClusterWarpStrings.h"
+
+#include "Components/DisplayClusterCameraComponent.h"
+#include "Components/DisplayClusterInFrustumFitCameraComponent.h"
+#include "Components/DisplayClusterPreviewComponent.h"
+#include "Components/MeshComponent.h"
+#include "Materials/MaterialInstanceDynamic.h"
+
+// Experimental: Enable static view direction for mpcdi 2d
+int32 GDisplayClusterWarpInFrustumFitPolicyUseStaticViewDirectionForMPCDIProfile2D = 1;
+static FAutoConsoleVariableRef CVarDisplayClusterWarpInFrustumFitPolicyUseStaticViewDirectionForMPCDIProfile2D(
+	TEXT("nDisplay.warp.InFrustumFit.UseStaticViewDirectionForMPCDIProfile2D"),
+	GDisplayClusterWarpInFrustumFitPolicyUseStaticViewDirectionForMPCDIProfile2D,
+	TEXT("Experimental: Enable static view direction for mpcdi 2d (0 - disable)\n"),
+	ECVF_Default
+);
+
+// Experimental: Enable single view target for group of viewports
+int32 GDisplayClusterWarpInFrustumFitPolicyUseGroupViewTarget = 1;
+static FAutoConsoleVariableRef CVarDisplayClusterWarpInFrustumFitPolicyUseGroupViewTarget(
+	TEXT("nDisplay.warp.InFrustumFit.UseGroupViewTarget"),
+	GDisplayClusterWarpInFrustumFitPolicyUseGroupViewTarget,
+	TEXT("Experimental: Enable single view target for group of viewports (0 - disable)\n"),
+	ECVF_Default
+);
+
+int32 GDisplayClusterWarpInFrustumFitPolicyDrawFrustum = 0;
+static FAutoConsoleVariableRef CVarDisplayClusterWarpInFrustumFitPolicyDrawFrustum(
+	TEXT("nDisplay.warp.InFrustumFit.DrawFrustum"),
+	GDisplayClusterWarpInFrustumFitPolicyDrawFrustum,
+	TEXT("Toggles drawing the stage geometry frustum and bounding box\n"),
+	ECVF_Default
+);
+
+//-------------------------------------------------------------------
+// FDisplayClusterWarpInFrustumFitPolicy
+//-------------------------------------------------------------------
+FDisplayClusterWarpInFrustumFitPolicy::FDisplayClusterWarpInFrustumFitPolicy(const FString& InWarpPolicyName)
+	: FDisplayClusterWarpPolicyBase(GetType(), InWarpPolicyName)
+{ }
+
+const FString& FDisplayClusterWarpInFrustumFitPolicy::GetType() const
+{
+	static const FString Type(DisplayClusterWarpStrings::warp::InFrustumFit);
+
+	return Type;
+}
+
+void FDisplayClusterWarpInFrustumFitPolicy::HandleNewFrame(const TArray<TSharedPtr<IDisplayClusterViewport, ESPMode::ThreadSafe>>& InViewports)
+{
+	// Todo: get real scale
+	const float WorldToMeters = 100.0f;
+
+	const float WorldScale = WorldToMeters / 100.f;
+	
+	bGeometryContextsUpdated = false;
+	if (GDisplayClusterWarpInFrustumFitPolicyUseGroupViewTarget)
+	{
+		// Calculate AABB for group of viewports:
+		GroupAABBox = FDisplayClusterWarpAABB();
+		for (const TSharedPtr<IDisplayClusterViewport, ESPMode::ThreadSafe>& Viewport : InViewports)
+		{
+			if (Viewport.IsValid() && Viewport->GetProjectionPolicy().IsValid())
+			{
+				TSharedPtr<IDisplayClusterWarpBlend, ESPMode::ThreadSafe> WarpBlend;
+				if (Viewport->GetProjectionPolicy()->GetWarpBlendInterface(WarpBlend))
+				{
+					if (WarpBlend->UpdateGeometryContext(WorldScale))
+					{
+						GroupAABBox.UpdateAABB(WarpBlend->GetGeometryContext().AABBox);
+					}
+				}
+			}
+		}
+
+		bGeometryContextsUpdated = true;
+	}
+
+	const uint32 ContextNum = 0; // calculate for a single context
+
+	bValidGroupFrustum = false;
+
+	// Recalculate warp projection angles
+	GroupGeometryWarpProjection.ResetProjectionAngles();
+	FittedViewports.Empty();
+	SymmetricForwardCorrection.Reset();
+
+	bool bCanCalcFrustumContext = true;
+	bool bHasFixedViewDirection = false;
+
+	// Use center of GroupAABB as target viewpoint, and setup to use same viewprojection plane for all viewports:
+	for (const TSharedPtr<IDisplayClusterViewport, ESPMode::ThreadSafe>& Viewport : InViewports)
+	{
+		if (bCanCalcFrustumContext && Viewport.IsValid() && Viewport->GetProjectionPolicy().IsValid())
+		{
+			// Note: This code is partially copied from FDisplayClusterProjectionMPCDIPolicy::CalculateView().
+
+			// Override viewpoint
+			// MPCDI always expects the location of the viewpoint component (eye location from the real world)
+			FVector ViewOffset = FVector::ZeroVector;
+			FVector InOutViewLocation;
+			FRotator InOutViewRotation;
+			if (!Viewport->GetViewPointCameraEye(ContextNum, InOutViewLocation, InOutViewRotation, ViewOffset))
+			{
+				continue;
+			}
+
+			if (UDisplayClusterInFrustumFitCameraComponent* WarpViewPointComponent = Cast<UDisplayClusterInFrustumFitCameraComponent>(Viewport->GetViewPointCameraComponent()))
+			{
+				bHasFixedViewDirection = WarpViewPointComponent->CameraViewTarget == EDisplayClusterWarpCameraViewTarget::MatchViewOrigin;
+			}
+
+			TSharedPtr<IDisplayClusterWarpBlend, ESPMode::ThreadSafe> WarpBlend;
+			if (Viewport->GetProjectionPolicy()->GetWarpBlendInterface(WarpBlend))
+			{
+				const USceneComponent* const OriginComp = Viewport->GetProjectionPolicy()->GetOriginComponent();
+
+				TSharedPtr<FDisplayClusterWarpEye, ESPMode::ThreadSafe> WarpEye = MakeShared<FDisplayClusterWarpEye, ESPMode::ThreadSafe>(Viewport, 0);
+
+				WarpEye->World2LocalTransform = (OriginComp ? OriginComp->GetComponentTransform() : FTransform::Identity);
+
+				// Get our base camera location and view offset in local space (MPCDI space)
+				WarpEye->ViewPoint.Location = WarpEye->World2LocalTransform.InverseTransformPosition(InOutViewLocation - ViewOffset);
+				WarpEye->ViewPoint.EyeOffset = WarpEye->World2LocalTransform.InverseTransformPosition(InOutViewLocation) - WarpEye->ViewPoint.Location;
+				WarpEye->ViewPoint.Rotation = WarpEye->World2LocalTransform.InverseTransformRotation(InOutViewRotation.Quaternion()).Rotator();
+
+				WarpEye->WorldScale = WorldScale;
+
+				WarpEye->WarpPolicy = SharedThis(this);
+
+				if (!WarpBlend->CalcFrustumContext(WarpEye))
+				{
+					bCanCalcFrustumContext = false;
+				}
+				else
+				{
+					const FDisplayClusterWarpData& WarpData = WarpBlend->GetWarpData(ContextNum);
+					GroupGeometryWarpProjection.ExpandProjectionAngles(WarpData.GeometryWarpProjection);
+				}
+			}
+		}
+	}
+
+	bValidGroupFrustum = bCanCalcFrustumContext;
+
+	// Recalculate the group frustum so that it is symmetric
+	MakeGroupFrustumSymmetrical(bHasFixedViewDirection);
+}
+
+void FDisplayClusterWarpInFrustumFitPolicy::Tick(IDisplayClusterViewportManager* InViewportManager, float DeltaSeconds)
+{
+#if WITH_EDITOR
+	bool bHasDrawnDebugFrustum = false;
+
+	TArray<TSharedPtr<IDisplayClusterViewport, ESPMode::ThreadSafe>> Viewports = InViewportManager->GetViewportsForWarpPolicy(SharedThis(this));
+	for (const TSharedPtr<IDisplayClusterViewport, ESPMode::ThreadSafe>& Viewport : Viewports)
+	{
+		// If the viewport doesn't have a valid warp projection yet, skip it
+		if (!FittedViewports.Contains(Viewport->GetId()))
+		{
+			continue;
+		}
+
+		TSharedPtr<IDisplayClusterWarpBlend, ESPMode::ThreadSafe> WarpBlend;
+		if (Viewport->GetProjectionPolicy()->GetWarpBlendInterface(WarpBlend))
+		{
+			FDisplayClusterWarpData& WarpData = WarpBlend->GetWarpData(0);
+
+			const FTransform CameraTransform(WarpData.WarpProjection.CameraRotation.Quaternion(), WarpData.WarpProjection.CameraLocation);
+
+			const float HScale = (WarpData.WarpProjection.Left - WarpData.WarpProjection.Right) / (WarpData.GeometryWarpProjection.Left - WarpData.GeometryWarpProjection.Right);
+			const float VScale = (WarpData.WarpProjection.Top - WarpData.WarpProjection.Bottom) / (WarpData.GeometryWarpProjection.Top - WarpData.GeometryWarpProjection.Bottom);
+
+			checkf(FMath::IsNearlyEqual(HScale, VScale), TEXT("Streching the stage geometry to fit a different aspect ratio is not supported!"));
+
+			const FVector Scale = FVector(1, HScale, HScale);
+
+			// Compute the relative transform from the view origin to the geometry
+			FTransform RelativeTransform(WarpData.WarpContext.MeshToStageMatrix * WarpData.Local2World.Inverse());
+			RelativeTransform.ScaleTranslation(Scale);
+
+			// Final transform is computed from the relative transform of the geometry to the view point, the frustum fit transform
+			// which will scale and position the geometry based on the fitted frustum, and the camera transform
+			const FTransform FinalTransform = RelativeTransform * CameraTransform;
+
+			UMeshComponent* MovableMeshComponent = Viewport->GetProjectionPolicy()->GetOrCreatePreviewMovableMeshComponent(Viewport.Get());
+			MovableMeshComponent->SetRelativeTransform(FinalTransform);
+
+			// Since the mesh needs to be skewed to scale appropriately, and since Unreal Engine does not support a skew transform
+			// through FTransform, the mesh needs to be skewed through the vertex shader using WorldPositionOffset,
+			// so pass in the "global" scale to the preview mesh's material instance
+			if (UMaterialInstanceDynamic* MaterialInstance = Cast<UMaterialInstanceDynamic>(MovableMeshComponent->GetMaterial(0)))
+			{
+				FMatrix CameraBasis = FRotationMatrix::Make(WarpData.WarpProjection.CameraRotation).Inverse();
+
+				MaterialInstance->SetVectorParameterValue(TEXT("GlobalScale"), Scale);
+				MaterialInstance->SetVectorParameterValue(TEXT("GlobalForward"), CameraBasis.GetUnitAxis(EAxis::X));
+				MaterialInstance->SetVectorParameterValue(TEXT("GlobalRight"), CameraBasis.GetUnitAxis(EAxis::Y));
+				MaterialInstance->SetVectorParameterValue(TEXT("GlobalUp"), CameraBasis.GetUnitAxis(EAxis::Z));
+			}
+
+			if (GDisplayClusterWarpInFrustumFitPolicyDrawFrustum && !bHasDrawnDebugFrustum)
+			{
+				DrawDebugGroupFrustum(InViewportManager->GetRootActor(), Cast<UDisplayClusterInFrustumFitCameraComponent>(Viewport->GetViewPointCameraComponent()), FColor::Blue);
+				bHasDrawnDebugFrustum = true;
+			}
+		}
+	}
+
+	if (GDisplayClusterWarpInFrustumFitPolicyDrawFrustum)
+	{
+		DrawDebugGroupBoundingBox(InViewportManager->GetRootActor(), FColor::Red);
+	}
+#endif
+}
+
+void FDisplayClusterWarpInFrustumFitPolicy::BeginCalcFrustum(IDisplayClusterViewport* InViewport, const uint32 ContextNum)
+{
+	if (InViewport && InViewport->GetProjectionPolicy().IsValid())
+	{
+		TSharedPtr<IDisplayClusterWarpBlend, ESPMode::ThreadSafe> WarpBlend;
+		if (InViewport->GetProjectionPolicy()->GetWarpBlendInterface(WarpBlend))
+		{
+			if (UDisplayClusterInFrustumFitCameraComponent* WarpViewPointComponent = Cast<UDisplayClusterInFrustumFitCameraComponent>(InViewport->GetViewPointCameraComponent()))
+			{
+				FDisplayClusterWarpData& WarpData = WarpBlend->GetWarpData(ContextNum);
+				if (WarpData.WarpEye.IsValid())
+				{
+					// geometry context already updated.
+					WarpData.WarpEye->bUpdateGeometryContext = !bGeometryContextsUpdated;
+
+					if (WarpViewPointComponent->CameraViewTarget == EDisplayClusterWarpCameraViewTarget::MatchViewOrigin)
+					{
+						WarpData.WarpEye->OverrideViewDirection = WarpData.WarpEye->ViewPoint.Rotation.RotateVector(FVector::XAxisVector);
+					}
+					else
+					{
+						if (GDisplayClusterWarpInFrustumFitPolicyUseGroupViewTarget)
+						{
+							// If we have a correction to the view forward vector to make the frustum symmetric, apply it
+							if (SymmetricForwardCorrection.IsSet())
+							{
+								const FVector AABBForward = (GroupAABBox.GetCenter() - WarpData.WarpEye->ViewPoint.GetEyeLocation()).GetSafeNormal();
+								WarpData.WarpEye->OverrideViewDirection = SymmetricForwardCorrection->RotateVector(AABBForward);
+							}
+							else
+							{
+								// Use the same view target for all viewports
+								WarpData.WarpEye->OverrideViewTarget = GroupAABBox.GetCenter();
+							}
+						}
+					}
+
+					if (GDisplayClusterWarpInFrustumFitPolicyUseStaticViewDirectionForMPCDIProfile2D)
+					{
+						// [experimental] use static view direction for MPCDI profile 2D
+						if (WarpBlend->GetWarpProfileType() == EDisplayClusterWarpProfileType::warp_2D)
+						{
+							WarpData.WarpEye->OverrideViewDirection = FVector(1, 0, 0);
+						}
+					}
+
+					// Todo: This feature now not supported here. need to be fixed.
+					WarpData.bEnabledRotateFrustumToFitContextSize = false;
+				}
+			}
+		}
+	}
+}
+
+void FDisplayClusterWarpInFrustumFitPolicy::EndCalcFrustum(IDisplayClusterViewport* InViewport, const uint32 ContextNum)
+{
+	if (bValidGroupFrustum && InViewport && InViewport->GetProjectionPolicy().IsValid())
+	{
+		TSharedPtr<IDisplayClusterWarpBlend, ESPMode::ThreadSafe> WarpBlend;
+		if (InViewport->GetProjectionPolicy()->GetWarpBlendInterface(WarpBlend))
+		{
+			if (UDisplayClusterCameraComponent* ViewPointComponent = InViewport->GetViewPointCameraComponent())
+			{
+				if (ViewPointComponent->IsA<UDisplayClusterInFrustumFitCameraComponent>())
+				{
+					if (UDisplayClusterInFrustumFitCameraComponent* WarpViewPointComponent = CastChecked<UDisplayClusterInFrustumFitCameraComponent>(ViewPointComponent))
+					{
+						// Change warp settings:
+						FDisplayClusterWarpData& WarpData = WarpBlend->GetWarpData(ContextNum);
+
+						// Apply camera frustum fitting:
+						FDisplayClusterWarpProjection NewWarpProjection = ApplyInFrustumFit(WarpViewPointComponent, WarpData.WarpEye->World2LocalTransform, WarpData.WarpProjection);
+
+						WarpData.WarpProjection = NewWarpProjection;
+
+						FittedViewports.Add(InViewport->GetId());
+					}
+				}
+			}
+		}
+	}
+}
+
+#if WITH_EDITOR
+void FDisplayClusterWarpInFrustumFitPolicy::DrawDebugGroupBoundingBox(ADisplayClusterRootActor* RootActor, const FColor& Color)
+{
+	if (RootActor)
+	{
+		if (UWorld* World = RootActor->GetWorld())
+		{
+			FBox WorldBox = GroupAABBox.TransformBy(RootActor->GetActorTransform());
+			DrawDebugBox(World, WorldBox.GetCenter(), WorldBox.GetExtent(), Color);
+			DrawDebugPoint(World, WorldBox.GetCenter(), 5, Color);
+		}
+	}
+}
+
+void FDisplayClusterWarpInFrustumFitPolicy::DrawDebugGroupFrustum(ADisplayClusterRootActor* RootActor, UDisplayClusterInFrustumFitCameraComponent* CameraComponent, const FColor& Color)
+{
+	if (RootActor && CameraComponent)
+	{
+		if (UWorld* World = RootActor->GetWorld())
+		{
+			const float NearPlane = 10;
+			const float FarPlane = 1000;
+
+			const FVector CameraLoc = CameraComponent->GetComponentLocation();
+			FVector ViewDirection;
+
+			if (CameraComponent->CameraViewTarget == EDisplayClusterWarpCameraViewTarget::MatchViewOrigin)
+			{
+				ViewDirection = CameraComponent->GetComponentRotation().RotateVector(FVector::XAxisVector);
+			}
+			else
+			{
+				const FBox WorldBox = GroupAABBox.TransformBy(RootActor->GetActorTransform());
+				ViewDirection = (WorldBox.GetCenter() - CameraComponent->GetComponentLocation()).GetSafeNormal();
+
+				if (SymmetricForwardCorrection.IsSet())
+				{
+					ViewDirection = SymmetricForwardCorrection->RotateVector(ViewDirection);
+				}
+			}
+
+			DrawDebugLine(World, CameraLoc, CameraLoc + ViewDirection * 50, Color);
+
+			const FRotator ViewRotator = ViewDirection.ToOrientationRotator();
+			const FVector FrustumTopLeft = ViewRotator.RotateVector(FVector(GroupGeometryWarpProjection.ZNear, GroupGeometryWarpProjection.Left, GroupGeometryWarpProjection.Top) / GroupGeometryWarpProjection.ZNear);
+			const FVector FrustumTopRight = ViewRotator.RotateVector(FVector(GroupGeometryWarpProjection.ZNear, GroupGeometryWarpProjection.Right, GroupGeometryWarpProjection.Top) / GroupGeometryWarpProjection.ZNear);
+			const FVector FrustumBottomLeft = ViewRotator.RotateVector(FVector(GroupGeometryWarpProjection.ZNear, GroupGeometryWarpProjection.Left, GroupGeometryWarpProjection.Bottom) / GroupGeometryWarpProjection.ZNear);
+			const FVector FrustumBottomRight = ViewRotator.RotateVector(FVector(GroupGeometryWarpProjection.ZNear, GroupGeometryWarpProjection.Right, GroupGeometryWarpProjection.Bottom) / GroupGeometryWarpProjection.ZNear);
+
+			const FVector FrustumVertices[8] =
+			{
+				CameraLoc + FrustumTopLeft * NearPlane,
+				CameraLoc + FrustumTopRight * NearPlane,
+				CameraLoc + FrustumBottomRight * NearPlane,
+				CameraLoc + FrustumBottomLeft * NearPlane,
+
+				CameraLoc + FrustumTopLeft * FarPlane,
+				CameraLoc + FrustumTopRight * FarPlane,
+				CameraLoc + FrustumBottomRight * FarPlane,
+				CameraLoc + FrustumBottomLeft * FarPlane,
+			};
+
+			// Near plane rectangle
+			DrawDebugLine(World, FrustumVertices[0], FrustumVertices[1], Color);
+			DrawDebugLine(World, FrustumVertices[1], FrustumVertices[2], Color);
+			DrawDebugLine(World, FrustumVertices[2], FrustumVertices[3], Color);
+			DrawDebugLine(World, FrustumVertices[3], FrustumVertices[0], Color);
+
+			// Frustum
+			DrawDebugLine(World, FrustumVertices[0], FrustumVertices[4], Color);
+			DrawDebugLine(World, FrustumVertices[1], FrustumVertices[5], Color);
+			DrawDebugLine(World, FrustumVertices[2], FrustumVertices[6], Color);
+			DrawDebugLine(World, FrustumVertices[3], FrustumVertices[7], Color);
+
+			// Far plane rectangle
+			DrawDebugLine(World, FrustumVertices[4], FrustumVertices[5], Color);
+			DrawDebugLine(World, FrustumVertices[5], FrustumVertices[6], Color);
+			DrawDebugLine(World, FrustumVertices[6], FrustumVertices[7], Color);
+			DrawDebugLine(World, FrustumVertices[7], FrustumVertices[4], Color);
+		}
+	}
+}
+#endif
diff --git a/Engine/Plugins/Runtime/nDisplay/Source/DisplayClusterWarp/Private/Policy/InFrustumFit/DisplayClusterWarpInFrustumFitPolicy.h b/Engine/Plugins/Runtime/nDisplay/Source/DisplayClusterWarp/Private/Policy/InFrustumFit/DisplayClusterWarpInFrustumFitPolicy.h
new file mode 100644
--- /dev/null
+++ b/Engine/Plugins/Runtime/nDisplay/Source/DisplayClusterWarp/Private/Policy/InFrustumFit/DisplayClusterWarpInFrustumFitPolicy.h
@@ -0,0 +1,69 @@
+// Copyright Epic Games, Inc. All Rights Reserved.
+
+#pragma once
+
+#include "Policy/DisplayClusterWarpPolicyBase.h"
+#include "Blueprints/DisplayClusterWarpBlueprint_Enums.h"
+
+#include "Containers/DisplayClusterWarpContext.h"
+
+class ADisplayClusterRootActor;
+class UDisplayClusterInFrustumFitCameraComponent;
+class FDisplayClusterWarpEye;
+
+/**
+ * FurstumFit warp policy
+ */
+class FDisplayClusterWarpInFrustumFitPolicy
+	: public FDisplayClusterWarpPolicyBase
+{
+public:
+	FDisplayClusterWarpInFrustumFitPolicy(const FString& InWarpPolicyName);
+
+public:
+	//~ Begin IDisplayClusterWarpPolicy
+	virtual const FString& GetType() const override;
+	virtual void HandleNewFrame(const TArray<TSharedPtr<IDisplayClusterViewport, ESPMode::ThreadSafe>>& InViewports) override;
+	virtual void Tick(IDisplayClusterViewportManager* InViewportManager, float DeltaSeconds) override;
+
+	virtual void BeginCalcFrustum(IDisplayClusterViewport* InViewport, const uint32 ContextNum) override;
+	virtual void EndCalcFrustum(IDisplayClusterViewport* InViewport, const uint32 ContextNum) override;
+	//~~ End IDisplayClusterWarpPolicy
+
+private:
+	/** Converts an asymmetric group frustum to a symmetric group frustum and computes the correction rotator to the view forward vector needed for a symmetric frustum  */
+	void MakeGroupFrustumSymmetrical(bool bFixedViewTarget);
+
+	/** Apply frustum fit to the specified warp projection */
+	FDisplayClusterWarpProjection ApplyInFrustumFit(UDisplayClusterInFrustumFitCameraComponent* InComponent, const FTransform& World2OriginTransform, const FDisplayClusterWarpProjection& InWarpProjection);
+
+	/** Find final projection scale.*/
+	FVector2D FindFrustumFit(const EDisplayClusterWarpCameraProjectionMode InProjectionMode, const FVector2D& InCameraFOV, const FVector2D& InGeometryFOV);
+
+#if WITH_EDITOR
+	/** Renders a debug visualization for the group bounding box */
+	void DrawDebugGroupBoundingBox(ADisplayClusterRootActor* RootActor, const FColor& Color);
+
+	/** Renders a debug visualization for the group frustum */
+	void DrawDebugGroupFrustum(ADisplayClusterRootActor* RootActor, UDisplayClusterInFrustumFitCameraComponent* CameraComponent, const FColor& Color);
+#endif
+
+private:
+	// Warp projection data
+	FDisplayClusterWarpProjection GroupGeometryWarpProjection;
+
+	/**
+	 * Due to the piecemeal nature of the viewport rendering, not all viewports will have a valid warp projection every tick. This list
+	 * stores the viewports that have been fitted since the last frame was started.
+	 */
+	TSet<FString> FittedViewports;
+
+	/** The AABB of the group computed in DCRA space */
+	FDisplayClusterWarpAABB GroupAABBox;
+
+	/** If set, stores the correction rotator to apply to the view direction to make the group frustum symmetrical */
+	TOptional<FRotator> SymmetricForwardCorrection;
+
+	bool bValidGroupFrustum = false;
+	bool bGeometryContextsUpdated = false;
+};
diff --git a/Engine/Plugins/Runtime/nDisplay/Source/DisplayClusterWarp/Private/Policy/InFrustumFit/DisplayClusterWarpInFrustumFitPolicyFactory.cpp b/Engine/Plugins/Runtime/nDisplay/Source/DisplayClusterWarp/Private/Policy/InFrustumFit/DisplayClusterWarpInFrustumFitPolicyFactory.cpp
new file mode 100644
--- /dev/null
+++ b/Engine/Plugins/Runtime/nDisplay/Source/DisplayClusterWarp/Private/Policy/InFrustumFit/DisplayClusterWarpInFrustumFitPolicyFactory.cpp
@@ -0,0 +1,27 @@
+// Copyright Epic Games, Inc. All Rights Reserved.
+
+#include "Policy/InFrustumFit/DisplayClusterWarpInFrustumFitPolicyFactory.h"
+#include "Policy/InFrustumFit/DisplayClusterWarpInFrustumFitPolicy.h"
+
+#include "DisplayClusterWarpStrings.h"
+#include "DisplayClusterWarpLog.h"
+
+//////////////////////////////////////////////////////////////////////////////////////////////
+// FDisplayClusterWarpCameraFitPolicyFactory
+//////////////////////////////////////////////////////////////////////////////////////////////
+TSharedPtr<IDisplayClusterWarpPolicy, ESPMode::ThreadSafe> FDisplayClusterWarpInFrustumFitPolicyFactory::Create(const FString& InWarpPolicyType, const FString& InWarpPolicyName)
+{
+	TSharedPtr<IDisplayClusterWarpPolicy, ESPMode::ThreadSafe> Result;
+
+	if (InWarpPolicyType == DisplayClusterWarpStrings::warp::InFrustumFit)
+	{
+		Result = MakeShared<FDisplayClusterWarpInFrustumFitPolicy, ESPMode::ThreadSafe>(InWarpPolicyName);
+	}
+
+	if (Result.IsValid())
+	{
+		UE_LOG(LogDisplayClusterWarpInFrustumFitPolicy, Verbose, TEXT("Instantiating warp policy <%s> id='%s'"), *InWarpPolicyType, *InWarpPolicyName);
+	}
+
+	return  Result;
+}
diff --git a/Engine/Plugins/Runtime/nDisplay/Source/DisplayClusterWarp/Private/Policy/InFrustumFit/DisplayClusterWarpInFrustumFitPolicyFactory.h b/Engine/Plugins/Runtime/nDisplay/Source/DisplayClusterWarp/Private/Policy/InFrustumFit/DisplayClusterWarpInFrustumFitPolicyFactory.h
new file mode 100644
--- /dev/null
+++ b/Engine/Plugins/Runtime/nDisplay/Source/DisplayClusterWarp/Private/Policy/InFrustumFit/DisplayClusterWarpInFrustumFitPolicyFactory.h
@@ -0,0 +1,17 @@
+// Copyright Epic Games, Inc. All Rights Reserved.
+
+#pragma once
+#include "Render/Warp/IDisplayClusterWarpPolicyFactory.h"
+
+/**
+ * Implements warp policy factory for the 'InFrustumFit' policy
+ */
+class FDisplayClusterWarpInFrustumFitPolicyFactory
+	: public IDisplayClusterWarpPolicyFactory
+{
+public:
+	//////////////////////////////////////////////////////////////////////////////////////////////
+	// IDisplayClusterWarpPolicyFactory
+	//////////////////////////////////////////////////////////////////////////////////////////////
+	virtual TSharedPtr<IDisplayClusterWarpPolicy, ESPMode::ThreadSafe> Create(const FString& InWarpPolicyType, const FString& InWarpPolicyName) override;
+};
diff --git a/Engine/Plugins/Runtime/nDisplay/Source/DisplayClusterWarp/Private/Policy/InFrustumFit/DisplayClusterWarpInFrustumFitPolicy_Math.cpp b/Engine/Plugins/Runtime/nDisplay/Source/DisplayClusterWarp/Private/Policy/InFrustumFit/DisplayClusterWarpInFrustumFitPolicy_Math.cpp
new file mode 100644
--- /dev/null
+++ b/Engine/Plugins/Runtime/nDisplay/Source/DisplayClusterWarp/Private/Policy/InFrustumFit/DisplayClusterWarpInFrustumFitPolicy_Math.cpp
@@ -0,0 +1,161 @@
+// Copyright Epic Games, Inc. All Rights Reserved.
+
+#include "Policy/InFrustumFit/DisplayClusterWarpInFrustumFitPolicy.h"
+
+#include "Components/DisplayClusterInFrustumFitCameraComponent.h"
+#include "Camera/CameraTypes.h"
+#include "HAL/IConsoleManager.h"
+#include "IDisplayClusterWarpBlend.h"
+
+// Debug: enable camera position fittig
+int32 GDisplayClusterWarpInFrustumFitPolicyEnableCameraPositionFit = 1;
+static FAutoConsoleVariableRef CVarDisplayClusterWarpInFrustumFitPolicyEnableCameraPositionFit(
+	TEXT("nDisplay.warp.InFrustumFit.EnableCameraPositionFit"),
+	GDisplayClusterWarpInFrustumFitPolicyEnableCameraPositionFit,
+	TEXT("(debug) Enable projection angles fitting (0 - disable)\n"),
+	ECVF_Default
+);
+
+// Debug: enable projection angle fitting
+int32 GDisplayClusterWarpInFrustumFitPolicyEnableProjectionAnglesFit = 1;
+static FAutoConsoleVariableRef CVarDisplayClusterWarpInFrustumFitPolicyEnableProjectionAnglesFit(
+	TEXT("nDisplay.warp.InFrustumFit.EnableProjectionAnglesFit"),
+	GDisplayClusterWarpInFrustumFitPolicyEnableProjectionAnglesFit,
+	TEXT("(debug) Enable projection angles fitting (0 - disable)\n"),
+	ECVF_Default
+);
+
+void FDisplayClusterWarpInFrustumFitPolicy::MakeGroupFrustumSymmetrical(bool bFixedViewTarget)
+{
+	if (bValidGroupFrustum)
+	{
+		if (bFixedViewTarget)
+		{
+			// If the view target is set to a fixed value instead of being computed by the group AABB, we do not want to alter the view direction,
+			// but make the frustum symmetric around that fixed direction. This involves expanding the asymmetric frustum so that its left and right,
+			// top and bottom angles are equal
+
+			const float MaxHorizontal = FMath::Max(FMath::Abs(GroupGeometryWarpProjection.Left), FMath::Abs(GroupGeometryWarpProjection.Right));
+			const float MaxVertical = FMath::Max(FMath::Abs(GroupGeometryWarpProjection.Top), FMath::Abs(GroupGeometryWarpProjection.Bottom));
+
+			GroupGeometryWarpProjection.Left = -MaxHorizontal;
+			GroupGeometryWarpProjection.Right = MaxHorizontal;
+
+			GroupGeometryWarpProjection.Bottom = -MaxVertical;
+			GroupGeometryWarpProjection.Top = MaxVertical;
+		}
+		else
+		{
+			// Otherwise, symmetrize the view frustum by computing a new view direction within the existing frustum
+			// First, compute the change in the view forward vector needed to make the frustum symmetric. This is done my storing a rotation on the original view forward vector
+			// computed from the group AABB that can then be applied to the each viewport's view direction when the render pass is performed for that viewport. The rotation
+			// angles can be computed by computing the difference between the current frustum's minimum edges and the half angle of the frustum's FOV
+			FDisplayClusterWarpProjection GroupWarpProjectionInDegrees = GroupGeometryWarpProjection;
+			GroupWarpProjectionInDegrees.ConvertProjectionAngles(EDisplayClusterWarpAngleUnit::Degrees);
+
+			const FVector2D GroupHalfFOVDegrees = 0.5 * FVector2D(GroupWarpProjectionInDegrees.Right - GroupWarpProjectionInDegrees.Left, GroupWarpProjectionInDegrees.Top - GroupWarpProjectionInDegrees.Bottom);
+			const FVector2D CorrectionAngles = GroupHalfFOVDegrees - FVector2D(FMath::Abs(GroupWarpProjectionInDegrees.Left), FMath::Abs(GroupWarpProjectionInDegrees.Bottom));
+
+			SymmetricForwardCorrection = FRotator(CorrectionAngles.Y, CorrectionAngles.X, 0);
+
+			// Finally, make the frustum angles themselves symmetric
+			GroupWarpProjectionInDegrees.Left = -0.5 * (GroupWarpProjectionInDegrees.Right - GroupWarpProjectionInDegrees.Left);
+			GroupWarpProjectionInDegrees.Right = -GroupWarpProjectionInDegrees.Left;
+
+			GroupWarpProjectionInDegrees.Bottom = -0.5 * (GroupWarpProjectionInDegrees.Top - GroupWarpProjectionInDegrees.Bottom);
+			GroupWarpProjectionInDegrees.Top = -GroupWarpProjectionInDegrees.Bottom;
+
+			GroupWarpProjectionInDegrees.ConvertProjectionAngles(EDisplayClusterWarpAngleUnit::Default);
+			GroupGeometryWarpProjection = GroupWarpProjectionInDegrees;
+		}
+	}
+}
+
+FDisplayClusterWarpProjection FDisplayClusterWarpInFrustumFitPolicy::ApplyInFrustumFit(UDisplayClusterInFrustumFitCameraComponent* InComponent, const FTransform& World2OriginTransform, const FDisplayClusterWarpProjection& InWarpProjection)
+{
+	check(InComponent);
+
+	FDisplayClusterWarpProjection OutWarpProjection(InWarpProjection);
+
+	FMinimalViewInfo CameraViewInfo;
+	InComponent->GetDesiredView(CameraViewInfo);
+
+	// Use camera position to render:
+	if (GDisplayClusterWarpInFrustumFitPolicyEnableCameraPositionFit)
+	{
+		OutWarpProjection.CameraRotation = World2OriginTransform.InverseTransformRotation(CameraViewInfo.Rotation.Quaternion()).Rotator();
+		OutWarpProjection.CameraLocation = World2OriginTransform.InverseTransformPosition(CameraViewInfo.Location);
+	}
+
+	// Fit the frustum to the rules:
+
+	FDisplayClusterWarpProjection ViewportAngles = InWarpProjection;
+	FDisplayClusterWarpProjection GroupAngles = GroupGeometryWarpProjection;
+
+	const FVector2D GeometryFOV(FMath::Abs(GroupAngles.Right - GroupAngles.Left), FMath::Abs(GroupAngles.Top - GroupAngles.Bottom));
+
+	// Convert frustum angles to group FOV space, normalized to 0..1
+	const FVector2D ViewportMin((ViewportAngles.Left - GroupAngles.Left) / GeometryFOV.X, (ViewportAngles.Bottom - GroupAngles.Bottom) / GeometryFOV.Y);
+	const FVector2D ViewportMax((ViewportAngles.Right - GroupAngles.Left) / GeometryFOV.X, (ViewportAngles.Top - GroupAngles.Bottom) / GeometryFOV.Y);
+
+	// And convert back to camera space:
+	const float CameraHalfFOVDegrees = CameraViewInfo.FOV * 0.5f;
+	const float CameraHalfFOVProjection = GroupGeometryWarpProjection.ConvertDegreesToProjection(CameraHalfFOVDegrees);
+	const FVector2D CameraHalfFOV(CameraHalfFOVProjection, CameraHalfFOVProjection / CameraViewInfo.AspectRatio);
+	const FVector2D GeometryHalfFOV = GeometryFOV * 0.5;
+
+	const FVector2D FinalHalfFOV = FindFrustumFit(InComponent->CameraProjectionMode, CameraHalfFOV, GeometryHalfFOV);
+
+	// Sample code: convert back to projection angles
+	if (GDisplayClusterWarpInFrustumFitPolicyEnableProjectionAnglesFit)
+	{
+		OutWarpProjection.Left = -FinalHalfFOV.X + ViewportMin.X * FinalHalfFOV.X * 2;
+		OutWarpProjection.Right = -FinalHalfFOV.X + ViewportMax.X * FinalHalfFOV.X * 2;
+		OutWarpProjection.Top = -FinalHalfFOV.Y + ViewportMax.Y * FinalHalfFOV.Y * 2;
+		OutWarpProjection.Bottom = -FinalHalfFOV.Y + ViewportMin.Y * FinalHalfFOV.Y * 2;
+	}
+
+	return OutWarpProjection;
+}
+
+FVector2D FDisplayClusterWarpInFrustumFitPolicy::FindFrustumFit(const EDisplayClusterWarpCameraProjectionMode InProjectionMode, const FVector2D& InCameraFOV, const FVector2D& InGeometryFOV)
+{
+	FVector2D OutFOV(InCameraFOV);
+
+	double DestAspectRatio = InGeometryFOV.X / InGeometryFOV.Y;
+
+	switch (InProjectionMode)
+	{
+	case EDisplayClusterWarpCameraProjectionMode::Fit:
+		if (InCameraFOV.Y * DestAspectRatio < InCameraFOV.X)
+		{
+			OutFOV.Y = InCameraFOV.Y;
+			OutFOV.X = OutFOV.Y * DestAspectRatio;
+		}
+		else
+		{
+
+			OutFOV.X = InCameraFOV.X;
+			OutFOV.Y = OutFOV.X / DestAspectRatio;
+		}
+		break;
+
+	case EDisplayClusterWarpCameraProjectionMode::Fill:
+		if (InCameraFOV.Y * DestAspectRatio > InCameraFOV.X)
+		{
+			OutFOV.Y = InCameraFOV.Y;
+			OutFOV.X = OutFOV.Y * DestAspectRatio;
+
+			// todo : add check ve 180 degree fov
+		}
+		else
+		{
+
+			OutFOV.X = InCameraFOV.X;
+			OutFOV.Y = OutFOV.X / DestAspectRatio;
+		}
+		break;
+	}
+
+	return OutFOV;
+}
diff --git a/Engine/Plugins/Runtime/nDisplay/Source/DisplayClusterWarp/Private/WarpBlend/DisplayClusterWarpBlend_GeometryContext.cpp b/Engine/Plugins/Runtime/nDisplay/Source/DisplayClusterWarp/Private/WarpBlend/DisplayClusterWarpBlend_GeometryContext.cpp
--- a/Engine/Plugins/Runtime/nDisplay/Source/DisplayClusterWarp/Private/WarpBlend/DisplayClusterWarpBlend_GeometryContext.cpp
+++ b/Engine/Plugins/Runtime/nDisplay/Source/DisplayClusterWarp/Private/WarpBlend/DisplayClusterWarpBlend_GeometryContext.cpp
@@ -19,8 +19,12 @@ bool FDisplayClusterWarpBlend_GeometryContext::UpdateGeometryContext(const doubl
 	Context.SurfaceViewNormal = Context.GeometryToOrigin.TransformVector(GeometryProxy.GeometryCache.SurfaceViewNormal);
 	Context.SurfaceViewPlane = Context.GeometryToOrigin.TransformVector(GeometryProxy.GeometryCache.SurfaceViewPlane);
 
-	Context.AABBox.Max = Context.GeometryToOrigin.TransformPosition(GeometryProxy.GeometryCache.AABBox.Max) * InWorldScale;
-	Context.AABBox.Min = Context.GeometryToOrigin.TransformPosition(GeometryProxy.GeometryCache.AABBox.Min) * InWorldScale;
+	// Use the built in bounding box transform method, which will correctly transform the box using origin and extent instead
+	// of just transforming the min and max properties, which will not correctly transform the box
+	FBox TransformedBox = GeometryProxy.GeometryCache.AABBox.TransformBy(Context.GeometryToOrigin);
+
+	Context.AABBox.Max = TransformedBox.Max * InWorldScale;
+	Context.AABBox.Min = TransformedBox.Min * InWorldScale;
 
 	return true;
 }
diff --git a/Engine/Plugins/Runtime/nDisplay/Source/DisplayClusterWarp/Public/Blueprints/DisplayClusterWarpBlueprint_Enums.h b/Engine/Plugins/Runtime/nDisplay/Source/DisplayClusterWarp/Public/Blueprints/DisplayClusterWarpBlueprint_Enums.h
--- a/Engine/Plugins/Runtime/nDisplay/Source/DisplayClusterWarp/Public/Blueprints/DisplayClusterWarpBlueprint_Enums.h
+++ b/Engine/Plugins/Runtime/nDisplay/Source/DisplayClusterWarp/Public/Blueprints/DisplayClusterWarpBlueprint_Enums.h
@@ -15,18 +15,20 @@
 UENUM()
 enum class EDisplayClusterWarpCameraProjectionMode : uint8
 {
-	/** Use camera as source of image.Preserving the camera aspect ratio.
-	 * The camera frame touches the geometry inside
-	 */
-	TouchInside UMETA(DisplayName = "Touch Inside"),
+	/** Fit the stage geometry entirely within the camera's frustum */
+	Fit UMETA(DisplayName = "Fit"),
 
-	/** Use camera as source of image.Preserving the camera aspect ratio.
-	 * The camera frame touches the geometry outside
-	 */
-	TouchOutside  UMETA(DisplayName = "Touch Outside"),
+	/** Fill the camera's frustum entire with the stage geometry */
+	Fill  UMETA(DisplayName = "Fill"),
+};
+
+/** A set of modes used to determine the view target of the stage's geometry frustum */
+UENUM()
+enum class EDisplayClusterWarpCameraViewTarget
+{
+	/** The camera will point in the direction of the geometric center of the stage's geometry */
+	GeometricCenter UMETA(DisplayName = "Geometric Center"),
 
-	/** Use camera as source of image.
-	 * Stretch the camera image to geometry (the camera aspect ratio will be lost)
-	 */
-	Stretch  UMETA(DisplayName = "Stretch"),
+	/** The camera will point in the same direction as the frustum fit view origin */
+	MatchViewOrigin UMETA(DisplayName = "Match View Origin")
 };
diff --git a/Engine/Plugins/Runtime/nDisplay/Source/DisplayClusterWarp/Public/Components/DisplayClusterInFrustumFitCameraComponent.h b/Engine/Plugins/Runtime/nDisplay/Source/DisplayClusterWarp/Public/Components/DisplayClusterInFrustumFitCameraComponent.h
new file mode 100644
--- /dev/null
+++ b/Engine/Plugins/Runtime/nDisplay/Source/DisplayClusterWarp/Public/Components/DisplayClusterInFrustumFitCameraComponent.h
@@ -0,0 +1,88 @@
+// Copyright Epic Games, Inc. All Rights Reserved.
+
+#pragma once
+#include "CoreMinimal.h"
+
+#include "Components/DisplayClusterCameraComponent.h"
+#include "Blueprints/DisplayClusterWarpBlueprint_Enums.h"
+
+#include "DisplayClusterInFrustumFitCameraComponent.generated.h"
+
+class ACineCameraActor;
+class UCameraComponent;
+class IDisplayClusterWarpPolicy;
+
+/**
+ * 3D point in space used to project the camera view onto a group of nDisplay viewports.
+ * Support projection policies: mpcdi/pfm 2d/a3d, mesh.
+ */
+UCLASS(ClassGroup = (DisplayCluster), meta = (BlueprintSpawnableComponent, DisplayName = "nDisplay InFrustumFit View Origin"))
+class DISPLAYCLUSTERWARP_API UDisplayClusterInFrustumFitCameraComponent : public UDisplayClusterCameraComponent
+{
+	GENERATED_BODY()
+
+public:
+	UDisplayClusterInFrustumFitCameraComponent(const FObjectInitializer& ObjectInitializer);
+
+	//~Begin UActorComponent
+	virtual void OnRegister() override;
+	virtual void TickComponent(float DeltaTime, ELevelTick TickType, FActorComponentTickFunction* ThisTickFunction) override;
+	//~~~End UActorComponent
+
+	//~Begin UDisplayClusterCameraComponent
+	virtual void GetDesiredView(FMinimalViewInfo& InOutViewInfo, float* OutCustomNearClippingPlane = nullptr) override;
+	virtual bool ShouldUseEntireClusterViewports(IDisplayClusterViewportManager* InViewportManager) const override;
+	virtual IDisplayClusterWarpPolicy* GetWarpPolicy(IDisplayClusterViewportManager* InViewportManager) override;
+	//~~End UDisplayClusterCameraComponent
+
+#if WITH_EDITOR
+	virtual void PostEditChangeProperty(FPropertyChangedEvent& PropertyChangedEvent) override;
+	virtual bool GetEditorPreviewInfo(float DeltaTime, FMinimalViewInfo& ViewOut) override;
+	virtual TSharedPtr<SWidget> GetCustomEditorPreviewWidget() override;
+
+	void InvalidatePreviewState() { bRefreshPreviewState = true; }
+#endif
+
+public:
+	/** true, if camera projection is used. */
+	bool IsEnabled() const;
+
+	/** Return external camera component. */
+	UCameraComponent* GetExternalCameraComponent() const;
+
+public:
+	/** Camera projection mode is used. */
+	UPROPERTY(EditAnywhere, Category = "Projection")
+	bool bEnableCameraProjection = true;
+
+	/** Enable special rendering mode for all viewports using this viewpoint. */
+	UPROPERTY(EditAnywhere, Category = "Projection", meta = (EditCondition = "bEnableCameraProjection"))
+	EDisplayClusterWarpCameraProjectionMode CameraProjectionMode = EDisplayClusterWarpCameraProjectionMode::Fit;
+
+	/** Indicates which camera facing mode is used when frustum fitting the stage geometry */
+	UPROPERTY(EditAnywhere, Category = "Projection", meta = (EditCondition = "bEnableCameraProjection"))
+	EDisplayClusterWarpCameraViewTarget CameraViewTarget = EDisplayClusterWarpCameraViewTarget::GeometricCenter;
+
+	/** Use a specific actor camera instead of a game camera. */
+	UPROPERTY(EditAnywhere, Category = "Projection", meta = (EditCondition = "bEnableCameraProjection"))
+	TSoftObjectPtr<ACineCameraActor> ExternalCameraActor;
+
+	/** Allows you to override the PostProcess of the viewport from the camera. */
+	UPROPERTY(EditAnywhere, Category = "Projection", meta = (EditCondition = "bEnableCameraProjection"))
+	bool bUseCameraPostprocess = false;
+
+#if WITH_EDITORONLY_DATA
+	UPROPERTY(EditAnywhere, Category = "Projection", meta = (EditCondition = "bEnableCameraProjection"))
+	bool bShowPreviewFrustumFit = false;
+#endif
+
+private:
+	// a unique type of warp policy for this component
+	// this policy class knows the properties of the component and implements the corresponding logic
+	TSharedPtr<IDisplayClusterWarpPolicy, ESPMode::ThreadSafe> WarpPolicy;
+
+#if WITH_EDITOR
+	/** Indicates that the state of the preview meshes for the frustum fit are invalid and need to be refreshed on the next tick */
+	bool bRefreshPreviewState = true;
+#endif
+};
