nDisplay: Add checkbox to enable TechViz for preview rendering.

diff --git a/Engine/Plugins/Runtime/nDisplay/Source/DisplayCluster/Private/Game/EngineClasses/Scene/DisplayClusterPreviewComponent.cpp b/Engine/Plugins/Runtime/nDisplay/Source/DisplayCluster/Private/Game/EngineClasses/Scene/DisplayClusterPreviewComponent.cpp
--- a/Engine/Plugins/Runtime/nDisplay/Source/DisplayCluster/Private/Game/EngineClasses/Scene/DisplayClusterPreviewComponent.cpp
+++ b/Engine/Plugins/Runtime/nDisplay/Source/DisplayCluster/Private/Game/EngineClasses/Scene/DisplayClusterPreviewComponent.cpp
@@ -285,6 +285,14 @@ bool UDisplayClusterPreviewComponent::UpdatePreviewMesh()
 				if (PreviewMesh != nullptr)
 				{
 					PreviewMesh->SetCastShadow(false);
+
+					if (RootActor && RootActor->IsTechVizPreviewRenderingEnabled())
+					{
+						// Force set these properties to on with techviz rendering enabled. It won't work properly
+						// with them off. We don't need to force them off when techviz rendering is disabled, PostEditChangeProperty
+						// on the root actor will take care of that.
+						UE::DisplayClusterDisplayDeviceUtils::ConfigureTechVizForMesh(PreviewMesh, true);
+					}
 				}
 
 				if (MovablePreviewMesh != nullptr)
@@ -357,10 +365,13 @@ void UDisplayClusterPreviewComponent::UpdatePreviewMaterial()
 		}
 
 		// Allow display device to perform any processing on the material instance.
-		if (UDisplayClusterDisplayDeviceBaseComponent* CachedComponent =
-			Cast<UDisplayClusterDisplayDeviceBaseComponent>(CachedDisplayDevice.GetComponent(GetOwner())))
+		if (RootActor && RootActor->IsTechVizPreviewRenderingEnabled())
 		{
-			CachedComponent->OnUpdatePreviewMaterialInstance(PreviewMaterialInstance);
+			if (UDisplayClusterDisplayDeviceBaseComponent* CachedComponent =
+				Cast<UDisplayClusterDisplayDeviceBaseComponent>(CachedDisplayDevice.GetComponent(GetOwner())))
+			{
+				CachedComponent->OnUpdatePreviewMaterialInstance(PreviewMaterialInstance);
+			}
 		}
 	}
 
@@ -433,14 +444,17 @@ void UDisplayClusterPreviewComponent::UpdatePreviewRenderTarget()
 			DestinationRenderTarget = SourceRenderTarget;
 
 			// Perform any render passes on the render target by the display device
-			if (UDisplayClusterDisplayDeviceBaseComponent* CachedComponent =
-			Cast<UDisplayClusterDisplayDeviceBaseComponent>(CachedDisplayDevice.GetComponent(GetOwner())))
+			if (RootActor->IsTechVizPreviewRenderingEnabled())
 			{
-				if (CachedComponent->IsRenderPassEnabled() && !RootActor->IsTemplate())
+				if (UDisplayClusterDisplayDeviceBaseComponent* CachedComponent =
+				Cast<UDisplayClusterDisplayDeviceBaseComponent>(CachedDisplayDevice.GetComponent(GetOwner())))
 				{
-					UpdateRenderTargetImpl(&DisplayDeviceRenderTarget);
-					CachedComponent->RenderPass_GameThread(SourceRenderTarget, DisplayDeviceRenderTarget);
-					DestinationRenderTarget = DisplayDeviceRenderTarget;
+					if (CachedComponent->IsRenderPassEnabled() && !RootActor->IsTemplate())
+					{
+						UpdateRenderTargetImpl(&DisplayDeviceRenderTarget);
+						CachedComponent->RenderPass_GameThread(SourceRenderTarget, DisplayDeviceRenderTarget);
+						DestinationRenderTarget = DisplayDeviceRenderTarget;
+					}
 				}
 			}
 		}
diff --git a/Engine/Plugins/Runtime/nDisplay/Source/DisplayCluster/Private/Game/EngineClasses/Scene/DisplayClusterRootActorEditor.cpp b/Engine/Plugins/Runtime/nDisplay/Source/DisplayCluster/Private/Game/EngineClasses/Scene/DisplayClusterRootActorEditor.cpp
--- a/Engine/Plugins/Runtime/nDisplay/Source/DisplayCluster/Private/Game/EngineClasses/Scene/DisplayClusterRootActorEditor.cpp
+++ b/Engine/Plugins/Runtime/nDisplay/Source/DisplayCluster/Private/Game/EngineClasses/Scene/DisplayClusterRootActorEditor.cpp
@@ -101,6 +101,20 @@ void ADisplayClusterRootActor::ResetPreviewInternals_Editor()
 	ResetClusterNodePreviewRendering_Editor();
 }
 
+void ADisplayClusterRootActor::ConfigureTechViz_Editor()
+{
+	TArray<UDisplayClusterPreviewComponent*> Components;
+	GetComponents(Components);
+
+	for (UDisplayClusterPreviewComponent* PreviewComponent : Components)
+	{
+		if (UMeshComponent* Mesh = PreviewComponent->GetPreviewMesh())
+		{
+			UE::DisplayClusterDisplayDeviceUtils::ConfigureTechVizForMesh(Mesh, bEnablePreviewTechViz);
+		}
+	}
+}
+
 void ADisplayClusterRootActor::Constructor_Editor()
 {
 	// Allow tick in editor for preview rendering
@@ -898,6 +912,11 @@ void ADisplayClusterRootActor::PostEditChangeProperty(FPropertyChangedEvent& Pro
 		PreviewRenderFrame.Reset();
 		bReinitializeActor = false;
 	}
+	else if (PropertyName == GET_MEMBER_NAME_CHECKED(ADisplayClusterRootActor, bEnablePreviewTechViz))
+	{
+		ConfigureTechViz_Editor();
+	}
+	
 	if (bReinitializeActor)
 	{
 		InitializeRootActor();
diff --git a/Engine/Plugins/Runtime/nDisplay/Source/DisplayCluster/Private/Game/EngineClasses/Scene/DisplayDevice/DisplayClusterDisplayDeviceUtils.cpp b/Engine/Plugins/Runtime/nDisplay/Source/DisplayCluster/Private/Game/EngineClasses/Scene/DisplayDevice/DisplayClusterDisplayDeviceUtils.cpp
--- a/Engine/Plugins/Runtime/nDisplay/Source/DisplayCluster/Private/Game/EngineClasses/Scene/DisplayDevice/DisplayClusterDisplayDeviceUtils.cpp
+++ b/Engine/Plugins/Runtime/nDisplay/Source/DisplayCluster/Private/Game/EngineClasses/Scene/DisplayDevice/DisplayClusterDisplayDeviceUtils.cpp
@@ -96,4 +96,26 @@ UDisplayClusterDisplayDeviceBaseComponent* UE::DisplayClusterDisplayDeviceUtils:
 	return nullptr;
 }
 
+void UE::DisplayClusterDisplayDeviceUtils::ConfigureTechVizForMesh(UMeshComponent* InMeshComponent,
+	bool bIsTechVizEnabled)
+{
+	check(InMeshComponent);
+	
+	if (bIsTechVizEnabled)
+	{
+		InMeshComponent->bAffectDynamicIndirectLighting = true;
+		InMeshComponent->bAffectIndirectLightingWhileHidden = true;
+	}
+	else
+	{
+		// When disabling just revert to the archetype values. The original values at the time of enabling
+		// techviz aren't saved, and techviz is enabled by default so they'll always be set to true on new meshes.
+		if (const UMeshComponent* MeshArchetype = Cast<UMeshComponent>(InMeshComponent->GetArchetype()))
+		{
+			InMeshComponent->bAffectDynamicIndirectLighting = MeshArchetype->bAffectDynamicIndirectLighting;
+			InMeshComponent->bAffectIndirectLightingWhileHidden = MeshArchetype->bAffectIndirectLightingWhileHidden;
+		}
+	}
+}
+
 #endif
diff --git a/Engine/Plugins/Runtime/nDisplay/Source/DisplayCluster/Private/Game/EngineClasses/Scene/DisplayDevice/DisplayClusterDisplayDeviceUtils.h b/Engine/Plugins/Runtime/nDisplay/Source/DisplayCluster/Private/Game/EngineClasses/Scene/DisplayDevice/DisplayClusterDisplayDeviceUtils.h
--- a/Engine/Plugins/Runtime/nDisplay/Source/DisplayCluster/Private/Game/EngineClasses/Scene/DisplayDevice/DisplayClusterDisplayDeviceUtils.h
+++ b/Engine/Plugins/Runtime/nDisplay/Source/DisplayCluster/Private/Game/EngineClasses/Scene/DisplayDevice/DisplayClusterDisplayDeviceUtils.h
@@ -4,6 +4,7 @@
 
 class UDisplayClusterConfigurationViewport;
 class UDisplayClusterDisplayDeviceBaseComponent;
+class UMeshComponent;
 
 namespace UE::DisplayClusterDisplayDeviceUtils
 {
@@ -18,5 +19,13 @@ namespace UE::DisplayClusterDisplayDeviceUtils
 	 * @return The display device this viewport uses or nullptr.
 	 */
 	UDisplayClusterDisplayDeviceBaseComponent* FindAndSyncDisplayDeviceFromViewport(UDisplayClusterConfigurationViewport* InViewport, bool* bOutIsDefaultDevice = nullptr);
+
+	/**
+	 * Configure techviz parameters on a given mesh component.
+	 * 
+	 * @param InMeshComponent The impacted mesh component.
+	 * @param bIsTechVizEnabled Whether techviz should be enabled or disabled on this mesh.
+	 */
+	void ConfigureTechVizForMesh(UMeshComponent* InMeshComponent, bool bIsTechVizEnabled);
 #endif
 }
diff --git a/Engine/Plugins/Runtime/nDisplay/Source/DisplayCluster/Public/DisplayClusterRootActor.h b/Engine/Plugins/Runtime/nDisplay/Source/DisplayCluster/Public/DisplayClusterRootActor.h
--- a/Engine/Plugins/Runtime/nDisplay/Source/DisplayCluster/Public/DisplayClusterRootActor.h
+++ b/Engine/Plugins/Runtime/nDisplay/Source/DisplayCluster/Public/DisplayClusterRootActor.h
@@ -384,6 +384,10 @@ public:
 	UPROPERTY(EditAnywhere, BlueprintReadWrite, Category = "Editor Preview", meta = (DisplayName = "Enable Post Process"), BlueprintSetter = SetPreviewEnablePostProcess)
 	bool bPreviewEnablePostProcess = false;
 
+	/** Configure the root actor for tech viz rendering with preview components. */
+	UPROPERTY(EditAnywhere, BlueprintReadWrite, Category = "Editor Preview")
+	bool bEnablePreviewTechViz = true;
+	
 	/** Freeze preview render.  This will impact editor performance. */
 	UPROPERTY(EditAnywhere, BlueprintReadWrite, Category = "Editor Preview", meta = (DisplayName = "Freeze Editor Preview"))
 	bool bFreezePreviewRender = false;
@@ -552,11 +556,17 @@ public:
 
 	/** Retrieve the default display device, creating it if it doesn't exist */
 	UDisplayClusterDisplayDeviceBaseComponent* GetDefaultDisplayDevice();
+
+	/** If the preview should use techviz rendering. */
+	bool IsTechVizPreviewRenderingEnabled() const { return bEnablePreviewTechViz; }
 	
 protected:
 	FString GeneratePreviewComponentName_Editor(const FString& NodeId, const FString& ViewportId) const;
 	void ResetPreviewInternals_Editor();
 
+	/** Configure preview components for techviz rendering. */
+	void ConfigureTechViz_Editor();
+
 	bool ImplUpdatePreviewConfiguration_Editor(const FString& InClusterNodeId);
 
 	void ImplRenderPreview_Editor();
