Fix nDisplay InCameraVFX rendering with Substrate.

diff --git a/Engine/Plugins/Runtime/nDisplay/Shaders/Private/MeshProjectionShaders.usf b/Engine/Plugins/Runtime/nDisplay/Shaders/Private/MeshProjectionShaders.usf
--- a/Engine/Plugins/Runtime/nDisplay/Shaders/Private/MeshProjectionShaders.usf
+++ b/Engine/Plugins/Runtime/nDisplay/Shaders/Private/MeshProjectionShaders.usf
@@ -121,6 +121,50 @@ void ClipDivergentPixel(FMeshProjectionVSToPS Input, float DivergenceThreshold)
 	clip(DivergenceThreshold - Divergence);
 }
 
+float4 GetMaterialProperties(in FPixelMaterialInputs PixelMaterialInputs, in FMaterialPixelParameters MaterialParameters)
+{
+	float3 OutEmissive = 0;
+	float OutOpacity = 0;
+#if TEMPLATE_USES_SUBSTRATE
+	// Initialise a Substrate header with normal in registers
+	FSubstrateData SubstrateData = PixelMaterialInputs.GetFrontSubstrateData();
+	FSubstratePixelHeader SubstratePixelHeader = MaterialParameters.GetFrontSubstrateHeader();
+	SubstratePixelHeader.IrradianceAO.MaterialAO = GetMaterialAmbientOcclusion(PixelMaterialInputs);
+
+	float TotalCoverage = 1.f;
+	if (SubstratePixelHeader.SubstrateTree.BSDFCount > 0)
+	{
+		const float3 V = MaterialParameters.CameraVector;
+		const FSubstrateIntegrationSettings Settings = InitSubstrateIntegrationSettings(false /*bForceFullyRough*/, true /*bRoughDiffuseEnabled*/, 0 /*PeelLayersAboveDepth*/, false/*bRoughnessTracking*/);
+		float3 TotalTransmittancePreCoverage = 0;
+		SubstratePixelHeader.SubstrateUpdateTree(SubstrateData, V, Settings, TotalCoverage, TotalTransmittancePreCoverage);
+		OutOpacity = TotalCoverage;
+
+		// Clip pixels that will never receive any lumen front material reflections.
+		// We cannot generalise and use GetMaterialClippingVelocity for that because this function clips with a 1/255 threshold...that can result in NaN refelctions if some pixels are not genrated correctly.
+		clip(TotalCoverage <= 0.0f ? -1.0f : 1.0f);
+
+		// Accumulate emissive contribution
+		float TopLayerTotalWeight = 0.0f;
+		SUBSTRATE_UNROLL_N(SUBSTRATE_CLAMPED_BSDF_COUNT)
+		for (int BSDFIdx = 0; BSDFIdx < SubstratePixelHeader.SubstrateTree.BSDFCount; ++BSDFIdx)
+		{
+			#define CurrentBSDF SubstratePixelHeader.SubstrateTree.BSDFs[BSDFIdx]
+			if (SubstrateIsBSDFVisible(CurrentBSDF))
+			{
+				OutEmissive += CurrentBSDF.LuminanceWeightV * BSDF_GETEMISSIVE(CurrentBSDF);
+			}
+			#undef CurrentBSDF
+		}
+	}
+#else
+	OutEmissive = GetMaterialEmissive(PixelMaterialInputs);
+	OutOpacity = GetMaterialOpacity(PixelMaterialInputs);
+#endif
+
+	return float4(OutEmissive, OutOpacity);
+}
+
 void MainPS(FMeshProjectionVSToPS Input OPTIONAL_IsFrontFace, out float4 OutputColor : SV_Target0)
 {
 	FMaterialPixelParameters MaterialParameters = GetMaterialPixelParameters(Input.FactoryInterpolants, Input.Position);
@@ -128,11 +172,8 @@ void MainPS(FMeshProjectionVSToPS Input OPTIONAL_IsFrontFace, out float4 OutputC
 	CalcMaterialParameters(MaterialParameters, PixelMaterialInputs, Input.Position, bIsFrontFace);
 
 	GetMaterialCoverageAndClipping(MaterialParameters, PixelMaterialInputs);
-	
-	const float3 Emissive = GetMaterialEmissive(PixelMaterialInputs);
-	const float Opacity = GetMaterialOpacity(PixelMaterialInputs);
-	
-	OutputColor = float4(Emissive, Opacity);
+
+	OutputColor = GetMaterialProperties(PixelMaterialInputs, MaterialParameters);
 	
 	static const float AzimuthFeatherStart = radians(180 - 45);
 	static const float AzimuthFeatherEnd   = radians(180 - 30);
