nDisplay: Refactor some private headers to public or internal. Update namespaces with UE prefix. Break ClusterUtils into private and public declarations since several generic cluster/viewport utility functions were already exported and licensees may find them useful.

diff --git a/Engine/Plugins/Runtime/nDisplay/Source/DisplayCluster/Internal/Cluster/IPDisplayClusterClusterManager.h b/Engine/Plugins/Runtime/nDisplay/Source/DisplayCluster/Internal/Cluster/IPDisplayClusterClusterManager.h
new file mode 100644
--- /dev/null
+++ b/Engine/Plugins/Runtime/nDisplay/Source/DisplayCluster/Internal/Cluster/IPDisplayClusterClusterManager.h
@@ -0,0 +1,49 @@
+// Copyright Epic Games, Inc. All Rights Reserved.
+
+#pragma once
+
+#include "Cluster/IDisplayClusterClusterManager.h"
+#include "IPDisplayClusterManager.h"
+
+#include "Misc/Timecode.h"
+#include "Misc/FrameRate.h"
+
+class IDisplayClusterClusterNodeController;
+class IDisplayClusterFailoverNodeController;
+class FJsonObject;
+struct FQualifiedFrameTime;
+
+
+/**
+ * Cluster manager private interface
+ */
+class IPDisplayClusterClusterManager :
+	public IDisplayClusterClusterManager,
+	public IPDisplayClusterManager
+{
+public:
+	virtual ~IPDisplayClusterClusterManager() = default;
+
+public:
+	virtual IDisplayClusterClusterNodeController*  GetClusterNodeController() const = 0;
+	virtual IDisplayClusterFailoverNodeController* GetFailoverNodeController() const = 0;
+
+	// Time data sync
+	virtual void SyncTimeData() = 0;
+	virtual void ExportTimeData(      double& OutDeltaTime,      double& OutGameTime,      TOptional<FQualifiedFrameTime>& OutFrameTime) = 0;
+	virtual void ImportTimeData(const double& InDeltaTime, const double& InGameTime, const TOptional<FQualifiedFrameTime>& InFrameTime) = 0;
+
+	// Objects sync
+	virtual void SyncObjects(EDisplayClusterSyncGroup SyncGroup) = 0;
+	virtual void ExportObjectsData(const EDisplayClusterSyncGroup InSyncGroup,       TMap<FString, FString>& OutObjectsData) = 0;
+	virtual void ImportObjectsData(const EDisplayClusterSyncGroup InSyncGroup, const TMap<FString, FString>& InObjectsData) = 0;
+
+	// Cluster events sync
+	virtual void SyncEvents() = 0;
+	virtual void ExportEventsData(      TArray<TSharedPtr<FDisplayClusterClusterEventJson, ESPMode::ThreadSafe>>& OutJsonEvents,      TArray<TSharedPtr<FDisplayClusterClusterEventBinary, ESPMode::ThreadSafe>>& OutBinaryEvents) = 0;
+	virtual void ImportEventsData(const TArray<TSharedPtr<FDisplayClusterClusterEventJson, ESPMode::ThreadSafe>>& InJsonEvents, const TArray<TSharedPtr<FDisplayClusterClusterEventBinary, ESPMode::ThreadSafe>>& InBinaryEvents) = 0;
+
+	// Native input sync
+	virtual void ExportNativeInputData(TMap<FString, FString>&  OutNativeInputData) = 0;
+	virtual void ImportNativeInputData(TMap<FString, FString>&& InNativeInputData) = 0;
+};
diff --git a/Engine/Plugins/Runtime/nDisplay/Source/DisplayCluster/Internal/IPDisplayCluster.h b/Engine/Plugins/Runtime/nDisplay/Source/DisplayCluster/Internal/IPDisplayCluster.h
new file mode 100644
--- /dev/null
+++ b/Engine/Plugins/Runtime/nDisplay/Source/DisplayCluster/Internal/IPDisplayCluster.h
@@ -0,0 +1,30 @@
+// Copyright Epic Games, Inc. All Rights Reserved.
+
+#pragma once
+
+#include "IDisplayCluster.h"
+#include "IPDisplayClusterManager.h"
+
+class IPDisplayClusterRenderManager;
+class IPDisplayClusterClusterManager;
+class IPDisplayClusterConfigManager;
+class IPDisplayClusterGameManager;
+class ADisplayClusterSettings;
+
+
+/**
+ * Private module interface
+ */
+class IPDisplayCluster
+	: public IDisplayCluster
+	, public IPDisplayClusterManager
+{
+public:
+	virtual ~IPDisplayCluster() = default;
+
+public:
+	virtual IPDisplayClusterRenderManager*    GetPrivateRenderMgr() const = 0;
+	virtual IPDisplayClusterClusterManager*   GetPrivateClusterMgr() const = 0;
+	virtual IPDisplayClusterConfigManager*    GetPrivateConfigMgr() const = 0;
+	virtual IPDisplayClusterGameManager*      GetPrivateGameMgr() const = 0;
+};
diff --git a/Engine/Plugins/Runtime/nDisplay/Source/DisplayCluster/Internal/IPDisplayClusterManager.h b/Engine/Plugins/Runtime/nDisplay/Source/DisplayCluster/Internal/IPDisplayClusterManager.h
new file mode 100644
--- /dev/null
+++ b/Engine/Plugins/Runtime/nDisplay/Source/DisplayCluster/Internal/IPDisplayClusterManager.h
@@ -0,0 +1,64 @@
+// Copyright Epic Games, Inc. All Rights Reserved.
+
+#pragma once
+
+#include "CoreMinimal.h"
+
+#include "DisplayClusterEnums.h"
+
+class UDisplayClusterConfigurationData;
+
+
+/**
+ * Private manager interface
+ */
+class IPDisplayClusterManager
+{
+public:
+	virtual ~IPDisplayClusterManager() = default;
+
+public:
+	// Called at start to initialize internals
+	virtual bool Init(EDisplayClusterOperationMode OperationMode)
+	{ return true; }
+
+	// Called before application/Editor exit to release internals
+	virtual void Release()
+	{ }
+
+	// Called on each session start before first level start (before the first tick)
+	virtual bool StartSession(UDisplayClusterConfigurationData* InConfigData, const FString& InNodeId)
+	{ return true; }
+
+	// Called on each session end at early step before exit (before UGameEngine::Preexit)
+	virtual void EndSession()
+	{ }
+
+	// Called each time a new game level starts
+	virtual bool StartScene(UWorld* InWorld)
+	{ return true; }
+
+	// Called when current level is going to be closed (i.e. when loading new map)
+	virtual void EndScene()
+	{ }
+
+	// Called before start every frame
+	virtual void StartFrame(uint64 FrameNum)
+	{ }
+
+	// Called before start every frame
+	virtual void EndFrame(uint64 FrameNum)
+	{ }
+
+	// Called every frame before world Tick
+	virtual void PreTick(float DeltaSeconds)
+	{ }
+
+	// Called every frame during world Tick
+	virtual void Tick(float DeltaSeconds)
+	{ }
+
+	// Called every frame during world Tick
+	virtual void PostTick(float DeltaSeconds)
+	{ }
+};
diff --git a/Engine/Plugins/Runtime/nDisplay/Source/DisplayCluster/Private/Cluster/IPDisplayClusterClusterManager.h b/Engine/Plugins/Runtime/nDisplay/Source/DisplayCluster/Private/Cluster/IPDisplayClusterClusterManager.h
deleted file mode 100644
diff --git a/Engine/Plugins/Runtime/nDisplay/Source/DisplayCluster/Private/IPDisplayCluster.h b/Engine/Plugins/Runtime/nDisplay/Source/DisplayCluster/Private/IPDisplayCluster.h
deleted file mode 100644
diff --git a/Engine/Plugins/Runtime/nDisplay/Source/DisplayCluster/Private/IPDisplayClusterManager.h b/Engine/Plugins/Runtime/nDisplay/Source/DisplayCluster/Private/IPDisplayClusterManager.h
deleted file mode 100644
diff --git a/Engine/Plugins/Runtime/nDisplay/Source/DisplayClusterConfigurator/Private/ClusterConfiguration/DisplayClusterConfiguratorClusterEditorUtils.cpp b/Engine/Plugins/Runtime/nDisplay/Source/DisplayClusterConfigurator/Private/ClusterConfiguration/DisplayClusterConfiguratorClusterEditorUtils.cpp
new file mode 100644
--- /dev/null
+++ b/Engine/Plugins/Runtime/nDisplay/Source/DisplayClusterConfigurator/Private/ClusterConfiguration/DisplayClusterConfiguratorClusterEditorUtils.cpp
@@ -0,0 +1,810 @@
+// Copyright Epic Games, Inc. All Rights Reserved.
+
+#include "DisplayClusterConfiguratorClusterEditorUtils.h"
+
+#include "ClusterConfiguration/DisplayClusterConfiguratorClusterUtils.h"
+#include "DisplayClusterConfiguratorBlueprintEditor.h"
+#include "DisplayClusterConfigurationTypes.h"
+#include "DisplayClusterConfiguratorPropertyUtils.h"
+#include "DisplayClusterConfiguratorStyle.h"
+#include "ClusterConfiguration/SDisplayClusterConfiguratorNewClusterItemDialog.h"
+#include "ClusterConfiguration/ViewModels/DisplayClusterConfiguratorClusterNodeViewModel.h"
+#include "Views/DragDrop/DisplayClusterConfiguratorClusterNodeDragDropOp.h"
+#include "Views/DragDrop/DisplayClusterConfiguratorViewportDragDropOp.h"
+
+#include "Factories.h"
+#include "ScopedTransaction.h"
+#include "Framework/Application/SlateApplication.h"
+#include "HAL/PlatformApplicationMisc.h"
+#include "UnrealExporter.h"
+#include "Editor/Transactor.h"
+#include "Editor/TransBuffer.h"
+#include "Exporters/Exporter.h"
+#include "Input/DragAndDrop.h"
+#include "Widgets/SBoxPanel.h"
+#include "Widgets/Input/SCheckBox.h"
+#include "Widgets/Text/STextBlock.h"
+
+#define LOCTEXT_NAMESPACE "FDisplayClusterConfiguratorClusterEditorUtils"
+
+namespace UE::DisplayClusterConfiguratorClusterEditorUtils
+{
+	// All transactions performed within this scope can be considered temporary and will be cleared when out of scope.
+	// Transactions prior to the scope cannot be undone while in scope.
+	struct FScopedTemporaryTransactions
+	{
+		// The index of the transaction queue when entering the scope
+		int32 StartTransactionIdx = 0;
+		
+		FScopedTemporaryTransactions()
+		{
+			if (GEditor)
+			{
+				// Prevent undoing any actions before the scope started
+				GEditor->Trans->SetUndoBarrier();
+
+				// Record the transaction idx before we start our scoped transactions
+				StartTransactionIdx = GEditor->Trans->GetCurrentUndoBarrier();
+
+				// Broadcasting changes here can take a big performance hit and isn't necessary since we only use this in a modal window
+				GEditor->bSuspendBroadcastPostUndoRedo = true;
+			}
+		}
+
+		~FScopedTemporaryTransactions()
+		{
+			if (GEditor)
+			{
+				GEditor->bSuspendBroadcastPostUndoRedo = false;
+
+				// Clear all transactions that occurred while in scope, accounting for transactions that were
+				// undone. None of these transactions should be allowed to be undone or redone from this point forward
+				{
+					UTransBuffer* TransBuffer = CastChecked<UTransBuffer>(GEditor->Trans);
+					const int32 ScopedTransactionsPerformed = TransBuffer->UndoBuffer.Num() - StartTransactionIdx;
+				
+					if (ScopedTransactionsPerformed > 0
+						&& StartTransactionIdx >= 0
+						&& StartTransactionIdx < TransBuffer->UndoBuffer.Num()
+						&& (StartTransactionIdx + ScopedTransactionsPerformed - 1) < TransBuffer->UndoBuffer.Num())
+					{
+						// Remove only the transactions performed within scope
+						TransBuffer->UndoBuffer.RemoveAt(StartTransactionIdx, ScopedTransactionsPerformed);
+
+						// We can't restore redo items prior to the scope since they would have been cleared while
+						// performing transactions within the scope
+						TransBuffer->UndoCount = 0;
+
+						// Need to broadcast for other systems since we are changing the undo buffer manually
+						TransBuffer->OnUndoBufferChanged().Broadcast();
+					}
+				}
+
+				// Allow undos again from before the scope started
+				GEditor->Trans->RemoveUndoBarrier();
+			}
+		}
+	};
+
+	const FVector2D NewClusterItemDialogSize = FVector2D(410, 512);
+	const FString DefaultNewClusterNodeName = TEXT("Node");
+	const FString DefaultNewViewportName = TEXT("VP");
+}
+
+using namespace UE::DisplayClusterConfiguratorClusterUtils;
+
+UDisplayClusterConfigurationClusterNode* UE::DisplayClusterConfiguratorClusterEditorUtils::CreateNewClusterNodeFromDialog(
+	const TSharedRef<FDisplayClusterConfiguratorBlueprintEditor>& Toolkit,
+	UDisplayClusterConfigurationCluster* Cluster,
+	const FDisplayClusterConfigurationRectangle& PresetRect,
+	TSharedPtr<FScopedTransaction>& OutTransaction,
+	FString PresetHost)
+{
+	UDisplayClusterConfigurationClusterNode* NodeTemplate = NewObject<UDisplayClusterConfigurationClusterNode>(Toolkit->GetBlueprintObj(), NAME_None, RF_Transactional | RF_ArchetypeObject | RF_Public);
+	NodeTemplate->WindowRect = FDisplayClusterConfigurationRectangle(PresetRect);
+	NodeTemplate->Host = PresetHost;
+
+	TArray<FString> ParentItems;
+	ParentItems.Add("Cluster");
+
+	const FString InitialName = GetUniqueNameForClusterNode(DefaultNewClusterNodeName, Cluster, true);
+	
+	bool bAutoposition = true;
+	bool bAddViewport = true;
+	const TSharedRef<SWidget> ClusterNodeFooter = SNew(SVerticalBox)
+		+SVerticalBox::Slot()
+		.AutoHeight()
+		.Padding(2.0f)
+		[
+			SNew(SCheckBox)
+			.IsChecked_Lambda([&bAutoposition]() { return bAutoposition ? ECheckBoxState::Checked : ECheckBoxState::Unchecked; })
+			.OnCheckStateChanged_Lambda([&bAutoposition](ECheckBoxState NewState) { bAutoposition = NewState == ECheckBoxState::Checked; })
+			.ToolTipText(LOCTEXT("AddNewClusterNode_AutoPositionToolTip", "When checked, auto-positions the new cluster node to prevent overlap with other cluster nodes in the same host"))
+			[
+				SNew(STextBlock)
+				.Text(LOCTEXT("AddNewClusterNode_AutoPositionLabel", "Adjust Cluster Node Position to Prevent Overlap"))
+			]
+		]
+
+		+SVerticalBox::Slot()
+		.AutoHeight()
+		.Padding(2.0f)
+		[
+			SNew(SCheckBox)
+			.IsChecked_Lambda([&bAddViewport]() { return bAddViewport ? ECheckBoxState::Checked : ECheckBoxState::Unchecked; })
+			.OnCheckStateChanged_Lambda([&bAddViewport](ECheckBoxState NewState) { bAddViewport = NewState == ECheckBoxState::Checked; })
+			.ToolTipText(LOCTEXT("AddNewClusterNode_AddViewportToolTip", "When checked, adds a new viewport to the cluster node when created that matches the size of the Node application window"))
+			[
+				SNew(STextBlock)
+				.Text(LOCTEXT("AddNewClusterNode_AddViewportLabel", "Add Viewport to New Cluster Node"))
+			]
+		];
+
+	const TSharedRef<SDisplayClusterConfiguratorNewClusterItemDialog> DialogContent = SNew(SDisplayClusterConfiguratorNewClusterItemDialog, NodeTemplate)
+		.ParentItemOptions(ParentItems)
+		.InitiallySelectedParentItem("Cluster")
+		.PresetItemOptions(FDisplayClusterConfiguratorPresetSize::CommonPresets)
+		.InitiallySelectedPreset(FDisplayClusterConfiguratorPresetSize::CommonPresets[FDisplayClusterConfiguratorPresetSize::DefaultPreset])
+		.InitialName(InitialName)
+		.MaxWindowWidth(NewClusterItemDialogSize.X)
+		.MaxWindowHeight(NewClusterItemDialogSize.Y)
+		.FooterContent(ClusterNodeFooter)
+		.OnPresetChanged_Lambda([=](FVector2D Size) { NodeTemplate->WindowRect.W = Size.X; NodeTemplate->WindowRect.H = Size.Y; });
+
+	{
+		UE::DisplayClusterConfiguratorClusterEditorUtils::FScopedTemporaryTransactions ScopedTransactionsForModal;
+		ShowNewClusterItemDialogWindow(DialogContent, nullptr, LOCTEXT("AddNewClusterNode_DialogTitle", "Add New Cluster Node"), NewClusterItemDialogSize);
+	}
+	
+	UDisplayClusterConfigurationClusterNode* NewNode = nullptr;
+
+	if (DialogContent->WasAccepted())
+	{
+		OutTransaction = MakeShared<FScopedTransaction>(LOCTEXT("AddClusterNode", "Add Cluster Node"));
+		
+		const FString ItemName = DialogContent->GetItemName();
+		NodeTemplate->SetFlags(RF_Transactional);
+
+		if (bAutoposition)
+		{
+			const FVector2D DesiredPosition = FVector2D(NodeTemplate->WindowRect.X, NodeTemplate->WindowRect.Y);
+			const FVector2D DesiredSize = FVector2D(NodeTemplate->WindowRect.W, NodeTemplate->WindowRect.H);
+			const FVector2D NewPosition = FindNextAvailablePositionForClusterNode(Cluster, NodeTemplate->Host, DesiredPosition, DesiredSize);
+
+			NodeTemplate->WindowRect.X = NewPosition.X;
+			NodeTemplate->WindowRect.Y = NewPosition.Y;
+		}
+
+		NewNode = AddClusterNodeToCluster(NodeTemplate, Cluster, ItemName);
+
+		// If the newly added cluster node is the only node in the cluster, it should be the primary node, so set it as the primary
+		if (Cluster->Nodes.Num() == 1)
+		{
+			SetClusterNodeAsPrimary(NewNode);
+		}
+
+		if (bAddViewport)
+		{
+			const FString ViewportName = GetUniqueNameForViewport(DefaultNewViewportName, NewNode, true);
+			UDisplayClusterConfigurationViewport* NewViewport = NewObject<UDisplayClusterConfigurationViewport>(NewNode, NAME_None, RF_Transactional | RF_ArchetypeObject | RF_Public);
+			NewViewport->Region.W = NewNode->WindowRect.W;
+			NewViewport->Region.H = NewNode->WindowRect.H;
+
+			NewViewport = AddViewportToClusterNode(NewViewport, NewNode, ViewportName);
+		}
+	}
+
+	// Clean up the template object so that it gets GCed once this schema action has been destroyed
+	NodeTemplate->Rename(nullptr, GetTransientPackage(), REN_DoNotDirty | REN_DontCreateRedirectors | REN_ForceNoResetLoaders);
+	NodeTemplate->SetFlags(RF_Transient);
+
+	return NewNode;
+}
+
+UDisplayClusterConfigurationViewport* UE::DisplayClusterConfiguratorClusterEditorUtils::CreateNewViewportFromDialog(
+	const TSharedRef<FDisplayClusterConfiguratorBlueprintEditor>& Toolkit,
+	UDisplayClusterConfigurationClusterNode* ClusterNode,
+	const FDisplayClusterConfigurationRectangle& PresetRect,
+	TSharedPtr<FScopedTransaction>& OutTransaction)
+{
+	UDisplayClusterConfigurationViewport* ViewportTemplate = NewObject<UDisplayClusterConfigurationViewport>(Toolkit->GetBlueprintObj(), NAME_None, RF_Transactional | RF_ArchetypeObject | RF_Public);
+	ViewportTemplate->Region = FDisplayClusterConfigurationRectangle(PresetRect);
+
+	UDisplayClusterConfigurationCluster* Cluster = Toolkit->GetEditorData()->Cluster;
+
+	TArray<FString> ParentItems;
+	for (const TPair<FString, TObjectPtr<UDisplayClusterConfigurationClusterNode>>& Node : Cluster->Nodes)
+	{
+		ParentItems.Add(Node.Key);
+	}
+
+	FString ClusterNodeName = "";
+	if (const FString* Key = Cluster->Nodes.FindKey(ClusterNode))
+	{
+		ClusterNodeName = *Key;
+	}
+
+	const FString InitialName = ClusterNode ? 
+		GetUniqueNameForViewport(DefaultNewViewportName, ClusterNode, true) : 
+		FString::Printf(TEXT("%s_%d"), *DefaultNewViewportName, 0);
+
+	bool bAutoposition = true;
+	bool bExpandClusterNode = true;
+	TSharedRef<SWidget> ViewportFooter = SNew(SVerticalBox)
+		+SVerticalBox::Slot()
+		.AutoHeight()
+		.Padding(2.0f)
+		[
+			SNew(SCheckBox)
+			.IsChecked_Lambda([&bAutoposition]() { return bAutoposition ? ECheckBoxState::Checked : ECheckBoxState::Unchecked; })
+			.OnCheckStateChanged_Lambda([&bAutoposition](ECheckBoxState NewState) { bAutoposition = NewState == ECheckBoxState::Checked; })
+			.ToolTipText(LOCTEXT("AddNewViewport_AutoPositionToolTip", "When checked, auto-positions the new Viewport to prevent overlap with other Viewports in the same cluster Node"))
+			[
+				SNew(STextBlock)
+				.Text(LOCTEXT("AddNewViewport_AutoPositionLabel", "Adjust Viewport Position to Prevent Overlap"))
+			]
+		]
+	
+		+SVerticalBox::Slot()
+		.AutoHeight()
+		.Padding(2.0f)
+		[
+			SNew(SCheckBox)
+			.IsChecked_Lambda([&bExpandClusterNode]() { return bExpandClusterNode ? ECheckBoxState::Checked : ECheckBoxState::Unchecked; })
+			.OnCheckStateChanged_Lambda([&bExpandClusterNode](ECheckBoxState NewState) { bExpandClusterNode = NewState == ECheckBoxState::Checked; })
+			.ToolTipText(LOCTEXT("AddNewViewport_ExpandClusterNodeToolTip", "When checked, expands the parent cluster Node to contain the new Viewport"))
+			[
+				SNew(STextBlock)
+				.Text(LOCTEXT("AddNewViewport_ExpandClusterNodeLabel", "Expand Cluster Node to Fit New Viewport"))
+			]
+		];
+
+	TSharedRef<SDisplayClusterConfiguratorNewClusterItemDialog> DialogContent = SNew(SDisplayClusterConfiguratorNewClusterItemDialog, ViewportTemplate)
+		.ParentItemOptions(ParentItems)
+		.InitiallySelectedParentItem(ClusterNodeName)
+		.PresetItemOptions(FDisplayClusterConfiguratorPresetSize::CommonPresets)
+		.InitiallySelectedPreset(FDisplayClusterConfiguratorPresetSize::CommonPresets[FDisplayClusterConfiguratorPresetSize::DefaultPreset])
+		.InitialName(InitialName)
+		.MaxWindowWidth(NewClusterItemDialogSize.X)
+		.MaxWindowHeight(NewClusterItemDialogSize.Y)
+		.FooterContent(ViewportFooter)
+		.OnPresetChanged_Lambda([=](FVector2D Size) { ViewportTemplate->Region.W = Size.X; ViewportTemplate->Region.H = Size.Y; });
+
+	{
+		UE::DisplayClusterConfiguratorClusterEditorUtils::FScopedTemporaryTransactions ScopedTransactionsForModal;
+		ShowNewClusterItemDialogWindow(DialogContent, nullptr, LOCTEXT("AddNewViewport_DialogTitle", "Add New Viewport"), NewClusterItemDialogSize);
+	}
+	
+	UDisplayClusterConfigurationViewport* NewViewport = nullptr;
+
+	if (DialogContent->WasAccepted())
+	{
+		OutTransaction = MakeShared<FScopedTransaction>(LOCTEXT("AddViewport", "Add Viewport"));
+		
+		const FString ParentItem = DialogContent->GetSelectedParentItem();
+		const FString ItemName = DialogContent->GetItemName();
+
+		check(Cluster->Nodes.Contains(ParentItem));
+
+		UDisplayClusterConfigurationClusterNode* ParentNode = Cluster->Nodes[ParentItem];
+
+		if (bAutoposition)
+		{
+			FVector2D DesiredPosition = FVector2D(ViewportTemplate->Region.X, ViewportTemplate->Region.Y);
+			FVector2D DesiredSize = FVector2D(ViewportTemplate->Region.W, ViewportTemplate->Region.H);
+			FVector2D NewPosition = FindNextAvailablePositionForViewport(ParentNode, DesiredPosition, DesiredSize);
+
+			ViewportTemplate->Region.X = NewPosition.X;
+			ViewportTemplate->Region.Y = NewPosition.Y;
+		}
+
+		NewViewport = DuplicateObject(ViewportTemplate, ParentNode);
+		NewViewport->SetFlags(RF_Transactional);
+		NewViewport = AddViewportToClusterNode(NewViewport, ParentNode, ItemName);
+
+		if (bExpandClusterNode)
+		{
+			FIntRect ViewportRegion = NewViewport->Region.ToRect();
+			FIntRect ClusterWindowRect = FIntRect(0, 0, ParentNode->WindowRect.W, ParentNode->WindowRect.H);
+
+			ClusterWindowRect.Union(ViewportRegion);
+			FIntPoint WindowSize = ClusterWindowRect.Max;
+			FDisplayClusterConfigurationRectangle NewWindowRect(ParentNode->WindowRect.X, ParentNode->WindowRect.Y, WindowSize.X, WindowSize.Y);
+
+			FDisplayClusterConfiguratorClusterNodeViewModel ClusterNodeVM(ParentNode);
+			ClusterNodeVM.SetWindowRect(NewWindowRect);
+		}
+	}
+
+	// Clean up the template object so that it gets GCed once this schema action has been destroyed
+	ViewportTemplate->Rename(nullptr, GetTransientPackage(), REN_DoNotDirty | REN_DontCreateRedirectors | REN_ForceNoResetLoaders);
+	ViewportTemplate->SetFlags(RF_Transient);
+	
+	return NewViewport;
+}
+
+void UE::DisplayClusterConfiguratorClusterEditorUtils::ShowNewClusterItemDialogWindow(TSharedRef<SDisplayClusterConfiguratorNewClusterItemDialog> DialogContent, TSharedPtr<SWidget> ParentElement, FText WindowTitle, FVector2D WindowSize)
+{
+	TSharedPtr<SWidget> ParentWindow = FSlateApplication::Get().FindBestParentWindowForDialogs(ParentElement);
+
+	const FSlateRect WorkAreaRect = FSlateApplicationBase::Get().GetPreferredWorkArea();
+	const FVector2D DisplayTopLeft(WorkAreaRect.Left, WorkAreaRect.Top);
+	const FVector2D DisplaySize(WorkAreaRect.Right - WorkAreaRect.Left, WorkAreaRect.Bottom - WorkAreaRect.Top);
+
+	const float ScaleFactor = FPlatformApplicationMisc::GetDPIScaleFactorAtPoint(DisplayTopLeft.X, DisplayTopLeft.Y);
+	const FVector2D DialogWindowSize = WindowSize * ScaleFactor;
+
+	const FVector2D WindowPosition = (DisplayTopLeft + 0.5f * (DisplaySize - DialogWindowSize)) / ScaleFactor;
+
+	TSharedRef<SWindow> Window = SNew(SWindow)
+		.Title(WindowTitle)
+		.SizingRule(ESizingRule::Autosized)
+		.AutoCenter(EAutoCenter::None)
+		.ClientSize(DialogWindowSize)
+		.ScreenPosition(WindowPosition);
+
+	DialogContent->SetParentWindow(Window);
+	Window->SetContent(DialogContent);
+
+	FSlateApplication::Get().AddModalWindow(Window, ParentWindow, false);
+}
+
+TSharedPtr<FDragDropOperation> UE::DisplayClusterConfiguratorClusterEditorUtils::MakeDragDropOperation(const TArray<UObject*>& SelectedObjects)
+{
+	TArray<UDisplayClusterConfigurationViewport*> Viewports;
+	TArray<UDisplayClusterConfigurationClusterNode*> ClusterNodes;
+
+	for (UObject* SelectedObject : SelectedObjects)
+	{
+		if (UDisplayClusterConfigurationViewport* Viewport = Cast<UDisplayClusterConfigurationViewport>(SelectedObject))
+		{
+			Viewports.Add(Viewport);
+		}
+		else if (UDisplayClusterConfigurationClusterNode* ClusterNode = Cast<UDisplayClusterConfigurationClusterNode>(SelectedObject))
+		{
+			ClusterNodes.Add(ClusterNode);
+		}
+		else
+		{
+			// Only cluster nodes and viewports can be dragged and dropped. If anything else is selected, abort the drag/drop operation
+			return TSharedPtr<FDragDropOperation>();
+		}
+	}
+
+	if (Viewports.Num() && !ClusterNodes.Num())
+	{
+		// If only viewports are selected, return a viewport drag/drop operation
+		return FDisplayClusterConfiguratorViewportDragDropOp::New(Viewports);
+	}
+	else if (ClusterNodes.Num() && !Viewports.Num())
+	{
+		// If only cluster nodes are selected, return a cluster node drag/drop operation
+		return FDisplayClusterConfiguratorClusterNodeDragDropOp::New(ClusterNodes);
+	}
+	else
+	{
+		// If both cluster nodes and viewports are selected, only allow a cluster node drag/drop operation if all selected viewports are children of the selected cluster nodes
+		bool bOnlyChildViewports = true;
+		for (UDisplayClusterConfigurationViewport* Viewport : Viewports)
+		{
+			bool bIsChild = ClusterNodes.ContainsByPredicate([=](UDisplayClusterConfigurationClusterNode* ClusterNode)
+			{
+				return ClusterNode == Viewport->GetOuter();
+			});
+
+			if (!bIsChild)
+			{
+				bOnlyChildViewports = false;
+				break;
+			}
+		}
+
+		if (bOnlyChildViewports)
+		{
+			return FDisplayClusterConfiguratorClusterNodeDragDropOp::New(ClusterNodes);
+		}
+		else
+		{
+			return TSharedPtr<FDragDropOperation>();
+		}
+	}
+}
+
+namespace
+{
+	FVector2D FindAvailableSpace(const TArray<FBox2D>& OtherBounds, const FBox2D& DesiredBounds)
+	{
+		FBox2D CurrentBounds(DesiredBounds);
+
+		// Similar to FBox2D::Intersects, but ignores the case where the box edges are touching.
+		auto IntrudesFunc = [](const FBox2D& BoxA, const FBox2D& BoxB)
+		{
+
+			// Special case if both boxes are directly on top of each other, which is considered an intrusion.
+			if (BoxA == BoxB)
+			{
+				return true;
+			}
+
+			if ((BoxA.Min.X >= BoxB.Max.X) || (BoxB.Min.X >= BoxA.Max.X))
+			{
+				return false;
+			}
+
+			if ((BoxA.Min.Y >= BoxB.Max.Y) || (BoxB.Min.Y >= BoxA.Max.Y))
+			{
+				return false;
+			}
+
+			return true;
+		};
+
+		bool bIntersectsBounds;
+		do
+		{
+			bIntersectsBounds = false;
+
+			for (const FBox2D& Bounds : OtherBounds)
+			{
+				if (IntrudesFunc(CurrentBounds, Bounds))
+				{
+					bIntersectsBounds = true;
+
+					float Shift = Bounds.GetSize().X;
+					CurrentBounds.Min.X += Shift;
+					CurrentBounds.Max.X += Shift;
+
+					break;
+				}
+			}
+
+		} while (bIntersectsBounds);
+
+		return CurrentBounds.Min;
+	}
+}
+
+FVector2D UE::DisplayClusterConfiguratorClusterEditorUtils::FindNextAvailablePositionForClusterNode(UDisplayClusterConfigurationCluster* Cluster, const FString& DesiredHost, const FVector2D& DesiredPosition, const FVector2D& DesiredSize)
+{
+	TMap<FString, TMap<FString, UDisplayClusterConfigurationClusterNode*>> ClusterNodesByHost;
+	SortClusterNodesByHost(Cluster->Nodes, ClusterNodesByHost);
+
+	// If the desired host doesn't exist yet, simply return the node's desired position
+	if (!ClusterNodesByHost.Contains(DesiredHost))
+	{
+		return DesiredPosition;
+	}
+
+	const TMap<FString, UDisplayClusterConfigurationClusterNode*>& ExistingClusterNodes = ClusterNodesByHost[DesiredHost];
+	if (!ExistingClusterNodes.Num())
+	{
+		return DesiredPosition;
+	}
+
+	FBox2D DesiredBounds = FBox2D(DesiredPosition, DesiredPosition + DesiredSize);
+	TArray<FBox2D> ClusterNodeBounds;
+	for (const TPair<FString, UDisplayClusterConfigurationClusterNode*>& ClusterNodeKeyPair : ExistingClusterNodes)
+	{
+		const FDisplayClusterConfigurationRectangle& WindowRect = ClusterNodeKeyPair.Value->WindowRect;
+		ClusterNodeBounds.Add(FBox2D(FVector2D(WindowRect.X, WindowRect.Y), FVector2D(WindowRect.X + WindowRect.W, WindowRect.Y + WindowRect.H)));
+	}
+
+	return FindAvailableSpace(ClusterNodeBounds, DesiredBounds);
+}
+
+FVector2D UE::DisplayClusterConfiguratorClusterEditorUtils::FindNextAvailablePositionForViewport(UDisplayClusterConfigurationClusterNode* ClusterNode, const FVector2D& DesiredPosition, const FVector2D& DesiredSize)
+{
+	if (!ClusterNode)
+	{
+		return DesiredPosition;
+	}
+
+	FBox2D DesiredBounds = FBox2D(DesiredPosition, DesiredPosition + DesiredSize);
+	TArray<FBox2D> ViewportBounds;
+	for (const TPair<FString, TObjectPtr<UDisplayClusterConfigurationViewport>>& ViewportKeyPair : ClusterNode->Viewports)
+	{
+		const FDisplayClusterConfigurationRectangle& Region = ViewportKeyPair.Value->Region;
+		ViewportBounds.Add(FBox2D(FVector2D(Region.X, Region.Y), FVector2D(Region.X + Region.W, Region.Y + Region.H)));
+	}
+
+	return FindAvailableSpace(ViewportBounds, DesiredBounds);
+}
+
+void UE::DisplayClusterConfiguratorClusterEditorUtils::CopyClusterItemsToClipboard(const TArray<UObject*>& ClusterItemsToCopy)
+{
+	TSet<UObject*> UniqueItemsToCopy;
+
+	for (UObject* ClusterItem : ClusterItemsToCopy)
+	{
+		// Only allow viewports and cluster nodes to be copied at this point.
+		if (ClusterItem->IsA<UDisplayClusterConfigurationViewport>())
+		{
+			// If a viewport is being copied, make sure its parent cluster node is not also being copied; if its parent is being copied,
+			// the viewport doesn't need to be copied directly as it will be copied with its parent.
+			if (!ClusterItemsToCopy.Contains(ClusterItem->GetOuter()))
+			{
+				UniqueItemsToCopy.Add(ClusterItem);
+			}
+		}
+		else if (ClusterItem->IsA<UDisplayClusterConfigurationClusterNode>())
+		{
+			UniqueItemsToCopy.Add(ClusterItem);
+		}
+	}
+
+	UnMarkAllObjects(EObjectMark(OBJECTMARK_TagExp | OBJECTMARK_TagImp));
+
+	FStringOutputDevice Archive;
+	const FExportObjectInnerContext Context;
+
+	for (UObject* ClusterItem : UniqueItemsToCopy)
+	{
+		UExporter::ExportToOutputDevice(&Context, ClusterItem, nullptr, Archive, TEXT("copy"), 0, PPF_ExportsNotFullyQualified | PPF_Copy | PPF_Delimited, false, nullptr);
+	}
+
+	FString CopiedText = Archive;
+	FPlatformApplicationMisc::ClipboardCopy(*CopiedText);
+}
+
+class FClusterNodeTextFactory : public FCustomizableTextObjectFactory
+{
+public:
+	FClusterNodeTextFactory()
+		: FCustomizableTextObjectFactory(GWarn)
+	{
+	}
+
+	// FCustomizableTextObjectFactory implementation
+
+	virtual bool CanCreateClass(UClass* ObjectClass, bool& bOmitSubObjs) const override
+	{
+		const bool isClusterNode = ObjectClass->IsChildOf(UDisplayClusterConfigurationClusterNode::StaticClass());
+		const bool isViewportNode = ObjectClass->IsChildOf(UDisplayClusterConfigurationViewport::StaticClass());
+
+		return isClusterNode || isViewportNode;
+	}
+
+	virtual void ProcessConstructedObject(UObject* NewObject) override
+	{
+		check(NewObject);
+
+		if (UDisplayClusterConfigurationClusterNode* ClusterNode = Cast<UDisplayClusterConfigurationClusterNode>(NewObject))
+		{
+			ClusterNodes.Add(ClusterNode);
+		}
+		else if (UDisplayClusterConfigurationViewport* ViewportNode = Cast<UDisplayClusterConfigurationViewport>(NewObject))
+		{
+			Viewports.Add(ViewportNode);
+		}
+	}
+
+	// FCustomizableTextObjectFactory (end)
+
+public:
+	TArray<UDisplayClusterConfigurationClusterNode*> ClusterNodes;
+	TArray<UDisplayClusterConfigurationViewport*> Viewports;
+};
+
+bool UE::DisplayClusterConfiguratorClusterEditorUtils::CanPasteClusterItemsFromClipboard(const TArray<UObject*>& TargetClusterItems, int32& OutNumItems)
+{
+	// Can't paste unless the clipboard has a string in it
+	FString ClipboardContent;
+	FPlatformApplicationMisc::ClipboardPaste(ClipboardContent);
+	if (!ClipboardContent.IsEmpty())
+	{
+		FClusterNodeTextFactory Factory;
+		Factory.ProcessBuffer(GetTransientPackage(), RF_Transactional, ClipboardContent);
+
+		OutNumItems = Factory.ClusterNodes.Num() + Factory.Viewports.Num();
+
+		bool bCanPaste = false;
+		if (Factory.ClusterNodes.Num() > 0)
+		{
+			// If we are trying to paste cluster nodes, the target should be a cluster root, a host display data, or a cluster node.
+			bCanPaste = TargetClusterItems.ContainsByPredicate([=](const UObject* Item)
+			{
+				return
+					Item->IsA<UDisplayClusterConfigurationCluster>() ||
+					Item->IsA<UDisplayClusterConfigurationClusterNode>() ||
+					Item->IsA<UDisplayClusterConfigurationHostDisplayData>();
+			});
+		}
+
+		if (Factory.Viewports.Num() > 0)
+		{
+			// If we are trying to paste viewports, the target should either be a cluster node or a viewport.
+			bCanPaste = TargetClusterItems.ContainsByPredicate([=](const UObject* Item)
+			{
+				return Item->IsA<UDisplayClusterConfigurationClusterNode>() || Item->IsA<UDisplayClusterConfigurationViewport>();
+			});
+		}
+
+		return bCanPaste;
+	}
+
+	return false;
+}
+
+TArray<UObject*> UE::DisplayClusterConfiguratorClusterEditorUtils::PasteClusterItemsFromClipboard(const TArray<UObject*>& TargetClusterItems, TOptional<FVector2D> PasteLocation)
+{
+	FString ClipboardContent;
+	FPlatformApplicationMisc::ClipboardPaste(ClipboardContent);
+
+	FClusterNodeTextFactory Factory;
+	Factory.ProcessBuffer(GetTransientPackage(), RF_Transactional, ClipboardContent);
+
+	// Filter out the target items and find all unique cluster roots and cluster nodes. For each target cluster node, also add its parent root,
+	// and for each target viewport, add its parent cluster node. This allows users to paste cluster nodes and viewports by selecting their sibling nodes,
+	// which will then add the pasted node to the parent of the target node.
+
+	// Pair up the cluster roots with a host string, which indicates which host the cluster nodes should be pasted into (the cluster nodes' Host strings will be set when duplicated).
+	// If the paired host string is empty, the pasted cluster nodes' original host values will be preserved.
+	TSet<TTuple<UDisplayClusterConfigurationCluster*, FString>> SelectedRootHosts;
+	TSet<UDisplayClusterConfigurationClusterNode*> SelectedNodes;
+
+	for (UObject* TargetItem : TargetClusterItems)
+	{
+		if (UDisplayClusterConfigurationCluster* ClusterRoot = Cast<UDisplayClusterConfigurationCluster>(TargetItem))
+		{
+			SelectedRootHosts.Add(TTuple<UDisplayClusterConfigurationCluster*, FString>(ClusterRoot, ""));
+		}
+		else if (UDisplayClusterConfigurationHostDisplayData* Host = Cast<UDisplayClusterConfigurationHostDisplayData>(TargetItem))
+		{
+			// If the target is a host display data, that indicates the user wants to paste the cluster node into the host the display data represents,
+			// so add an entry to the cluster root list with the parent cluster object and the host string of the host display data.
+			if (UDisplayClusterConfigurationCluster* ParentRoot = Cast<UDisplayClusterConfigurationCluster>(TargetItem->GetOuter()))
+			{
+				FString HostStr = FString();
+				if (const FString* Key = ParentRoot->HostDisplayData.FindKey(Host))
+				{
+					HostStr = *Key;
+				}
+
+				SelectedRootHosts.Add(TTuple<UDisplayClusterConfigurationCluster*, FString>(ParentRoot, HostStr));
+			}
+		}
+		else if (UDisplayClusterConfigurationClusterNode* ClusterNode = Cast<UDisplayClusterConfigurationClusterNode>(TargetItem))
+		{
+			SelectedNodes.Add(ClusterNode);
+
+			if (UDisplayClusterConfigurationCluster* ParentRoot = Cast<UDisplayClusterConfigurationCluster>(TargetItem->GetOuter()))
+			{
+				// If pasting as a sibling of a target cluster node, set the pasted node's host to be the same as the sibling node's host.
+				SelectedRootHosts.Add(TTuple<UDisplayClusterConfigurationCluster*, FString>(ParentRoot, ClusterNode->Host));
+			}
+		}
+		else if (UDisplayClusterConfigurationViewport* Viewport = Cast<UDisplayClusterConfigurationViewport>(TargetItem))
+		{
+			if (UDisplayClusterConfigurationClusterNode* ParentNode = Cast<UDisplayClusterConfigurationClusterNode>(TargetItem->GetOuter()))
+			{
+				SelectedNodes.Add(ParentNode);
+			}
+		}
+	}
+
+	TArray<UObject*> CopiedObjects;
+
+	// Keep track of a couple reference positions for pasting the items to a specfic location to correctly offset pasted location in the case where
+	// more than one item is being pasted or an item is being pasted into more than one parent. The ReferencePosition variable keeps track of the 
+	// relative position of the item to be pasted, while the ReferenceTargetPosition keeps track of the relative position of the target item being pasted
+	// into. The way these work is that, if they aren't set when in the pasting loop below, then they are set to the current iteration item's position. Every
+	// other iteration will then subtract the reference position from the pasted position. This allows each item being pasted to be correctly offset to maintain
+	// both its relative position in the group of items being pasted and its relative position to the item it is being pasted into.
+	TOptional<FVector2D> ReferencePosition = TOptional<FVector2D>();
+	TOptional<FVector2D> ReferenceTargetPosition = TOptional<FVector2D>();
+
+	for (const TTuple<UDisplayClusterConfigurationCluster*, FString>& ClusterRootHost : SelectedRootHosts)
+	{
+		UDisplayClusterConfigurationCluster* ClusterRoot = ClusterRootHost.Key;
+		FString HostStr = ClusterRootHost.Value;
+
+		for (UDisplayClusterConfigurationClusterNode* ClusterNode : Factory.ClusterNodes)
+		{
+			UDisplayClusterConfigurationClusterNode* ClusterNodeCopy = DuplicateObject(ClusterNode, ClusterRoot);
+			ClusterNodeCopy->SetFlags(RF_Transactional);
+			ClusterNodeCopy = AddClusterNodeToCluster(ClusterNodeCopy, ClusterRoot, ClusterNode->GetName());
+
+			// If the host string isn't empty, it means the user is pasting the cluster node into a particular host, so set the cluster node's host string.
+			if (!HostStr.IsEmpty())
+			{
+				ClusterNodeCopy->Host = HostStr;
+			}
+
+			if (PasteLocation.IsSet())
+			{
+				FVector2D NewLocation = *PasteLocation;
+
+				// If we are pasting into an existing host, offset the reference target position so that the position of the node inside the host
+				// is consistent for every host we paste into.
+				if (UDisplayClusterConfigurationHostDisplayData* HostDisplayData = GetHostDisplayDataForClusterNode(ClusterNodeCopy))
+				{
+					// If the reference target position isn't set, set it to the current host's position.
+					if (!ReferenceTargetPosition.IsSet())
+					{
+						ReferenceTargetPosition = HostDisplayData->Position;
+					}
+
+					// Offset the location to keep this cluster node's position relative to its parent host consistent.
+					NewLocation -= *ReferenceTargetPosition;
+				}
+
+				// If the reference position isn't set, set it to the current pasted cluster node's position.
+				if (!ReferencePosition.IsSet())
+				{
+					ReferencePosition = FVector2D(ClusterNode->WindowRect.X, ClusterNode->WindowRect.Y);
+				}
+
+				// Offset the location to keep this cluster node's relative position with the group of pasted nodes consistent.
+				NewLocation -= FVector2D(ClusterNode->WindowRect.X, ClusterNode->WindowRect.Y) - *ReferencePosition;
+
+				ClusterNodeCopy->WindowRect.X = NewLocation.X;
+				ClusterNodeCopy->WindowRect.Y = NewLocation.Y;
+			}
+
+			// Rename viewports in the new cluster node to ensure they're unique. Copy the viewport pointers into an array
+			// first since renaming would cause the dictionary to change while we're iterating it.
+			TArray<typename decltype(ClusterNodeCopy->Viewports)::ValueType> CopiedViewports;
+			ClusterNodeCopy->Viewports.GenerateValueArray(CopiedViewports);
+			
+			for (UDisplayClusterConfigurationViewport* Viewport : CopiedViewports)
+			{
+				RenameViewport(Viewport, *GetUniqueNameForViewport(*Viewport->GetName(), ClusterNodeCopy));
+			}
+
+			CopiedObjects.Add(ClusterNodeCopy);
+		}
+	}
+
+	// Reset the reference positions for the viewports pasting loop
+	ReferencePosition = TOptional<FVector2D>();
+	ReferenceTargetPosition = TOptional<FVector2D>();
+
+	for (UDisplayClusterConfigurationClusterNode* ClusterNode : SelectedNodes)
+	{
+		for (UDisplayClusterConfigurationViewport* Viewport : Factory.Viewports)
+		{
+			UDisplayClusterConfigurationViewport* ViewportCopy = DuplicateObject(Viewport, ClusterNode);
+			ViewportCopy->SetFlags(RF_Transactional);
+			ViewportCopy = AddViewportToClusterNode(ViewportCopy, ClusterNode, Viewport->GetName());
+
+			if (PasteLocation.IsSet())
+			{
+				// If the reference target position isn't set, set it to the current cluster node's position.
+				if (!ReferenceTargetPosition.IsSet())
+				{
+					FVector2D ClusterNodeLocation = FVector2D(ClusterNode->WindowRect.X, ClusterNode->WindowRect.Y);
+
+					// Cluster node window rectangles are local to the node's host, so factor in its host position, if it has one.
+					if (UDisplayClusterConfigurationHostDisplayData* HostDisplayData = GetHostDisplayDataForClusterNode(ClusterNode))
+					{
+						ClusterNodeLocation += HostDisplayData->Position;
+					}
+
+					ReferenceTargetPosition = ClusterNodeLocation;
+				}
+
+				// Offset the location to keep this viewport's position relative to its parent cluster node consistent.
+				FVector2D NewLocation = *PasteLocation - *ReferenceTargetPosition;
+
+				// If the reference position isn't set, set it to the current pasted viewport's position.
+				if (!ReferencePosition.IsSet())
+				{
+					ReferencePosition = FVector2D(Viewport->Region.X, Viewport->Region.Y);
+				}
+
+				// Offset the location to keep this viewport's relative position with the group of pasted nodes consistent.
+				NewLocation -= FVector2D(Viewport->Region.X, Viewport->Region.Y) - *ReferencePosition;
+
+				// Viewports cannot have negative position
+				ViewportCopy->Region.X = FMath::Max(NewLocation.X, 0.0f);
+				ViewportCopy->Region.Y = FMath::Max(NewLocation.Y, 0.0f);
+			}
+
+			CopiedObjects.Add(ViewportCopy);
+		}
+	}
+
+	return CopiedObjects;
+}
+
+#undef LOCTEXT_NAMESPACE
\ No newline at end of file
diff --git a/Engine/Plugins/Runtime/nDisplay/Source/DisplayClusterConfigurator/Private/ClusterConfiguration/DisplayClusterConfiguratorClusterEditorUtils.h b/Engine/Plugins/Runtime/nDisplay/Source/DisplayClusterConfigurator/Private/ClusterConfiguration/DisplayClusterConfiguratorClusterEditorUtils.h
new file mode 100644
--- /dev/null
+++ b/Engine/Plugins/Runtime/nDisplay/Source/DisplayClusterConfigurator/Private/ClusterConfiguration/DisplayClusterConfiguratorClusterEditorUtils.h
@@ -0,0 +1,103 @@
+// Copyright Epic Games, Inc. All Rights Reserved.
+
+#pragma once
+
+#include "CoreMinimal.h"
+#include "UObject/ObjectPtr.h"
+
+class SWidget;
+class FDragDropOperation;
+class FDisplayClusterConfiguratorBlueprintEditor;
+class FScopedTransaction;
+class UDisplayClusterConfigurationData_Base;
+class UDisplayClusterConfigurationCluster;
+class UDisplayClusterConfigurationHostDisplayData;
+class UDisplayClusterConfigurationClusterNode;
+class UDisplayClusterConfigurationViewport;
+class SDisplayClusterConfiguratorNewClusterItemDialog;
+struct FDisplayClusterConfigurationRectangle;
+
+#define NDISPLAY_DEFAULT_CLUSTER_HOST "127.0.0.1"
+
+namespace UE::DisplayClusterConfiguratorClusterEditorUtils
+{
+	/**
+	 * Creates a new cluster node, presenting a dialog box to allow the user to customize its properties, and adds it to the cluster.
+	 * @param Toolkit - The blueprint editor toolkit being used.
+	 * @param Cluster - The initial cluster the user wants to add the cluster node to.
+	 * @param PresetRect - The initial rectangle to configure the new cluster node with in the dialog box.
+	 * @param OutTransaction - The transaction created when adding a cluster node. May be null if the user canceled the operation.
+	 * @param PresetHost - Optional initial host string to configure the new cluster node with in the dialog box.
+	 * @return The newly created cluster node, or null if the user cancelled out of the dialog box.
+	 */
+	UDisplayClusterConfigurationClusterNode* CreateNewClusterNodeFromDialog(const TSharedRef<FDisplayClusterConfiguratorBlueprintEditor>& Toolkit,
+		UDisplayClusterConfigurationCluster* Cluster, const FDisplayClusterConfigurationRectangle& PresetRect,
+		TSharedPtr<FScopedTransaction>& OutTransaction,
+		FString PresetHost = NDISPLAY_DEFAULT_CLUSTER_HOST);
+
+	/**
+	 * Creates a new viewport, presenting a dialog box to allow the user to customize its properties, and adds it to the user chosen cluster node.
+	 * @param Toolkit - The blueprint editor toolkit being used.
+	 * @param ClusterNode - The initial cluster node the user wants to add the viewport to, or null to use the first cluster node in the cluster.
+	 * @param PresetRect - The initial rectangle to configure the new viewport with in the dialog box.
+	 * @param OutTransaction - The transaction created when adding a viewport. May be null if the user canceled the operation.
+	 * @return The newly created viewport, or null if the user cancelled out of the dialog box.
+	 */
+	UDisplayClusterConfigurationViewport* CreateNewViewportFromDialog(const TSharedRef<FDisplayClusterConfiguratorBlueprintEditor>& Toolkit,
+		UDisplayClusterConfigurationClusterNode* ClusterNode, const FDisplayClusterConfigurationRectangle& PresetRect,
+		TSharedPtr<FScopedTransaction>& OutTransaction);
+
+	/**
+	 * Displays a dialog box used to configure new cluster item properties before adding the new item to the cluster configuration.
+	 * @param DialogContent - The widget to display in the dialog box's main section.
+	 * @param ParentElement - The UI element to display the dialog box from. If null, will use the active window as a parent.
+	 * @param WindowTitle - The title displayed at the top of the dialog box.
+	 * @param WindowSize - The size of the dialog box.
+	 */
+	void ShowNewClusterItemDialogWindow(TSharedRef<SDisplayClusterConfiguratorNewClusterItemDialog> DialogContent, TSharedPtr<SWidget> ParentElement, FText WindowTitle, FVector2D WindowSize);
+
+	/**
+	 * Attempts to create an appropriate drag and drop operation for the selected objects.
+	 * @param SelectedObjects - The objects to create the drag and drop operation for.
+	 * @return A valid pointer to the drag and drop operation if an operation could be made, and invalid pointer if not.
+	 */
+	TSharedPtr<FDragDropOperation> MakeDragDropOperation(const TArray<UObject*>& SelectedObjects);
+
+	/**
+	 * Finds the next available position for a cluster node within a host so that the cluster node doesn't overlap any existing nodes.
+	 * @param Cluster - The cluster to find the position within.
+	 * @param DesiredHost - The host the cluster node wants to be positioned in.
+	 * @param DesiredPosition - The desired position of the cluster node.
+	 * @param DesiredSize - The desired size of the cluster node.
+	 */
+	FVector2D FindNextAvailablePositionForClusterNode(UDisplayClusterConfigurationCluster* Cluster, const FString& DesiredHost, const FVector2D& DesiredPosition, const FVector2D& DesiredSize);
+
+	/**
+	 * Finds the next available position for a viewport within a cluster node so that the viewport doesn't overlap any existing viewports.
+	 * @param ClusterNode - The cluster node to find the position within.
+	 * @param DesiredPosition - The desired position of the viewport.
+	 * @param DesiredSize - The desired size of the viewport.
+	 */
+	FVector2D FindNextAvailablePositionForViewport(UDisplayClusterConfigurationClusterNode* ClusterNode, const FVector2D& DesiredPosition, const FVector2D& DesiredSize);
+
+	/**
+	 * Copies the specified list of cluster items to the clipboard.
+	 * @param ClusterItemsToCopy - The list of cluster items to copy.
+	 */
+	void CopyClusterItemsToClipboard(const TArray<UObject*>& ClusterItemsToCopy);
+
+	/**
+	 * Checks to see if the clipboard contains cluster items that can be pasted into the target cluster items.
+	 * @param TargetClusterItems - The cluster items to paste the copied cluster items into.
+	 * @param OutNumItems - The number of cluster items that can be pasted from the clipboard.
+	 */
+	bool CanPasteClusterItemsFromClipboard(const TArray<UObject*>& TargetClusterItems, int32& OutNumItems);
+
+	/**
+	 * Pastes any cluster items in the clipboard into the specified target items, if the target items can accept the pasted items.
+	 * @param TargetClusterItems - The cluster items to paste the copied cluster items into, if compatible.
+	 * @param PasteLocation - Optional location to paste the cluster items at.
+	 * @return The pasted cluster items.
+	 */
+	TArray<UObject*> PasteClusterItemsFromClipboard(const TArray<UObject*>& TargetClusterItems, TOptional<FVector2D> PasteLocation = TOptional<FVector2D>());
+};
\ No newline at end of file
diff --git a/Engine/Plugins/Runtime/nDisplay/Source/DisplayClusterConfigurator/Private/ClusterConfiguration/DisplayClusterConfiguratorClusterUtils.cpp b/Engine/Plugins/Runtime/nDisplay/Source/DisplayClusterConfigurator/Private/ClusterConfiguration/DisplayClusterConfiguratorClusterUtils.cpp
--- a/Engine/Plugins/Runtime/nDisplay/Source/DisplayClusterConfigurator/Private/ClusterConfiguration/DisplayClusterConfiguratorClusterUtils.cpp
+++ b/Engine/Plugins/Runtime/nDisplay/Source/DisplayClusterConfigurator/Private/ClusterConfiguration/DisplayClusterConfiguratorClusterUtils.cpp
@@ -1,557 +1,87 @@
 // Copyright Epic Games, Inc. All Rights Reserved.
 
-#include "DisplayClusterConfiguratorClusterUtils.h"
+#include "ClusterConfiguration/DisplayClusterConfiguratorClusterUtils.h"
 
 #include "DisplayClusterConfiguratorBlueprintEditor.h"
 #include "DisplayClusterConfigurationTypes.h"
 #include "DisplayClusterConfiguratorPropertyUtils.h"
 #include "DisplayClusterConfiguratorStyle.h"
 #include "DisplayClusterConfiguratorUtils.h"
-#include "ClusterConfiguration/SDisplayClusterConfiguratorNewClusterItemDialog.h"
-#include "ClusterConfiguration/ViewModels/DisplayClusterConfiguratorClusterNodeViewModel.h"
-#include "Views/DragDrop/DisplayClusterConfiguratorClusterNodeDragDropOp.h"
-#include "Views/DragDrop/DisplayClusterConfiguratorViewportDragDropOp.h"
 
 #include "Factories.h"
-#include "ISinglePropertyView.h"
-#include "ScopedTransaction.h"
 #include "ObjectTools.h"
-#include "Framework/Application/SlateApplication.h"
-#include "HAL/PlatformApplicationMisc.h"
-#include "UnrealExporter.h"
 #include "Editor/Transactor.h"
-#include "Editor/TransBuffer.h"
-#include "Exporters/Exporter.h"
-#include "Input/DragAndDrop.h"
-#include "Widgets/SBoxPanel.h"
-#include "Widgets/Input/SCheckBox.h"
-#include "Widgets/Text/STextBlock.h"
 
 #define LOCTEXT_NAMESPACE "FDisplayClusterConfiguratorClusterUtils"
 
-const FVector2D FDisplayClusterConfiguratorClusterUtils::NewClusterItemDialogSize = FVector2D(410, 512);
-const FString FDisplayClusterConfiguratorClusterUtils::DefaultNewHostName = TEXT("Host");
-const FString FDisplayClusterConfiguratorClusterUtils::DefaultNewClusterNodeName = TEXT("Node");
-const FString FDisplayClusterConfiguratorClusterUtils::DefaultNewViewportName = TEXT("VP");
-
 namespace UE::DisplayClusterConfiguratorClusterUtils
 {
-	// All transactions performed within this scope can be considered temporary and will be cleared when out of scope.
-	// Transactions prior to the scope cannot be undone while in scope.
-	struct FScopedTemporaryTransactions
-	{
-		// The index of the transaction queue when entering the scope
-		int32 StartTransactionIdx = 0;
-		
-		FScopedTemporaryTransactions()
-		{
-			if (GEditor)
-			{
-				// Prevent undoing any actions before the scope started
-				GEditor->Trans->SetUndoBarrier();
-
-				// Record the transaction idx before we start our scoped transactions
-				StartTransactionIdx = GEditor->Trans->GetCurrentUndoBarrier();
-
-				// Broadcasting changes here can take a big performance hit and isn't necessary since we only use this in a modal window
-				GEditor->bSuspendBroadcastPostUndoRedo = true;
-			}
-		}
-
-		~FScopedTemporaryTransactions()
-		{
-			if (GEditor)
-			{
-				GEditor->bSuspendBroadcastPostUndoRedo = false;
-
-				// Clear all transactions that occurred while in scope, accounting for transactions that were
-				// undone. None of these transactions should be allowed to be undone or redone from this point forward
-				{
-					UTransBuffer* TransBuffer = CastChecked<UTransBuffer>(GEditor->Trans);
-					const int32 ScopedTransactionsPerformed = TransBuffer->UndoBuffer.Num() - StartTransactionIdx;
-				
-					if (ScopedTransactionsPerformed > 0
-						&& StartTransactionIdx >= 0
-						&& StartTransactionIdx < TransBuffer->UndoBuffer.Num()
-						&& (StartTransactionIdx + ScopedTransactionsPerformed - 1) < TransBuffer->UndoBuffer.Num())
-					{
-						// Remove only the transactions performed within scope
-						TransBuffer->UndoBuffer.RemoveAt(StartTransactionIdx, ScopedTransactionsPerformed);
-
-						// We can't restore redo items prior to the scope since they would have been cleared while
-						// performing transactions within the scope
-						TransBuffer->UndoCount = 0;
-
-						// Need to broadcast for other systems since we are changing the undo buffer manually
-						TransBuffer->OnUndoBufferChanged().Broadcast();
-					}
-				}
-
-				// Allow undos again from before the scope started
-				GEditor->Trans->RemoveUndoBarrier();
-			}
-		}
-	};
-}
-
-using namespace DisplayClusterConfiguratorPropertyUtils;
-
-UDisplayClusterConfigurationClusterNode* FDisplayClusterConfiguratorClusterUtils::CreateNewClusterNodeFromDialog(
-	const TSharedRef<FDisplayClusterConfiguratorBlueprintEditor>& Toolkit,
-	UDisplayClusterConfigurationCluster* Cluster,
-	const FDisplayClusterConfigurationRectangle& PresetRect,
-	TSharedPtr<FScopedTransaction>& OutTransaction,
-	FString PresetHost)
-{
-	UDisplayClusterConfigurationClusterNode* NodeTemplate = NewObject<UDisplayClusterConfigurationClusterNode>(Toolkit->GetBlueprintObj(), NAME_None, RF_Transactional | RF_ArchetypeObject | RF_Public);
-	NodeTemplate->WindowRect = FDisplayClusterConfigurationRectangle(PresetRect);
-	NodeTemplate->Host = PresetHost;
-
-	TArray<FString> ParentItems;
-	ParentItems.Add("Cluster");
-
-	const FString InitialName = GetUniqueNameForClusterNode(DefaultNewClusterNodeName, Cluster, true);
-	
-	bool bAutoposition = true;
-	bool bAddViewport = true;
-	const TSharedRef<SWidget> ClusterNodeFooter = SNew(SVerticalBox)
-		+SVerticalBox::Slot()
-		.AutoHeight()
-		.Padding(2.0f)
-		[
-			SNew(SCheckBox)
-			.IsChecked_Lambda([&bAutoposition]() { return bAutoposition ? ECheckBoxState::Checked : ECheckBoxState::Unchecked; })
-			.OnCheckStateChanged_Lambda([&bAutoposition](ECheckBoxState NewState) { bAutoposition = NewState == ECheckBoxState::Checked; })
-			.ToolTipText(LOCTEXT("AddNewClusterNode_AutoPositionToolTip", "When checked, auto-positions the new cluster node to prevent overlap with other cluster nodes in the same host"))
-			[
-				SNew(STextBlock)
-				.Text(LOCTEXT("AddNewClusterNode_AutoPositionLabel", "Adjust Cluster Node Position to Prevent Overlap"))
-			]
-		]
-
-		+SVerticalBox::Slot()
-		.AutoHeight()
-		.Padding(2.0f)
-		[
-			SNew(SCheckBox)
-			.IsChecked_Lambda([&bAddViewport]() { return bAddViewport ? ECheckBoxState::Checked : ECheckBoxState::Unchecked; })
-			.OnCheckStateChanged_Lambda([&bAddViewport](ECheckBoxState NewState) { bAddViewport = NewState == ECheckBoxState::Checked; })
-			.ToolTipText(LOCTEXT("AddNewClusterNode_AddViewportToolTip", "When checked, adds a new viewport to the cluster node when created that matches the size of the Node application window"))
-			[
-				SNew(STextBlock)
-				.Text(LOCTEXT("AddNewClusterNode_AddViewportLabel", "Add Viewport to New Cluster Node"))
-			]
-		];
-
-	const TSharedRef<SDisplayClusterConfiguratorNewClusterItemDialog> DialogContent = SNew(SDisplayClusterConfiguratorNewClusterItemDialog, NodeTemplate)
-		.ParentItemOptions(ParentItems)
-		.InitiallySelectedParentItem("Cluster")
-		.PresetItemOptions(FDisplayClusterConfiguratorPresetSize::CommonPresets)
-		.InitiallySelectedPreset(FDisplayClusterConfiguratorPresetSize::CommonPresets[FDisplayClusterConfiguratorPresetSize::DefaultPreset])
-		.InitialName(InitialName)
-		.MaxWindowWidth(NewClusterItemDialogSize.X)
-		.MaxWindowHeight(NewClusterItemDialogSize.Y)
-		.FooterContent(ClusterNodeFooter)
-		.OnPresetChanged_Lambda([=](FVector2D Size) { NodeTemplate->WindowRect.W = Size.X; NodeTemplate->WindowRect.H = Size.Y; });
-
-	{
-		UE::DisplayClusterConfiguratorClusterUtils::FScopedTemporaryTransactions ScopedTransactionsForModal;
-		ShowNewClusterItemDialogWindow(DialogContent, nullptr, LOCTEXT("AddNewClusterNode_DialogTitle", "Add New Cluster Node"), NewClusterItemDialogSize);
-	}
-	
-	UDisplayClusterConfigurationClusterNode* NewNode = nullptr;
-
-	if (DialogContent->WasAccepted())
-	{
-		OutTransaction = MakeShared<FScopedTransaction>(LOCTEXT("AddClusterNode", "Add Cluster Node"));
-		
-		const FString ItemName = DialogContent->GetItemName();
-		NodeTemplate->SetFlags(RF_Transactional);
-
-		if (bAutoposition)
-		{
-			const FVector2D DesiredPosition = FVector2D(NodeTemplate->WindowRect.X, NodeTemplate->WindowRect.Y);
-			const FVector2D DesiredSize = FVector2D(NodeTemplate->WindowRect.W, NodeTemplate->WindowRect.H);
-			const FVector2D NewPosition = FindNextAvailablePositionForClusterNode(Cluster, NodeTemplate->Host, DesiredPosition, DesiredSize);
-
-			NodeTemplate->WindowRect.X = NewPosition.X;
-			NodeTemplate->WindowRect.Y = NewPosition.Y;
-		}
-
-		NewNode = AddClusterNodeToCluster(NodeTemplate, Cluster, ItemName);
-
-		// If the newly added cluster node is the only node in the cluster, it should be the primary node, so set it as the primary
-		if (Cluster->Nodes.Num() == 1)
-		{
-			SetClusterNodeAsPrimary(NewNode);
-		}
-
-		if (bAddViewport)
-		{
-			const FString ViewportName = GetUniqueNameForViewport(DefaultNewViewportName, NewNode, true);
-			UDisplayClusterConfigurationViewport* NewViewport = NewObject<UDisplayClusterConfigurationViewport>(NewNode, NAME_None, RF_Transactional | RF_ArchetypeObject | RF_Public);
-			NewViewport->Region.W = NewNode->WindowRect.W;
-			NewViewport->Region.H = NewNode->WindowRect.H;
-
-			NewViewport = AddViewportToClusterNode(NewViewport, NewNode, ViewportName);
-		}
-	}
-
-	// Clean up the template object so that it gets GCed once this schema action has been destroyed
-	NodeTemplate->Rename(nullptr, GetTransientPackage(), REN_DoNotDirty | REN_DontCreateRedirectors | REN_ForceNoResetLoaders);
-	NodeTemplate->SetFlags(RF_Transient);
-
-	return NewNode;
-}
-
-UDisplayClusterConfigurationViewport* FDisplayClusterConfiguratorClusterUtils::CreateNewViewportFromDialog(
-	const TSharedRef<FDisplayClusterConfiguratorBlueprintEditor>& Toolkit,
-	UDisplayClusterConfigurationClusterNode* ClusterNode,
-	const FDisplayClusterConfigurationRectangle& PresetRect,
-	TSharedPtr<FScopedTransaction>& OutTransaction)
-{
-	UDisplayClusterConfigurationViewport* ViewportTemplate = NewObject<UDisplayClusterConfigurationViewport>(Toolkit->GetBlueprintObj(), NAME_None, RF_Transactional | RF_ArchetypeObject | RF_Public);
-	ViewportTemplate->Region = FDisplayClusterConfigurationRectangle(PresetRect);
-
-	UDisplayClusterConfigurationCluster* Cluster = Toolkit->GetEditorData()->Cluster;
-
-	TArray<FString> ParentItems;
-	for (const TPair<FString, TObjectPtr<UDisplayClusterConfigurationClusterNode>>& Node : Cluster->Nodes)
-	{
-		ParentItems.Add(Node.Key);
-	}
-
-	FString ClusterNodeName = "";
-	if (const FString* Key = Cluster->Nodes.FindKey(ClusterNode))
+	/**
+	 * Gets a unique name for something given a list of existing names.
+	 * @param InitialName - The initial name the object wants
+	 * @param UsedNames - The list of names that have already been used
+	 * @param Class - The class to generate a name for
+	 * @param Parent - The parent to check for in-memory objects that collide with the name
+	 * @param bAddZero - Whether to add an "_0" to the initial name if it is unique
+	 * @returns A unique name for the object
+	 */
+	FString GetUniqueName(const FString& InitialName, const TArray<FString>& UsedNames, const UClass* Class, UObject* Parent, bool bAddZero = false)
 	{
-		ClusterNodeName = *Key;
-	}
-
-	const FString InitialName = ClusterNode ? 
-		FDisplayClusterConfiguratorClusterUtils::GetUniqueNameForViewport(DefaultNewViewportName, ClusterNode, true) : 
-		FString::Printf(TEXT("%s_%d"), *DefaultNewViewportName, 0);
-
-	bool bAutoposition = true;
-	bool bExpandClusterNode = true;
-	TSharedRef<SWidget> ViewportFooter = SNew(SVerticalBox)
-		+SVerticalBox::Slot()
-		.AutoHeight()
-		.Padding(2.0f)
-		[
-			SNew(SCheckBox)
-			.IsChecked_Lambda([&bAutoposition]() { return bAutoposition ? ECheckBoxState::Checked : ECheckBoxState::Unchecked; })
-			.OnCheckStateChanged_Lambda([&bAutoposition](ECheckBoxState NewState) { bAutoposition = NewState == ECheckBoxState::Checked; })
-			.ToolTipText(LOCTEXT("AddNewViewport_AutoPositionToolTip", "When checked, auto-positions the new Viewport to prevent overlap with other Viewports in the same cluster Node"))
-			[
-				SNew(STextBlock)
-				.Text(LOCTEXT("AddNewViewport_AutoPositionLabel", "Adjust Viewport Position to Prevent Overlap"))
-			]
-		]
-	
-		+SVerticalBox::Slot()
-		.AutoHeight()
-		.Padding(2.0f)
-		[
-			SNew(SCheckBox)
-			.IsChecked_Lambda([&bExpandClusterNode]() { return bExpandClusterNode ? ECheckBoxState::Checked : ECheckBoxState::Unchecked; })
-			.OnCheckStateChanged_Lambda([&bExpandClusterNode](ECheckBoxState NewState) { bExpandClusterNode = NewState == ECheckBoxState::Checked; })
-			.ToolTipText(LOCTEXT("AddNewViewport_ExpandClusterNodeToolTip", "When checked, expands the parent cluster Node to contain the new Viewport"))
-			[
-				SNew(STextBlock)
-				.Text(LOCTEXT("AddNewViewport_ExpandClusterNodeLabel", "Expand Cluster Node to Fit New Viewport"))
-			]
-		];
-
-	TSharedRef<SDisplayClusterConfiguratorNewClusterItemDialog> DialogContent = SNew(SDisplayClusterConfiguratorNewClusterItemDialog, ViewportTemplate)
-		.ParentItemOptions(ParentItems)
-		.InitiallySelectedParentItem(ClusterNodeName)
-		.PresetItemOptions(FDisplayClusterConfiguratorPresetSize::CommonPresets)
-		.InitiallySelectedPreset(FDisplayClusterConfiguratorPresetSize::CommonPresets[FDisplayClusterConfiguratorPresetSize::DefaultPreset])
-		.InitialName(InitialName)
-		.MaxWindowWidth(NewClusterItemDialogSize.X)
-		.MaxWindowHeight(NewClusterItemDialogSize.Y)
-		.FooterContent(ViewportFooter)
-		.OnPresetChanged_Lambda([=](FVector2D Size) { ViewportTemplate->Region.W = Size.X; ViewportTemplate->Region.H = Size.Y; });
-
-	{
-		UE::DisplayClusterConfiguratorClusterUtils::FScopedTemporaryTransactions ScopedTransactionsForModal;
-		ShowNewClusterItemDialogWindow(DialogContent, nullptr, LOCTEXT("AddNewViewport_DialogTitle", "Add New Viewport"), NewClusterItemDialogSize);
-	}
-	
-	UDisplayClusterConfigurationViewport* NewViewport = nullptr;
-
-	if (DialogContent->WasAccepted())
-	{
-		OutTransaction = MakeShared<FScopedTransaction>(LOCTEXT("AddViewport", "Add Viewport"));
-		
-		const FString ParentItem = DialogContent->GetSelectedParentItem();
-		const FString ItemName = DialogContent->GetItemName();
-
-		check(Cluster->Nodes.Contains(ParentItem));
-
-		UDisplayClusterConfigurationClusterNode* ParentNode = Cluster->Nodes[ParentItem];
-
-		if (bAutoposition)
+		FString NewName;
+		if (!bAddZero && !UsedNames.Contains(InitialName))
 		{
-			FVector2D DesiredPosition = FVector2D(ViewportTemplate->Region.X, ViewportTemplate->Region.Y);
-			FVector2D DesiredSize = FVector2D(ViewportTemplate->Region.W, ViewportTemplate->Region.H);
-			FVector2D NewPosition = FindNextAvailablePositionForViewport(ParentNode, DesiredPosition, DesiredSize);
-
-			ViewportTemplate->Region.X = NewPosition.X;
-			ViewportTemplate->Region.Y = NewPosition.Y;
-		}
-
-		NewViewport = DuplicateObject(ViewportTemplate, ParentNode);
-		NewViewport->SetFlags(RF_Transactional);
-		NewViewport = AddViewportToClusterNode(NewViewport, ParentNode, ItemName);
-
-		if (bExpandClusterNode)
-		{
-			FIntRect ViewportRegion = NewViewport->Region.ToRect();
-			FIntRect ClusterWindowRect = FIntRect(0, 0, ParentNode->WindowRect.W, ParentNode->WindowRect.H);
-
-			ClusterWindowRect.Union(ViewportRegion);
-			FIntPoint WindowSize = ClusterWindowRect.Max;
-			FDisplayClusterConfigurationRectangle NewWindowRect(ParentNode->WindowRect.X, ParentNode->WindowRect.Y, WindowSize.X, WindowSize.Y);
-
-			FDisplayClusterConfiguratorClusterNodeViewModel ClusterNodeVM(ParentNode);
-			ClusterNodeVM.SetWindowRect(NewWindowRect);
-		}
-	}
-
-	// Clean up the template object so that it gets GCed once this schema action has been destroyed
-	ViewportTemplate->Rename(nullptr, GetTransientPackage(), REN_DoNotDirty | REN_DontCreateRedirectors | REN_ForceNoResetLoaders);
-	ViewportTemplate->SetFlags(RF_Transient);
-	
-	return NewViewport;
-}
-
-void FDisplayClusterConfiguratorClusterUtils::ShowNewClusterItemDialogWindow(TSharedRef<SDisplayClusterConfiguratorNewClusterItemDialog> DialogContent, TSharedPtr<SWidget> ParentElement, FText WindowTitle, FVector2D WindowSize)
-{
-	TSharedPtr<SWidget> ParentWindow = FSlateApplication::Get().FindBestParentWindowForDialogs(ParentElement);
-
-	const FSlateRect WorkAreaRect = FSlateApplicationBase::Get().GetPreferredWorkArea();
-	const FVector2D DisplayTopLeft(WorkAreaRect.Left, WorkAreaRect.Top);
-	const FVector2D DisplaySize(WorkAreaRect.Right - WorkAreaRect.Left, WorkAreaRect.Bottom - WorkAreaRect.Top);
-
-	const float ScaleFactor = FPlatformApplicationMisc::GetDPIScaleFactorAtPoint(DisplayTopLeft.X, DisplayTopLeft.Y);
-	const FVector2D DialogWindowSize = WindowSize * ScaleFactor;
-
-	const FVector2D WindowPosition = (DisplayTopLeft + 0.5f * (DisplaySize - DialogWindowSize)) / ScaleFactor;
-
-	TSharedRef<SWindow> Window = SNew(SWindow)
-		.Title(WindowTitle)
-		.SizingRule(ESizingRule::Autosized)
-		.AutoCenter(EAutoCenter::None)
-		.ClientSize(DialogWindowSize)
-		.ScreenPosition(WindowPosition);
-
-	DialogContent->SetParentWindow(Window);
-	Window->SetContent(DialogContent);
-
-	FSlateApplication::Get().AddModalWindow(Window, ParentWindow, false);
-}
-
-TSharedPtr<FDragDropOperation> FDisplayClusterConfiguratorClusterUtils::MakeDragDropOperation(const TArray<UObject*>& SelectedObjects)
-{
-	TArray<UDisplayClusterConfigurationViewport*> Viewports;
-	TArray<UDisplayClusterConfigurationClusterNode*> ClusterNodes;
-
-	for (UObject* SelectedObject : SelectedObjects)
-	{
-		if (UDisplayClusterConfigurationViewport* Viewport = Cast<UDisplayClusterConfigurationViewport>(SelectedObject))
-		{
-			Viewports.Add(Viewport);
-		}
-		else if (UDisplayClusterConfigurationClusterNode* ClusterNode = Cast<UDisplayClusterConfigurationClusterNode>(SelectedObject))
-		{
-			ClusterNodes.Add(ClusterNode);
+			// Name doesn't need to be modified
+			NewName = InitialName;
 		}
 		else
 		{
-			// Only cluster nodes and viewports can be dragged and dropped. If anything else is selected, abort the drag/drop operation
-			return TSharedPtr<FDragDropOperation>();
-		}
-	}
-
-	if (Viewports.Num() && !ClusterNodes.Num())
-	{
-		// If only viewports are selected, return a viewport drag/drop operation
-		return FDisplayClusterConfiguratorViewportDragDropOp::New(Viewports);
-	}
-	else if (ClusterNodes.Num() && !Viewports.Num())
-	{
-		// If only cluster nodes are selected, return a cluster node drag/drop operation
-		return FDisplayClusterConfiguratorClusterNodeDragDropOp::New(ClusterNodes);
-	}
-	else
-	{
-		// If both cluster nodes and viewports are selected, only allow a cluster node drag/drop operation if all selected viewports are children of the selected cluster nodes
-		bool bOnlyChildViewports = true;
-		for (UDisplayClusterConfigurationViewport* Viewport : Viewports)
-		{
-			bool bIsChild = ClusterNodes.ContainsByPredicate([=](UDisplayClusterConfigurationClusterNode* ClusterNode)
-			{
-				return ClusterNode == Viewport->GetOuter();
-			});
+			int32 Counter = bAddZero ? 0 : 1;
 
-			if (!bIsChild)
+			// Find the start of the existing numeric suffix
+			int32 Index = InitialName.Len();
+			while (Index > 0 && InitialName[Index-1] >= '0' && InitialName[Index-1] <= '9')
 			{
-				bOnlyChildViewports = false;
-				break;
+				--Index;
 			}
-		}
 
-		if (bOnlyChildViewports)
-		{
-			return FDisplayClusterConfiguratorClusterNodeDragDropOp::New(ClusterNodes);
-		}
-		else
-		{
-			return TSharedPtr<FDragDropOperation>();
-		}
-	}
-}
-
-namespace
-{
-	FVector2D FindAvailableSpace(const TArray<FBox2D>& OtherBounds, const FBox2D& DesiredBounds)
-	{
-		FBox2D CurrentBounds(DesiredBounds);
-
-		// Similar to FBox2D::Intersects, but ignores the case where the box edges are touching.
-		auto IntrudesFunc = [](const FBox2D& BoxA, const FBox2D& BoxB)
-		{
-
-			// Special case if both boxes are directly on top of each other, which is considered an intrusion.
-			if (BoxA == BoxB)
+			FString BaseName = InitialName;
+			if (Index < BaseName.Len())
 			{
-				return true;
+				// Strip away the suffix and store the value in the counter so we can count up from there
+				FString NumericSuffix = BaseName.RightChop(Index);
+				Counter = FCString::Atoi(*NumericSuffix);
+				NumericSuffix = FString::FromInt(Counter); // Restringify the counter to account for leading 0s that we don't want to remove
+				BaseName.RemoveAt(BaseName.Len() - NumericSuffix.Len(), NumericSuffix.Len(), false);
 			}
-
-			if ((BoxA.Min.X >= BoxB.Max.X) || (BoxB.Min.X >= BoxA.Max.X))
+			else
 			{
-				return false;
+				// No existing suffix, so add our underscore separator
+				BaseName += "_";
 			}
 
-			if ((BoxA.Min.Y >= BoxB.Max.Y) || (BoxB.Min.Y >= BoxA.Max.Y))
+			do
 			{
-				return false;
+				NewName = FString::Printf(TEXT("%s%d"), *BaseName, Counter);
+				++Counter;
 			}
-
-			return true;
-		};
-
-		bool bIntersectsBounds;
-		do
-		{
-			bIntersectsBounds = false;
-
-			for (const FBox2D& Bounds : OtherBounds)
-			{
-				if (IntrudesFunc(CurrentBounds, Bounds))
-				{
-					bIntersectsBounds = true;
-
-					float Shift = Bounds.GetSize().X;
-					CurrentBounds.Min.X += Shift;
-					CurrentBounds.Max.X += Shift;
-
-					break;
-				}
-			}
-
-		} while (bIntersectsBounds);
-
-		return CurrentBounds.Min;
-	}
-}
-
-FVector2D FDisplayClusterConfiguratorClusterUtils::FindNextAvailablePositionForClusterNode(UDisplayClusterConfigurationCluster* Cluster, const FString& DesiredHost, const FVector2D& DesiredPosition, const FVector2D& DesiredSize)
-{
-	TMap<FString, TMap<FString, UDisplayClusterConfigurationClusterNode*>> ClusterNodesByHost;
-	SortClusterNodesByHost(Cluster->Nodes, ClusterNodesByHost);
-
-	// If the desired host doesn't exist yet, simply return the node's desired position
-	if (!ClusterNodesByHost.Contains(DesiredHost))
-	{
-		return DesiredPosition;
-	}
-
-	const TMap<FString, UDisplayClusterConfigurationClusterNode*>& ExistingClusterNodes = ClusterNodesByHost[DesiredHost];
-	if (!ExistingClusterNodes.Num())
-	{
-		return DesiredPosition;
-	}
-
-	FBox2D DesiredBounds = FBox2D(DesiredPosition, DesiredPosition + DesiredSize);
-	TArray<FBox2D> ClusterNodeBounds;
-	for (const TPair<FString, UDisplayClusterConfigurationClusterNode*>& ClusterNodeKeyPair : ExistingClusterNodes)
-	{
-		const FDisplayClusterConfigurationRectangle& WindowRect = ClusterNodeKeyPair.Value->WindowRect;
-		ClusterNodeBounds.Add(FBox2D(FVector2D(WindowRect.X, WindowRect.Y), FVector2D(WindowRect.X + WindowRect.W, WindowRect.Y + WindowRect.H)));
-	}
-
-	return FindAvailableSpace(ClusterNodeBounds, DesiredBounds);
-}
-
-FVector2D FDisplayClusterConfiguratorClusterUtils::FindNextAvailablePositionForViewport(UDisplayClusterConfigurationClusterNode* ClusterNode, const FVector2D& DesiredPosition, const FVector2D& DesiredSize)
-{
-	if (!ClusterNode)
-	{
-		return DesiredPosition;
-	}
-
-	FBox2D DesiredBounds = FBox2D(DesiredPosition, DesiredPosition + DesiredSize);
-	TArray<FBox2D> ViewportBounds;
-	for (const TPair<FString, TObjectPtr<UDisplayClusterConfigurationViewport>>& ViewportKeyPair : ClusterNode->Viewports)
-	{
-		const FDisplayClusterConfigurationRectangle& Region = ViewportKeyPair.Value->Region;
-		ViewportBounds.Add(FBox2D(FVector2D(Region.X, Region.Y), FVector2D(Region.X + Region.W, Region.Y + Region.H)));
-	}
-
-	return FindAvailableSpace(ViewportBounds, DesiredBounds);
-}
-
-void FDisplayClusterConfiguratorClusterUtils::SortClusterNodesByHost(const TMap<FString, UDisplayClusterConfigurationClusterNode*>& InClusterNodes, TMap<FString, TMap<FString, UDisplayClusterConfigurationClusterNode*>>& OutSortedNodes)
-{
-	for (const TPair<FString, UDisplayClusterConfigurationClusterNode*>& ClusterNodePair : InClusterNodes)
-	{
-		check(ClusterNodePair.Value)
-
-		FString Host = ClusterNodePair.Value->Host;
-		if (!OutSortedNodes.Contains(Host))
-		{
-			OutSortedNodes.Add(Host, TMap<FString, UDisplayClusterConfigurationClusterNode*>());
+			while (UsedNames.Contains(NewName));
 		}
-
-		OutSortedNodes[Host].Add(ClusterNodePair);
-	}
-
-	// Sort the hosts by the host address
-	OutSortedNodes.KeySort(TLess<FString>());
-}
-
-void FDisplayClusterConfiguratorClusterUtils::SortClusterNodesByHost(const TMap<FString, TObjectPtr<UDisplayClusterConfigurationClusterNode>>& InClusterNodes, TMap<FString, TMap<FString, UDisplayClusterConfigurationClusterNode*>>& OutSortedNodes)
-{
-	for (const TPair<FString, TObjectPtr<UDisplayClusterConfigurationClusterNode>>& ClusterNodePair : InClusterNodes)
-	{
-		check(ClusterNodePair.Value)
-
-			FString Host = ClusterNodePair.Value->Host;
-		if (!OutSortedNodes.Contains(Host))
+	
+		// If there is already an in-memory object connected to the parent cluster with our generated name, we need to use a globally unique object name.
+		if (StaticFindObject(nullptr, Parent, *NewName, true))
 		{
-			OutSortedNodes.Add(Host, TMap<FString, UDisplayClusterConfigurationClusterNode*>());
+			NewName = MakeUniqueObjectName(Parent, Class, *NewName).ToString();
 		}
 
-		OutSortedNodes[Host].Add(ClusterNodePair);
+		return NewName;
 	}
-
-	// Sort the hosts by the host address
-	OutSortedNodes.KeySort(TLess<FString>());
+	
+	const FString DefaultNewHostName = TEXT("Host");
 }
 
-UDisplayClusterConfigurationHostDisplayData* FDisplayClusterConfiguratorClusterUtils::FindOrCreateHostDisplayData(UDisplayClusterConfigurationCluster* Cluster, FString HostIPAddress)
+using namespace UE::DisplayClusterConfiguratorPropertyUtils;
+
+UDisplayClusterConfigurationHostDisplayData* UE::DisplayClusterConfiguratorClusterUtils::FindOrCreateHostDisplayData(UDisplayClusterConfigurationCluster* Cluster, FString HostIPAddress)
 {
 	// In some cases, existing host display data may be pending kill, such as if the user recently performed an undo to a state
 	// prior to the data's existence. In this case, simply remove existing host data that is pending kill and create a new one to use.
@@ -582,7 +112,7 @@ UDisplayClusterConfigurationHostDisplayData* FDisplayClusterConfiguratorClusterU
 	return Cluster->HostDisplayData[HostIPAddress];
 }
 
-bool FDisplayClusterConfiguratorClusterUtils::RemoveUnusedHostDisplayData(UDisplayClusterConfigurationCluster* Cluster)
+bool UE::DisplayClusterConfiguratorClusterUtils::RemoveUnusedHostDisplayData(UDisplayClusterConfigurationCluster* Cluster)
 {
 	bool bHostDataRemoved = false;
 
@@ -610,7 +140,7 @@ bool FDisplayClusterConfiguratorClusterUtils::RemoveUnusedHostDisplayData(UDispl
 	return bHostDataRemoved;
 }
 
-FString FDisplayClusterConfiguratorClusterUtils::GetUniqueNameForHost(FString InitialName, UDisplayClusterConfigurationCluster* ParentCluster, bool bAddZero)
+FString UE::DisplayClusterConfiguratorClusterUtils::GetUniqueNameForHost(FString InitialName, UDisplayClusterConfigurationCluster* ParentCluster, bool bAddZero)
 {
 	TArray<FString> UsedNames;
 	for (TPair<FString, TObjectPtr<UDisplayClusterConfigurationHostDisplayData>>& HostPair : ParentCluster->HostDisplayData)
@@ -621,7 +151,7 @@ FString FDisplayClusterConfiguratorClusterUtils::GetUniqueNameForHost(FString In
 	return GetUniqueName(InitialName, UsedNames, UDisplayClusterConfigurationClusterNode::StaticClass(), ParentCluster, bAddZero);
 }
 
-FString FDisplayClusterConfiguratorClusterUtils::GetAddressForHost(UDisplayClusterConfigurationHostDisplayData* HostDisplayData)
+FString UE::DisplayClusterConfiguratorClusterUtils::GetAddressForHost(UDisplayClusterConfigurationHostDisplayData* HostDisplayData)
 {
 	FString HostAddress;
 	if (UDisplayClusterConfigurationCluster* HostParent = Cast<UDisplayClusterConfigurationCluster>(HostDisplayData->GetOuter()))
@@ -635,7 +165,7 @@ FString FDisplayClusterConfiguratorClusterUtils::GetAddressForHost(UDisplayClust
 	return HostAddress;
 }
 
-bool FDisplayClusterConfiguratorClusterUtils::RemoveHost(UDisplayClusterConfigurationCluster* Cluster, FString Host)
+bool UE::DisplayClusterConfiguratorClusterUtils::RemoveHost(UDisplayClusterConfigurationCluster* Cluster, FString Host)
 {
 	bool bDataRemoved = false;
 
@@ -668,7 +198,7 @@ bool FDisplayClusterConfiguratorClusterUtils::RemoveHost(UDisplayClusterConfigur
 	return bDataRemoved;
 }
 
-UDisplayClusterConfigurationHostDisplayData* FDisplayClusterConfiguratorClusterUtils::GetHostDisplayDataForClusterNode(UDisplayClusterConfigurationClusterNode* ClusterNode)
+UDisplayClusterConfigurationHostDisplayData* UE::DisplayClusterConfiguratorClusterUtils::GetHostDisplayDataForClusterNode(UDisplayClusterConfigurationClusterNode* ClusterNode)
 {
 	UDisplayClusterConfigurationHostDisplayData* HostDisplayData = nullptr;
 	if (UDisplayClusterConfigurationCluster* ClusterNodeParent = Cast<UDisplayClusterConfigurationCluster>(ClusterNode->GetOuter()))
@@ -682,7 +212,7 @@ UDisplayClusterConfigurationHostDisplayData* FDisplayClusterConfiguratorClusterU
 	return HostDisplayData;
 }
 
-FString FDisplayClusterConfiguratorClusterUtils::GetClusterNodeName(UDisplayClusterConfigurationClusterNode* ClusterNode)
+FString UE::DisplayClusterConfiguratorClusterUtils::GetClusterNodeName(UDisplayClusterConfigurationClusterNode* ClusterNode)
 {
 	if (UDisplayClusterConfigurationCluster* ClusterNodeParent = Cast<UDisplayClusterConfigurationCluster>(ClusterNode->GetOuter()))
 	{
@@ -695,7 +225,7 @@ FString FDisplayClusterConfiguratorClusterUtils::GetClusterNodeName(UDisplayClus
 	return "";
 }
 
-bool FDisplayClusterConfiguratorClusterUtils::IsClusterNodePrimary(UDisplayClusterConfigurationClusterNode* ClusterNode)
+bool UE::DisplayClusterConfiguratorClusterUtils::IsClusterNodePrimary(UDisplayClusterConfigurationClusterNode* ClusterNode)
 {
 	if (const UDisplayClusterConfigurationCluster* Cluster = Cast<UDisplayClusterConfigurationCluster>(ClusterNode->GetOuter()))
 	{
@@ -709,7 +239,7 @@ bool FDisplayClusterConfiguratorClusterUtils::IsClusterNodePrimary(UDisplayClust
 	return false;
 }
 
-FString FDisplayClusterConfiguratorClusterUtils::GetUniqueNameForClusterNode(FString InitialName, UDisplayClusterConfigurationCluster* ParentCluster, bool bAddZero)
+FString UE::DisplayClusterConfiguratorClusterUtils::GetUniqueNameForClusterNode(FString InitialName, UDisplayClusterConfigurationCluster* ParentCluster, bool bAddZero)
 {
 	InitialName = ObjectTools::SanitizeObjectName(InitialName);
 	
@@ -719,7 +249,7 @@ FString FDisplayClusterConfiguratorClusterUtils::GetUniqueNameForClusterNode(FSt
 	return GetUniqueName(InitialName, UsedNames, UDisplayClusterConfigurationClusterNode::StaticClass(), ParentCluster, bAddZero);
 }
 
-UDisplayClusterConfigurationClusterNode* FDisplayClusterConfiguratorClusterUtils::AddClusterNodeToCluster(UDisplayClusterConfigurationClusterNode* ClusterNode, UDisplayClusterConfigurationCluster* Cluster, FString NewClusterNodeName)
+UDisplayClusterConfigurationClusterNode* UE::DisplayClusterConfiguratorClusterUtils::AddClusterNodeToCluster(UDisplayClusterConfigurationClusterNode* ClusterNode, UDisplayClusterConfigurationCluster* Cluster, FString NewClusterNodeName)
 {
 	FString ClusterNodeName = "";
 
@@ -763,7 +293,7 @@ UDisplayClusterConfigurationClusterNode* FDisplayClusterConfiguratorClusterUtils
 	return NewClusterNode;
 }
 
-bool FDisplayClusterConfiguratorClusterUtils::RemoveClusterNodeFromCluster(UDisplayClusterConfigurationClusterNode* ClusterNode)
+bool UE::DisplayClusterConfiguratorClusterUtils::RemoveClusterNodeFromCluster(UDisplayClusterConfigurationClusterNode* ClusterNode)
 {
 	if (UDisplayClusterConfigurationCluster* ClusterNodeParent = Cast<UDisplayClusterConfigurationCluster>(ClusterNode->GetOuter()))
 	{
@@ -802,7 +332,7 @@ bool FDisplayClusterConfiguratorClusterUtils::RemoveClusterNodeFromCluster(UDisp
 	return false;
 }
 
-bool FDisplayClusterConfiguratorClusterUtils::RenameClusterNode(UDisplayClusterConfigurationClusterNode* ClusterNode, FString NewClusterNodeName)
+bool UE::DisplayClusterConfiguratorClusterUtils::RenameClusterNode(UDisplayClusterConfigurationClusterNode* ClusterNode, FString NewClusterNodeName)
 {
 	if (UDisplayClusterConfigurationCluster* ClusterNodeParent = Cast<UDisplayClusterConfigurationCluster>(ClusterNode->GetOuter()))
 	{
@@ -845,7 +375,7 @@ bool FDisplayClusterConfiguratorClusterUtils::RenameClusterNode(UDisplayClusterC
 	return false;
 }
 
-bool FDisplayClusterConfiguratorClusterUtils::SetClusterNodeAsPrimary(UDisplayClusterConfigurationClusterNode* ClusterNode)
+bool UE::DisplayClusterConfiguratorClusterUtils::SetClusterNodeAsPrimary(UDisplayClusterConfigurationClusterNode* ClusterNode)
 {
 	if (UDisplayClusterConfigurationCluster* ClusterNodeParent = Cast<UDisplayClusterConfigurationCluster>(ClusterNode->GetOuter()))
 	{
@@ -862,7 +392,7 @@ bool FDisplayClusterConfiguratorClusterUtils::SetClusterNodeAsPrimary(UDisplayCl
 	return false;
 }
 
-FString FDisplayClusterConfiguratorClusterUtils::GetViewportName(UDisplayClusterConfigurationViewport* Viewport)
+FString UE::DisplayClusterConfiguratorClusterUtils::GetViewportName(UDisplayClusterConfigurationViewport* Viewport)
 {
 	if (UDisplayClusterConfigurationClusterNode* ViewportParent = Cast<UDisplayClusterConfigurationClusterNode>(Viewport->GetOuter()))
 	{
@@ -875,7 +405,7 @@ FString FDisplayClusterConfiguratorClusterUtils::GetViewportName(UDisplayCluster
 	return "";
 }
 
-FString FDisplayClusterConfiguratorClusterUtils::GetUniqueNameForViewport(FString InitialName, UDisplayClusterConfigurationClusterNode* ParentClusterNode, bool bAddZero)
+FString UE::DisplayClusterConfiguratorClusterUtils::GetUniqueNameForViewport(FString InitialName, UDisplayClusterConfigurationClusterNode* ParentClusterNode, bool bAddZero)
 {
 	InitialName = ObjectTools::SanitizeObjectName(InitialName);
 	
@@ -908,7 +438,7 @@ FString FDisplayClusterConfiguratorClusterUtils::GetUniqueNameForViewport(FStrin
 	return GetUniqueName(InitialName, UsedNamesArray, UDisplayClusterConfigurationViewport::StaticClass(), ParentClusterNode, bAddZero);
 }
 
-UDisplayClusterConfigurationViewport* FDisplayClusterConfiguratorClusterUtils::AddViewportToClusterNode(UDisplayClusterConfigurationViewport* Viewport, UDisplayClusterConfigurationClusterNode* ClusterNode, FString NewViewportName)
+UDisplayClusterConfigurationViewport* UE::DisplayClusterConfiguratorClusterUtils::AddViewportToClusterNode(UDisplayClusterConfigurationViewport* Viewport, UDisplayClusterConfigurationClusterNode* ClusterNode, FString NewViewportName)
 {
 	FString ViewportName = "";
 
@@ -953,7 +483,7 @@ UDisplayClusterConfigurationViewport* FDisplayClusterConfiguratorClusterUtils::A
 	return NewViewport;
 }
 
-bool FDisplayClusterConfiguratorClusterUtils::RemoveViewportFromClusterNode(UDisplayClusterConfigurationViewport* Viewport)
+bool UE::DisplayClusterConfiguratorClusterUtils::RemoveViewportFromClusterNode(UDisplayClusterConfigurationViewport* Viewport)
 {
 	if (UDisplayClusterConfigurationClusterNode* ViewportParent = Cast<UDisplayClusterConfigurationClusterNode>(Viewport->GetOuter()))
 	{
@@ -975,7 +505,7 @@ bool FDisplayClusterConfiguratorClusterUtils::RemoveViewportFromClusterNode(UDis
 	return false;
 }
 
-bool FDisplayClusterConfiguratorClusterUtils::RenameViewport(UDisplayClusterConfigurationViewport* Viewport, FString NewViewportName)
+bool UE::DisplayClusterConfiguratorClusterUtils::RenameViewport(UDisplayClusterConfigurationViewport* Viewport, FString NewViewportName)
 {
 	if (UDisplayClusterConfigurationClusterNode* ViewportParent = Cast<UDisplayClusterConfigurationClusterNode>(Viewport->GetOuter()))
 	{
@@ -1016,352 +546,42 @@ bool FDisplayClusterConfiguratorClusterUtils::RenameViewport(UDisplayClusterConf
 	return false;
 }
 
-void FDisplayClusterConfiguratorClusterUtils::CopyClusterItemsToClipboard(const TArray<UObject*>& ClusterItemsToCopy)
+void UE::DisplayClusterConfiguratorClusterUtils::SortClusterNodesByHost(const TMap<FString, UDisplayClusterConfigurationClusterNode*>& InClusterNodes, TMap<FString, TMap<FString, UDisplayClusterConfigurationClusterNode*>>& OutSortedNodes)
 {
-	TSet<UObject*> UniqueItemsToCopy;
-
-	for (UObject* ClusterItem : ClusterItemsToCopy)
-	{
-		// Only allow viewports and cluster nodes to be copied at this point.
-		if (ClusterItem->IsA<UDisplayClusterConfigurationViewport>())
-		{
-			// If a viewport is being copied, make sure its parent cluster node is not also being copied; if its parent is being copied,
-			// the viewport doesn't need to be copied directly as it will be copied with its parent.
-			if (!ClusterItemsToCopy.Contains(ClusterItem->GetOuter()))
-			{
-				UniqueItemsToCopy.Add(ClusterItem);
-			}
-		}
-		else if (ClusterItem->IsA<UDisplayClusterConfigurationClusterNode>())
-		{
-			UniqueItemsToCopy.Add(ClusterItem);
-		}
-	}
-
-	UnMarkAllObjects(EObjectMark(OBJECTMARK_TagExp | OBJECTMARK_TagImp));
-
-	FStringOutputDevice Archive;
-	const FExportObjectInnerContext Context;
-
-	for (UObject* ClusterItem : UniqueItemsToCopy)
-	{
-		UExporter::ExportToOutputDevice(&Context, ClusterItem, nullptr, Archive, TEXT("copy"), 0, PPF_ExportsNotFullyQualified | PPF_Copy | PPF_Delimited, false, nullptr);
-	}
-
-	FString CopiedText = Archive;
-	FPlatformApplicationMisc::ClipboardCopy(*CopiedText);
-}
-
-class FClusterNodeTextFactory : public FCustomizableTextObjectFactory
-{
-public:
-	FClusterNodeTextFactory()
-		: FCustomizableTextObjectFactory(GWarn)
-	{
-	}
-
-	// FCustomizableTextObjectFactory implementation
-
-	virtual bool CanCreateClass(UClass* ObjectClass, bool& bOmitSubObjs) const override
-	{
-		const bool isClusterNode = ObjectClass->IsChildOf(UDisplayClusterConfigurationClusterNode::StaticClass());
-		const bool isViewportNode = ObjectClass->IsChildOf(UDisplayClusterConfigurationViewport::StaticClass());
-
-		return isClusterNode || isViewportNode;
-	}
-
-	virtual void ProcessConstructedObject(UObject* NewObject) override
-	{
-		check(NewObject);
-
-		if (UDisplayClusterConfigurationClusterNode* ClusterNode = Cast<UDisplayClusterConfigurationClusterNode>(NewObject))
-		{
-			ClusterNodes.Add(ClusterNode);
-		}
-		else if (UDisplayClusterConfigurationViewport* ViewportNode = Cast<UDisplayClusterConfigurationViewport>(NewObject))
-		{
-			Viewports.Add(ViewportNode);
-		}
-	}
-
-	// FCustomizableTextObjectFactory (end)
-
-public:
-	TArray<UDisplayClusterConfigurationClusterNode*> ClusterNodes;
-	TArray<UDisplayClusterConfigurationViewport*> Viewports;
-};
-
-bool FDisplayClusterConfiguratorClusterUtils::CanPasteClusterItemsFromClipboard(const TArray<UObject*>& TargetClusterItems, int32& OutNumItems)
-{
-	// Can't paste unless the clipboard has a string in it
-	FString ClipboardContent;
-	FPlatformApplicationMisc::ClipboardPaste(ClipboardContent);
-	if (!ClipboardContent.IsEmpty())
-	{
-		FClusterNodeTextFactory Factory;
-		Factory.ProcessBuffer(GetTransientPackage(), RF_Transactional, ClipboardContent);
-
-		OutNumItems = Factory.ClusterNodes.Num() + Factory.Viewports.Num();
-
-		bool bCanPaste = false;
-		if (Factory.ClusterNodes.Num() > 0)
-		{
-			// If we are trying to paste cluster nodes, the target should be a cluster root, a host display data, or a cluster node.
-			bCanPaste = TargetClusterItems.ContainsByPredicate([=](const UObject* Item)
-			{
-				return
-					Item->IsA<UDisplayClusterConfigurationCluster>() ||
-					Item->IsA<UDisplayClusterConfigurationClusterNode>() ||
-					Item->IsA<UDisplayClusterConfigurationHostDisplayData>();
-			});
-		}
-
-		if (Factory.Viewports.Num() > 0)
-		{
-			// If we are trying to paste viewports, the target should either be a cluster node or a viewport.
-			bCanPaste = TargetClusterItems.ContainsByPredicate([=](const UObject* Item)
-			{
-				return Item->IsA<UDisplayClusterConfigurationClusterNode>() || Item->IsA<UDisplayClusterConfigurationViewport>();
-			});
-		}
-
-		return bCanPaste;
-	}
-
-	return false;
-}
-
-TArray<UObject*> FDisplayClusterConfiguratorClusterUtils::PasteClusterItemsFromClipboard(const TArray<UObject*>& TargetClusterItems, TOptional<FVector2D> PasteLocation)
-{
-	FString ClipboardContent;
-	FPlatformApplicationMisc::ClipboardPaste(ClipboardContent);
-
-	FClusterNodeTextFactory Factory;
-	Factory.ProcessBuffer(GetTransientPackage(), RF_Transactional, ClipboardContent);
-
-	// Filter out the target items and find all unique cluster roots and cluster nodes. For each target cluster node, also add its parent root,
-	// and for each target viewport, add its parent cluster node. This allows users to paste cluster nodes and viewports by selecting their sibling nodes,
-	// which will then add the pasted node to the parent of the target node.
-
-	// Pair up the cluster roots with a host string, which indicates which host the cluster nodes should be pasted into (the cluster nodes' Host strings will be set when duplicated).
-	// If the paired host string is empty, the pasted cluster nodes' original host values will be preserved.
-	TSet<TTuple<UDisplayClusterConfigurationCluster*, FString>> SelectedRootHosts;
-	TSet<UDisplayClusterConfigurationClusterNode*> SelectedNodes;
-
-	for (UObject* TargetItem : TargetClusterItems)
-	{
-		if (UDisplayClusterConfigurationCluster* ClusterRoot = Cast<UDisplayClusterConfigurationCluster>(TargetItem))
-		{
-			SelectedRootHosts.Add(TTuple<UDisplayClusterConfigurationCluster*, FString>(ClusterRoot, ""));
-		}
-		else if (UDisplayClusterConfigurationHostDisplayData* Host = Cast<UDisplayClusterConfigurationHostDisplayData>(TargetItem))
-		{
-			// If the target is a host display data, that indicates the user wants to paste the cluster node into the host the display data represents,
-			// so add an entry to the cluster root list with the parent cluster object and the host string of the host display data.
-			if (UDisplayClusterConfigurationCluster* ParentRoot = Cast<UDisplayClusterConfigurationCluster>(TargetItem->GetOuter()))
-			{
-				FString HostStr = FString();
-				if (const FString* Key = ParentRoot->HostDisplayData.FindKey(Host))
-				{
-					HostStr = *Key;
-				}
-
-				SelectedRootHosts.Add(TTuple<UDisplayClusterConfigurationCluster*, FString>(ParentRoot, HostStr));
-			}
-		}
-		else if (UDisplayClusterConfigurationClusterNode* ClusterNode = Cast<UDisplayClusterConfigurationClusterNode>(TargetItem))
-		{
-			SelectedNodes.Add(ClusterNode);
-
-			if (UDisplayClusterConfigurationCluster* ParentRoot = Cast<UDisplayClusterConfigurationCluster>(TargetItem->GetOuter()))
-			{
-				// If pasting as a sibling of a target cluster node, set the pasted node's host to be the same as the sibling node's host.
-				SelectedRootHosts.Add(TTuple<UDisplayClusterConfigurationCluster*, FString>(ParentRoot, ClusterNode->Host));
-			}
-		}
-		else if (UDisplayClusterConfigurationViewport* Viewport = Cast<UDisplayClusterConfigurationViewport>(TargetItem))
-		{
-			if (UDisplayClusterConfigurationClusterNode* ParentNode = Cast<UDisplayClusterConfigurationClusterNode>(TargetItem->GetOuter()))
-			{
-				SelectedNodes.Add(ParentNode);
-			}
-		}
-	}
-
-	TArray<UObject*> CopiedObjects;
-
-	// Keep track of a couple reference positions for pasting the items to a specfic location to correctly offset pasted location in the case where
-	// more than one item is being pasted or an item is being pasted into more than one parent. The ReferencePosition variable keeps track of the 
-	// relative position of the item to be pasted, while the ReferenceTargetPosition keeps track of the relative position of the target item being pasted
-	// into. The way these work is that, if they aren't set when in the pasting loop below, then they are set to the current iteration item's position. Every
-	// other iteration will then subtract the reference position from the pasted position. This allows each item being pasted to be correctly offset to maintain
-	// both its relative position in the group of items being pasted and its relative position to the item it is being pasted into.
-	TOptional<FVector2D> ReferencePosition = TOptional<FVector2D>();
-	TOptional<FVector2D> ReferenceTargetPosition = TOptional<FVector2D>();
-
-	for (const TTuple<UDisplayClusterConfigurationCluster*, FString>& ClusterRootHost : SelectedRootHosts)
+	for (const TPair<FString, UDisplayClusterConfigurationClusterNode*>& ClusterNodePair : InClusterNodes)
 	{
-		UDisplayClusterConfigurationCluster* ClusterRoot = ClusterRootHost.Key;
-		FString HostStr = ClusterRootHost.Value;
+		check(ClusterNodePair.Value)
 
-		for (UDisplayClusterConfigurationClusterNode* ClusterNode : Factory.ClusterNodes)
+		FString Host = ClusterNodePair.Value->Host;
+		if (!OutSortedNodes.Contains(Host))
 		{
-			UDisplayClusterConfigurationClusterNode* ClusterNodeCopy = DuplicateObject(ClusterNode, ClusterRoot);
-			ClusterNodeCopy->SetFlags(RF_Transactional);
-			ClusterNodeCopy = AddClusterNodeToCluster(ClusterNodeCopy, ClusterRoot, ClusterNode->GetName());
-
-			// If the host string isn't empty, it means the user is pasting the cluster node into a particular host, so set the cluster node's host string.
-			if (!HostStr.IsEmpty())
-			{
-				ClusterNodeCopy->Host = HostStr;
-			}
-
-			if (PasteLocation.IsSet())
-			{
-				FVector2D NewLocation = *PasteLocation;
-
-				// If we are pasting into an existing host, offset the reference target position so that the position of the node inside the host
-				// is consistent for every host we paste into.
-				if (UDisplayClusterConfigurationHostDisplayData* HostDisplayData = GetHostDisplayDataForClusterNode(ClusterNodeCopy))
-				{
-					// If the reference target position isn't set, set it to the current host's position.
-					if (!ReferenceTargetPosition.IsSet())
-					{
-						ReferenceTargetPosition = HostDisplayData->Position;
-					}
-
-					// Offset the location to keep this cluster node's position relative to its parent host consistent.
-					NewLocation -= *ReferenceTargetPosition;
-				}
-
-				// If the reference position isn't set, set it to the current pasted cluster node's position.
-				if (!ReferencePosition.IsSet())
-				{
-					ReferencePosition = FVector2D(ClusterNode->WindowRect.X, ClusterNode->WindowRect.Y);
-				}
-
-				// Offset the location to keep this cluster node's relative position with the group of pasted nodes consistent.
-				NewLocation -= FVector2D(ClusterNode->WindowRect.X, ClusterNode->WindowRect.Y) - *ReferencePosition;
-
-				ClusterNodeCopy->WindowRect.X = NewLocation.X;
-				ClusterNodeCopy->WindowRect.Y = NewLocation.Y;
-			}
-
-			// Rename viewports in the new cluster node to ensure they're unique. Copy the viewport pointers into an array
-			// first since renaming would cause the dictionary to change while we're iterating it.
-			TArray<typename decltype(ClusterNodeCopy->Viewports)::ValueType> CopiedViewports;
-			ClusterNodeCopy->Viewports.GenerateValueArray(CopiedViewports);
-			
-			for (UDisplayClusterConfigurationViewport* Viewport : CopiedViewports)
-			{
-				RenameViewport(Viewport, *GetUniqueNameForViewport(*Viewport->GetName(), ClusterNodeCopy));
-			}
-
-			CopiedObjects.Add(ClusterNodeCopy);
+			OutSortedNodes.Add(Host, TMap<FString, UDisplayClusterConfigurationClusterNode*>());
 		}
-	}
 
-	// Reset the reference positions for the viewports pasting loop
-	ReferencePosition = TOptional<FVector2D>();
-	ReferenceTargetPosition = TOptional<FVector2D>();
-
-	for (UDisplayClusterConfigurationClusterNode* ClusterNode : SelectedNodes)
-	{
-		for (UDisplayClusterConfigurationViewport* Viewport : Factory.Viewports)
-		{
-			UDisplayClusterConfigurationViewport* ViewportCopy = DuplicateObject(Viewport, ClusterNode);
-			ViewportCopy->SetFlags(RF_Transactional);
-			ViewportCopy = AddViewportToClusterNode(ViewportCopy, ClusterNode, Viewport->GetName());
-
-			if (PasteLocation.IsSet())
-			{
-				// If the reference target position isn't set, set it to the current cluster node's position.
-				if (!ReferenceTargetPosition.IsSet())
-				{
-					FVector2D ClusterNodeLocation = FVector2D(ClusterNode->WindowRect.X, ClusterNode->WindowRect.Y);
-
-					// Cluster node window rectangles are local to the node's host, so factor in its host position, if it has one.
-					if (UDisplayClusterConfigurationHostDisplayData* HostDisplayData = GetHostDisplayDataForClusterNode(ClusterNode))
-					{
-						ClusterNodeLocation += HostDisplayData->Position;
-					}
-
-					ReferenceTargetPosition = ClusterNodeLocation;
-				}
-
-				// Offset the location to keep this viewport's position relative to its parent cluster node consistent.
-				FVector2D NewLocation = *PasteLocation - *ReferenceTargetPosition;
-
-				// If the reference position isn't set, set it to the current pasted viewport's position.
-				if (!ReferencePosition.IsSet())
-				{
-					ReferencePosition = FVector2D(Viewport->Region.X, Viewport->Region.Y);
-				}
-
-				// Offset the location to keep this viewport's relative position with the group of pasted nodes consistent.
-				NewLocation -= FVector2D(Viewport->Region.X, Viewport->Region.Y) - *ReferencePosition;
-
-				// Viewports cannot have negative position
-				ViewportCopy->Region.X = FMath::Max(NewLocation.X, 0.0f);
-				ViewportCopy->Region.Y = FMath::Max(NewLocation.Y, 0.0f);
-			}
-
-			CopiedObjects.Add(ViewportCopy);
-		}
+		OutSortedNodes[Host].Add(ClusterNodePair);
 	}
 
-	return CopiedObjects;
+	// Sort the hosts by the host address
+	OutSortedNodes.KeySort(TLess<FString>());
 }
 
-FString FDisplayClusterConfiguratorClusterUtils::GetUniqueName(const FString& InitialName, const TArray<FString>& UsedNames, const UClass* Class, UObject* Parent, bool bAddZero)
+void UE::DisplayClusterConfiguratorClusterUtils::SortClusterNodesByHost(const TMap<FString, TObjectPtr<UDisplayClusterConfigurationClusterNode>>& InClusterNodes, TMap<FString, TMap<FString, UDisplayClusterConfigurationClusterNode*>>& OutSortedNodes)
 {
-	FString NewName;
-	if (!bAddZero && !UsedNames.Contains(InitialName))
-	{
-		// Name doesn't need to be modified
-		NewName = InitialName;
-	}
-	else
+	for (const TPair<FString, TObjectPtr<UDisplayClusterConfigurationClusterNode>>& ClusterNodePair : InClusterNodes)
 	{
-		int32 Counter = bAddZero ? 0 : 1;
-
-		// Find the start of the existing numeric suffix
-		int32 Index = InitialName.Len();
-		while (Index > 0 && InitialName[Index-1] >= '0' && InitialName[Index-1] <= '9')
-		{
-			--Index;
-		}
+		check(ClusterNodePair.Value)
 
-		FString BaseName = InitialName;
-		if (Index < BaseName.Len())
-		{
-			// Strip away the suffix and store the value in the counter so we can count up from there
-			FString NumericSuffix = BaseName.RightChop(Index);
-			Counter = FCString::Atoi(*NumericSuffix);
-			NumericSuffix = FString::FromInt(Counter); // Restringify the counter to account for leading 0s that we don't want to remove
-			BaseName.RemoveAt(BaseName.Len() - NumericSuffix.Len(), NumericSuffix.Len(), false);
-		}
-		else
+			FString Host = ClusterNodePair.Value->Host;
+		if (!OutSortedNodes.Contains(Host))
 		{
-			// No existing suffix, so add our underscore separator
-			BaseName += "_";
+			OutSortedNodes.Add(Host, TMap<FString, UDisplayClusterConfigurationClusterNode*>());
 		}
 
-		do
-		{
-			NewName = FString::Printf(TEXT("%s%d"), *BaseName, Counter);
-			++Counter;
-		}
-		while (UsedNames.Contains(NewName));
-	}
-	
-	// If there is already an in-memory object connected to the parent cluster with our generated name, we need to use a globally unique object name.
-	if (StaticFindObject(nullptr, Parent, *NewName, true))
-	{
-		NewName = MakeUniqueObjectName(Parent, Class, *NewName).ToString();
+		OutSortedNodes[Host].Add(ClusterNodePair);
 	}
 
-	return NewName;
+	// Sort the hosts by the host address
+	OutSortedNodes.KeySort(TLess<FString>());
 }
 
 #undef LOCTEXT_NAMESPACE
\ No newline at end of file
diff --git a/Engine/Plugins/Runtime/nDisplay/Source/DisplayClusterConfigurator/Private/ClusterConfiguration/DisplayClusterConfiguratorClusterUtils.h b/Engine/Plugins/Runtime/nDisplay/Source/DisplayClusterConfigurator/Private/ClusterConfiguration/DisplayClusterConfiguratorClusterUtils.h
deleted file mode 100644
diff --git a/Engine/Plugins/Runtime/nDisplay/Source/DisplayClusterConfigurator/Private/ClusterConfiguration/ViewModels/DisplayClusterConfiguratorProjectionPolicyViewModel.cpp b/Engine/Plugins/Runtime/nDisplay/Source/DisplayClusterConfigurator/Private/ClusterConfiguration/ViewModels/DisplayClusterConfiguratorProjectionPolicyViewModel.cpp
--- a/Engine/Plugins/Runtime/nDisplay/Source/DisplayClusterConfigurator/Private/ClusterConfiguration/ViewModels/DisplayClusterConfiguratorProjectionPolicyViewModel.cpp
+++ b/Engine/Plugins/Runtime/nDisplay/Source/DisplayClusterConfigurator/Private/ClusterConfiguration/ViewModels/DisplayClusterConfiguratorProjectionPolicyViewModel.cpp
@@ -14,7 +14,7 @@ FDisplayClusterConfiguratorProjectionPolicyViewModel::FDisplayClusterConfigurato
 {
 	ViewportPtr = InViewport;
 
-	ProjectionPolicyView = DisplayClusterConfiguratorPropertyUtils::GetPropertyView(InViewport, GET_MEMBER_NAME_CHECKED(UDisplayClusterConfigurationViewport, ProjectionPolicy));
+	ProjectionPolicyView = UE::DisplayClusterConfiguratorPropertyUtils::GetPropertyView(InViewport, GET_MEMBER_NAME_CHECKED(UDisplayClusterConfigurationViewport, ProjectionPolicy));
 	ProjectionPolicyHandle = ProjectionPolicyView->GetPropertyHandle();
 }
 
@@ -65,7 +65,7 @@ void FDisplayClusterConfiguratorProjectionPolicyViewModel::SetParameterValue(con
 	check(ParametersHandle);
 
 	uint8* MapContainer = StructProperty->ContainerPtrToValuePtr<uint8>(Viewport);
-	DisplayClusterConfiguratorPropertyUtils::AddKeyValueToMap(MapContainer, ParametersHandle, ParameterKey, ParameterValue);
+	UE::DisplayClusterConfiguratorPropertyUtils::AddKeyValueToMap(MapContainer, ParametersHandle, ParameterKey, ParameterValue);
 
 	Viewport->MarkPackageDirty();
 }
@@ -83,7 +83,7 @@ void FDisplayClusterConfiguratorProjectionPolicyViewModel::ClearParameters()
 	check(ParametersHandle);
 
 	uint8* MapContainer = StructProperty->ContainerPtrToValuePtr<uint8>(Viewport);
-	DisplayClusterConfiguratorPropertyUtils::EmptyMap(MapContainer, ParametersHandle);
+	UE::DisplayClusterConfiguratorPropertyUtils::EmptyMap(MapContainer, ParametersHandle);
 
 	Viewport->MarkPackageDirty();
 }
\ No newline at end of file
diff --git a/Engine/Plugins/Runtime/nDisplay/Source/DisplayClusterConfigurator/Private/ClusterConfiguration/ViewModels/DisplayClusterConfiguratorViewModelMacros.h b/Engine/Plugins/Runtime/nDisplay/Source/DisplayClusterConfigurator/Private/ClusterConfiguration/ViewModels/DisplayClusterConfiguratorViewModelMacros.h
--- a/Engine/Plugins/Runtime/nDisplay/Source/DisplayClusterConfigurator/Private/ClusterConfiguration/ViewModels/DisplayClusterConfiguratorViewModelMacros.h
+++ b/Engine/Plugins/Runtime/nDisplay/Source/DisplayClusterConfigurator/Private/ClusterConfiguration/ViewModels/DisplayClusterConfiguratorViewModelMacros.h
@@ -7,5 +7,5 @@ TSharedPtr<class ISinglePropertyView> PropertyName##View; \
 TSharedPtr<class IPropertyHandle> PropertyName##Handle;
 
 #define INIT_PROPERTY_HANDLE(ObjectClass, Object, Property) \
-Property##View = DisplayClusterConfiguratorPropertyUtils::GetPropertyView(Object, GET_MEMBER_NAME_CHECKED(ObjectClass, Property)); \
+Property##View = UE::DisplayClusterConfiguratorPropertyUtils::GetPropertyView(Object, GET_MEMBER_NAME_CHECKED(ObjectClass, Property)); \
 Property##Handle = Property##View->GetPropertyHandle();
\ No newline at end of file
diff --git a/Engine/Plugins/Runtime/nDisplay/Source/DisplayClusterConfigurator/Private/DisplayClusterConfiguratorBlueprintEditor.cpp b/Engine/Plugins/Runtime/nDisplay/Source/DisplayClusterConfigurator/Private/DisplayClusterConfiguratorBlueprintEditor.cpp
--- a/Engine/Plugins/Runtime/nDisplay/Source/DisplayClusterConfigurator/Private/DisplayClusterConfiguratorBlueprintEditor.cpp
+++ b/Engine/Plugins/Runtime/nDisplay/Source/DisplayClusterConfigurator/Private/DisplayClusterConfiguratorBlueprintEditor.cpp
@@ -796,7 +796,7 @@ void FDisplayClusterConfiguratorBlueprintEditor::OnRenameVariable(UBlueprint* Bl
 					TMap<FString, FString> PolicyCopy = Viewport->ProjectionPolicy.Parameters;
 
 					const TSharedPtr<ISinglePropertyView> ProjectPolicyView =
-						DisplayClusterConfiguratorPropertyUtils::GetPropertyView(
+						UE::DisplayClusterConfiguratorPropertyUtils::GetPropertyView(
 							Viewport, GET_MEMBER_NAME_CHECKED(UDisplayClusterConfigurationViewport, ProjectionPolicy));
 					check(ProjectPolicyView);
 					
@@ -811,20 +811,20 @@ void FDisplayClusterConfiguratorBlueprintEditor::OnRenameVariable(UBlueprint* Bl
 					{
 						if (PolicyParameters.Value == OldVariableName.ToString())
 						{
-							DisplayClusterConfiguratorPropertyUtils::RemoveKeyFromMap(MapContainer, ParametersHandle, PolicyParameters.Key);
-							DisplayClusterConfiguratorPropertyUtils::AddKeyValueToMap(MapContainer, ParametersHandle, PolicyParameters.Key, NewVariableName.ToString());
+							UE::DisplayClusterConfiguratorPropertyUtils::RemoveKeyFromMap(MapContainer, ParametersHandle, PolicyParameters.Key);
+							UE::DisplayClusterConfiguratorPropertyUtils::AddKeyValueToMap(MapContainer, ParametersHandle, PolicyParameters.Key, NewVariableName.ToString());
 						}
 					}
 				}
 
 				if (Viewport->Camera == OldVariableName.ToString())
 				{
-					DisplayClusterConfiguratorPropertyUtils::SetPropertyHandleValue(Viewport, GET_MEMBER_NAME_CHECKED(UDisplayClusterConfigurationViewport, Camera), NewVariableName.ToString());
+					UE::DisplayClusterConfiguratorPropertyUtils::SetPropertyHandleValue(Viewport, GET_MEMBER_NAME_CHECKED(UDisplayClusterConfigurationViewport, Camera), NewVariableName.ToString());
 				}
 
 				if (Viewport->DisplayDeviceName == OldVariableName.ToString())
 				{
-					DisplayClusterConfiguratorPropertyUtils::SetPropertyHandleValue(Viewport, GET_MEMBER_NAME_CHECKED(UDisplayClusterConfigurationViewport, DisplayDeviceName), NewVariableName.ToString());
+					UE::DisplayClusterConfiguratorPropertyUtils::SetPropertyHandleValue(Viewport, GET_MEMBER_NAME_CHECKED(UDisplayClusterConfigurationViewport, DisplayDeviceName), NewVariableName.ToString());
 				}
 			}
 		}
@@ -1150,7 +1150,7 @@ void FDisplayClusterConfiguratorBlueprintEditor::OnClose()
 	if (UDisplayClusterConfigurationData* Config = GetConfig())
 	{
 		bool bIsDirty = LoadedBlueprint->GetOutermost()->IsDirty();
-		bool bHostDataRemoved = FDisplayClusterConfiguratorClusterUtils::RemoveUnusedHostDisplayData(Config->Cluster);
+		bool bHostDataRemoved = UE::DisplayClusterConfiguratorClusterUtils::RemoveUnusedHostDisplayData(Config->Cluster);
 
 		// If the blueprint wasn't dirty before, removing the unused host display data will have make it dirty, which is confusing to the user.
 		// In this case, immediately save the host display data removal to the blueprint.
diff --git a/Engine/Plugins/Runtime/nDisplay/Source/DisplayClusterConfigurator/Private/DisplayClusterConfiguratorFactory.h b/Engine/Plugins/Runtime/nDisplay/Source/DisplayClusterConfigurator/Private/DisplayClusterConfiguratorFactory.h
deleted file mode 100644
diff --git a/Engine/Plugins/Runtime/nDisplay/Source/DisplayClusterConfigurator/Private/DisplayClusterConfiguratorPropertyUtils.cpp b/Engine/Plugins/Runtime/nDisplay/Source/DisplayClusterConfigurator/Private/DisplayClusterConfiguratorPropertyUtils.cpp
--- a/Engine/Plugins/Runtime/nDisplay/Source/DisplayClusterConfigurator/Private/DisplayClusterConfiguratorPropertyUtils.cpp
+++ b/Engine/Plugins/Runtime/nDisplay/Source/DisplayClusterConfigurator/Private/DisplayClusterConfiguratorPropertyUtils.cpp
@@ -6,7 +6,7 @@
 
 #define LOCTEXT_NAMESPACE "DisplayClusterConfiguratorPropertyUtils"
 
-TSharedPtr<ISinglePropertyView> DisplayClusterConfiguratorPropertyUtils::GetPropertyView(UObject* Owner,
+TSharedPtr<ISinglePropertyView> UE::DisplayClusterConfiguratorPropertyUtils::GetPropertyView(UObject* Owner,
 	const FName& FieldName)
 {
 	check(Owner);
@@ -21,7 +21,7 @@ TSharedPtr<ISinglePropertyView> DisplayClusterConfiguratorPropertyUtils::GetProp
 	return PropView;
 }
 
-UObject* DisplayClusterConfiguratorPropertyUtils::AddKeyWithInstancedValueToMap(UObject* MapOwner, const FName& FieldName,
+UObject* UE::DisplayClusterConfiguratorPropertyUtils::AddKeyWithInstancedValueToMap(UObject* MapOwner, const FName& FieldName,
                                                                                 const FString& Key, UObject* Value)
 {
 	check(MapOwner);
@@ -40,7 +40,7 @@ UObject* DisplayClusterConfiguratorPropertyUtils::AddKeyWithInstancedValueToMap(
 	return NewValue;
 }
 
-TSharedPtr<IPropertyHandle> DisplayClusterConfiguratorPropertyUtils::AddKeyValueToMap(uint8* MapOwner, TSharedPtr<IPropertyHandle> PropertyHandle, const FString& Key, const FString& Value, EPropertyValueSetFlags::Type SetFlags)
+TSharedPtr<IPropertyHandle> UE::DisplayClusterConfiguratorPropertyUtils::AddKeyValueToMap(uint8* MapOwner, TSharedPtr<IPropertyHandle> PropertyHandle, const FString& Key, const FString& Value, EPropertyValueSetFlags::Type SetFlags)
 {
 	check(MapOwner);
 	check(PropertyHandle.IsValid());
@@ -86,7 +86,7 @@ TSharedPtr<IPropertyHandle> DisplayClusterConfiguratorPropertyUtils::AddKeyValue
 	return PropertyHandle->GetChildHandle(IndexToModify);
 }
 
-bool DisplayClusterConfiguratorPropertyUtils::RemoveKeyFromMap(UObject* MapOwner, const FName& FieldName,
+bool UE::DisplayClusterConfiguratorPropertyUtils::RemoveKeyFromMap(UObject* MapOwner, const FName& FieldName,
                                                                const FString& Key)
 {
 	check(MapOwner);
@@ -98,7 +98,7 @@ bool DisplayClusterConfiguratorPropertyUtils::RemoveKeyFromMap(UObject* MapOwner
 	return RemoveKeyFromMap((uint8*)MapOwner, PropHandle, Key);
 }
 
-bool DisplayClusterConfiguratorPropertyUtils::RemoveKeyFromMap(uint8* MapOwner,
+bool UE::DisplayClusterConfiguratorPropertyUtils::RemoveKeyFromMap(uint8* MapOwner,
 	TSharedPtr<IPropertyHandle> PropertyHandle, const FString& Key)
 {
 	check(MapOwner);
@@ -128,7 +128,7 @@ bool DisplayClusterConfiguratorPropertyUtils::RemoveKeyFromMap(uint8* MapOwner,
 	return true;
 }
 
-bool DisplayClusterConfiguratorPropertyUtils::EmptyMap(uint8* MapOwner, TSharedPtr<IPropertyHandle> PropertyHandle)
+bool UE::DisplayClusterConfiguratorPropertyUtils::EmptyMap(uint8* MapOwner, TSharedPtr<IPropertyHandle> PropertyHandle)
 {
 	check(MapOwner);
 	check(PropertyHandle.IsValid());
@@ -151,7 +151,7 @@ bool DisplayClusterConfiguratorPropertyUtils::EmptyMap(uint8* MapOwner, TSharedP
 	return true;
 }
 
-int32 DisplayClusterConfiguratorPropertyUtils::FindMapHandleIndexFromKey(TSharedPtr<IPropertyHandleMap> MapHandle, FScriptMapHelper& MapHelper,
+int32 UE::DisplayClusterConfiguratorPropertyUtils::FindMapHandleIndexFromKey(TSharedPtr<IPropertyHandleMap> MapHandle, FScriptMapHelper& MapHelper,
 	const FString& Key)
 {
 	check(MapHandle.IsValid());
diff --git a/Engine/Plugins/Runtime/nDisplay/Source/DisplayClusterConfigurator/Private/DisplayClusterConfiguratorPropertyUtils.h b/Engine/Plugins/Runtime/nDisplay/Source/DisplayClusterConfigurator/Private/DisplayClusterConfiguratorPropertyUtils.h
--- a/Engine/Plugins/Runtime/nDisplay/Source/DisplayClusterConfigurator/Private/DisplayClusterConfiguratorPropertyUtils.h
+++ b/Engine/Plugins/Runtime/nDisplay/Source/DisplayClusterConfigurator/Private/DisplayClusterConfiguratorPropertyUtils.h
@@ -11,7 +11,7 @@
 /**
  * Utilities to assist with modifying properties through their handle so values are propagated across instances.
  */
-namespace DisplayClusterConfiguratorPropertyUtils
+namespace UE::DisplayClusterConfiguratorPropertyUtils
 {
 	/**
 	 * Create a temporary property handle for a given property. Call GetPropertyHandle() from here.
diff --git a/Engine/Plugins/Runtime/nDisplay/Source/DisplayClusterConfigurator/Private/MPCDI/DisplayClusterConfiguratorMPCDIImporter.cpp b/Engine/Plugins/Runtime/nDisplay/Source/DisplayClusterConfigurator/Private/MPCDI/DisplayClusterConfiguratorMPCDIImporter.cpp
--- a/Engine/Plugins/Runtime/nDisplay/Source/DisplayClusterConfigurator/Private/MPCDI/DisplayClusterConfiguratorMPCDIImporter.cpp
+++ b/Engine/Plugins/Runtime/nDisplay/Source/DisplayClusterConfigurator/Private/MPCDI/DisplayClusterConfiguratorMPCDIImporter.cpp
@@ -207,10 +207,10 @@ UDisplayClusterConfigurationViewport* FDisplayClusterConfiguratorMPCDIImporter::
 	{
 		const FString NodeName = GetClusterNodeNameForRegion(RegionId);
 		UDisplayClusterConfigurationClusterNode* ClusterNodeTemplate = NewObject<UDisplayClusterConfigurationClusterNode>(InBlueprint, NAME_None, RF_Transactional | RF_ArchetypeObject | RF_Public);
-		UDisplayClusterConfigurationClusterNode* NewClusterNode = FDisplayClusterConfiguratorClusterUtils::AddClusterNodeToCluster(ClusterNodeTemplate, Cluster, NodeName);
+		UDisplayClusterConfigurationClusterNode* NewClusterNode = UE::DisplayClusterConfiguratorClusterUtils::AddClusterNodeToCluster(ClusterNodeTemplate, Cluster, NodeName);
 
 		UDisplayClusterConfigurationViewport* ViewportTemplate = NewObject<UDisplayClusterConfigurationViewport>(InBlueprint, NAME_None, RF_Transactional | RF_ArchetypeObject | RF_Public);
-		Viewport = FDisplayClusterConfiguratorClusterUtils::AddViewportToClusterNode(ViewportTemplate, NewClusterNode, ViewportName);
+		Viewport = UE::DisplayClusterConfiguratorClusterUtils::AddViewportToClusterNode(ViewportTemplate, NewClusterNode, ViewportName);
 	}
 
 	return Viewport;
diff --git a/Engine/Plugins/Runtime/nDisplay/Source/DisplayClusterConfigurator/Private/Views/Details/Policies/DisplayClusterConfiguratorPolicyDetailCustomization.cpp b/Engine/Plugins/Runtime/nDisplay/Source/DisplayClusterConfigurator/Private/Views/Details/Policies/DisplayClusterConfiguratorPolicyDetailCustomization.cpp
--- a/Engine/Plugins/Runtime/nDisplay/Source/DisplayClusterConfigurator/Private/Views/Details/Policies/DisplayClusterConfiguratorPolicyDetailCustomization.cpp
+++ b/Engine/Plugins/Runtime/nDisplay/Source/DisplayClusterConfigurator/Private/Views/Details/Policies/DisplayClusterConfiguratorPolicyDetailCustomization.cpp
@@ -1269,7 +1269,7 @@ void FDisplayClusterConfiguratorRenderSyncPolicyCustomization::AddToParameterMap
 	check(RenderStructProperty);
 
 	uint8* MapContainer = RenderStructProperty->ContainerPtrToValuePtr<uint8>(&ConfigurationCluster->Sync);
-	DisplayClusterConfiguratorPropertyUtils::AddKeyValueToMap(MapContainer, ParametersHandle, Key, Value);
+	UE::DisplayClusterConfiguratorPropertyUtils::AddKeyValueToMap(MapContainer, ParametersHandle, Key, Value);
 }
 
 void FDisplayClusterConfiguratorRenderSyncPolicyCustomization::RemoveFromParameterMap(const FString& Key)
@@ -1284,7 +1284,7 @@ void FDisplayClusterConfiguratorRenderSyncPolicyCustomization::RemoveFromParamet
 	check(RenderStructProperty);
 
 	uint8* MapContainer = RenderStructProperty->ContainerPtrToValuePtr<uint8>(&ConfigurationCluster->Sync);
-	DisplayClusterConfiguratorPropertyUtils::RemoveKeyFromMap(MapContainer, ParametersHandle, Key);
+	UE::DisplayClusterConfiguratorPropertyUtils::RemoveKeyFromMap(MapContainer, ParametersHandle, Key);
 }
 
 //////////////////////////////////////////////////////////////////////////////////////////////
diff --git a/Engine/Plugins/Runtime/nDisplay/Source/DisplayClusterConfigurator/Private/Views/Details/Policies/DisplayClusterConfiguratorPolicyParameterCustomization.cpp b/Engine/Plugins/Runtime/nDisplay/Source/DisplayClusterConfigurator/Private/Views/Details/Policies/DisplayClusterConfiguratorPolicyParameterCustomization.cpp
--- a/Engine/Plugins/Runtime/nDisplay/Source/DisplayClusterConfigurator/Private/Views/Details/Policies/DisplayClusterConfiguratorPolicyParameterCustomization.cpp
+++ b/Engine/Plugins/Runtime/nDisplay/Source/DisplayClusterConfigurator/Private/Views/Details/Policies/DisplayClusterConfiguratorPolicyParameterCustomization.cpp
@@ -163,7 +163,7 @@ void FPolicyParameterInfo::UpdateCustomParameterValueText(const FString& NewValu
 		FStructProperty* StructProperty = FindFProperty<FStructProperty>(ConfigurationViewport->GetClass(), GET_MEMBER_NAME_CHECKED(UDisplayClusterConfigurationViewport, ProjectionPolicy));
 		check(StructProperty);
 
-		const TSharedPtr<ISinglePropertyView> ProjectionPolicyView = DisplayClusterConfiguratorPropertyUtils::GetPropertyView(
+		const TSharedPtr<ISinglePropertyView> ProjectionPolicyView = UE::DisplayClusterConfiguratorPropertyUtils::GetPropertyView(
 			ConfigurationViewport, GET_MEMBER_NAME_CHECKED(UDisplayClusterConfigurationViewport, ProjectionPolicy));
 
 		check(ProjectionPolicyView.IsValid());
@@ -172,7 +172,7 @@ void FPolicyParameterInfo::UpdateCustomParameterValueText(const FString& NewValu
 		check(PropertyHandle.IsValid());
 		
 		uint8* MapContainer = StructProperty->ContainerPtrToValuePtr<uint8>(ConfigurationViewport);
-		DisplayClusterConfiguratorPropertyUtils::AddKeyValueToMap(MapContainer, PropertyHandle, Key, NewValue);
+		UE::DisplayClusterConfiguratorPropertyUtils::AddKeyValueToMap(MapContainer, PropertyHandle, Key, NewValue);
 	}
 
 	if (bNotify && bHasChanged)
diff --git a/Engine/Plugins/Runtime/nDisplay/Source/DisplayClusterConfigurator/Private/Views/OutputMapping/DisplayClusterConfiguratorGraph.cpp b/Engine/Plugins/Runtime/nDisplay/Source/DisplayClusterConfigurator/Private/Views/OutputMapping/DisplayClusterConfiguratorGraph.cpp
--- a/Engine/Plugins/Runtime/nDisplay/Source/DisplayClusterConfigurator/Private/Views/OutputMapping/DisplayClusterConfiguratorGraph.cpp
+++ b/Engine/Plugins/Runtime/nDisplay/Source/DisplayClusterConfigurator/Private/Views/OutputMapping/DisplayClusterConfiguratorGraph.cpp
@@ -62,7 +62,7 @@ void UDisplayClusterConfiguratorGraph::RebuildGraph()
 		UDisplayClusterConfiguratorCanvasNode* CanvasNode = BuildCanvasNode(Config->Cluster);
 
 		TMap<FString, TMap<FString, UDisplayClusterConfigurationClusterNode*>> SortedClusterNodes;
-		FDisplayClusterConfiguratorClusterUtils::SortClusterNodesByHost(Config->Cluster->Nodes, SortedClusterNodes);
+		UE::DisplayClusterConfiguratorClusterUtils::SortClusterNodesByHost(Config->Cluster->Nodes, SortedClusterNodes);
 
 		// Keep track of the count of each cluster item to use as a z-index for the node.
 		int HostIndex = 0;
@@ -71,7 +71,7 @@ void UDisplayClusterConfiguratorGraph::RebuildGraph()
 
 		for (const TPair<FString, TMap<FString, UDisplayClusterConfigurationClusterNode*>>& SortedPair : SortedClusterNodes)
 		{
-			UDisplayClusterConfigurationHostDisplayData* DisplayData = FDisplayClusterConfiguratorClusterUtils::FindOrCreateHostDisplayData(Config->Cluster, SortedPair.Key);
+			UDisplayClusterConfigurationHostDisplayData* DisplayData = UE::DisplayClusterConfiguratorClusterUtils::FindOrCreateHostDisplayData(Config->Cluster, SortedPair.Key);
 			UDisplayClusterConfiguratorHostNode* HostNode = BuildHostNode(CanvasNode, SortedPair.Key, HostIndex, DisplayData);
 			++HostIndex;
 
diff --git a/Engine/Plugins/Runtime/nDisplay/Source/DisplayClusterConfigurator/Private/Views/OutputMapping/DisplayClusterConfiguratorGraphSchema.cpp b/Engine/Plugins/Runtime/nDisplay/Source/DisplayClusterConfigurator/Private/Views/OutputMapping/DisplayClusterConfiguratorGraphSchema.cpp
--- a/Engine/Plugins/Runtime/nDisplay/Source/DisplayClusterConfigurator/Private/Views/OutputMapping/DisplayClusterConfiguratorGraphSchema.cpp
+++ b/Engine/Plugins/Runtime/nDisplay/Source/DisplayClusterConfigurator/Private/Views/OutputMapping/DisplayClusterConfiguratorGraphSchema.cpp
@@ -5,7 +5,7 @@
 #include "DisplayClusterConfiguratorBlueprintEditor.h"
 #include "DisplayClusterConfiguratorGraph.h"
 #include "DisplayClusterConfigurationTypes.h"
-#include "ClusterConfiguration/DisplayClusterConfiguratorClusterUtils.h"
+#include "ClusterConfiguration/DisplayClusterConfiguratorClusterEditorUtils.h"
 #include "ClusterConfiguration/SDisplayClusterConfiguratorNewClusterItemDialog.h"
 
 #include "ScopedTransaction.h"
@@ -36,11 +36,11 @@ UEdGraphNode* FDisplayClusterConfiguratorSchemaAction_NewNode::PerformAction(UEd
 	switch (ItemType)
 	{
 	case EClusterItemType::ClusterNode:
-		CreatedObject = FDisplayClusterConfiguratorClusterUtils::CreateNewClusterNodeFromDialog(Toolkit.ToSharedRef(), Cluster, PresetRect, Transaction);
+		CreatedObject = UE::DisplayClusterConfiguratorClusterEditorUtils::CreateNewClusterNodeFromDialog(Toolkit.ToSharedRef(), Cluster, PresetRect, Transaction);
 		break;
 
 	case EClusterItemType::Viewport:
-		CreatedObject = FDisplayClusterConfiguratorClusterUtils::CreateNewViewportFromDialog(Toolkit.ToSharedRef(), nullptr, PresetRect, Transaction);
+		CreatedObject = UE::DisplayClusterConfiguratorClusterEditorUtils::CreateNewViewportFromDialog(Toolkit.ToSharedRef(), nullptr, PresetRect, Transaction);
 		break;
 	}
 
diff --git a/Engine/Plugins/Runtime/nDisplay/Source/DisplayClusterConfigurator/Private/Views/OutputMapping/EdNodes/DisplayClusterConfiguratorHostNode.cpp b/Engine/Plugins/Runtime/nDisplay/Source/DisplayClusterConfigurator/Private/Views/OutputMapping/EdNodes/DisplayClusterConfiguratorHostNode.cpp
--- a/Engine/Plugins/Runtime/nDisplay/Source/DisplayClusterConfigurator/Private/Views/OutputMapping/EdNodes/DisplayClusterConfiguratorHostNode.cpp
+++ b/Engine/Plugins/Runtime/nDisplay/Source/DisplayClusterConfigurator/Private/Views/OutputMapping/EdNodes/DisplayClusterConfiguratorHostNode.cpp
@@ -207,7 +207,7 @@ void UDisplayClusterConfiguratorHostNode::DeleteObject()
 	UDisplayClusterConfigurationHostDisplayData* HostDisplayData = GetObjectChecked<UDisplayClusterConfigurationHostDisplayData>();
 	if (UDisplayClusterConfigurationCluster* Cluster = Cast<UDisplayClusterConfigurationCluster>(HostDisplayData->GetOuter()))
 	{
-		FDisplayClusterConfiguratorClusterUtils::RemoveHost(Cluster, HostStr);
+		UE::DisplayClusterConfiguratorClusterUtils::RemoveHost(Cluster, HostStr);
 	}
 }
 
diff --git a/Engine/Plugins/Runtime/nDisplay/Source/DisplayClusterConfigurator/Private/Views/OutputMapping/EdNodes/DisplayClusterConfiguratorViewportNode.cpp b/Engine/Plugins/Runtime/nDisplay/Source/DisplayClusterConfigurator/Private/Views/OutputMapping/EdNodes/DisplayClusterConfiguratorViewportNode.cpp
--- a/Engine/Plugins/Runtime/nDisplay/Source/DisplayClusterConfigurator/Private/Views/OutputMapping/EdNodes/DisplayClusterConfiguratorViewportNode.cpp
+++ b/Engine/Plugins/Runtime/nDisplay/Source/DisplayClusterConfigurator/Private/Views/OutputMapping/EdNodes/DisplayClusterConfiguratorViewportNode.cpp
@@ -97,7 +97,7 @@ void UDisplayClusterConfiguratorViewportNode::DeleteObject()
 	}
 
 	UDisplayClusterConfigurationViewport* Viewport = GetObjectChecked<UDisplayClusterConfigurationViewport>();
-	FDisplayClusterConfiguratorClusterUtils::RemoveViewportFromClusterNode(Viewport);
+	UE::DisplayClusterConfiguratorClusterUtils::RemoveViewportFromClusterNode(Viewport);
 }
 
 void UDisplayClusterConfiguratorViewportNode::WriteNodeStateToObject()
diff --git a/Engine/Plugins/Runtime/nDisplay/Source/DisplayClusterConfigurator/Private/Views/OutputMapping/EdNodes/DisplayClusterConfiguratorWindowNode.cpp b/Engine/Plugins/Runtime/nDisplay/Source/DisplayClusterConfigurator/Private/Views/OutputMapping/EdNodes/DisplayClusterConfiguratorWindowNode.cpp
--- a/Engine/Plugins/Runtime/nDisplay/Source/DisplayClusterConfigurator/Private/Views/OutputMapping/EdNodes/DisplayClusterConfiguratorWindowNode.cpp
+++ b/Engine/Plugins/Runtime/nDisplay/Source/DisplayClusterConfigurator/Private/Views/OutputMapping/EdNodes/DisplayClusterConfiguratorWindowNode.cpp
@@ -78,7 +78,7 @@ bool UDisplayClusterConfiguratorWindowNode::IsPrimary() const
 	}
 
 	UDisplayClusterConfigurationClusterNode* ClusterNode = GetObjectChecked<UDisplayClusterConfigurationClusterNode>();
-	return FDisplayClusterConfiguratorClusterUtils::IsClusterNodePrimary(ClusterNode);
+	return UE::DisplayClusterConfiguratorClusterUtils::IsClusterNodePrimary(ClusterNode);
 }
 
 FDelegateHandle UDisplayClusterConfiguratorWindowNode::RegisterOnPreviewImageChanged(const FOnPreviewImageChangedDelegate& Delegate)
@@ -147,7 +147,7 @@ void UDisplayClusterConfiguratorWindowNode::DeleteObject()
 	}
 
 	UDisplayClusterConfigurationClusterNode* ClusterNode = GetObjectChecked<UDisplayClusterConfigurationClusterNode>();
-	FDisplayClusterConfiguratorClusterUtils::RemoveClusterNodeFromCluster(ClusterNode);
+	UE::DisplayClusterConfiguratorClusterUtils::RemoveClusterNodeFromCluster(ClusterNode);
 }
 
 void UDisplayClusterConfiguratorWindowNode::WriteNodeStateToObject()
diff --git a/Engine/Plugins/Runtime/nDisplay/Source/DisplayClusterConfigurator/Private/Views/OutputMapping/SDisplayClusterConfiguratorGraphEditor.cpp b/Engine/Plugins/Runtime/nDisplay/Source/DisplayClusterConfigurator/Private/Views/OutputMapping/SDisplayClusterConfiguratorGraphEditor.cpp
--- a/Engine/Plugins/Runtime/nDisplay/Source/DisplayClusterConfigurator/Private/Views/OutputMapping/SDisplayClusterConfiguratorGraphEditor.cpp
+++ b/Engine/Plugins/Runtime/nDisplay/Source/DisplayClusterConfigurator/Private/Views/OutputMapping/SDisplayClusterConfiguratorGraphEditor.cpp
@@ -6,7 +6,7 @@
 #include "DisplayClusterConfiguratorBlueprintEditor.h"
 #include "DisplayClusterConfigurationTypes.h"
 #include "IDisplayClusterConfigurator.h"
-#include "ClusterConfiguration/DisplayClusterConfiguratorClusterUtils.h"
+#include "ClusterConfiguration/DisplayClusterConfiguratorClusterEditorUtils.h"
 #include "ClusterConfiguration/SDisplayClusterConfiguratorNewClusterItemDialog.h"
 #include "Views/OutputMapping/DisplayClusterConfiguratorGraph.h"
 #include "Views/OutputMapping/DisplayClusterConfiguratorOutputMappingCommands.h"
@@ -412,7 +412,7 @@ void SDisplayClusterConfiguratorGraphEditor::AddNewClusterNode()
 
 		TSharedPtr<FScopedTransaction> Transaction;
 		if (UDisplayClusterConfigurationClusterNode* NewNode =
-			FDisplayClusterConfiguratorClusterUtils::CreateNewClusterNodeFromDialog(Toolkit.ToSharedRef(), Cluster, PresetRect,Transaction, Host))
+			UE::DisplayClusterConfiguratorClusterEditorUtils::CreateNewClusterNodeFromDialog(Toolkit.ToSharedRef(), Cluster, PresetRect,Transaction, Host))
 		{
 			// Mark the cluster configuration data as dirty, allowing user to save the changes, and fire off a cluster changed event to let other
 			// parts of the UI update as well
@@ -450,7 +450,7 @@ void SDisplayClusterConfiguratorGraphEditor::AddNewViewport()
 		const FDisplayClusterConfigurationRectangle PresetRect = FDisplayClusterConfigurationRectangle(0, 0, PresetSize.X, PresetSize.Y);
 
 		TSharedPtr<FScopedTransaction> Transaction;
-		if (UDisplayClusterConfigurationViewport* NewViewport = FDisplayClusterConfiguratorClusterUtils::CreateNewViewportFromDialog(Toolkit.ToSharedRef(), SelectedClusterNode, PresetRect, Transaction))
+		if (UDisplayClusterConfigurationViewport* NewViewport = UE::DisplayClusterConfiguratorClusterEditorUtils::CreateNewViewportFromDialog(Toolkit.ToSharedRef(), SelectedClusterNode, PresetRect, Transaction))
 		{
 			// Mark the cluster configuration data as dirty, allowing user to save the changes, and fire off a cluster changed event to let other
 			// parts of the UI update as well
@@ -539,7 +539,7 @@ void SDisplayClusterConfiguratorGraphEditor::CopySelectedNodes()
 			}
 		}
 
-		FDisplayClusterConfiguratorClusterUtils::CopyClusterItemsToClipboard(ObjectsToCopy);
+		UE::DisplayClusterConfiguratorClusterEditorUtils::CopyClusterItemsToClipboard(ObjectsToCopy);
 	}
 }
 
@@ -615,11 +615,11 @@ void SDisplayClusterConfiguratorGraphEditor::PasteNodes()
 	}
 
 	int32 NumClusterItems;
-	if (FDisplayClusterConfiguratorClusterUtils::CanPasteClusterItemsFromClipboard(TargetObjects, NumClusterItems))
+	if (UE::DisplayClusterConfiguratorClusterEditorUtils::CanPasteClusterItemsFromClipboard(TargetObjects, NumClusterItems))
 	{
 		FScopedTransaction Transaction(FText::Format(LOCTEXT("PasteClusterItems", "Paste Cluster {0}|plural(one=Item, other=Items)"), NumClusterItems));
 
-		TArray<UObject*> PastedObjects = FDisplayClusterConfiguratorClusterUtils::PasteClusterItemsFromClipboard(TargetObjects, GetPasteLocation() / ViewScale);
+		TArray<UObject*> PastedObjects = UE::DisplayClusterConfiguratorClusterEditorUtils::PasteClusterItemsFromClipboard(TargetObjects, GetPasteLocation() / ViewScale);
 		if (PastedObjects.Num() > 0)
 		{
 			Toolkit->GetEditorData()->MarkPackageDirty();
@@ -665,7 +665,7 @@ bool SDisplayClusterConfiguratorGraphEditor::CanPasteNodes() const
 	}
 
 	int32 NumClusterItems;
-	return FDisplayClusterConfiguratorClusterUtils::CanPasteClusterItemsFromClipboard(TargetObjects, NumClusterItems);
+	return UE::DisplayClusterConfiguratorClusterEditorUtils::CanPasteClusterItemsFromClipboard(TargetObjects, NumClusterItems);
 }
 
 void SDisplayClusterConfiguratorGraphEditor::DuplicateNodes()
diff --git a/Engine/Plugins/Runtime/nDisplay/Source/DisplayClusterConfigurator/Private/Views/TreeViews/Cluster/DisplayClusterConfiguratorViewCluster.cpp b/Engine/Plugins/Runtime/nDisplay/Source/DisplayClusterConfigurator/Private/Views/TreeViews/Cluster/DisplayClusterConfiguratorViewCluster.cpp
--- a/Engine/Plugins/Runtime/nDisplay/Source/DisplayClusterConfigurator/Private/Views/TreeViews/Cluster/DisplayClusterConfiguratorViewCluster.cpp
+++ b/Engine/Plugins/Runtime/nDisplay/Source/DisplayClusterConfigurator/Private/Views/TreeViews/Cluster/DisplayClusterConfiguratorViewCluster.cpp
@@ -5,6 +5,7 @@
 #include "DisplayClusterConfiguratorBlueprintEditor.h"
 #include "DisplayClusterConfiguratorUtils.h"
 #include "DisplayClusterConfiguratorCommands.h"
+#include "ClusterConfiguration/DisplayClusterConfiguratorClusterEditorUtils.h"
 #include "ClusterConfiguration/DisplayClusterConfiguratorClusterUtils.h"
 #include "ClusterConfiguration/SDisplayClusterConfiguratorNewClusterItemDialog.h"
 #include "Views/TreeViews/IDisplayClusterConfiguratorTreeItem.h"
@@ -281,7 +282,7 @@ void FDisplayClusterConfiguratorViewCluster::CopySelectedNodes()
 		}
 	}
 
-	FDisplayClusterConfiguratorClusterUtils::CopyClusterItemsToClipboard(ObjectsToCopy);
+	UE::DisplayClusterConfiguratorClusterEditorUtils::CopyClusterItemsToClipboard(ObjectsToCopy);
 }
 
 bool FDisplayClusterConfiguratorViewCluster::CanCopySelectedNodes() const
@@ -309,11 +310,11 @@ void FDisplayClusterConfiguratorViewCluster::PasteNodes()
 	}
 
 	int32 NumClusterItems;
-	if (FDisplayClusterConfiguratorClusterUtils::CanPasteClusterItemsFromClipboard(TargetObjects, NumClusterItems))
+	if (UE::DisplayClusterConfiguratorClusterEditorUtils::CanPasteClusterItemsFromClipboard(TargetObjects, NumClusterItems))
 	{
 		FScopedTransaction Transaction(FText::Format(LOCTEXT("PasteClusterItems", "Paste Cluster {0}|plural(one=Item, other=Items)"), NumClusterItems));
 
-		TArray<UObject*> PastedObjects = FDisplayClusterConfiguratorClusterUtils::PasteClusterItemsFromClipboard(TargetObjects);
+		TArray<UObject*> PastedObjects = UE::DisplayClusterConfiguratorClusterEditorUtils::PasteClusterItemsFromClipboard(TargetObjects);
 		if (PastedObjects.Num() > 0)
 		{
 			// Mark the cluster configuration data as dirty, allowing user to save the changes, and fire off a cluster changed event to let other
@@ -355,7 +356,7 @@ bool FDisplayClusterConfiguratorViewCluster::CanPasteNodes() const
 	}
 
 	int32 NumClusterItems;
-	return FDisplayClusterConfiguratorClusterUtils::CanPasteClusterItemsFromClipboard(TargetObjects, NumClusterItems);
+	return UE::DisplayClusterConfiguratorClusterEditorUtils::CanPasteClusterItemsFromClipboard(TargetObjects, NumClusterItems);
 }
 
 void FDisplayClusterConfiguratorViewCluster::DuplicateSelectedNodes()
@@ -443,7 +444,7 @@ void FDisplayClusterConfiguratorViewCluster::SetAsPrimary()
 		{
 			FScopedTransaction Transaction(LOCTEXT("SetPrimaryNode", "Set Primary Node"));
 
-			if (!FDisplayClusterConfiguratorClusterUtils::SetClusterNodeAsPrimary(ClusterNode))
+			if (!UE::DisplayClusterConfiguratorClusterUtils::SetClusterNodeAsPrimary(ClusterNode))
 			{
 				Transaction.Cancel();
 			}
@@ -571,12 +572,12 @@ void FDisplayClusterConfiguratorViewCluster::AddNewClusterNode(FVector2D PresetS
 		UObject* SelectedObject = SelectedTreeItems[0]->GetObject();
 		if (UDisplayClusterConfigurationHostDisplayData* HostData = Cast<UDisplayClusterConfigurationHostDisplayData>(SelectedObject))
 		{
-			HostAddress = FDisplayClusterConfiguratorClusterUtils::GetAddressForHost(HostData);
+			HostAddress = UE::DisplayClusterConfiguratorClusterUtils::GetAddressForHost(HostData);
 		}
 	}
 
 	TSharedPtr<FScopedTransaction> Transaction;
-	if (UDisplayClusterConfigurationClusterNode* NewNode = FDisplayClusterConfiguratorClusterUtils::CreateNewClusterNodeFromDialog(Toolkit.ToSharedRef(), Cluster, PresetRect, Transaction, HostAddress))
+	if (UDisplayClusterConfigurationClusterNode* NewNode = UE::DisplayClusterConfiguratorClusterEditorUtils::CreateNewClusterNodeFromDialog(Toolkit.ToSharedRef(), Cluster, PresetRect, Transaction, HostAddress))
 	{
 		// Mark the cluster configuration data as dirty, allowing user to save the changes, and fire off a cluster changed event to let other
 		// parts of the UI update as well
@@ -628,7 +629,7 @@ void FDisplayClusterConfiguratorViewCluster::AddNewViewport(FVector2D PresetSize
 	const TSharedPtr<FDisplayClusterConfiguratorBlueprintEditor> Toolkit = ToolkitPtr.Pin();
 	const FDisplayClusterConfigurationRectangle PresetRect = FDisplayClusterConfigurationRectangle(0, 0, PresetSize.X, PresetSize.Y);
 	TSharedPtr<FScopedTransaction> Transaction;
-	if (const UDisplayClusterConfigurationViewport* NewViewport = FDisplayClusterConfiguratorClusterUtils::CreateNewViewportFromDialog(Toolkit.ToSharedRef(), SelectedClusterNode, PresetRect, Transaction))
+	if (const UDisplayClusterConfigurationViewport* NewViewport = UE::DisplayClusterConfiguratorClusterEditorUtils::CreateNewViewportFromDialog(Toolkit.ToSharedRef(), SelectedClusterNode, PresetRect, Transaction))
 	{
 		// Mark the cluster configuration data as dirty, allowing user to save the changes, and fire off a cluster changed event to let other
 		// parts of the UI update as well
diff --git a/Engine/Plugins/Runtime/nDisplay/Source/DisplayClusterConfigurator/Private/Views/TreeViews/Cluster/DisplayClusterConfiguratorViewClusterBuilder.cpp b/Engine/Plugins/Runtime/nDisplay/Source/DisplayClusterConfigurator/Private/Views/TreeViews/Cluster/DisplayClusterConfiguratorViewClusterBuilder.cpp
--- a/Engine/Plugins/Runtime/nDisplay/Source/DisplayClusterConfigurator/Private/Views/TreeViews/Cluster/DisplayClusterConfiguratorViewClusterBuilder.cpp
+++ b/Engine/Plugins/Runtime/nDisplay/Source/DisplayClusterConfigurator/Private/Views/TreeViews/Cluster/DisplayClusterConfiguratorViewClusterBuilder.cpp
@@ -45,11 +45,11 @@ void FDisplayClusterConfiguratorViewClusterBuilder::AddCluster(FDisplayClusterCo
 void FDisplayClusterConfiguratorViewClusterBuilder::AddClusterNodes(FDisplayClusterConfiguratorTreeBuilderOutput& Output, UDisplayClusterConfigurationData* InConfig, UObject* InObjectToEdit)
 {
 	TMap<FString, TMap<FString, UDisplayClusterConfigurationClusterNode*>> SortedClusterNodes;
-	FDisplayClusterConfiguratorClusterUtils::SortClusterNodesByHost(InConfig->Cluster->Nodes, SortedClusterNodes);
+	UE::DisplayClusterConfiguratorClusterUtils::SortClusterNodesByHost(InConfig->Cluster->Nodes, SortedClusterNodes);
 
 	for (const auto& HostPair : SortedClusterNodes)
 	{
-		UDisplayClusterConfigurationHostDisplayData* DisplayData = FDisplayClusterConfiguratorClusterUtils::FindOrCreateHostDisplayData(InConfig->Cluster, HostPair.Key);
+		UDisplayClusterConfigurationHostDisplayData* DisplayData = UE::DisplayClusterConfiguratorClusterUtils::FindOrCreateHostDisplayData(InConfig->Cluster, HostPair.Key);
 
 		const FName HostParentName = "Cluster";
 		const FName HostNodeName = *DisplayData->HostName.ToString();
diff --git a/Engine/Plugins/Runtime/nDisplay/Source/DisplayClusterConfigurator/Private/Views/TreeViews/Cluster/TreeItems/DisplayClusterConfiguratorTreeItemClusterNode.cpp b/Engine/Plugins/Runtime/nDisplay/Source/DisplayClusterConfigurator/Private/Views/TreeViews/Cluster/TreeItems/DisplayClusterConfiguratorTreeItemClusterNode.cpp
--- a/Engine/Plugins/Runtime/nDisplay/Source/DisplayClusterConfigurator/Private/Views/TreeViews/Cluster/TreeItems/DisplayClusterConfiguratorTreeItemClusterNode.cpp
+++ b/Engine/Plugins/Runtime/nDisplay/Source/DisplayClusterConfigurator/Private/Views/TreeViews/Cluster/TreeItems/DisplayClusterConfiguratorTreeItemClusterNode.cpp
@@ -7,6 +7,7 @@
 #include "DisplayClusterConfiguratorBlueprintEditor.h"
 #include "DisplayClusterConfiguratorPropertyUtils.h"
 #include "ISinglePropertyView.h"
+#include "ClusterConfiguration/DisplayClusterConfiguratorClusterEditorUtils.h"
 #include "ClusterConfiguration/DisplayClusterConfiguratorClusterUtils.h"
 #include "Views/DragDrop/DisplayClusterConfiguratorValidatedDragDropOp.h"
 #include "Views/DragDrop/DisplayClusterConfiguratorClusterNodeDragDropOp.h"
@@ -50,7 +51,7 @@ void FDisplayClusterConfiguratorTreeItemClusterNode::SetVisible(bool bIsVisible)
 	// Use SaveToTransactionBuffer to avoid marking the package as dirty
 	SaveToTransactionBuffer(ClusterNode, false);
 
-	const TSharedPtr<ISinglePropertyView> PropertyView = DisplayClusterConfiguratorPropertyUtils::GetPropertyView(
+	const TSharedPtr<ISinglePropertyView> PropertyView = UE::DisplayClusterConfiguratorPropertyUtils::GetPropertyView(
 		ClusterNode, GET_MEMBER_NAME_CHECKED(UDisplayClusterConfigurationClusterNode, bIsVisible));
 	
 	PropertyView->GetPropertyHandle()->SetValue(bIsVisible);
@@ -73,7 +74,7 @@ void FDisplayClusterConfiguratorTreeItemClusterNode::SetUnlocked(bool bIsUnlocke
 	// Use SaveToTransactionBuffer to avoid marking the package as dirty
 	SaveToTransactionBuffer(ClusterNode, false);
 
-	const TSharedPtr<ISinglePropertyView> PropertyView = DisplayClusterConfiguratorPropertyUtils::GetPropertyView(
+	const TSharedPtr<ISinglePropertyView> PropertyView = UE::DisplayClusterConfiguratorPropertyUtils::GetPropertyView(
 		ClusterNode, GET_MEMBER_NAME_CHECKED(UDisplayClusterConfigurationClusterNode, bIsUnlocked));
 
 	PropertyView->GetPropertyHandle()->SetValue(bIsUnlocked);
@@ -92,7 +93,7 @@ void FDisplayClusterConfiguratorTreeItemClusterNode::SetUnlocked(bool bIsUnlocke
 void FDisplayClusterConfiguratorTreeItemClusterNode::DeleteItem() const
 {
 	UDisplayClusterConfigurationClusterNode* ClusterNode = GetObjectChecked<UDisplayClusterConfigurationClusterNode>();
-	FDisplayClusterConfiguratorClusterUtils::RemoveClusterNodeFromCluster(ClusterNode);
+	UE::DisplayClusterConfiguratorClusterUtils::RemoveClusterNodeFromCluster(ClusterNode);
 }
 
 FReply FDisplayClusterConfiguratorTreeItemClusterNode::HandleDragDetected(const FGeometry& MyGeometry, const FPointerEvent& MouseEvent)
@@ -114,7 +115,7 @@ FReply FDisplayClusterConfiguratorTreeItemClusterNode::HandleDragDetected(const
 			SelectedObjects.Add(SelectedItem->GetObject());
 		}
 
-		TSharedPtr<FDragDropOperation> DragDropOp = FDisplayClusterConfiguratorClusterUtils::MakeDragDropOperation(SelectedObjects);
+		TSharedPtr<FDragDropOperation> DragDropOp = UE::DisplayClusterConfiguratorClusterEditorUtils::MakeDragDropOperation(SelectedObjects);
 
 		if (DragDropOp.IsValid())
 		{
@@ -226,7 +227,7 @@ FReply FDisplayClusterConfiguratorTreeItemClusterNode::HandleAcceptDrop(const FD
 					continue;
 				}
 
-				FDisplayClusterConfiguratorClusterUtils::AddViewportToClusterNode(Viewport.Get(), ClusterNode);
+				UE::DisplayClusterConfiguratorClusterUtils::AddViewportToClusterNode(Viewport.Get(), ClusterNode);
 				bClusterModified = true;
 			}
 		}
@@ -283,7 +284,7 @@ void FDisplayClusterConfiguratorTreeItemClusterNode::OnDisplayNameCommitted(cons
 	FScopedTransaction Transaction(LOCTEXT("RenameClusterNode", "Rename Cluster Node"));
 	FString NewName = NewText.ToString();
 
-	if (FDisplayClusterConfiguratorClusterUtils::RenameClusterNode(ClusterNode, NewName))
+	if (UE::DisplayClusterConfiguratorClusterUtils::RenameClusterNode(ClusterNode, NewName))
 	{
 		Name = *NewName;
 
@@ -299,7 +300,7 @@ void FDisplayClusterConfiguratorTreeItemClusterNode::OnDisplayNameCommitted(cons
 EVisibility FDisplayClusterConfiguratorTreeItemClusterNode::GetPrimaryLabelVisibility() const
 {
 	UDisplayClusterConfigurationClusterNode* ClusterNode = GetObjectChecked<UDisplayClusterConfigurationClusterNode>();
-	const bool bIsPrimary = FDisplayClusterConfiguratorClusterUtils::IsClusterNodePrimary(ClusterNode);
+	const bool bIsPrimary = UE::DisplayClusterConfiguratorClusterUtils::IsClusterNodePrimary(ClusterNode);
 
 	return bIsPrimary ? EVisibility::Visible : EVisibility::Collapsed;
 }
@@ -316,7 +317,7 @@ bool FDisplayClusterConfiguratorTreeItemClusterNode::CanDropViewports(TSharedPtr
 		{
 			if (Viewport.IsValid())
 			{
-				FString ViewportName = FDisplayClusterConfiguratorClusterUtils::GetViewportName(Viewport.Get());
+				FString ViewportName = UE::DisplayClusterConfiguratorClusterUtils::GetViewportName(Viewport.Get());
 
 				if (ClusterNode->Viewports.Contains(ViewportName) && ClusterNode->Viewports[ViewportName] != Viewport.Get())
 				{
@@ -348,7 +349,7 @@ void FDisplayClusterConfiguratorTreeItemClusterNode::ToggleClusterItemVisibility
 	// Use SaveToTransactionBuffer to avoid marking the package as dirty
 	SaveToTransactionBuffer(ClusterNode, false);
 
-	const TSharedPtr<ISinglePropertyView> PropertyView = DisplayClusterConfiguratorPropertyUtils::GetPropertyView(
+	const TSharedPtr<ISinglePropertyView> PropertyView = UE::DisplayClusterConfiguratorPropertyUtils::GetPropertyView(
 		ClusterNode, GET_MEMBER_NAME_CHECKED(UDisplayClusterConfigurationClusterNode, bIsVisible));
 
 	PropertyView->GetPropertyHandle()->SetValue(!ClusterNode->bIsVisible);
@@ -380,7 +381,7 @@ void FDisplayClusterConfiguratorTreeItemClusterNode::ToggleClusterItemLock()
 	// Use SaveToTransactionBuffer to avoid marking the package as dirty
 	SaveToTransactionBuffer(ClusterNode, false);
 
-	const TSharedPtr<ISinglePropertyView> PropertyView = DisplayClusterConfiguratorPropertyUtils::GetPropertyView(
+	const TSharedPtr<ISinglePropertyView> PropertyView = UE::DisplayClusterConfiguratorPropertyUtils::GetPropertyView(
 		ClusterNode, GET_MEMBER_NAME_CHECKED(UDisplayClusterConfigurationClusterNode, bIsUnlocked));
 
 	PropertyView->GetPropertyHandle()->SetValue(!ClusterNode->bIsUnlocked);
diff --git a/Engine/Plugins/Runtime/nDisplay/Source/DisplayClusterConfigurator/Private/Views/TreeViews/Cluster/TreeItems/DisplayClusterConfiguratorTreeItemHost.cpp b/Engine/Plugins/Runtime/nDisplay/Source/DisplayClusterConfigurator/Private/Views/TreeViews/Cluster/TreeItems/DisplayClusterConfiguratorTreeItemHost.cpp
--- a/Engine/Plugins/Runtime/nDisplay/Source/DisplayClusterConfigurator/Private/Views/TreeViews/Cluster/TreeItems/DisplayClusterConfiguratorTreeItemHost.cpp
+++ b/Engine/Plugins/Runtime/nDisplay/Source/DisplayClusterConfigurator/Private/Views/TreeViews/Cluster/TreeItems/DisplayClusterConfiguratorTreeItemHost.cpp
@@ -35,8 +35,8 @@ void FDisplayClusterConfiguratorTreeItemHost::DeleteItem() const
 	UDisplayClusterConfigurationHostDisplayData* Host = GetObjectChecked<UDisplayClusterConfigurationHostDisplayData>();
 	if (UDisplayClusterConfigurationCluster* Cluster = Cast<UDisplayClusterConfigurationCluster>(Host->GetOuter()))
 	{
-		FString HostAddress = FDisplayClusterConfiguratorClusterUtils::GetAddressForHost(Host);
-		FDisplayClusterConfiguratorClusterUtils::RemoveHost(Cluster, HostAddress);
+		FString HostAddress = UE::DisplayClusterConfiguratorClusterUtils::GetAddressForHost(Host);
+		UE::DisplayClusterConfiguratorClusterUtils::RemoveHost(Cluster, HostAddress);
 	}
 }
 
@@ -100,7 +100,7 @@ FReply FDisplayClusterConfiguratorTreeItemHost::HandleAcceptDrop(const FDragDrop
 		FScopedTransaction Transaction(FText::Format(LOCTEXT("DropClusterNodes", "Drop {0}|plural(one=Cluster Node, other=Cluster Nodes)"), ClusterNodeDragDropOp->GetDraggedClusterNodes().Num()));
 
 		UDisplayClusterConfigurationHostDisplayData* HostDisplayData = GetObjectChecked<UDisplayClusterConfigurationHostDisplayData>();
-		const FString HostAddress = FDisplayClusterConfiguratorClusterUtils::GetAddressForHost(HostDisplayData);
+		const FString HostAddress = UE::DisplayClusterConfiguratorClusterUtils::GetAddressForHost(HostDisplayData);
 		bool bClusterModified = false;
 		for (TWeakObjectPtr<UDisplayClusterConfigurationClusterNode> ClusterNode : ClusterNodeDragDropOp->GetDraggedClusterNodes())
 		{
@@ -139,7 +139,7 @@ FName FDisplayClusterConfiguratorTreeItemHost::GetAttachName() const
 	// Use the host's address as the attach name, to ensure that the name is unique (the host's HostName property has no uniqueness enforced on it)
 	// so that child cluster nodes are correctly associated with their parent host nodes
 	UDisplayClusterConfigurationHostDisplayData* Host = GetObjectChecked<UDisplayClusterConfigurationHostDisplayData>();
-	FName AttachName = *FDisplayClusterConfiguratorClusterUtils::GetAddressForHost(Host);
+	FName AttachName = *UE::DisplayClusterConfiguratorClusterUtils::GetAddressForHost(Host);
 	return AttachName;
 }
 
@@ -269,7 +269,7 @@ FReply FDisplayClusterConfiguratorTreeItemHost::OnClusterItemGroupClicked(const
 FText FDisplayClusterConfiguratorTreeItemHost::GetHostAddress() const
 {
 	UDisplayClusterConfigurationHostDisplayData* Host = GetObjectChecked<UDisplayClusterConfigurationHostDisplayData>();
-	return FText::Format(LOCTEXT("HostAddressFormattedLabel", "({0})"), FText::FromString(FDisplayClusterConfiguratorClusterUtils::GetAddressForHost(Host)));
+	return FText::Format(LOCTEXT("HostAddressFormattedLabel", "({0})"), FText::FromString(UE::DisplayClusterConfiguratorClusterUtils::GetAddressForHost(Host)));
 }
 
 bool FDisplayClusterConfiguratorTreeItemHost::CanDropClusterNodes(TSharedPtr<FDisplayClusterConfiguratorClusterNodeDragDropOp> ClusterNodeDragDropOp, FText& OutErrorMessage) const
diff --git a/Engine/Plugins/Runtime/nDisplay/Source/DisplayClusterConfigurator/Private/Views/TreeViews/Cluster/TreeItems/DisplayClusterConfiguratorTreeItemViewport.cpp b/Engine/Plugins/Runtime/nDisplay/Source/DisplayClusterConfigurator/Private/Views/TreeViews/Cluster/TreeItems/DisplayClusterConfiguratorTreeItemViewport.cpp
--- a/Engine/Plugins/Runtime/nDisplay/Source/DisplayClusterConfigurator/Private/Views/TreeViews/Cluster/TreeItems/DisplayClusterConfiguratorTreeItemViewport.cpp
+++ b/Engine/Plugins/Runtime/nDisplay/Source/DisplayClusterConfigurator/Private/Views/TreeViews/Cluster/TreeItems/DisplayClusterConfiguratorTreeItemViewport.cpp
@@ -7,6 +7,7 @@
 #include "DisplayClusterProjectionStrings.h"
 #include "DisplayClusterConfiguratorBlueprintEditor.h"
 #include "DisplayClusterConfiguratorPropertyUtils.h"
+#include "ClusterConfiguration/DisplayClusterConfiguratorClusterEditorUtils.h"
 #include "ClusterConfiguration/DisplayClusterConfiguratorClusterUtils.h"
 #include "Views/DragDrop/DisplayClusterConfiguratorValidatedDragDropOp.h"
 #include "Views/DragDrop/DisplayClusterConfiguratorViewportDragDropOp.h"
@@ -35,7 +36,7 @@ void FDisplayClusterConfiguratorTreeItemViewport::SetVisible(bool bIsVisible)
 	// Use SaveToTransactionBuffer to avoid marking the package as dirty
 	SaveToTransactionBuffer(Viewport, false);
 
-	const TSharedPtr<ISinglePropertyView> PropertyView = DisplayClusterConfiguratorPropertyUtils::GetPropertyView(
+	const TSharedPtr<ISinglePropertyView> PropertyView = UE::DisplayClusterConfiguratorPropertyUtils::GetPropertyView(
 		Viewport, GET_MEMBER_NAME_CHECKED(UDisplayClusterConfigurationViewport, bIsVisible));
 
 	PropertyView->GetPropertyHandle()->SetValue(bIsVisible);
@@ -48,7 +49,7 @@ void FDisplayClusterConfiguratorTreeItemViewport::SetUnlocked(bool bIsUnlocked)
 	// Use SaveToTransactionBuffer to avoid marking the package as dirty
 	SaveToTransactionBuffer(Viewport, false);
 
-	const TSharedPtr<ISinglePropertyView> PropertyView = DisplayClusterConfiguratorPropertyUtils::GetPropertyView(
+	const TSharedPtr<ISinglePropertyView> PropertyView = UE::DisplayClusterConfiguratorPropertyUtils::GetPropertyView(
 		Viewport, GET_MEMBER_NAME_CHECKED(UDisplayClusterConfigurationViewport, bIsUnlocked));
 	
 	PropertyView->GetPropertyHandle()->SetValue(bIsUnlocked);
@@ -86,7 +87,7 @@ void FDisplayClusterConfiguratorTreeItemViewport::OnSelection()
 void FDisplayClusterConfiguratorTreeItemViewport::DeleteItem() const
 {
 	UDisplayClusterConfigurationViewport* Viewport = GetObjectChecked<UDisplayClusterConfigurationViewport>();
-	FDisplayClusterConfiguratorClusterUtils::RemoveViewportFromClusterNode(Viewport);
+	UE::DisplayClusterConfiguratorClusterUtils::RemoveViewportFromClusterNode(Viewport);
 }
 
 FReply FDisplayClusterConfiguratorTreeItemViewport::HandleDragDetected(const FGeometry& MyGeometry, const FPointerEvent& MouseEvent)
@@ -108,7 +109,7 @@ FReply FDisplayClusterConfiguratorTreeItemViewport::HandleDragDetected(const FGe
 			SelectedObjects.Add(SelectedItem->GetObject());
 		}
 
-		TSharedPtr<FDragDropOperation> DragDropOp = FDisplayClusterConfiguratorClusterUtils::MakeDragDropOperation(SelectedObjects);
+		TSharedPtr<FDragDropOperation> DragDropOp = UE::DisplayClusterConfiguratorClusterEditorUtils::MakeDragDropOperation(SelectedObjects);
 
 		if (DragDropOp.IsValid())
 		{
@@ -198,7 +199,7 @@ void FDisplayClusterConfiguratorTreeItemViewport::OnDisplayNameCommitted(const F
 	FScopedTransaction Transaction(NSLOCTEXT("FDisplayClusterConfiguratorViewCluster", "RenameClusterNode", "Rename Cluster Node"));
 	FString NewName = NewText.ToString();
 
-	if (FDisplayClusterConfiguratorClusterUtils::RenameViewport(Viewport, NewName))
+	if (UE::DisplayClusterConfiguratorClusterUtils::RenameViewport(Viewport, NewName))
 	{
 		Name = *NewName;
 
@@ -226,7 +227,7 @@ void FDisplayClusterConfiguratorTreeItemViewport::ToggleClusterItemVisibility()
 	// Use SaveToTransactionBuffer to avoid marking the package as dirty
 	SaveToTransactionBuffer(Viewport, false);
 
-	const TSharedPtr<ISinglePropertyView> PropertyView = DisplayClusterConfiguratorPropertyUtils::GetPropertyView(
+	const TSharedPtr<ISinglePropertyView> PropertyView = UE::DisplayClusterConfiguratorPropertyUtils::GetPropertyView(
 		Viewport, GET_MEMBER_NAME_CHECKED(UDisplayClusterConfigurationViewport, bIsVisible));
 
 	PropertyView->GetPropertyHandle()->SetValue(!Viewport->bIsVisible);
@@ -247,7 +248,7 @@ void FDisplayClusterConfiguratorTreeItemViewport::ToggleClusterItemLock()
 	// Use SaveToTransactionBuffer to avoid marking the package as dirty
 	SaveToTransactionBuffer(Viewport, false);
 
-	const TSharedPtr<ISinglePropertyView> PropertyView = DisplayClusterConfiguratorPropertyUtils::GetPropertyView(
+	const TSharedPtr<ISinglePropertyView> PropertyView = UE::DisplayClusterConfiguratorPropertyUtils::GetPropertyView(
 		Viewport, GET_MEMBER_NAME_CHECKED(UDisplayClusterConfigurationViewport, bIsUnlocked));
 
 	PropertyView->GetPropertyHandle()->SetValue(!Viewport->bIsUnlocked);
diff --git a/Engine/Plugins/Runtime/nDisplay/Source/DisplayClusterConfigurator/Private/Views/Viewport/DisplayClusterConfiguratorSCSEditorViewportClient.cpp b/Engine/Plugins/Runtime/nDisplay/Source/DisplayClusterConfigurator/Private/Views/Viewport/DisplayClusterConfiguratorSCSEditorViewportClient.cpp
--- a/Engine/Plugins/Runtime/nDisplay/Source/DisplayClusterConfigurator/Private/Views/Viewport/DisplayClusterConfiguratorSCSEditorViewportClient.cpp
+++ b/Engine/Plugins/Runtime/nDisplay/Source/DisplayClusterConfigurator/Private/Views/Viewport/DisplayClusterConfiguratorSCSEditorViewportClient.cpp
@@ -1164,7 +1164,7 @@ void FDisplayClusterConfiguratorSCSEditorViewportClient::SetPreviewResolutionSca
 	{
 		if (ADisplayClusterRootActor* CDO = BlueprintEditor->GetDefaultRootActor())
 		{
-			DisplayClusterConfiguratorPropertyUtils::SetPropertyHandleValue(
+			UE::DisplayClusterConfiguratorPropertyUtils::SetPropertyHandleValue(
 				CDO, GET_MEMBER_NAME_CHECKED(ADisplayClusterRootActor, PreviewRenderTargetRatioMult), InScale);
 		}
 	}
diff --git a/Engine/Plugins/Runtime/nDisplay/Source/DisplayClusterConfigurator/Public/ClusterConfiguration/DisplayClusterConfiguratorClusterUtils.h b/Engine/Plugins/Runtime/nDisplay/Source/DisplayClusterConfigurator/Public/ClusterConfiguration/DisplayClusterConfiguratorClusterUtils.h
new file mode 100644
--- /dev/null
+++ b/Engine/Plugins/Runtime/nDisplay/Source/DisplayClusterConfigurator/Public/ClusterConfiguration/DisplayClusterConfiguratorClusterUtils.h
@@ -0,0 +1,166 @@
+// Copyright Epic Games, Inc. All Rights Reserved.
+
+#pragma once
+
+#include "CoreMinimal.h"
+
+class UDisplayClusterConfigurationData_Base;
+class UDisplayClusterConfigurationCluster;
+class UDisplayClusterConfigurationHostDisplayData;
+class UDisplayClusterConfigurationClusterNode;
+class UDisplayClusterConfigurationViewport;
+
+namespace UE::DisplayClusterConfiguratorClusterUtils
+{
+	/**
+	 * Adds a cluster node to a cluster, removing it from its current cluster if it belongs to one.
+	 * @param ClusterNode - The cluster node to add
+	 * @param Cluster - The cluster to add the cluster node to
+	 * @param NewClusterNodeName - Optional name to add the new cluster node under in the cluster. If not supplied, the cluster node's current name is used, if possible.
+	 * @return An updated pointer to the cluster node added to the cluster.
+	 */
+	DISPLAYCLUSTERCONFIGURATOR_API UDisplayClusterConfigurationClusterNode* AddClusterNodeToCluster(UDisplayClusterConfigurationClusterNode* ClusterNode, UDisplayClusterConfigurationCluster* Cluster, FString NewClusterNodeName = "");
+
+	/**
+	 * Removes a cluster node from the cluster that owns it. The owning cluster is retrieved using the cluster node's Outer.
+	 * @param ClusterNode - The cluster node to remove
+	 * @returns true if the cluster node was successfully removed from its owning cluster
+	 */
+	DISPLAYCLUSTERCONFIGURATOR_API bool RemoveClusterNodeFromCluster(UDisplayClusterConfigurationClusterNode* ClusterNode);
+
+	/**
+	 * Renames a cluster node, which involves changing the key that the cluster node is stored under in its parent cluster.
+	 * @param ClusterNode - The cluster node to rename
+	 * @param NewClusterNodeName - The desired name to change to. The final name will be a unique-ified version of this name.
+	 * @return true if the cluster node was successfully renamed
+	 */
+	DISPLAYCLUSTERCONFIGURATOR_API bool RenameClusterNode(UDisplayClusterConfigurationClusterNode* ClusterNode, FString NewClusterNodeName);
+
+	/**
+	 * Sets the cluster node as the primary node.
+	 * @param ClusterNode - The cluster node to set as the primary node
+	 * @return true if the cluster node was successfully set as the primary node
+	 */
+	DISPLAYCLUSTERCONFIGURATOR_API bool SetClusterNodeAsPrimary(UDisplayClusterConfigurationClusterNode* ClusterNode);
+
+	/**
+	 * Determines if a cluster node is the primary node in its cluster.
+	 * @param ClusterNode - The cluster node to check.
+	 * @return true if the node is the primary; otherwise, false.
+	 */
+	DISPLAYCLUSTERCONFIGURATOR_API bool IsClusterNodePrimary(UDisplayClusterConfigurationClusterNode* ClusterNode);
+	
+	/**
+	 * Adds a viewport to a cluster node, removing it from its current cluster node if it belongs to one.
+	 * @param Viewport - The viewport to add
+	 * @param ClusterNode - The cluster node to add the viewport to
+	 * @param NewViewportName - Optional name to add the new viewport under in the cluster node. If not supplied, the viewport's current name is used, if possible.
+	 * @return An updated pointer to the viewport added to the cluster.
+	 */
+	DISPLAYCLUSTERCONFIGURATOR_API UDisplayClusterConfigurationViewport* AddViewportToClusterNode(UDisplayClusterConfigurationViewport* Viewport, UDisplayClusterConfigurationClusterNode* ClusterNode, FString NewViewportName = "");
+
+	/**
+	 * Removes a viewport from the cluster node that owns it. The owning cluster node is retrieved using the viewport's Outer.
+	 * @param Viewport - The viewport to remove
+	 * @returns true if the viewport was successfully removed from its owning cluster
+	 */
+	DISPLAYCLUSTERCONFIGURATOR_API bool RemoveViewportFromClusterNode(UDisplayClusterConfigurationViewport* Viewport);
+
+	/**
+	 * Renames a viewport, which involves changing the key that the viewport is stored under in its parent cluster node.
+	 * @param Viewport - The viewport to rename
+	 * @param NewViewportName - The desired name to change to. The final name will be a unique-ified version of this name.
+	 * @return true if the viewport was successfully renamed
+	 */
+	DISPLAYCLUSTERCONFIGURATOR_API bool RenameViewport(UDisplayClusterConfigurationViewport* Viewport, FString NewViewportName);
+
+	/**
+	 * Attempts to find a preexisting display data object for the specified host, and if one can't be found, creates a new one.
+	 * @param Cluster - The cluster to find the host display data object in
+	 * @param HostIPAddress - The IP address of the host to find the display data for
+	 * @return The display data that was found or created
+	 */
+	DISPLAYCLUSTERCONFIGURATOR_API UDisplayClusterConfigurationHostDisplayData* FindOrCreateHostDisplayData(UDisplayClusterConfigurationCluster* Cluster, FString HostIPAddress);
+
+	/**
+	 * Removes all host display data objects from the cluster that aren't being used by any of the cluster nodes.
+	 * @param Cluster - The cluster to clean up the host display data for
+	 * @return true if any host display data was removed; otherwise, false
+	 */
+	DISPLAYCLUSTERCONFIGURATOR_API bool RemoveUnusedHostDisplayData(UDisplayClusterConfigurationCluster* Cluster);
+
+	/**
+	 * Gets a unique name for a host owned by the specified cluster.
+	 * @param InitialName - The initial name the host wants
+	 * @param ParentCluster - The cluster the host belongs to
+	 * @param bAddZero - Whether to add an "_0" to the initial name if it is unique
+	 * @returns A unique name for the host
+	 */
+	DISPLAYCLUSTERCONFIGURATOR_API FString GetUniqueNameForHost(FString InitialName, UDisplayClusterConfigurationCluster* ParentCluster, bool bAddZero = false);
+
+	/**
+	 * Gets the address of a host using its host display data.
+	 * @param HostDisplayData - The display data of the host to get the address for
+	 * @return The address of the host, or an empty string if no host address was found
+	 */
+	DISPLAYCLUSTERCONFIGURATOR_API FString GetAddressForHost(UDisplayClusterConfigurationHostDisplayData* HostDisplayData);
+
+	/**
+	 * Removes all cluster nodes and host display data from the cluster for the specified host.
+	 * @param Cluster - The cluster to remove the host from.
+	 * @param Host - The host to remove.
+	 * @return true if any cluster nodes or host display data was removed; otherwise, false
+	 */
+	DISPLAYCLUSTERCONFIGURATOR_API bool RemoveHost(UDisplayClusterConfigurationCluster* Cluster, FString Host);
+
+	/**
+	 * Attempts to get the host display data for the host of a cluster node.
+	 * @param ClusterNode - The cluster node whose host to get the display data for.
+	 * @return The display data for the cluster node's host, or null if no display data was found.
+	 */
+	DISPLAYCLUSTERCONFIGURATOR_API UDisplayClusterConfigurationHostDisplayData* GetHostDisplayDataForClusterNode(UDisplayClusterConfigurationClusterNode* ClusterNode);
+
+	/**
+	 * Gets the name that the cluster node is stored under in its parent cluster.
+	 * @return The cluster node's storage name
+	 */
+	DISPLAYCLUSTERCONFIGURATOR_API FString GetClusterNodeName(UDisplayClusterConfigurationClusterNode* ClusterNode);
+
+	/**
+	 * Gets a unique name for a cluster node owned by the specified cluster.
+	 * @param InitialName - The initial name the cluster node wants
+	 * @param ParentCluster - The cluster the cluster node belongs to
+	 * @param bAddZero - Whether to add an "_0" to the initial name if it is unique
+	 * @returns A unique name for the cluster node
+	 */
+	DISPLAYCLUSTERCONFIGURATOR_API FString GetUniqueNameForClusterNode(FString InitialName, UDisplayClusterConfigurationCluster* ParentCluster, bool bAddZero = false);
+
+	/**
+	 * Gets the name that the viewport is stored under in its parent cluster node.
+	 * @return The viewport's storage name
+	 */
+	DISPLAYCLUSTERCONFIGURATOR_API FString GetViewportName(UDisplayClusterConfigurationViewport* Viewport);
+
+	/**
+	 * Gets a unique name for a viewport owned by the specified cluster node.
+	 * @param InitialName - The initial name the viewport wants
+	 * @param ParentClusterNode - The cluster node the viewport belongs to
+	 * @param bAddZero - Whether to add an "_0" to the initial name if it is unique
+	 * @returns A unique name for the viewport
+	 */
+	DISPLAYCLUSTERCONFIGURATOR_API FString GetUniqueNameForViewport(FString InitialName, UDisplayClusterConfigurationClusterNode* ParentClusterNode, bool bAddZero = false);
+
+	/**
+	 * Sorts a list of cluster nodes by their host.
+	 * @param InClusterNodes - The list of nodes to sort
+	 * @param OutSortedNodes - The sorted nodes, indexed by the Host property of the cluster nodes
+	 */
+	DISPLAYCLUSTERCONFIGURATOR_API void SortClusterNodesByHost(const TMap<FString, UDisplayClusterConfigurationClusterNode*>& InClusterNodes, TMap<FString, TMap<FString, UDisplayClusterConfigurationClusterNode*>>& OutSortedNodes);
+
+	/**
+	 * Sorts a list of cluster nodes by their host.
+	 * @param InClusterNodes - The list of nodes to sort
+	 * @param OutSortedNodes - The sorted nodes, indexed by the Host property of the cluster nodes
+	 */
+	DISPLAYCLUSTERCONFIGURATOR_API void SortClusterNodesByHost(const TMap<FString, TObjectPtr<UDisplayClusterConfigurationClusterNode>>& InClusterNodes, TMap<FString, TMap<FString, UDisplayClusterConfigurationClusterNode*>>& OutSortedNodes);
+};
\ No newline at end of file
diff --git a/Engine/Plugins/Runtime/nDisplay/Source/DisplayClusterConfigurator/Public/DisplayClusterConfiguratorFactory.h b/Engine/Plugins/Runtime/nDisplay/Source/DisplayClusterConfigurator/Public/DisplayClusterConfiguratorFactory.h
new file mode 100644
--- /dev/null
+++ b/Engine/Plugins/Runtime/nDisplay/Source/DisplayClusterConfigurator/Public/DisplayClusterConfiguratorFactory.h
@@ -0,0 +1,72 @@
+// Copyright Epic Games, Inc. All Rights Reserved.
+
+#pragma once
+
+#include "CoreMinimal.h"
+#include "UObject/ObjectMacros.h"
+#include "Factories/Factory.h"
+#include "EditorReimportHandler.h"
+
+#include "DisplayClusterConfiguratorFactory.generated.h"
+
+class UDisplayClusterBlueprint;
+
+UCLASS(MinimalApi)
+class UDisplayClusterConfiguratorFactory
+	: public UFactory
+{
+	GENERATED_BODY()
+
+public:
+	UDisplayClusterConfiguratorFactory();
+	~UDisplayClusterConfiguratorFactory();
+	
+	//~ Begin UFactory Interface
+	virtual bool ConfigureProperties() override;
+	virtual UObject* FactoryCreateNew(UClass* Class, UObject* InParent, FName Name, EObjectFlags Flags, UObject* Context,
+		FFeedbackContext* Warn, FName CallingContext) override;
+	virtual UObject* FactoryCreateNew(UClass* Class, UObject* InParent, FName Name, EObjectFlags Flags, UObject* Context,
+		FFeedbackContext* Warn) override;
+	virtual bool DoesSupportClass(UClass* Class) override;
+	virtual UClass* ResolveSupportedClass() override;
+	virtual UObject* FactoryCreateFile(UClass* InClass, UObject* InParent, FName InName, EObjectFlags Flags, const FString& InFilename, const TCHAR* Parms, FFeedbackContext* Warn, bool& bOutOperationCanceled) override;
+	virtual bool FactoryCanImport(const FString& Filename) override { return true; }
+	virtual FString GetDefaultNewAssetName() const override;
+	//~ End UFactory Interface
+
+	/** Create initial config data and components. */
+	static void SetupNewBlueprint(UDisplayClusterBlueprint* NewBlueprint);
+
+	/** Perform one time setup of open documents. */
+	static void SetupInitialBlueprintDocuments(UDisplayClusterBlueprint* NewBlueprint);
+private:
+	void OnConfigureNewAssetRequest(UFactory* InFactory);
+
+private:
+	// The parent class of the created blueprint
+	UPROPERTY()
+	TSubclassOf<class ADisplayClusterRootActor> ParentClass;
+
+	UPROPERTY(Transient)
+	TObjectPtr<UDisplayClusterBlueprint> BlueprintToCopy;
+
+	FDelegateHandle OnConfigureNewAssetRequestHandle;
+	bool bIsConfigureNewAssetRequest;
+};
+
+UCLASS(MinimalApi)
+class UDisplayClusterConfiguratorReimportFactory
+	: public UDisplayClusterConfiguratorFactory
+	, public FReimportHandler
+{
+	GENERATED_BODY()
+
+	UDisplayClusterConfiguratorReimportFactory();
+
+	//~ Begin FReimportHandler Interface
+	virtual bool CanReimport(UObject* Obj, TArray<FString>& OutFilenames) override;
+	virtual void SetReimportPaths(UObject* Obj, const TArray<FString>& NewReimportPaths) override;
+	virtual EReimportResult::Type Reimport(UObject* Obj) override;
+	virtual int32 GetPriority() const override;
+	//~ End FReimportHandler Interface
+};
diff --git a/Engine/Plugins/Runtime/nDisplay/Source/DisplayClusterEditor/DisplayClusterEditor.Build.cs b/Engine/Plugins/Runtime/nDisplay/Source/DisplayClusterEditor/DisplayClusterEditor.Build.cs
--- a/Engine/Plugins/Runtime/nDisplay/Source/DisplayClusterEditor/DisplayClusterEditor.Build.cs
+++ b/Engine/Plugins/Runtime/nDisplay/Source/DisplayClusterEditor/DisplayClusterEditor.Build.cs
@@ -17,8 +17,5 @@ public class DisplayClusterEditor : ModuleRules
 				"Engine",
 				"UnrealEd"
 			});
-
-		// TODO: Should not be including private headers from other modules
-		PrivateIncludePaths.Add(Path.Combine(GetModuleDirectory("DisplayCluster"), "Private")); // For IPDisplayCluster.h
 	}
 }
diff --git a/Engine/Plugins/Runtime/nDisplay/Source/DisplayClusterTests/DisplayClusterTests.Build.cs b/Engine/Plugins/Runtime/nDisplay/Source/DisplayClusterTests/DisplayClusterTests.Build.cs
--- a/Engine/Plugins/Runtime/nDisplay/Source/DisplayClusterTests/DisplayClusterTests.Build.cs
+++ b/Engine/Plugins/Runtime/nDisplay/Source/DisplayClusterTests/DisplayClusterTests.Build.cs
@@ -17,15 +17,10 @@ public class DisplayClusterTests : ModuleRules
 				"Engine",
 				"OpenColorIO",
 				"SubobjectDataInterface",
-				"UnrealEd"
+				"UnrealEd",
+				"SlateCore"
 			});
-
-		// TODO: Should not be including private headers from a different module
-		PrivateIncludePaths.AddRange(
-			new string[] {
-				Path.Combine(GetModuleDirectory("DisplayClusterConfigurator"), "Private"), // For DisplayClusterConfiguratorPropertyUtils.h
-			});
-
+		
 		OptimizeCode = CodeOptimization.Never;
 		PCHUsage = PCHUsageMode.NoPCHs;
 	}
diff --git a/Engine/Plugins/Runtime/nDisplay/Source/DisplayClusterTests/Private/DisplayCluster.spec.cpp b/Engine/Plugins/Runtime/nDisplay/Source/DisplayClusterTests/Private/DisplayCluster.spec.cpp
--- a/Engine/Plugins/Runtime/nDisplay/Source/DisplayClusterTests/Private/DisplayCluster.spec.cpp
+++ b/Engine/Plugins/Runtime/nDisplay/Source/DisplayClusterTests/Private/DisplayCluster.spec.cpp
@@ -341,7 +341,7 @@ void FDisplayClusterSpec::Define()
 			// a new node to the instance's map results in an invalid access. This test catches 
 			
 			DisplayClusterTestUtils::AddClusterNodeToCluster(ClusterAsset, AssetRootCluster);
-			FDisplayClusterConfiguratorClusterUtils::RemoveClusterNodeFromCluster(AssetClusterNode);
+			UE::DisplayClusterConfiguratorClusterUtils::RemoveClusterNodeFromCluster(AssetClusterNode);
 			FKismetEditorUtilities::CompileBlueprint(ClusterAsset);
 
 			// We recompiled, so reacquire the instanced pointers
diff --git a/Engine/Plugins/Runtime/nDisplay/Source/DisplayClusterTests/Private/DisplayClusterTestUtils.cpp b/Engine/Plugins/Runtime/nDisplay/Source/DisplayClusterTests/Private/DisplayClusterTestUtils.cpp
--- a/Engine/Plugins/Runtime/nDisplay/Source/DisplayClusterTests/Private/DisplayClusterTestUtils.cpp
+++ b/Engine/Plugins/Runtime/nDisplay/Source/DisplayClusterTests/Private/DisplayClusterTestUtils.cpp
@@ -9,6 +9,7 @@
 #include "DisplayClusterTestsModule.h"
 #include "Engine/Engine.h"
 #include "Engine/World.h"
+#include "ISinglePropertyView.h"
 #include "ObjectTools.h"
 #include "PackageTools.h"
 
@@ -76,7 +77,7 @@ namespace DisplayClusterTestUtils
 		// This must be set to avoid errors about illegal cross-package references
 		NodeTemplate->SetFlags(RF_Transactional);
 
-		UDisplayClusterConfigurationClusterNode* NewNode = FDisplayClusterConfiguratorClusterUtils::AddClusterNodeToCluster(NodeTemplate, RootCluster, Name);
+		UDisplayClusterConfigurationClusterNode* NewNode = UE::DisplayClusterConfiguratorClusterUtils::AddClusterNodeToCluster(NodeTemplate, RootCluster, Name);
 
 		// Node template is no longer needed, leave it to be cleaned up
 		NodeTemplate->MarkAsGarbage();
@@ -97,7 +98,7 @@ namespace DisplayClusterTestUtils
 		// This must be set to avoid errors about illegal cross-package references
 		ViewportTemplate->SetFlags(RF_Transactional);
 
-		UDisplayClusterConfigurationViewport* NewViewport = FDisplayClusterConfiguratorClusterUtils::AddViewportToClusterNode(ViewportTemplate, Node, Name);
+		UDisplayClusterConfigurationViewport* NewViewport = UE::DisplayClusterConfiguratorClusterUtils::AddViewportToClusterNode(ViewportTemplate, Node, Name);
 
 		// Node template is no longer needed, leave it to be cleaned up
 		ViewportTemplate->MarkAsGarbage();
@@ -201,8 +202,13 @@ namespace DisplayClusterTestUtils
 			return false;
 		}
 		
-		// Get the root property view and handle
-		PropertyView = DisplayClusterConfiguratorPropertyUtils::GetPropertyView(Owner, FieldNames[0]);
+		// Get the root property view and handle -- equivalent to UE::DisplayClusterConfiguratorPropertyUtils::GetPropertyView
+		const FSinglePropertyParams InitParams;
+		FPropertyEditorModule& PropertyEditorModule = FModuleManager::GetModuleChecked<FPropertyEditorModule>("PropertyEditor");
+		PropertyView = PropertyEditorModule.CreateSingleProperty(
+			Owner,
+			FieldNames[0],
+			InitParams);
 		
 		if (!PropertyView.IsValid())
 		{
diff --git a/Engine/Plugins/Runtime/nDisplay/Source/DisplayClusterTests/Private/DisplayClusterTestUtils.h b/Engine/Plugins/Runtime/nDisplay/Source/DisplayClusterTests/Private/DisplayClusterTestUtils.h
--- a/Engine/Plugins/Runtime/nDisplay/Source/DisplayClusterTests/Private/DisplayClusterTestUtils.h
+++ b/Engine/Plugins/Runtime/nDisplay/Source/DisplayClusterTests/Private/DisplayClusterTestUtils.h
@@ -1,12 +1,14 @@
 // Copyright Epic Games, Inc. All Rights Reserved.
 
 #pragma once
+
 #include "DisplayClusterConfigurationTypes_Viewport.h"
 
 #if WITH_DEV_AUTOMATION_TESTS
 
-#include "DisplayClusterConfiguratorPropertyUtils.h"
 #include "Kismet2/BlueprintEditorUtils.h"
+#include "PropertyEditorModule.h"
+#include "PropertyHandle.h"
 
 class UDisplayClusterBlueprint;
 class UDisplayClusterConfigurationCluster;
