nDisplay: Added new interface IDisplayClusterViewportConfiguration
- Reworked all related logic and interfaces.
- Code cleanup
- Partially revised code related to preview rendering, which was under WITH_EDITOR macros. Now this logic can be used in Package\Standalone builds.
- Several changes to the warp policy interface and related logic.

diff --git a/Engine/Plugins/Runtime/nDisplay/Source/DisplayCluster/Private/Game/DisplayClusterGameManager.cpp b/Engine/Plugins/Runtime/nDisplay/Source/DisplayCluster/Private/Game/DisplayClusterGameManager.cpp
--- a/Engine/Plugins/Runtime/nDisplay/Source/DisplayCluster/Private/Game/DisplayClusterGameManager.cpp
+++ b/Engine/Plugins/Runtime/nDisplay/Source/DisplayCluster/Private/Game/DisplayClusterGameManager.cpp
@@ -194,7 +194,7 @@ ADisplayClusterRootActor* FDisplayClusterGameManager::GetRootActor() const
 //////////////////////////////////////////////////////////////////////////////////////////////
 // FDisplayClusterGameManager
 //////////////////////////////////////////////////////////////////////////////////////////////
-ADisplayClusterRootActor* FDisplayClusterGameManager::FindRootActor(UWorld* InWorld, UDisplayClusterConfigurationData* InConfigData)
+TArray<ADisplayClusterRootActor*> FDisplayClusterGameManager::GetAllRootActorsFromWorld(UWorld* InWorld)
 {
 	TArray<ADisplayClusterRootActor*> FoundActors;
 	FoundActors.Reserve(16);
@@ -214,6 +214,13 @@ ADisplayClusterRootActor* FDisplayClusterGameManager::FindRootActor(UWorld* InWo
 		}
 	}
 
+	return FoundActors;
+}
+
+ADisplayClusterRootActor* FDisplayClusterGameManager::FindRootActor(UWorld* InWorld, UDisplayClusterConfigurationData* InConfigData)
+{
+	TArray<ADisplayClusterRootActor*> FoundActors = FDisplayClusterGameManager::GetAllRootActorsFromWorld(InWorld);
+	
 #if WITH_EDITOR
 	if (InWorld->IsPlayInEditor())
 	{
@@ -286,7 +293,7 @@ void FDisplayClusterGameManager::FindRootActorsInWorld(UWorld* InWorld, TArray<A
 				ADisplayClusterRootActor* RootActor = Cast<ADisplayClusterRootActor>(*It);
 				if (RootActor != nullptr && !RootActor->IsTemplate())
 				{
-					UE_LOG(LogDisplayClusterGame, Log, TEXT("Found root actor - %s"), *RootActor->GetName());
+					UE_LOG(LogDisplayClusterGame, VeryVerbose, TEXT("Found root actor - %s"), *RootActor->GetName());
 					OutActors.Add(RootActor);
 				}
 			}
diff --git a/Engine/Plugins/Runtime/nDisplay/Source/DisplayCluster/Private/Game/DisplayClusterGameManager.h b/Engine/Plugins/Runtime/nDisplay/Source/DisplayCluster/Private/Game/DisplayClusterGameManager.h
--- a/Engine/Plugins/Runtime/nDisplay/Source/DisplayCluster/Private/Game/DisplayClusterGameManager.h
+++ b/Engine/Plugins/Runtime/nDisplay/Source/DisplayCluster/Private/Game/DisplayClusterGameManager.h
@@ -42,7 +42,7 @@ public:
 	// IDisplayClusterGameManager
 	//////////////////////////////////////////////////////////////////////////////////////////////
 	virtual ADisplayClusterRootActor*     GetRootActor() const override;
-
+	virtual TArray<ADisplayClusterRootActor*> GetAllRootActorsFromWorld(UWorld* InWorld) override;
 	virtual UWorld* GetWorld() const override
 	{ return CurrentWorld; }
 
diff --git a/Engine/Plugins/Runtime/nDisplay/Source/DisplayCluster/Private/Game/EngineClasses/Basics/DisplayClusterViewportClient.cpp b/Engine/Plugins/Runtime/nDisplay/Source/DisplayCluster/Private/Game/EngineClasses/Basics/DisplayClusterViewportClient.cpp
--- a/Engine/Plugins/Runtime/nDisplay/Source/DisplayCluster/Private/Game/EngineClasses/Basics/DisplayClusterViewportClient.cpp
+++ b/Engine/Plugins/Runtime/nDisplay/Source/DisplayCluster/Private/Game/EngineClasses/Basics/DisplayClusterViewportClient.cpp
@@ -1029,19 +1029,18 @@ bool UDisplayClusterViewportClient::Draw_PIE(FViewport* InViewport, FCanvas* Sce
 		return false;
 	}
 
-	FDisplayClusterPreviewSettings PreviewSettings;
-	PreviewSettings.bPreviewEnablePostProcess = true;
-	PreviewSettings.bIsPIE = true;
+	// update current world
+	ViewportManager->GetConfiguration().SetCurrentWorld(MyWorld);
 
 	// Update local node viewports (update\create\delete) and build new render frame
-	const EDisplayClusterRenderFrameMode PreviewRenderMode = RootActor->GetPreviewRenderMode();
-	if (ViewportManager->UpdateConfiguration(PreviewRenderMode, LocalNodeId, RootActor, &PreviewSettings) == false)
+	const EDisplayClusterRenderFrameMode RenderModeForPIE = ViewportManager->GetConfiguration().GetRenderModeForPIE();
+	if (ViewportManager->GetConfiguration().UpdateConfigurationForClusterNode(RenderModeForPIE, LocalNodeId) == false)
 	{
 		return false;
 	}
 
 	FDisplayClusterRenderFrame RenderFrame;
-	if (ViewportManager->BeginNewFrame(InViewport, MyWorld, RenderFrame) == false)
+	if (ViewportManager->BeginNewFrame(InViewport, RenderFrame) == false)
 	{
 		return false;
 	}
@@ -1056,4 +1055,3 @@ bool UDisplayClusterViewportClient::Draw_PIE(FViewport* InViewport, FCanvas* Sce
 	return true;
 }
 #endif /*WITH_EDITOR*/
-
diff --git a/Engine/Plugins/Runtime/nDisplay/Source/DisplayCluster/Private/Game/EngineClasses/Scene/DisplayClusterCameraComponent.cpp b/Engine/Plugins/Runtime/nDisplay/Source/DisplayCluster/Private/Game/EngineClasses/Scene/DisplayClusterCameraComponent.cpp
--- a/Engine/Plugins/Runtime/nDisplay/Source/DisplayCluster/Private/Game/EngineClasses/Scene/DisplayClusterCameraComponent.cpp
+++ b/Engine/Plugins/Runtime/nDisplay/Source/DisplayCluster/Private/Game/EngineClasses/Scene/DisplayClusterCameraComponent.cpp
@@ -3,6 +3,8 @@
 #include "Components/DisplayClusterCameraComponent.h"
 #include "Components/BillboardComponent.h"
 
+#include "DisplayClusterRootActor.h"
+
 #include "Engine/CollisionProfile.h"
 #include "Engine/Texture2D.h"
 #include "UObject/ConstructorHelpers.h"
@@ -30,7 +32,7 @@ UDisplayClusterCameraComponent::UDisplayClusterCameraComponent(const FObjectInit
 #endif
 }
 
-void UDisplayClusterCameraComponent::GetDesiredView(FMinimalViewInfo& InOutViewInfo, float* OutCustomNearClippingPlane)
+void UDisplayClusterCameraComponent::GetDesiredView(IDisplayClusterViewportConfiguration& InViewportConfiguration, FMinimalViewInfo& InOutViewInfo, float* OutCustomNearClippingPlane)
 {
 	// Use this component as a camera
 	InOutViewInfo.Location = GetComponentLocation();
diff --git a/Engine/Plugins/Runtime/nDisplay/Source/DisplayCluster/Private/Game/EngineClasses/Scene/DisplayClusterICVFXCameraComponent.cpp b/Engine/Plugins/Runtime/nDisplay/Source/DisplayCluster/Private/Game/EngineClasses/Scene/DisplayClusterICVFXCameraComponent.cpp
--- a/Engine/Plugins/Runtime/nDisplay/Source/DisplayCluster/Private/Game/EngineClasses/Scene/DisplayClusterICVFXCameraComponent.cpp
+++ b/Engine/Plugins/Runtime/nDisplay/Source/DisplayCluster/Private/Game/EngineClasses/Scene/DisplayClusterICVFXCameraComponent.cpp
@@ -112,49 +112,6 @@ FString UDisplayClusterICVFXCameraComponent::GetCameraUniqueId() const
 	return GetFName().ToString();
 }
 
-bool UDisplayClusterICVFXCameraComponent::IsICVFXEnabled() const
-{
-	// When rendering offscreen, we have an extended logic for camera rendering activation
-	static const bool bIsRunningClusterModeOffscreen =
-		(GDisplayCluster->GetOperationMode() == EDisplayClusterOperationMode::Cluster) &&
-		FParse::Param(FCommandLine::Get(), TEXT("RenderOffscreen"));
-
-	// If cluster mode + rendering offscreen, discover media output settings
-	if (bIsRunningClusterModeOffscreen)
-	{
-		// This cluster node ID
-		static const FString NodeId = GDisplayCluster->GetPrivateClusterMgr()->GetNodeId();
-
-		// First condition to render offscreen: it has media output assigned
-		const bool bUsesMediaOutput = (CameraSettings.RenderSettings.Media.bEnable && CameraSettings.RenderSettings.Media.IsMediaOutputAssigned(NodeId));
-
-		// Get backbuffer media settings
-		const FDisplayClusterConfigurationMedia* BackbufferMediaSettings = nullptr;
-		if (const ADisplayClusterRootActor* const RootActor = Cast<ADisplayClusterRootActor>(GetOwner()))
-		{
-			if (const UDisplayClusterConfigurationData* const ConfigData = RootActor->GetConfigData())
-			{
-				if (const UDisplayClusterConfigurationClusterNode* const NodeCfg = ConfigData->Cluster->GetNode(NodeId))
-				{
-					BackbufferMediaSettings = &NodeCfg->Media;
-				}
-			}
-		}
-
-		// Second condition to render offscreen: the backbuffer has media output assigned.
-		// This means the whole frame including ICVFX cameras need to be rendered.
-		const bool bIsBackbufferBeingCaptured = BackbufferMediaSettings ?
-			BackbufferMediaSettings->bEnable && BackbufferMediaSettings->IsMediaOutputAssigned() :
-			false;
-
-		// Finally make a decision if the camera should be rendered
-		return CameraSettings.bEnable && (bUsesMediaOutput || bIsBackbufferBeingCaptured);
-	}
-
-	// Otherwise the on/off condition only
-	return CameraSettings.bEnable;
-}
-
 #if WITH_EDITOR
 bool UDisplayClusterICVFXCameraComponent::GetEditorPreviewInfo(float DeltaTime, FMinimalViewInfo& ViewOut)
 {
diff --git a/Engine/Plugins/Runtime/nDisplay/Source/DisplayCluster/Private/Game/EngineClasses/Scene/DisplayClusterPreviewComponent.cpp b/Engine/Plugins/Runtime/nDisplay/Source/DisplayCluster/Private/Game/EngineClasses/Scene/DisplayClusterPreviewComponent.cpp
--- a/Engine/Plugins/Runtime/nDisplay/Source/DisplayCluster/Private/Game/EngineClasses/Scene/DisplayClusterPreviewComponent.cpp
+++ b/Engine/Plugins/Runtime/nDisplay/Source/DisplayCluster/Private/Game/EngineClasses/Scene/DisplayClusterPreviewComponent.cpp
@@ -17,7 +17,6 @@
 
 #include "Render/Viewport/DisplayClusterViewport.h"
 #include "Render/Viewport/DisplayClusterViewportManager.h"
-#include "Render/Viewport/Containers/DisplayClusterViewportReadPixels.h"
 #include "Render/Viewport/RenderFrame/DisplayClusterRenderFrameSettings.h"
 #include "Render/Viewport/DisplayClusterViewportHelpers.h"
 #include "Render/Viewport/Resource/DisplayClusterViewportResourceSettings.h"
@@ -221,6 +220,21 @@ void UDisplayClusterPreviewComponent::UpdatePreviewMeshReference()
 		// TODO: See if we can remove this hack
 		ReleasePreviewMesh();
 	}
+
+	// And search for new mesh reference
+	if (IDisplayClusterViewport* Viewport = GetCurrentViewport())
+	{
+		// Handle movable preview mesh:
+		if (!Viewport->GetProjectionPolicy()->HasPreviewMovableMesh(Viewport))
+		{
+			ReleaseMovablePreviewMesh();
+		}
+
+		if (!Viewport->GetProjectionPolicy()->HasPreviewMesh(Viewport))
+		{
+			ReleasePreviewMesh();
+		}
+	}
 }
 
 bool UDisplayClusterPreviewComponent::UpdatePreviewMesh()
@@ -248,7 +262,7 @@ bool UDisplayClusterPreviewComponent::UpdatePreviewMesh()
 		if (bOutputToPreviewMesh)
 		{
 			// Handle preview mesh:
-			if (Viewport->GetProjectionPolicy()->HasPreviewMesh())
+			if (Viewport->GetProjectionPolicy()->HasPreviewMesh(Viewport))
 			{
 				// create warp mesh or update changes
 				if (PreviewMesh != nullptr)
@@ -265,22 +279,16 @@ bool UDisplayClusterPreviewComponent::UpdatePreviewMesh()
 					// Get new mesh ptr
 					PreviewMesh = Viewport->GetProjectionPolicy()->GetOrCreatePreviewMeshComponent(Viewport, bIsRootActorPreviewMesh);
 
-					if (Viewport->GetProjectionPolicy()->HasPreviewMovableMesh())
-					{
-						// Get new movable mesh ptr
-						MovablePreviewMesh = Viewport->GetProjectionPolicy()->GetOrCreatePreviewMovableMeshComponent(Viewport);
-
-						if (MovablePreviewMesh)
-						{
-							// Make the movable preview mesh invisible by default
-							MovablePreviewMesh->SetVisibility(false);
-						}
-					}
-
 					// Update saved proj policy parameters
 					WarpMeshSavedProjectionPolicy = ViewportConfig->ProjectionPolicy;
 				}
 
+				if (MovablePreviewMesh == nullptr)
+				{
+					// Get new movable mesh ptr
+					MovablePreviewMesh = Viewport->GetProjectionPolicy()->GetOrCreatePreviewMovableMeshComponent(Viewport);
+				}
+
 				// disable shadow rendering for preview meshes
 				if (PreviewMesh != nullptr)
 				{
@@ -325,14 +333,27 @@ bool UDisplayClusterPreviewComponent::UpdatePreviewMesh()
 	return false;
 }
 
+void UDisplayClusterPreviewComponent::ReleaseMovablePreviewMesh()
+{
+	// Release the movable mesh component of the preview
+	if (MovablePreviewMesh)
+	{
+		MovablePreviewMesh->UnregisterComponent();
+		MovablePreviewMesh->DestroyComponent();
+	}
+}
+
 void UDisplayClusterPreviewComponent::ReleasePreviewMesh()
 {
+	// Release the mesh component of the preview
 	if (!bIsRootActorPreviewMesh && PreviewMesh)
 	{
 		PreviewMesh->UnregisterComponent();
 		PreviewMesh->DestroyComponent();
 	}
 
+	ReleaseMovablePreviewMesh();
+	
 	// Forget old mesh with material
 	PreviewMesh = nullptr;
 	MovablePreviewMesh = nullptr;
@@ -555,12 +576,12 @@ bool UDisplayClusterPreviewComponent::GetPreviewTextureSettings(FIntPoint& OutSi
 {
 	if (FDisplayClusterViewport* Viewport = static_cast<FDisplayClusterViewport*>(GetCurrentViewport()))
 	{
-		if (FDisplayClusterViewportManager* ViewportManager = Viewport->GetViewportManagerImpl())
+		if (FDisplayClusterViewportManager* ViewportManager = Viewport->Configuration->GetViewportManagerImpl())
 		{
 			// The viewport size is already capped for RenderSettings
 			if (!Viewport->GetContexts().IsEmpty())
 			{
-				FDisplayClusterViewportResourceSettings DefaultResourceSettings(ViewportManager->GetRenderFrameSettings(), nullptr);
+				FDisplayClusterViewportResourceSettings DefaultResourceSettings(ViewportManager->Configuration->GetRenderFrameSettings(), nullptr);
 
 				OutSize = Viewport->GetContexts()[0].FrameTargetRect.Size();
 				OutTextureFormat = DefaultResourceSettings.GetFormat();
diff --git a/Engine/Plugins/Runtime/nDisplay/Source/DisplayCluster/Private/Game/EngineClasses/Scene/DisplayClusterRootActor.cpp b/Engine/Plugins/Runtime/nDisplay/Source/DisplayCluster/Private/Game/EngineClasses/Scene/DisplayClusterRootActor.cpp
--- a/Engine/Plugins/Runtime/nDisplay/Source/DisplayCluster/Private/Game/EngineClasses/Scene/DisplayClusterRootActor.cpp
+++ b/Engine/Plugins/Runtime/nDisplay/Source/DisplayCluster/Private/Game/EngineClasses/Scene/DisplayClusterRootActor.cpp
@@ -59,7 +59,9 @@
 #include "AssetToolsModule.h"
 #endif
 
-
+//////////////////////////////////////////////////////////////////////////////////////////////
+// ADisplayClusterRootActor
+//////////////////////////////////////////////////////////////////////////////////////////////
 ADisplayClusterRootActor::ADisplayClusterRootActor(const FObjectInitializer& ObjectInitializer)
 	: Super(ObjectInitializer)
 	, OperationMode(EDisplayClusterOperationMode::Disabled)
@@ -125,13 +127,23 @@ IDisplayClusterViewportManager* ADisplayClusterRootActor::GetViewportManager() c
 	return ViewportManager.Get();
 }
 
+IDisplayClusterViewportConfiguration* ADisplayClusterRootActor::GetViewportConfiguration() const
+{
+	return ViewportManager.IsValid() ? &ViewportManager->GetConfiguration() : nullptr;
+}
+
 void ADisplayClusterRootActor::CreateViewportManagerImpl()
 {
 	if (!ViewportManager.IsValid())
 	{
 		ViewportManager = MakeShared<FDisplayClusterViewportManager, ESPMode::ThreadSafe>();
+
+		// After the constructor, we should always call this function to initialize internal references.
 		ViewportManager->Initialize();
 
+		// Set the owner's DCRA to the newly created viewport manager.
+		ViewportManager->GetConfiguration().SetRootActor(this, EDisplayClusterRootActorType::Any);
+
 		// Preview rendering depends on the DC VM
 		FDisplayClusterRootActorPreviewRenderingManager::HandleEvent(EDisplayClusterRootActorPreviewEvent::Create, this);
 	}
@@ -175,35 +187,6 @@ const FDisplayClusterConfigurationRenderFrame& ADisplayClusterRootActor::GetRend
 	return CurrentConfigData->RenderFrameSettings;
 }
 
-EDisplayClusterRenderFrameMode ADisplayClusterRootActor::GetRenderMode() const
-{
-	if (ViewportManager.IsValid())
-	{
-		return ViewportManager->GetRenderMode();
-	}
-
-	return EDisplayClusterRenderFrameMode::Unknown;
-}
-
-EDisplayClusterRenderFrameMode ADisplayClusterRootActor::GetPreviewRenderMode() const
-{
-#if WITH_EDITOR
-	switch (RenderMode)
-	{
-	case EDisplayClusterConfigurationRenderMode::SideBySide:
-		return EDisplayClusterRenderFrameMode::SideBySide;
-
-	case EDisplayClusterConfigurationRenderMode::TopBottom:
-		return EDisplayClusterRenderFrameMode::TopBottom;
-
-	default:
-		break;
-	}
-#endif
-
-	return EDisplayClusterRenderFrameMode::Mono;
-}
-
 void ADisplayClusterRootActor::InitializeFromConfig(UDisplayClusterConfigurationData* ConfigData)
 {
 	if (ConfigData)
diff --git a/Engine/Plugins/Runtime/nDisplay/Source/DisplayCluster/Private/Game/EngineClasses/Scene/DisplayClusterRootActorEditor.cpp b/Engine/Plugins/Runtime/nDisplay/Source/DisplayCluster/Private/Game/EngineClasses/Scene/DisplayClusterRootActorEditor.cpp
--- a/Engine/Plugins/Runtime/nDisplay/Source/DisplayCluster/Private/Game/EngineClasses/Scene/DisplayClusterRootActorEditor.cpp
+++ b/Engine/Plugins/Runtime/nDisplay/Source/DisplayCluster/Private/Game/EngineClasses/Scene/DisplayClusterRootActorEditor.cpp
@@ -60,31 +60,6 @@
 #include "EditorSupportDelegates.h"
 #include "LevelEditor.h"
 
-//////////////////////////////////////////////////////////////////////////////////////////////
-int32 GDisplayClusterPreviewAllowMultiGPURendering = 0;
-static FAutoConsoleVariableRef CVarDisplayClusterPreviewAllowMultiGPURendering(
-	TEXT("DC.Preview.AllowMultiGPURendering"),
-	GDisplayClusterPreviewAllowMultiGPURendering,
-	TEXT("Allow mGPU for preview rendering (0 == disabled)"),
-	ECVF_RenderThreadSafe
-);
-
-int32 GDisplayClusterPreviewMultiGPURenderingMinIndex = 0;
-static FAutoConsoleVariableRef CVarDisplayClusterPreviewMultiGPURenderingMinIndex(
-	TEXT("DC.Preview.MultiGPURenderingMinIndex"),
-	GDisplayClusterPreviewMultiGPURenderingMinIndex,
-	TEXT("Distribute mGPU render on GPU from #min to #max indices"),
-	ECVF_RenderThreadSafe
-);
-
-int32 GDisplayClusterPreviewMultiGPURenderingMaxIndex = 0;
-static FAutoConsoleVariableRef CVarDisplayClusterPreviewMultiGPURenderingMaxIndex(
-	TEXT("DC.Preview.MultiGPURenderingMaxIndex"),
-	GDisplayClusterPreviewMultiGPURenderingMaxIndex,
-	TEXT("Distribute mGPU render on GPU from #min to #max indices"),
-	ECVF_RenderThreadSafe
-);
-
 //////////////////////////////////////////////////////////////////////////////////////////////
 // ADisplayClusterRootActor
 //////////////////////////////////////////////////////////////////////////////////////////////
@@ -184,7 +159,7 @@ void ADisplayClusterRootActor::RenderPreview_Editor()
 		ResetPreviewInternals_Editor();
 		if (ViewportManager.IsValid())
 		{
-			if(ViewportManager->GetRenderMode() == EDisplayClusterRenderFrameMode::PreviewInScene)
+			if(ViewportManager->GetConfiguration().IsPreviewRendering())
 			{
 				// Release viewport manager with resources immediatelly for preview in scene
 				RemoveViewportManagerImpl();
@@ -386,6 +361,13 @@ IDisplayClusterViewport* ADisplayClusterRootActor::FindPreviewViewport(const FSt
 	return nullptr;
 }
 
+EDisplayClusterRenderFrameMode ADisplayClusterRootActor::GetPreviewRenderMode_Editor() const
+{
+	// Todo: we can add HitProxy rendering support to this function with EDisplayClusterRenderFrameMode::PreviewProxyHitInScene
+	
+	return EDisplayClusterRenderFrameMode::PreviewInScene;
+}
+
 bool ADisplayClusterRootActor::ImplUpdatePreviewConfiguration_Editor(const FString& InClusterNodeId)
 {
 	DECLARE_SCOPE_CYCLE_COUNTER(TEXT("ADisplayClusterRootActor::ImplUpdatePreviewConfiguration_Editor"), STAT_ImplUpdatePreviewConfiguration_Editor, STATGROUP_NDisplay);
@@ -399,19 +381,12 @@ bool ADisplayClusterRootActor::ImplUpdatePreviewConfiguration_Editor(const FStri
 		// Now we render this node
 		PreviewRenderFrameClusterNodeId = InClusterNodeId;
 
-		FDisplayClusterPreviewSettings PreviewSettings;
-		PreviewSettings.PreviewRenderTargetRatioMult = PreviewRenderTargetRatioMult;
-
-		PreviewSettings.bFreezePreviewRender = bFreezePreviewRender;
-		PreviewSettings.bPreviewEnablePostProcess = ShouldThisFrameOutputPreviewToPostProcessRenderTarget();
-
-		PreviewSettings.PreviewMaxTextureDimension = PreviewMaxTextureDimension;
-		
-		PreviewSettings.bAllowMultiGPURenderingInEditor = GDisplayClusterPreviewAllowMultiGPURendering != 0;
-		PreviewSettings.MinGPUIndex = FMath::Max(0, GDisplayClusterPreviewMultiGPURenderingMinIndex);
-		PreviewSettings.MaxGPUIndex = FMath::Max(0, GDisplayClusterPreviewMultiGPURenderingMaxIndex);
+		// Update current world from scene DCRA
+		ADisplayClusterRootActor* SceneRootActor = ViewportManager->GetConfiguration().GetRootActor(EDisplayClusterRootActorType::Scene);
+		ViewportManager->GetConfiguration().SetCurrentWorld(SceneRootActor ? SceneRootActor->GetWorld() : GetWorld());
 
-		return ViewportManager->UpdateConfiguration(EDisplayClusterRenderFrameMode::PreviewInScene, InClusterNodeId, this, &PreviewSettings);
+		// Update local node viewports (update\create\delete) and build new render frame
+		return ViewportManager->GetConfiguration().UpdateConfigurationForClusterNode(GetPreviewRenderMode_Editor(), InClusterNodeId);
 	}
 
 	return false;
@@ -445,11 +420,7 @@ bool ADisplayClusterRootActor::ImplUpdatePreviewRenderFrame_Editor(const FString
 
 	if (!PreviewRenderFrame.IsValid())
 	{
-		// Now always use RootActor world to preview.
-		UWorld* CurrentPreviewWorld = GetWorld();
-
-		// Begin render new frame for cluster node
-		if (CurrentPreviewWorld == nullptr || !ImplUpdatePreviewConfiguration_Editor(InClusterNodeId))
+		if (!ImplUpdatePreviewConfiguration_Editor(InClusterNodeId))
 		{
 			return false;
 		}
@@ -466,7 +437,7 @@ bool ADisplayClusterRootActor::ImplUpdatePreviewRenderFrame_Editor(const FString
 		PreviewRenderFrame = MakeUnique<FDisplayClusterRenderFrame>();
 
 		// Update preview viewports from settings
-		if (!ViewportManager->BeginNewFrame(nullptr, CurrentPreviewWorld, *PreviewRenderFrame))
+		if (!ViewportManager->BeginNewFrame(nullptr, *PreviewRenderFrame))
 		{
 			PreviewRenderFrame.Reset();
 
@@ -632,22 +603,12 @@ void ADisplayClusterRootActor::ImplRenderPreviewFrustums_Editor()
 	// collect incameras
 	if (bPreviewICVFXFrustums)
 	{
-		for (UActorComponent* ActorComponentIt : GetComponents())
+		// Iterate over rendered inner camera viewports (whole cluster)
+		for (const TSharedPtr<FDisplayClusterViewport, ESPMode::ThreadSafe>& InnerCameraViewportIt : FDisplayClusterViewportConfigurationHelpers_ICVFX::GetAllVisibleInnerCameraViewports(*ViewportManager->Configuration))
 		{
-			if (ActorComponentIt)
+			if (InnerCameraViewportIt.IsValid())
 			{
-				UDisplayClusterICVFXCameraComponent* CineCameraComponent = Cast<UDisplayClusterICVFXCameraComponent>(ActorComponentIt);
-				if (CineCameraComponent && CineCameraComponent->IsICVFXEnabled())
-				{
-					// Iterate over rendered incamera viewports (whole cluster)
-					for (const TSharedPtr<FDisplayClusterViewport, ESPMode::ThreadSafe>& InCameraViewportIt : FDisplayClusterViewportConfigurationHelpers_ICVFX::PreviewGetRenderedInCameraViewports(*this, *CineCameraComponent))
-					{
-						if (InCameraViewportIt.IsValid())
-						{
-							FrustumPreviewViewports.Add(InCameraViewportIt.Get());
-						}
-					}
-				}
+				FrustumPreviewViewports.Add(InnerCameraViewportIt.Get());
 			}
 		}
 	}
diff --git a/Engine/Plugins/Runtime/nDisplay/Source/DisplayCluster/Private/Render/Device/DisplayClusterDeviceBase.cpp b/Engine/Plugins/Runtime/nDisplay/Source/DisplayCluster/Private/Render/Device/DisplayClusterDeviceBase.cpp
--- a/Engine/Plugins/Runtime/nDisplay/Source/DisplayCluster/Private/Render/Device/DisplayClusterDeviceBase.cpp
+++ b/Engine/Plugins/Runtime/nDisplay/Source/DisplayCluster/Private/Render/Device/DisplayClusterDeviceBase.cpp
@@ -42,6 +42,17 @@
 
 #include <utility>
 
+namespace UE::DisplayCluster::DeviceBaseHelpers
+{
+	static inline IDisplayCluster& GetDisplayClusterAPI()
+	{
+		static IDisplayCluster& DisplayClusterAPISingleton = IDisplayCluster::Get();
+
+		return DisplayClusterAPISingleton;
+	}
+};
+using namespace UE::DisplayCluster;
+
 // Enable/Disable ClearTexture for RTT after resolving to the backbuffer
 static TAutoConsoleVariable<int32> CVarClearTextureEnabled(
 	TEXT("nDisplay.render.ClearTextureEnabled"),
@@ -214,7 +225,7 @@ void FDisplayClusterDeviceBase::AdjustViewRect(int32 ViewIndex, int32& X, int32&
 	check(IsInGameThread());
 
 	IDisplayClusterViewportManager* ViewportManager = GetViewportManager();
-	if (ViewportManager == nullptr || ViewportManager->IsSceneOpened() == false)
+	if (ViewportManager == nullptr || ViewportManager->GetConfiguration().IsSceneOpened() == false)
 	{
 		return;
 	}
@@ -266,7 +277,7 @@ void FDisplayClusterDeviceBase::CalculateStereoViewOffset(const int32 ViewIndex,
 	{
 #if WITH_EDITOR
 		// Hide spam in logs when configuring VP in editor [UE-114493]
-		static const bool bIsEditorOperationMode = IDisplayCluster::Get().GetOperationMode() == EDisplayClusterOperationMode::Editor;
+		static const bool bIsEditorOperationMode = DeviceBaseHelpers::GetDisplayClusterAPI().GetOperationMode() == EDisplayClusterOperationMode::Editor;
 		if (!bIsEditorOperationMode)
 #endif
 		{
@@ -288,7 +299,7 @@ FMatrix FDisplayClusterDeviceBase::GetStereoProjectionMatrix(const int32 ViewInd
 	// ViewIndex == eSSE_MONOSCOPIC(-1) is a special case called for ISR culling math.
 	// Since nDisplay is not ISR compatible, we ignore this request. This won't be neccessary once
 	// we stop using nDisplay as a stereoscopic rendering device (IStereoRendering).
-	if (ViewportManager && ViewportManager->IsSceneOpened() && ViewIndex >= 0)
+	if (ViewportManager && ViewportManager->GetConfiguration().IsSceneOpened() && ViewIndex >= 0)
 	{
 		uint32 ViewportContextNum = 0;
 		IDisplayClusterViewport* ViewportPtr = ViewportManager->FindViewport(ViewIndex, &ViewportContextNum);
@@ -311,17 +322,19 @@ bool FDisplayClusterDeviceBase::BeginNewFrame(FViewport* InViewport, UWorld* InW
 	check(IsInGameThread());
 	check(InViewport);
 
-	IDisplayCluster& DisplayCluster = IDisplayCluster::Get();
-	if (ADisplayClusterRootActor* RootActor = DisplayCluster.GetGameMgr()->GetRootActor())
+	if (ADisplayClusterRootActor* RootActor = DeviceBaseHelpers::GetDisplayClusterAPI().GetGameMgr()->GetRootActor())
 	{
 		if (IDisplayClusterViewportManager* ViewportManagerPtr = RootActor->GetViewportManager())
 		{
-			const FString LocalNodeId = DisplayCluster.GetConfigMgr()->GetLocalNodeId();
+			const FString LocalNodeId = DeviceBaseHelpers::GetDisplayClusterAPI().GetConfigMgr()->GetLocalNodeId();
+
+			// update current world
+			ViewportManagerPtr->GetConfiguration().SetCurrentWorld(InWorld);
 
 			// Update local node viewports (update\create\delete) and build new render frame
-			if (ViewportManagerPtr->UpdateConfiguration(RenderFrameMode, LocalNodeId, RootActor))
+			if (ViewportManagerPtr->GetConfiguration().UpdateConfigurationForClusterNode(RenderFrameMode, LocalNodeId))
 			{
-				if (ViewportManagerPtr->BeginNewFrame(InViewport, InWorld, OutRenderFrame))
+				if (ViewportManagerPtr->BeginNewFrame(InViewport, OutRenderFrame))
 				{
 					// update total number of views for this frame (in multiple families)
 					DesiredNumberOfViews = OutRenderFrame.DesiredNumberOfViews;
@@ -337,8 +350,7 @@ bool FDisplayClusterDeviceBase::BeginNewFrame(FViewport* InViewport, UWorld* InW
 
 void FDisplayClusterDeviceBase::InitializeNewFrame()
 {
-	IDisplayCluster& DisplayCluster = IDisplayCluster::Get();
-	if (ADisplayClusterRootActor* RootActor = DisplayCluster.GetGameMgr()->GetRootActor())
+	if (ADisplayClusterRootActor* RootActor = DeviceBaseHelpers::GetDisplayClusterAPI().GetGameMgr()->GetRootActor())
 	{
 		if (IDisplayClusterViewportManager* ViewportManager = RootActor->GetViewportManager())
 		{
diff --git a/Engine/Plugins/Runtime/nDisplay/Source/DisplayCluster/Private/Render/Viewport/Configuration/DisplayClusterViewportConfiguration.cpp b/Engine/Plugins/Runtime/nDisplay/Source/DisplayCluster/Private/Render/Viewport/Configuration/DisplayClusterViewportConfiguration.cpp
--- a/Engine/Plugins/Runtime/nDisplay/Source/DisplayCluster/Private/Render/Viewport/Configuration/DisplayClusterViewportConfiguration.cpp
+++ b/Engine/Plugins/Runtime/nDisplay/Source/DisplayCluster/Private/Render/Viewport/Configuration/DisplayClusterViewportConfiguration.cpp
@@ -2,15 +2,22 @@
 
 #include "DisplayClusterViewportConfiguration.h"
 
+#include "DisplayClusterViewportConfiguration_Viewport.h"
+#include "DisplayClusterViewportConfiguration_ViewportManager.h"
+#include "DisplayClusterViewportConfiguration_Postprocess.h"
+#include "DisplayClusterViewportConfiguration_ProjectionPolicy.h"
+#include "DisplayClusterViewportConfiguration_ICVFX.h"
+#include "DisplayClusterViewportConfiguration_ICVFXCamera.h"
+
+#include "DisplayClusterViewportConfigurationHelpers_RenderFrameSettings.h"
+
 #include "DisplayClusterConfigurationStrings.h"
-#include "DisplayClusterViewportConfigurationBase.h"
-#include "DisplayClusterViewportConfigurationICVFX.h"
-#include "DisplayClusterViewportConfigurationProjectionPolicy.h"
 
 #include "DisplayClusterRootActor.h"
 
 #include "Render/Viewport/DisplayClusterViewport.h"
 #include "Render/Viewport/DisplayClusterViewportManager.h"
+#include "Render/Viewport/DisplayClusterViewportManagerProxy.h"
 #include "Render/Viewport/Postprocess/DisplayClusterViewportPostProcessManager.h"
 
 #include "Render/Viewport/RenderFrame/DisplayClusterRenderFrameSettings.h"
@@ -19,368 +26,324 @@
 #include "DisplayClusterConfigurationTypes_Viewport.h"
 
 #include "Render/IPDisplayClusterRenderManager.h"
-#include "Misc/DisplayClusterGlobals.h"
 #include "Misc/DisplayClusterLog.h"
 
 #include "Engine/RendererSettings.h"
 
-
-///////////////////////////////////////////////////////////////////
-int32 GDisplayClusterCrossGPUTransferEnable = 0;
-static FAutoConsoleVariableRef CDisplayClusterCrossGPUTransferEnable(
-	TEXT("nDisplay.render.CrossGPUTransfer.Enable"),
-	GDisplayClusterCrossGPUTransferEnable,
-	TEXT("Enable cross-GPU transfers using nDisplay implementation (0 - disable, default) \n")
-	TEXT("That replaces the default cross-GPU transfers using UE Core for the nDisplay viewports viewfamilies.\n"),
-	ECVF_RenderThreadSafe
-);
-
-int32 GDisplayClusterCrossGPUTransferLockSteps = 1;
-static FAutoConsoleVariableRef CDisplayClusterCrossGPUTransferLockSteps(
-	TEXT("nDisplay.render.CrossGPUTransfer.LockSteps"),
-	GDisplayClusterCrossGPUTransferLockSteps,
-	TEXT("The bLockSteps parameter is simply passed to the FTransferResourceParams structure. (0 - disable)\n")
-	TEXT("Whether the GPUs must handshake before and after the transfer. Required if the texture rect is being written to in several render passes.\n")
-	TEXT("Otherwise, minimal synchronization will be used.\n"),
-	ECVF_RenderThreadSafe
-);
-
-int32 GDisplayClusterCrossGPUTransferPullData = 1;
-static FAutoConsoleVariableRef CVarDisplayClusterCrossGPUTransferPullData(
-	TEXT("nDisplay.render.CrossGPUTransfer.PullData"),
-	GDisplayClusterCrossGPUTransferPullData,
-	TEXT("The bPullData parameter is simply passed to the FTransferResourceParams structure. (0 - disable)\n")
-	TEXT("Whether the data is read by the dest GPU, or written by the src GPU (not allowed if the texture is a backbuffer)\n"),
-	ECVF_RenderThreadSafe
-);
-
-/**
-* This enum is for CVar values only and is used to process the logic that converts the values to the runtime enum in GetAlphaChannelCaptureMode().
-*/
-enum class ECVarDisplayClusterAlphaChannelCaptureMode : uint8
+namespace UE::DisplayCluster::ConfigurationHelpers
 {
-	/** [Disabled]
-	 * Disable alpha channel saving.
-	 */
-	Disabled,
-
-	/** [ThroughTonemapper]
-	 * When rendering with the PropagateAlpha experimental mode turned on, the alpha channel is forwarded to post-processes.
-	 * In this case, the alpha channel is anti-aliased along with the color.
-	 * Since some post-processing may change the alpha, it is copied at the beginning of the PP and restored after all post-processing is completed.
-	 */
-	ThroughTonemapper,
-
-	/** [FXAA]
-	 * Otherwise, if the PropagateAlpha mode is disabled in the project settings, we need to save the alpha before it becomes invalid.
-	 * The alpha is valid until the scene color is resolved (on the ResolvedSceneColor callback). Therefore, it is copied to a temporary resource on this cb.
-	 * Since we need to remove AA jittering (because alpha copied before AA), anti-aliasing is turned off for this viewport.
-	 * And finally the FXAA is used for smoothing.
-	 */
-	FXAA,
-
-	/** [Copy]
-	 * Disable AA and TAA, Copy alpha from scenecolor texture to final.
-	 * These experimental (temporary) modes for the performance tests.
-	 */
-	 Copy,
-
-	/** [CopyAA]
-	 * Use AA, disable TAA, Copy alpha from scenecolor texture to final.
-	 * These experimental (temporary) modes for the performance tests.
-	 */
-	 CopyAA,
-
-	COUNT
-};
+	static inline ADisplayClusterRootActor* ImplGetRootActor(const FDisplayClusterActorRef& InConfigurationRootActorRef)
+	{
+		if (AActor* ActorPtr = InConfigurationRootActorRef.GetOrFindSceneActor())
+		{
+			if (ActorPtr->IsA<ADisplayClusterRootActor>())
+			{
+				return static_cast<ADisplayClusterRootActor*>(ActorPtr);
+			}
+		}
 
-/**
- *Choose method to preserve alpha channel
- */
-int32 GDisplayClusterAlphaChannelCaptureMode = (uint8)ECVarDisplayClusterAlphaChannelCaptureMode::FXAA;
-static FAutoConsoleVariableRef CVarDisplayClusterAlphaChannelCaptureMode(
-	TEXT("nDisplay.render.AlphaChannelCaptureMode"),
-	GDisplayClusterAlphaChannelCaptureMode,
-	TEXT("Alpha channel capture mode (FXAA - default)\n")
-	TEXT("0 - Disabled\n")
-	TEXT("1 - ThroughTonemapper\n")
-	TEXT("2 - FXAA\n")
-	TEXT("3 - Copy [experimental]\n")
-	TEXT("4 - CopyAA [experimental]\n"),
-	ECVF_RenderThreadSafe
-);
+		return nullptr;
+	}
+
+	static inline bool ImplSetRootActor(ADisplayClusterRootActor* InRootActor, FDisplayClusterActorRef& InOutConfigurationRootActorRef)
+	{
+		if (ImplGetRootActor(InOutConfigurationRootActorRef) != InRootActor)
+		{
+			InOutConfigurationRootActorRef.SetSceneActor(InRootActor);
+
+			return true;
+		}
+
+		return false;
+	}
+};
+using namespace UE::DisplayCluster;
 
 ///////////////////////////////////////////////////////////////////
 // FDisplayClusterViewportConfiguration
 ///////////////////////////////////////////////////////////////////
-FDisplayClusterViewportConfiguration::FDisplayClusterViewportConfiguration(FDisplayClusterViewportManager& InViewportManager)
-	: ViewportManagerWeakPtr(InViewportManager.AsShared())
+FDisplayClusterViewportConfiguration::FDisplayClusterViewportConfiguration()
+	: Proxy(MakeShared<FDisplayClusterViewportConfigurationProxy, ESPMode::ThreadSafe>())
 { }
 
 FDisplayClusterViewportConfiguration::~FDisplayClusterViewportConfiguration()
 { }
 
-EDisplayClusterRenderFrameAlphaChannelCaptureMode FDisplayClusterViewportConfiguration::GetAlphaChannelCaptureMode() const
+void FDisplayClusterViewportConfiguration::Initialize(FDisplayClusterViewportManager& InViewportManager)
 {
-	ECVarDisplayClusterAlphaChannelCaptureMode AlphaChannelCaptureMode = (ECVarDisplayClusterAlphaChannelCaptureMode)FMath::Clamp(GDisplayClusterAlphaChannelCaptureMode, 0, (int32)ECVarDisplayClusterAlphaChannelCaptureMode::COUNT - 1);
+	// Set weak refs to viewport manager and proxy
+	ViewportManagerWeakPtr = InViewportManager.AsShared();
+	Proxy->Initialize_GameThread(InViewportManager.GetViewportManagerProxy());
+}
 
-	static const auto CVarPropagateAlpha = IConsoleManager::Get().FindTConsoleVariableDataInt(TEXT("r.PostProcessing.PropagateAlpha"));
-	const EAlphaChannelMode::Type PropagateAlpha = EAlphaChannelMode::FromInt(CVarPropagateAlpha->GetValueOnGameThread());
-	const bool bAllowThroughTonemapper = PropagateAlpha == EAlphaChannelMode::AllowThroughTonemapper;
+void FDisplayClusterViewportConfiguration::SetRootActor(ADisplayClusterRootActor* InRootActor, const EDisplayClusterRootActorType InRootActorType)
+{
+	check(IsInGameThread());
 
-	switch (AlphaChannelCaptureMode)
+	switch (InRootActorType)
 	{
-	case ECVarDisplayClusterAlphaChannelCaptureMode::ThroughTonemapper:
-		// Disable alpha capture if PropagateAlpha not valid
-		return bAllowThroughTonemapper ? EDisplayClusterRenderFrameAlphaChannelCaptureMode::ThroughTonemapper : EDisplayClusterRenderFrameAlphaChannelCaptureMode::None;
+	case EDisplayClusterRootActorType::Preview:
+		bCurrentSceneNeedsToBeUpdated |= ConfigurationHelpers::ImplSetRootActor(InRootActor, PreviewRootActorRef);
+		break;
 
-	case ECVarDisplayClusterAlphaChannelCaptureMode::FXAA:
-		return EDisplayClusterRenderFrameAlphaChannelCaptureMode::FXAA;
+	case EDisplayClusterRootActorType::Scene:
+		bCurrentSceneNeedsToBeUpdated |= ConfigurationHelpers::ImplSetRootActor(InRootActor, SceneRootActorRef);
+		break;
 
-	case ECVarDisplayClusterAlphaChannelCaptureMode::Copy:
-		return EDisplayClusterRenderFrameAlphaChannelCaptureMode::Copy;
+	case EDisplayClusterRootActorType::Configuration:
+		ConfigurationHelpers::ImplSetRootActor(InRootActor, ConfigurationRootActorRef);
+		break;
 
-	case ECVarDisplayClusterAlphaChannelCaptureMode::CopyAA:
-		return EDisplayClusterRenderFrameAlphaChannelCaptureMode::CopyAA;
+	case EDisplayClusterRootActorType::Any:
+		bCurrentSceneNeedsToBeUpdated |= ConfigurationHelpers::ImplSetRootActor(InRootActor, PreviewRootActorRef);
+		bCurrentSceneNeedsToBeUpdated |= ConfigurationHelpers::ImplSetRootActor(InRootActor, SceneRootActorRef);
+
+		ConfigurationHelpers::ImplSetRootActor(InRootActor, ConfigurationRootActorRef);
+		break;
 
-	case ECVarDisplayClusterAlphaChannelCaptureMode::Disabled:
 	default:
 		break;
 	}
-
-	return EDisplayClusterRenderFrameAlphaChannelCaptureMode::None;
 }
 
-bool FDisplayClusterViewportConfiguration::SetRootActor(ADisplayClusterRootActor* InRootActorPtr)
+ADisplayClusterRootActor* FDisplayClusterViewportConfiguration::GetRootActor(const EDisplayClusterRootActorType InRootActorType) const
 {
-	check(IsInGameThread());
-	check(InRootActorPtr);
-
-	if (!RootActorRef.IsDefinedSceneActor() || GetRootActor() != InRootActorPtr)
+	switch (InRootActorType)
 	{
-		// Update root actor reference:
-		RootActorRef.ResetSceneActor();
-		RootActorRef.SetSceneActor(InRootActorPtr);
+	case EDisplayClusterRootActorType::Preview:
+		return ConfigurationHelpers::ImplGetRootActor(PreviewRootActorRef);
 
-		// return true, if changed
-		return true;
-	}
+	case EDisplayClusterRootActorType::Scene:
+		return ConfigurationHelpers::ImplGetRootActor(SceneRootActorRef);
 
-	return false;
-}
+	case EDisplayClusterRootActorType::Configuration:
+		return ConfigurationHelpers::ImplGetRootActor(ConfigurationRootActorRef);
 
-ADisplayClusterRootActor* FDisplayClusterViewportConfiguration::GetRootActor() const
-{
-	check(IsInGameThread());
-
-	AActor* ActorPtr = RootActorRef.GetOrFindSceneActor();
-	if (ActorPtr)
+	case EDisplayClusterRootActorType::Any:
 	{
-		return static_cast<ADisplayClusterRootActor*>(ActorPtr);
+		ADisplayClusterRootActor* OutRootActor = ConfigurationHelpers::ImplGetRootActor(PreviewRootActorRef);
+		if (!OutRootActor)
+		{
+			OutRootActor = ConfigurationHelpers::ImplGetRootActor(SceneRootActorRef);
+			if (!OutRootActor)
+			{
+				OutRootActor = ConfigurationHelpers::ImplGetRootActor(ConfigurationRootActorRef);
+			}
+		}
+
+		return OutRootActor;
+	}
+
+	default:
+		break;
 	}
 
 	return nullptr;
 }
 
-bool FDisplayClusterViewportConfiguration::ImplUpdateConfiguration(EDisplayClusterRenderFrameMode InRenderMode, const FString& InClusterNodeId, const FDisplayClusterPreviewSettings* InPreviewSettings, const TArray<FString>* InViewportNames)
+void FDisplayClusterViewportConfiguration::SetCurrentWorld(UWorld* InWorld)
 {
-	check(IsInGameThread());
-
-
-	ADisplayClusterRootActor* RootActor = GetRootActor();
-	if (RootActor)
+	if (GetCurrentWorld() != InWorld)
 	{
-		const UDisplayClusterConfigurationData* ConfigurationData = RootActor->GetConfigData();
-		FDisplayClusterViewportManager* ViewportManager = GetViewportManager();
-		if (ConfigurationData && ViewportManager)
+		if (FDisplayClusterViewportManager* ViewportManager = GetViewportManagerImpl())
 		{
-			FDisplayClusterViewportConfigurationBase ConfigurationBase(*ViewportManager, *RootActor, *ConfigurationData, RenderFrameSettings);
-			FDisplayClusterViewportConfigurationICVFX ConfigurationICVFX(*RootActor);
-			FDisplayClusterViewportConfigurationProjectionPolicy ConfigurationProjectionPolicy(*ViewportManager, *RootActor, *ConfigurationData);
-
-			ImplUpdateRenderFrameConfiguration(RootActor->GetRenderFrameSettings());
-
-			// Set current rendering mode
-			RenderFrameSettings.RenderMode = InRenderMode;
+			const bool bNeedRestartScene = bCurrentSceneActive;
 
-			// Support alpha channel capture
-			RenderFrameSettings.AlphaChannelCaptureMode = GetAlphaChannelCaptureMode();
+			ViewportManager->HandleEndScene();
+			CurrentWorldRef.Reset();
 
-			if (InPreviewSettings != nullptr)
+			if (InWorld)
 			{
-				// Downscale resources with PreviewDownscaleRatio
-				RenderFrameSettings.PreviewRenderTargetRatioMult = InPreviewSettings->PreviewRenderTargetRatioMult;
-
-				// Limit preview textures max size
-				RenderFrameSettings.PreviewMaxTextureDimension = InPreviewSettings->PreviewMaxTextureDimension;
-
-				// Hack preview gamma.
-				// In a scene, PostProcess always renders on top of the preview textures.
-				// But in it, PostProcess is also rendered with the flag turned off.
-				RenderFrameSettings.bPreviewEnablePostProcess = InPreviewSettings->bPreviewEnablePostProcess;
-
-				// Support mGPU for preview rendering
-				RenderFrameSettings.bAllowMultiGPURenderingInEditor = InPreviewSettings->bAllowMultiGPURenderingInEditor;
-				RenderFrameSettings.PreviewMinGPUIndex = InPreviewSettings->MinGPUIndex;
-				RenderFrameSettings.PreviewMaxGPUIndex = InPreviewSettings->MaxGPUIndex;
-
-				RenderFrameSettings.bIsRenderingInEditor = true;
-				RenderFrameSettings.bIsPreviewRendering = true;
-				RenderFrameSettings.bFreezePreviewRender = InPreviewSettings->bFreezePreviewRender;
-				
-				if (InPreviewSettings->bIsPIE)
+				CurrentWorldRef = TWeakObjectPtr<UWorld>(InWorld);
+
+				if (bNeedRestartScene)
 				{
-					// Allow TextureShare+nDisplay from PIE
-					RenderFrameSettings.bIsPreviewRendering = false;
+					ViewportManager->HandleStartScene();
 				}
 			}
-			else
-			{
-				RenderFrameSettings.bPreviewEnablePostProcess = false;
-				RenderFrameSettings.bAllowMultiGPURenderingInEditor = false;
-				RenderFrameSettings.bIsRenderingInEditor = false;
-				RenderFrameSettings.bIsPreviewRendering = false;
-			}
-
-			if (InViewportNames)
-			{
-				ConfigurationBase.UpdateCustomViewports(*InViewportNames);
-			}
-			else
-			{
-				ConfigurationBase.UpdateClusterNodeViewports(InClusterNodeId);
-			}
+		}
+		else
+		{
+			CurrentWorldRef.Reset();
+		}
+	}
+}
 
-			ConfigurationICVFX.Update();
-			ConfigurationProjectionPolicy.Update();
-			ConfigurationICVFX.PostUpdate();
+UWorld* FDisplayClusterViewportConfiguration::GetCurrentWorld() const
+{
+	check(IsInGameThread());
 
-#if WITH_EDITOR
-			if (RenderFrameSettings.bIsPreviewRendering)
-			{
-				ConfigurationICVFX.PostUpdatePreview_Editor(*InPreviewSettings);
-			}
-#endif
+	if (!CurrentWorldRef.IsValid() || CurrentWorldRef.IsStale())
+	{
+		return nullptr;
+	}
 
-			ImplUpdateConfigurationVisibility(*RootActor, *ConfigurationData);
+	return CurrentWorldRef.Get();
+}
 
-			if (!InViewportNames)
-			{
-				// Update postprocess for current cluster node
-				ConfigurationBase.UpdateClusterNodePostProcess(InClusterNodeId);
-			}
+const UDisplayClusterConfigurationData* FDisplayClusterViewportConfiguration::GetConfigurationData() const
+{
+	ADisplayClusterRootActor* ConfigurationRootActor = GetRootActor(EDisplayClusterRootActorType::Configuration);
 
-			ImplPostUpdateRenderFrameConfiguration();
+	return ConfigurationRootActor ? ConfigurationRootActor->GetConfigData() : nullptr;
+}
 
-			return true;
-		}
+const FDisplayClusterConfigurationICVFX_StageSettings* FDisplayClusterViewportConfiguration::GetStageSettings() const
+{
+	if (const UDisplayClusterConfigurationData* ConfigurationData = GetConfigurationData())
+	{
+		return &ConfigurationData->StageSettings;
 	}
 
-	return false;
+	return nullptr;
 }
 
-bool FDisplayClusterViewportConfiguration::UpdateConfiguration(EDisplayClusterRenderFrameMode InRenderMode, const FString& InClusterNodeId)
+const FDisplayClusterConfigurationRenderFrame* FDisplayClusterViewportConfiguration::GetConfigurationRenderFrameSettings() const
 {
-	return ImplUpdateConfiguration(InRenderMode, InClusterNodeId, nullptr, nullptr);
-}
+	if (const UDisplayClusterConfigurationData* ConfigurationData = GetConfigurationData())
+	{
+		return &ConfigurationData->RenderFrameSettings;
+	}
 
-bool FDisplayClusterViewportConfiguration::UpdateCustomConfiguration(EDisplayClusterRenderFrameMode InRenderMode, const TArray<FString>& InViewportNames)
-{
-	return ImplUpdateConfiguration(InRenderMode, TEXT(""), nullptr, &InViewportNames);
+	return nullptr;
 }
 
-#if WITH_EDITOR
-bool FDisplayClusterViewportConfiguration::UpdatePreviewConfiguration(EDisplayClusterRenderFrameMode InRenderMode, const FString& InClusterNodeId, const FDisplayClusterPreviewSettings& InPreviewSettings)
+bool FDisplayClusterViewportConfiguration::IsCurrentWorldHasAnyType(const EWorldType::Type InWorldType1, const EWorldType::Type InWorldType2, const EWorldType::Type InWorldType3) const
 {
-	// Special cluster node names are now supported in the lower classes.
-	// Therefore, we don't need to go through the cluster nodes to initialize the whole cluster.
-	return ImplUpdateConfiguration(InRenderMode, InClusterNodeId, &InPreviewSettings, nullptr);
+	if (UWorld* CurrentWorld = GetCurrentWorld())
+	{
+		return (CurrentWorld->WorldType == InWorldType1 && InWorldType1 != EWorldType::None)
+			|| (CurrentWorld->WorldType == InWorldType2 && InWorldType2 != EWorldType::None)
+			|| (CurrentWorld->WorldType == InWorldType3 && InWorldType3 != EWorldType::None);
+	}
+
+	return false;
 }
-#endif
 
-void FDisplayClusterViewportConfiguration::ImplUpdateConfigurationVisibility(ADisplayClusterRootActor& RootActor, const UDisplayClusterConfigurationData& ConfigurationData) const
+bool FDisplayClusterViewportConfiguration::IsRootActorWorldHasAnyType(const EDisplayClusterRootActorType InRootActorType, const EWorldType::Type InWorldType1, const EWorldType::Type InWorldType2, const EWorldType::Type InWorldType3) const
 {
-	// Hide root actor components for all viewports
-	FDisplayClusterViewportManager* ViewportManager = GetViewportManager();
-
-	TSet<FPrimitiveComponentId> RootActorHidePrimitivesList;
-	if (ViewportManager && RootActor.GetHiddenInGamePrimitives(RootActorHidePrimitivesList))
+	if (ADisplayClusterRootActor* RootActor = GetRootActor(InRootActorType))
 	{
-		for (const TSharedPtr<FDisplayClusterViewport, ESPMode::ThreadSafe>& ViewportIt : ViewportManager->ImplGetCurrentRenderFrameViewports())
+		if (UWorld* CurrentWorld = RootActor->GetWorld())
 		{
-			if (ViewportIt.IsValid())
-			{
-				ViewportIt->VisibilitySettings.SetRootActorHideList(RootActorHidePrimitivesList);
-			}
+			return (CurrentWorld->WorldType == InWorldType1 && InWorldType1 != EWorldType::None)
+				|| (CurrentWorld->WorldType == InWorldType2 && InWorldType2 != EWorldType::None)
+				|| (CurrentWorld->WorldType == InWorldType3 && InWorldType3 != EWorldType::None);
 		}
 	}
+
+	return false;
 }
 
-void FDisplayClusterViewportConfiguration::ImplPostUpdateRenderFrameConfiguration()
+EDisplayClusterRenderFrameMode FDisplayClusterViewportConfiguration::GetRenderModeForPIE() const
 {
-	if (FDisplayClusterViewportManager* ViewportManager = GetViewportManager())
+#if WITH_EDITOR
+	if (ADisplayClusterRootActor* PreviewRootActor = GetRootActor(EDisplayClusterRootActorType::Preview))
 	{
-		// Some frame postprocess require additional render targetable resources
-		RenderFrameSettings.bShouldUseAdditionalFrameTargetableResource = ViewportManager->ShouldUseAdditionalFrameTargetableResource();
-		RenderFrameSettings.bShouldUseFullSizeFrameTargetableResource = ViewportManager->ShouldUseFullSizeFrameTargetableResource();
+		switch (PreviewRootActor->RenderMode)
+		{
+		case EDisplayClusterConfigurationRenderMode::SideBySide:
+			return EDisplayClusterRenderFrameMode::SideBySide;
+
+		case EDisplayClusterConfigurationRenderMode::TopBottom:
+			return EDisplayClusterRenderFrameMode::TopBottom;
+
+		default:
+			break;
+		}
 	}
+#endif
+
+	return EDisplayClusterRenderFrameMode::Mono;
+}
+
+IDisplayClusterViewportManager* FDisplayClusterViewportConfiguration::GetViewportManager() const
+{
+	return GetViewportManagerImpl();
 }
 
-void FDisplayClusterViewportConfiguration::ImplUpdateRenderFrameConfiguration(const FDisplayClusterConfigurationRenderFrame& InRenderFrameConfiguration)
+bool FDisplayClusterViewportConfiguration::ImplUpdateConfiguration(EDisplayClusterRenderFrameMode InRenderMode, const FString& InClusterNodeId, const TArray<FString>* InViewportNames)
 {
-	// Global RTT sizes mults
-	RenderFrameSettings.ClusterRenderTargetRatioMult = InRenderFrameConfiguration.ClusterRenderTargetRatioMult;
-	RenderFrameSettings.ClusterICVFXInnerViewportRenderTargetRatioMult = InRenderFrameConfiguration.ClusterICVFXInnerViewportRenderTargetRatioMult;
-	RenderFrameSettings.ClusterICVFXOuterViewportRenderTargetRatioMult = InRenderFrameConfiguration.ClusterICVFXOuterViewportRenderTargetRatioMult;
+	check(IsInGameThread());
+	
+	FDisplayClusterViewportManager* ViewportManager = GetViewportManagerImpl();
+	if (!ViewportManager || !FDisplayClusterViewportConfigurationHelpers_RenderFrameSettings::UpdateRenderFrameConfiguration(InRenderMode, *this))
+	{
+		return false;
+	}
 
-	// Global Buffer ratio mults
-	RenderFrameSettings.ClusterBufferRatioMult = InRenderFrameConfiguration.ClusterBufferRatioMult;
-	RenderFrameSettings.ClusterICVFXInnerFrustumBufferRatioMult = InRenderFrameConfiguration.ClusterICVFXInnerFrustumBufferRatioMult;
-	RenderFrameSettings.ClusterICVFXOuterViewportBufferRatioMult = InRenderFrameConfiguration.ClusterICVFXOuterViewportBufferRatioMult;
+	if (bCurrentSceneNeedsToBeUpdated)
+	{
+		// When the root actor changes, we have to ResetScene() to reinitialize the internal references of the projection policy.
+		ViewportManager->HandleEndScene();
+		ViewportManager->HandleStartScene();
+	}
 
-	// Allow warpblend render
-	RenderFrameSettings.bAllowWarpBlend = InRenderFrameConfiguration.bAllowWarpBlend;
+	FDisplayClusterViewportConfiguration_ViewportManager  ConfigurationViewportManager(*this);
+	FDisplayClusterViewportConfiguration_Postprocess      ConfigurationPostprocess(*this);
+	FDisplayClusterViewportConfiguration_ProjectionPolicy ConfigurationProjectionPolicy(*this);
+	FDisplayClusterViewportConfiguration_ICVFX            ConfigurationICVFX(*this);
 
-	// Performance: Allow merge multiple viewports on single RTT with atlasing (required for bAllowViewFamilyMergeOptimization)
-	RenderFrameSettings.bAllowRenderTargetAtlasing = InRenderFrameConfiguration.bAllowRenderTargetAtlasing;
+	// when InClusterNodeId==PreviewNodeAll, means that it is an undefined cluster node
+	const FString ClusterNodeId = InClusterNodeId == DisplayClusterConfigurationStrings::gui::preview::PreviewNodeAll ? TEXT("") : InClusterNodeId;
 
-	// Performance: Allow viewfamily merge optimization (render multiple viewports contexts within single family)
-	// [not implemented yet] Experimental
-	switch (InRenderFrameConfiguration.ViewFamilyMode)
+	if (InViewportNames)
 	{
-	case EDisplayClusterConfigurationRenderFamilyMode::AllowMergeForGroups:
-		RenderFrameSettings.ViewFamilyMode = EDisplayClusterRenderFamilyMode::AllowMergeForGroups;
-		break;
-	case EDisplayClusterConfigurationRenderFamilyMode::AllowMergeForGroupsAndStereo:
-		RenderFrameSettings.ViewFamilyMode = EDisplayClusterRenderFamilyMode::AllowMergeForGroupsAndStereo;
-		break;
-	case EDisplayClusterConfigurationRenderFamilyMode::MergeAnyPossible:
-		RenderFrameSettings.ViewFamilyMode = EDisplayClusterRenderFamilyMode::MergeAnyPossible;
-		break;
-	case EDisplayClusterConfigurationRenderFamilyMode::None:
-	default:
-		RenderFrameSettings.ViewFamilyMode = EDisplayClusterRenderFamilyMode::None;
-		break;
+		ConfigurationViewportManager.UpdateCustomViewports(*InViewportNames);
+
+		// Do not use the cluster node name for this pass type
+		SetClusterNodeId(FString());
 	}
+	else
+	{
+		ConfigurationViewportManager.UpdateClusterNodeViewports(ClusterNodeId);
+		// Use only valid values of cluster node id
+		SetClusterNodeId(ClusterNodeId);
 
-	// Performance: nDisplay has its own implementation of cross-GPU transfer.
-	RenderFrameSettings.CrossGPUTransfer.bEnable    = GDisplayClusterCrossGPUTransferEnable != 0;
-	RenderFrameSettings.CrossGPUTransfer.bLockSteps = GDisplayClusterCrossGPUTransferLockSteps != 0;
-	RenderFrameSettings.CrossGPUTransfer.bPullData  = GDisplayClusterCrossGPUTransferPullData != 0;
+	}
 
-	// Performance: Allow to use parent ViewFamily from parent viewport 
-	// (icvfx has child viewports: lightcard and chromakey with prj_view matrices copied from parent viewport. May sense to use same viewfamily?)
-	// [not implemented yet] Experimental
-	RenderFrameSettings.bShouldUseParentViewportRenderFamily = InRenderFrameConfiguration.bShouldUseParentViewportRenderFamily;
+	ConfigurationICVFX.Update();
+	ConfigurationProjectionPolicy.Update();
+	ConfigurationICVFX.PostUpdate();
 
-	RenderFrameSettings.bIsRenderingInEditor = false;
+	ImplUpdateConfigurationVisibility();
 
-#if WITH_EDITOR
-	// Use special renderin mode, if DCRenderDevice not used right now
-	const bool bIsNDisplayClusterMode = (GEngine->StereoRenderingDevice.IsValid() && GDisplayCluster->GetOperationMode() == EDisplayClusterOperationMode::Cluster);
-	if (bIsNDisplayClusterMode == false)
+	if (!InViewportNames)
 	{
-		RenderFrameSettings.bIsRenderingInEditor = true;
+		// Update postprocess for current cluster node
+		ConfigurationPostprocess.UpdateClusterNodePostProcess(ClusterNodeId);
+	}
+
+	FDisplayClusterViewportConfigurationHelpers_RenderFrameSettings::PostUpdateRenderFrameConfiguration(*this);
+
+	return true;
+}
+
+bool FDisplayClusterViewportConfiguration::UpdateConfigurationForClusterNode(EDisplayClusterRenderFrameMode InRenderMode, const FString& InClusterNodeId)
+{
+	return ImplUpdateConfiguration(InRenderMode, InClusterNodeId, nullptr);
+}
+
+bool FDisplayClusterViewportConfiguration::UpdateConfigurationForViewportsList(EDisplayClusterRenderFrameMode InRenderMode, const TArray<FString>& InViewportNames)
+{
+	return ImplUpdateConfiguration(InRenderMode, TEXT(""), &InViewportNames);
+}
+
+void FDisplayClusterViewportConfiguration::ImplUpdateConfigurationVisibility() const
+{
+	ADisplayClusterRootActor* SceneRootActor = GetRootActor(EDisplayClusterRootActorType::Scene);
+	FDisplayClusterViewportManager* ViewportManager = GetViewportManagerImpl();
+
+	// Hide root actor components for all viewports
+	TSet<FPrimitiveComponentId> RootActorHidePrimitivesList;
+	if (ViewportManager && SceneRootActor && SceneRootActor->GetHiddenInGamePrimitives(RootActorHidePrimitivesList))
+	{
+		for (const TSharedPtr<FDisplayClusterViewport, ESPMode::ThreadSafe>& ViewportIt : ViewportManager->ImplGetCurrentRenderFrameViewports())
+		{
+			if (ViewportIt.IsValid())
+			{
+				ViewportIt->GetVisibilitySettingsImpl().SetRootActorHideList(RootActorHidePrimitivesList);
+			}
+		}
 	}
-#endif /*WITH_EDITOR*/
 }
diff --git a/Engine/Plugins/Runtime/nDisplay/Source/DisplayCluster/Private/Render/Viewport/Configuration/DisplayClusterViewportConfiguration.h b/Engine/Plugins/Runtime/nDisplay/Source/DisplayCluster/Private/Render/Viewport/Configuration/DisplayClusterViewportConfiguration.h
--- a/Engine/Plugins/Runtime/nDisplay/Source/DisplayCluster/Private/Render/Viewport/Configuration/DisplayClusterViewportConfiguration.h
+++ b/Engine/Plugins/Runtime/nDisplay/Source/DisplayCluster/Private/Render/Viewport/Configuration/DisplayClusterViewportConfiguration.h
@@ -1,30 +1,81 @@
 // Copyright Epic Games, Inc. All Rights Reserved.
 
 #pragma once
+#include "Render/Viewport/IDisplayClusterViewportConfiguration.h"
+#include "Render/Viewport/Configuration/DisplayClusterViewportConfigurationProxy.h"
+#include "Render/Viewport/RenderFrame/DisplayClusterRenderFrameSettings.h"
 
-#include "CoreMinimal.h"
-
+#include "SceneTypes.h"
 #include "Misc/DisplayClusterObjectRef.h"
-#include "Render/Viewport/Containers/DisplayClusterViewport_Enums.h"
-#include "Render/Viewport/RenderFrame/DisplayClusterRenderFrameSettings.h"
-#include "Render/Viewport/Containers/DisplayClusterPreviewSettings.h"
+#include "Templates/SharedPointer.h"
 
+class FDisplayClusterViewportConfigurationProxy;
 class FDisplayClusterViewportManager;
+class FDisplayClusterViewportManagerProxy;
 class ADisplayClusterRootActor;
 class UDisplayClusterConfigurationData;
 struct FDisplayClusterConfigurationRenderFrame;
 
+/**
+* Implementation of the viewport manager configuration.
+*/
 class FDisplayClusterViewportConfiguration
+	: public IDisplayClusterViewportConfiguration
+	, public TSharedFromThis<FDisplayClusterViewportConfiguration, ESPMode::ThreadSafe>
 {
 public:
-	FDisplayClusterViewportConfiguration(FDisplayClusterViewportManager& InViewportManager);
-	~FDisplayClusterViewportConfiguration();
+	FDisplayClusterViewportConfiguration();
+	virtual ~FDisplayClusterViewportConfiguration();
+
+	void Initialize(FDisplayClusterViewportManager& InViewportManager);
+
+public:
+	//~ BEGIN IDisplayClusterViewportConfiguration
+	virtual void SetCurrentWorld(UWorld* InWorld) override;
+	virtual void SetRootActor(ADisplayClusterRootActor* InRootActor, const EDisplayClusterRootActorType InRootActorType) override;
+	virtual bool UpdateConfigurationForClusterNode(EDisplayClusterRenderFrameMode InRenderMode, const FString& InClusterNodeId) override;
+	virtual bool UpdateConfigurationForViewportsList(EDisplayClusterRenderFrameMode InRenderMode, const TArray<FString>& InViewportNames) override;
+
+	virtual UWorld* GetCurrentWorld() const override;
+	virtual ADisplayClusterRootActor* GetRootActor(const EDisplayClusterRootActorType InRootActorType) const override;
+	virtual IDisplayClusterViewportManager* GetViewportManager() const override;
+	virtual const UDisplayClusterConfigurationData* GetConfigurationData() const override;
+	virtual const FDisplayClusterConfigurationICVFX_StageSettings* GetStageSettings() const override;
+	virtual const FDisplayClusterConfigurationRenderFrame* GetConfigurationRenderFrameSettings() const  override;
+
+	virtual EDisplayClusterRenderFrameMode GetRenderModeForPIE() const override;
+	virtual bool IsCurrentWorldHasAnyType(const EWorldType::Type InWorldType1, const EWorldType::Type InWorldType2 = EWorldType::None, const EWorldType::Type InWorldType3 = EWorldType::None) const override;
+	virtual bool IsRootActorWorldHasAnyType(const EDisplayClusterRootActorType InRootActorType, const EWorldType::Type InWorldType1, const EWorldType::Type InWorldType2 = EWorldType::None, const EWorldType::Type InWorldType3 = EWorldType::None) const override;
+
+	virtual const IDisplayClusterViewportConfigurationProxy& GetProxy() const override
+	{
+		return *Proxy;
+	}
+
+	virtual bool IsSceneOpened() const override
+	{
+		return bCurrentSceneActive && GetCurrentWorld();
+	}
+
+	virtual const FString& GetClusterNodeId() const override
+	{
+		return RenderFrameSettings.ClusterNodeId;
+	}
+
+	virtual bool IsPreviewRendering() const override
+	{
+		return RenderFrameSettings.IsPreviewRendering();
+	}
+	// ~~END IDisplayClusterViewportConfiguration
 
 public:
-	// Return true, if root actor ref changed
-	bool SetRootActor(ADisplayClusterRootActor* InRootActorPtr);
-	ADisplayClusterRootActor* GetRootActor() const;
+	/** Get a pointer to the DC ViewportManager if it still exists. */
+	FDisplayClusterViewportManager* GetViewportManagerImpl() const
+	{
+		return ViewportManagerWeakPtr.IsValid() ? ViewportManagerWeakPtr.Pin().Get() : nullptr;
+	}
 
+	/** Gets the rendering frame settings. */
 	const FDisplayClusterRenderFrameSettings& GetRenderFrameSettings() const
 	{ 
 		check(IsInGameThread());
@@ -32,36 +83,70 @@ public:
 		return RenderFrameSettings;
 	}
 
-	bool UpdateConfiguration(EDisplayClusterRenderFrameMode InRenderMode, const FString& InClusterNodeId);
-	bool UpdateCustomConfiguration(EDisplayClusterRenderFrameMode InRenderMode, const TArray<FString>& InViewportNames);
+	/** Sets the rendering frame settings. */
+	void SetRenderFrameSettings(const FDisplayClusterRenderFrameSettings& InRenderFrameSettings)
+	{
+		check(IsInGameThread());
+
+		// Process cluster node name changes
+		SetClusterNodeId(InRenderFrameSettings.ClusterNodeId);
 
-#if WITH_EDITOR
-	bool UpdatePreviewConfiguration(EDisplayClusterRenderFrameMode InRenderMode, const FString& ClusterNodeId, const FDisplayClusterPreviewSettings& InPreviewSettings);
-#endif
+		RenderFrameSettings = InRenderFrameSettings;
+	}
 
-private:
-	bool ImplUpdateConfiguration(EDisplayClusterRenderFrameMode InRenderMode, const FString& InClusterNodeId, const FDisplayClusterPreviewSettings* InPreviewSettings, const TArray<FString>* InViewportNames);
+	/** Set the name of the current cluster node. */
+	void SetClusterNodeId(const FString& InClusterNodeId)
+	{
+		if (RenderFrameSettings.ClusterNodeId != InClusterNodeId)
+		{
+			bCurrentRenderFrameViewportsNeedsToBeUpdated = true;
+			RenderFrameSettings.ClusterNodeId = InClusterNodeId;
+		}
+	}
 
-	void ImplUpdateRenderFrameConfiguration(const FDisplayClusterConfigurationRenderFrame& InRenderFrameConfiguration);
-	void ImplPostUpdateRenderFrameConfiguration();
+private:
+	/** Update configuration implementation.
+	* 
+	* @param InRenderMode - rendering mode.
+	* @param InClusterNodeId - (opt) Configuring rendering for a cluster node.
+	* @param InViewportNames - (opt) Configuring rendering for a list of viewports.
+	*/
+	bool ImplUpdateConfiguration(EDisplayClusterRenderFrameMode InRenderMode, const FString& InClusterNodeId, const TArray<FString>* InViewportNames);
 
 	/** Hide DCRA components for nDisplay rendering.*/
-	void ImplUpdateConfigurationVisibility(ADisplayClusterRootActor& InRootActor, const UDisplayClusterConfigurationData& InConfigurationData) const;
+	void ImplUpdateConfigurationVisibility() const;
 
-	/** Get alpha channel capture mode (for LightCard, ChromaKey).*/
-	EDisplayClusterRenderFrameAlphaChannelCaptureMode GetAlphaChannelCaptureMode() const;
+public:
+	// Reference to configuration proxy object
+	const TSharedRef<FDisplayClusterViewportConfigurationProxy, ESPMode::ThreadSafe> Proxy;
 
-	/** Get a pointer to the DC ViewportManager if it still exists. */
-	FDisplayClusterViewportManager* GetViewportManager() const
-	{
-		return ViewportManagerWeakPtr.IsValid() ? ViewportManagerWeakPtr.Pin().Get() : nullptr;
-	}
+	// Is current scene started
+	bool bCurrentSceneActive = false;
 
-private:
-	/** A reference to the owning viewport manager */
-	const TWeakPtr<FDisplayClusterViewportManager> ViewportManagerWeakPtr;
+	// Is scene update is required
+	bool bCurrentSceneNeedsToBeUpdated = false;
+
+	// Whether the list of viewports of the current frame needs to be updated
+	bool bCurrentRenderFrameViewportsNeedsToBeUpdated = false;
 
-	FDisplayClusterActorRef            RootActorRef;
+private:
+	// Current render frame settings
 	FDisplayClusterRenderFrameSettings RenderFrameSettings;
-};
 
+	// A reference to the owning viewport manager
+	TWeakPtr<FDisplayClusterViewportManager, ESPMode::ThreadSafe> ViewportManagerWeakPtr;
+
+private:
+	// This DCRA will be used to render previews. The meshes and preview materials are created at runtime.
+	FDisplayClusterActorRef PreviewRootActorRef;
+
+	// (Optional) A reference to DCRA in the scene, used as a source for math calculations and references.
+	// Locations in the scene and math data are taken from this DCRA.
+	FDisplayClusterActorRef SceneRootActorRef;
+
+	// (Optional) Reference to DCRA, used as a source of configuration data from DCRA and its components.
+	FDisplayClusterActorRef ConfigurationRootActorRef;
+
+	// Pointer to the current world to be rendered in.
+	TWeakObjectPtr<UWorld> CurrentWorldRef;
+};
diff --git a/Engine/Plugins/Runtime/nDisplay/Source/DisplayCluster/Private/Render/Viewport/Configuration/DisplayClusterViewportConfigurationHelpers.cpp b/Engine/Plugins/Runtime/nDisplay/Source/DisplayCluster/Private/Render/Viewport/Configuration/DisplayClusterViewportConfigurationHelpers.cpp
--- a/Engine/Plugins/Runtime/nDisplay/Source/DisplayCluster/Private/Render/Viewport/Configuration/DisplayClusterViewportConfigurationHelpers.cpp
+++ b/Engine/Plugins/Runtime/nDisplay/Source/DisplayCluster/Private/Render/Viewport/Configuration/DisplayClusterViewportConfigurationHelpers.cpp
@@ -4,6 +4,8 @@
 #include "DisplayClusterViewportConfigurationHelpers_OpenColorIO.h"
 #include "DisplayClusterViewportConfigurationHelpers_Postprocess.h"
 
+#include "Render/Viewport/Configuration/DisplayClusterViewportConfiguration.h"
+
 #include "Render/Viewport/DisplayClusterViewport.h"
 #include "Render/Viewport/DisplayClusterViewportManager.h"
 #include "Render/Viewport/DisplayClusterViewportHelpers.h"
@@ -31,139 +33,23 @@
 #include "Render/Projection/IDisplayClusterProjectionPolicy.h"
 
 #include "DisplayClusterSceneViewExtensions.h"
-#include "OpenColorIODisplayExtension.h"
 
 #include "Misc/DisplayClusterLog.h"
 #include "TextureResource.h"
 
 #include "HAL/IConsoleManager.h"
 
-
-void FDisplayClusterViewportConfigurationHelpers::UpdateViewportStereoMode(FDisplayClusterViewport& DstViewport, const EDisplayClusterConfigurationViewport_StereoMode StereoMode)
+bool FDisplayClusterViewportConfigurationHelpers::IsForceMonoscopicRendering(const EDisplayClusterConfigurationViewport_StereoMode StereoMode)
 {
-	switch (StereoMode)
-	{
-	case EDisplayClusterConfigurationViewport_StereoMode::ForceMono:
-		DstViewport.RenderSettings.bForceMono = true;
-		break;
-	default:
-		DstViewport.RenderSettings.bForceMono = false;
-		break;
-	}
+	return StereoMode == EDisplayClusterConfigurationViewport_StereoMode::ForceMono;
 }
 
-void FDisplayClusterViewportConfigurationHelpers::UpdateViewportSetting_OverlayRenderSettings(FDisplayClusterViewport& DstViewport, const FDisplayClusterConfigurationICVFX_OverlayAdvancedRenderSettings& InOverlaySettings)
-{
-	DstViewport.SetViewportBufferRatio(InOverlaySettings.BufferRatio);
-	DstViewport.RenderSettings.RenderTargetRatio = InOverlaySettings.RenderTargetRatio;
-
-	DstViewport.RenderSettings.GPUIndex = InOverlaySettings.GPUIndex;
-	DstViewport.RenderSettings.StereoGPUIndex = InOverlaySettings.StereoGPUIndex;
-
-	UpdateViewportStereoMode(DstViewport, InOverlaySettings.StereoMode);
-
-	DstViewport.RenderSettings.RenderFamilyGroup = InOverlaySettings.RenderFamilyGroup;
-};
-
-void FDisplayClusterViewportConfigurationHelpers::UpdateViewportSetting_Override(FDisplayClusterViewport& DstViewport, const FDisplayClusterConfigurationPostRender_Override& InOverride)
+void FDisplayClusterViewportConfigurationHelpers::UpdateBaseViewportSetting(FDisplayClusterViewport& DstViewport, const UDisplayClusterConfigurationViewport& InConfigurationViewport)
 {
-	DstViewport.PostRenderSettings.Replace.TextureRHI.SafeRelease();
-
-	if (InOverride.bAllowReplace && InOverride.SourceTexture != nullptr)
-	{
-		FTextureResource* TextureResource = InOverride.SourceTexture->GetResource();
-		if(TextureResource)
-		{
-			FTextureRHIRef& TextureRHI = TextureResource->TextureRHI;
-
-			if (TextureRHI.IsValid())
-			{
-				DstViewport.PostRenderSettings.Replace.TextureRHI = TextureRHI;
-				FIntVector Size = TextureRHI->GetSizeXYZ();
+	// Gain direct access to internal settings of the viewport:
+	FDisplayClusterViewport_RenderSettings&           InOutRenderSettings = DstViewport.GetRenderSettingsImpl();
+	FDisplayClusterViewport_RenderSettingsICVFX& InOutRenderSettingsICVFX = DstViewport.GetRenderSettingsICVFXImpl();
 
-				DstViewport.PostRenderSettings.Replace.Rect = FDisplayClusterViewportHelpers::GetValidViewportRect((InOverride.bShouldUseTextureRegion) ? InOverride.TextureRegion.ToRect() : FIntRect(FIntPoint(0, 0), FIntPoint(Size.X, Size.Y)), DstViewport.GetId(), TEXT("Configuration Override"));
-			}
-		}
-	}
-};
-
-void FDisplayClusterViewportConfigurationHelpers::UpdateViewportSetting_PostprocessBlur(FDisplayClusterViewport& DstViewport, const FDisplayClusterConfigurationPostRender_BlurPostprocess& InBlurPostprocess)
-{
-	switch (InBlurPostprocess.Mode)
-	{
-	case EDisplayClusterConfiguration_PostRenderBlur::Gaussian:
-		DstViewport.PostRenderSettings.PostprocessBlur.Mode = EDisplayClusterShaderParameters_PostprocessBlur::Gaussian;
-		break;
-	case EDisplayClusterConfiguration_PostRenderBlur::Dilate:
-		DstViewport.PostRenderSettings.PostprocessBlur.Mode = EDisplayClusterShaderParameters_PostprocessBlur::Dilate;
-		break;
-	default:
-		DstViewport.PostRenderSettings.PostprocessBlur.Mode = EDisplayClusterShaderParameters_PostprocessBlur::None;
-		break;
-	}
-
-	DstViewport.PostRenderSettings.PostprocessBlur.KernelRadius = InBlurPostprocess.KernelRadius;
-	DstViewport.PostRenderSettings.PostprocessBlur.KernelScale = InBlurPostprocess.KernelScale;
-};
-
-void FDisplayClusterViewportConfigurationHelpers::UpdateViewportSetting_Overscan(const FDisplayClusterConfigurationViewport_Overscan& InOverscan, FDisplayClusterViewport_OverscanSettings& OutOverscanSettings)
-{
-	OutOverscanSettings.bEnabled = false;
-	OutOverscanSettings.bOversize = InOverscan.bOversize;
-	
-	if (InOverscan.bEnabled)
-	{
-		switch (InOverscan.Mode)
-		{
-		case EDisplayClusterConfigurationViewportOverscanMode::Percent:
-			OutOverscanSettings.bEnabled = InOverscan.bEnabled;
-			OutOverscanSettings.Unit = EDisplayClusterViewport_FrustumUnit::Percent;
-
-			// Scale 0..100% to 0..1 range
-			OutOverscanSettings.Left = .01f * InOverscan.Left;
-			OutOverscanSettings.Right = .01f * InOverscan.Right;
-			OutOverscanSettings.Top = .01f * InOverscan.Top;
-			OutOverscanSettings.Bottom = .01f * InOverscan.Bottom;
-			break;
-
-		case EDisplayClusterConfigurationViewportOverscanMode::Pixels:
-			OutOverscanSettings.bEnabled = InOverscan.bEnabled;
-			OutOverscanSettings.Unit = EDisplayClusterViewport_FrustumUnit::Pixels;
-
-			OutOverscanSettings.Left = InOverscan.Left;
-			OutOverscanSettings.Right = InOverscan.Right;
-			OutOverscanSettings.Top = InOverscan.Top;
-			OutOverscanSettings.Bottom = InOverscan.Bottom;
-			break;
-
-		default:
-			break;
-		}
-	}
-};
-
-void FDisplayClusterViewportConfigurationHelpers::UpdateViewportSetting_GenerateMips(FDisplayClusterViewport& DstViewport, const FDisplayClusterConfigurationPostRender_GenerateMips& InGenerateMips)
-{
-	if (InGenerateMips.bAutoGenerateMips)
-	{
-		DstViewport.PostRenderSettings.GenerateMips.bAutoGenerateMips = true;
-
-		DstViewport.PostRenderSettings.GenerateMips.MipsSamplerFilter = InGenerateMips.MipsSamplerFilter;
-
-		DstViewport.PostRenderSettings.GenerateMips.MipsAddressU = InGenerateMips.MipsAddressU;
-		DstViewport.PostRenderSettings.GenerateMips.MipsAddressV = InGenerateMips.MipsAddressV;
-
-		DstViewport.PostRenderSettings.GenerateMips.MaxNumMipsLimit = (InGenerateMips.bEnabledMaxNumMips) ? InGenerateMips.MaxNumMips : -1;
-	}
-	else
-	{
-		// Disable mips
-		DstViewport.PostRenderSettings.GenerateMips.bAutoGenerateMips = false;
-	}
-}
-
-void FDisplayClusterViewportConfigurationHelpers::UpdateBaseViewportSetting(FDisplayClusterViewport& DstViewport, ADisplayClusterRootActor& RootActor, const UDisplayClusterConfigurationViewport& InConfigurationViewport)
-{
 	// Reset runtime flags from prev frame:
 	DstViewport.ResetRuntimeParameters();
 
@@ -171,44 +57,43 @@ void FDisplayClusterViewportConfigurationHelpers::UpdateBaseViewportSetting(FDis
 	{
 		if (InConfigurationViewport.bAllowRendering == false)
 		{
-			DstViewport.RenderSettings.bEnable = false;
+			InOutRenderSettings.bEnable = false;
 		}
 
-		DstViewport.RenderSettings.CameraId = InConfigurationViewport.Camera;
-		DstViewport.RenderSettings.Rect = FDisplayClusterViewportHelpers::GetValidViewportRect(InConfigurationViewport.Region.ToRect(), DstViewport.GetId(), TEXT("Configuration Region"));
+		InOutRenderSettings.CameraId = InConfigurationViewport.Camera;
+		InOutRenderSettings.Rect = FDisplayClusterViewportHelpers::GetValidViewportRect(InConfigurationViewport.Region.ToRect(), DstViewport.GetId(), TEXT("Configuration Region"));
 
-		DstViewport.RenderSettings.bEnableCrossGPUTransfer = InConfigurationViewport.RenderSettings.bEnableCrossGPUTransfer;
+		InOutRenderSettings.bEnableCrossGPUTransfer = InConfigurationViewport.RenderSettings.bEnableCrossGPUTransfer;
 
-		DstViewport.RenderSettings.GPUIndex = InConfigurationViewport.GPUIndex;
-		DstViewport.RenderSettings.OverlapOrder = InConfigurationViewport.OverlapOrder;
+		InOutRenderSettings.GPUIndex = InConfigurationViewport.GPUIndex;
+		InOutRenderSettings.OverlapOrder = InConfigurationViewport.OverlapOrder;
 
 		// update viewport remap data
-		DstViewport.ViewportRemap.UpdateConfiguration(DstViewport, InConfigurationViewport.ViewportRemap);
+		DstViewport.UpdateConfiguration_ViewportRemap(InConfigurationViewport.ViewportRemap);
 	}
 
 	const FDisplayClusterConfigurationViewport_RenderSettings& InRenderSettings = InConfigurationViewport.RenderSettings;
 
 	// Update OCIO for Viewport
-	FDisplayClusterViewportConfigurationHelpers_OpenColorIO::UpdateBaseViewport(DstViewport, RootActor, InConfigurationViewport);
+	FDisplayClusterViewportConfigurationHelpers_OpenColorIO::UpdateBaseViewportOCIO(DstViewport, InConfigurationViewport);
 
 	// Additional per-viewport PostProcess
-	FDisplayClusterViewportConfigurationHelpers_Postprocess::UpdateCustomPostProcessSettings(DstViewport, RootActor, InRenderSettings.CustomPostprocess);
-	FDisplayClusterViewportConfigurationHelpers_Postprocess::UpdatePerViewportPostProcessSettings(DstViewport, RootActor);
+	FDisplayClusterViewportConfigurationHelpers_Postprocess::UpdateCustomPostProcessSettings(DstViewport, InRenderSettings.CustomPostprocess);
+	FDisplayClusterViewportConfigurationHelpers_Postprocess::UpdatePerViewportPostProcessSettings(DstViewport);
 
 	{
 		DstViewport.SetViewportBufferRatio(InRenderSettings.BufferRatio);
 
-		UpdateViewportSetting_Overscan(InRenderSettings.Overscan, DstViewport.RenderSettings.OverscanSettings);
+		DstViewport.UpdateConfiguration_Overscan(InRenderSettings.Overscan);
 
-		UpdateViewportSetting_Override(DstViewport, InRenderSettings.Replace);
-		UpdateViewportSetting_PostprocessBlur(DstViewport, InRenderSettings.PostprocessBlur);
-		UpdateViewportSetting_GenerateMips(DstViewport, InRenderSettings.GenerateMips);
+		DstViewport.UpdateConfiguration_PostRenderOverride(InRenderSettings.Replace);
+		DstViewport.UpdateConfiguration_PostRenderBlur(InRenderSettings.PostprocessBlur);
+		DstViewport.UpdateConfiguration_PostRenderGenerateMips(InRenderSettings.GenerateMips);
 
-		UpdateViewportStereoMode(DstViewport, InRenderSettings.StereoMode);
+		InOutRenderSettings.bForceMono = FDisplayClusterViewportConfigurationHelpers::IsForceMonoscopicRendering(InRenderSettings.StereoMode);
 
-		DstViewport.RenderSettings.StereoGPUIndex = InRenderSettings.StereoGPUIndex;
-		DstViewport.RenderSettings.RenderTargetRatio = InRenderSettings.RenderTargetRatio;
-		DstViewport.RenderSettings.RenderFamilyGroup = InRenderSettings.RenderFamilyGroup;
+		InOutRenderSettings.StereoGPUIndex = InRenderSettings.StereoGPUIndex;
+		InOutRenderSettings.RenderTargetRatio = InRenderSettings.RenderTargetRatio;
 	}
 
 	// Set media related configuration (runtime only for now)
@@ -223,48 +108,18 @@ void FDisplayClusterViewportConfigurationHelpers::UpdateBaseViewportSetting(FDis
 			if (MediaSettings.bEnable)
 			{
 				// Don't render the viewport if media input assigned
-				DstViewport.RenderSettings.bSkipSceneRenderingButLeaveResourcesAvailable = MediaSettings.IsMediaInputAssigned();
+				InOutRenderSettings.bSkipSceneRenderingButLeaveResourcesAvailable = MediaSettings.IsMediaInputAssigned();
 
 				// Mark this viewport is going to be captured by a capture device
-				DstViewport.RenderSettings.bIsBeingCaptured = MediaSettings.IsMediaOutputAssigned();
+				InOutRenderSettings.bIsBeingCaptured = MediaSettings.IsMediaOutputAssigned();
 			}
 		}
 	}
 
-
 	// FDisplayClusterConfigurationViewport_ICVFX property:
+	if(const FDisplayClusterConfigurationICVFX_StageSettings* StageSettings = DstViewport.Configuration->GetStageSettings())
 	{
-		const FDisplayClusterConfigurationICVFX_StageSettings& StageSettings = RootActor.GetStageSettings();
-
-		DstViewport.RenderSettingsICVFX.Flags = InConfigurationViewport.ICVFX.GetViewportICVFXFlags(StageSettings);
-		DstViewport.RenderSettingsICVFX.ICVFX.LightCardMode = InConfigurationViewport.ICVFX.GetLightCardRenderMode(StageSettings);
+		InOutRenderSettingsICVFX.Flags = InConfigurationViewport.ICVFX.GetViewportICVFXFlags(*StageSettings);
+		InOutRenderSettingsICVFX.ICVFX.LightCardMode = InConfigurationViewport.ICVFX.GetLightCardRenderMode(*StageSettings);
 	}
 }
-
-void FDisplayClusterViewportConfigurationHelpers::UpdateProjectionPolicy(FDisplayClusterViewport& DstViewport, const FDisplayClusterConfigurationProjection* InConfigurationProjectionPolicy)
-{
-	// Runtime update projection policy
-	const bool bNeedUpdateProjectionPolicy = InConfigurationProjectionPolicy && (
-			(DstViewport.ProjectionPolicy.IsValid() && DstViewport.ProjectionPolicy->IsConfigurationChanged(InConfigurationProjectionPolicy))
-		||  (DstViewport.UninitializedProjectionPolicy.IsValid() && DstViewport.UninitializedProjectionPolicy->IsConfigurationChanged(InConfigurationProjectionPolicy))
-		);
-
-	if (bNeedUpdateProjectionPolicy)
-	{
-		UE_LOG(LogDisplayClusterViewport, Verbose, TEXT("Update projection policy for viewport '%s'."), *DstViewport.GetId());
-
-		// Release current projection
-		DstViewport.HandleEndScene();
-		DstViewport.UninitializedProjectionPolicy.Reset();
-
-		// Create new projection type interface
-		DstViewport.UninitializedProjectionPolicy = FDisplayClusterViewportManager::CreateProjectionPolicy(DstViewport.GetId(), InConfigurationProjectionPolicy);
-	}
-
-	// If the projection policy is not initialized and the scene is open, initialize it immediately
-	if (!DstViewport.ProjectionPolicy.IsValid() && DstViewport.IsSceneOpened())
-			{
-				// Try initialize proj policy every tick (mesh deferred load, etc)
-				DstViewport.HandleStartScene();
-			}
-		}
diff --git a/Engine/Plugins/Runtime/nDisplay/Source/DisplayCluster/Private/Render/Viewport/Configuration/DisplayClusterViewportConfigurationHelpers.h b/Engine/Plugins/Runtime/nDisplay/Source/DisplayCluster/Private/Render/Viewport/Configuration/DisplayClusterViewportConfigurationHelpers.h
--- a/Engine/Plugins/Runtime/nDisplay/Source/DisplayCluster/Private/Render/Viewport/Configuration/DisplayClusterViewportConfigurationHelpers.h
+++ b/Engine/Plugins/Runtime/nDisplay/Source/DisplayCluster/Private/Render/Viewport/Configuration/DisplayClusterViewportConfigurationHelpers.h
@@ -3,34 +3,20 @@
 #pragma once
 
 #include "CoreMinimal.h"
-
 #include "DisplayClusterConfigurationTypes_Enums.h"
 
 class FDisplayClusterViewport;
-class ADisplayClusterRootActor;
 class UDisplayClusterConfigurationViewport;
 
-struct FDisplayClusterConfigurationICVFX_OverlayAdvancedRenderSettings;
-struct FDisplayClusterConfigurationPostRender_Override;
-struct FDisplayClusterConfigurationPostRender_BlurPostprocess;
-struct FDisplayClusterConfigurationViewport_Overscan;
-struct FDisplayClusterConfigurationPostRender_GenerateMips;
-struct FDisplayClusterConfigurationProjection;
-struct FDisplayClusterViewport_OverscanSettings;
-
+/**
+* Base configuration helper class.
+*/
 class FDisplayClusterViewportConfigurationHelpers
 {
 public:
-	static void UpdateViewportStereoMode(FDisplayClusterViewport& DstViewport, const EDisplayClusterConfigurationViewport_StereoMode StereoMode);
-	
-	static void UpdateViewportSetting_OverlayRenderSettings(FDisplayClusterViewport& DstViewport, const FDisplayClusterConfigurationICVFX_OverlayAdvancedRenderSettings& InOverlaySettings);
-	static void UpdateViewportSetting_Override(FDisplayClusterViewport& DstViewport, const FDisplayClusterConfigurationPostRender_Override& InOverride);
-
-	static void UpdateViewportSetting_PostprocessBlur(FDisplayClusterViewport& DstViewport, const FDisplayClusterConfigurationPostRender_BlurPostprocess& InBlurPostprocess);
-	static void UpdateViewportSetting_Overscan(const FDisplayClusterConfigurationViewport_Overscan& InOverscan, FDisplayClusterViewport_OverscanSettings& OutOverscanSettings);
-	static void UpdateViewportSetting_GenerateMips(FDisplayClusterViewport& DstViewport, const FDisplayClusterConfigurationPostRender_GenerateMips& InGenerateMips);
-
-	static void UpdateBaseViewportSetting(FDisplayClusterViewport& DstViewport, ADisplayClusterRootActor& RootActor, const UDisplayClusterConfigurationViewport& InConfigurationViewport);
+	/** Returns true if this viewport should be rendered in monoscopic mode. */
+	static bool IsForceMonoscopicRendering(const EDisplayClusterConfigurationViewport_StereoMode StereoMode);
 
-	static void UpdateProjectionPolicy(FDisplayClusterViewport& DstViewport, const FDisplayClusterConfigurationProjection* InConfigurationProjectionPolicy = nullptr);
+	/** Update base settings of DC viewport.*/
+	static void UpdateBaseViewportSetting(FDisplayClusterViewport& DstViewport, const UDisplayClusterConfigurationViewport& InConfigurationViewport);
 };
