nDisplay: Added new interface IDisplayClusterViewportConfiguration

diff --git a/Engine/Plugins/Runtime/nDisplay/Source/DisplayCluster/Private/Render/Viewport/Configuration/DisplayClusterViewportConfigurationHelpers_ICVFX.cpp b/Engine/Plugins/Runtime/nDisplay/Source/DisplayCluster/Private/Render/Viewport/Configuration/DisplayClusterViewportConfigurationHelpers_ICVFX.cpp
--- a/Engine/Plugins/Runtime/nDisplay/Source/DisplayCluster/Private/Render/Viewport/Configuration/DisplayClusterViewportConfigurationHelpers_ICVFX.cpp
+++ b/Engine/Plugins/Runtime/nDisplay/Source/DisplayCluster/Private/Render/Viewport/Configuration/DisplayClusterViewportConfigurationHelpers_ICVFX.cpp
@@ -2,6 +2,7 @@
 
 #include "Render/Viewport/Configuration/DisplayClusterViewportConfigurationHelpers_ICVFX.h"
 
+#include "Render/Viewport/Configuration/DisplayClusterViewportConfiguration.h"
 #include "Render/Viewport/Configuration/DisplayClusterViewportConfigurationHelpers.h"
 #include "Render/Viewport/Configuration/DisplayClusterViewportConfigurationHelpers_OpenColorIO.h"
 #include "Render/Viewport/Configuration/DisplayClusterViewportConfigurationHelpers_Postprocess.h"
@@ -41,37 +42,11 @@
 
 #include "Misc/DisplayClusterLog.h"
 #include "TextureResource.h"
-
 #include "HAL/IConsoleManager.h"
 
-
-////////////////////////////////////////////////////////////////////////////////
-// Experimental feature: to be approved after testing
-int32 GDisplayClusterPreviewEnableReuseViewportInCluster = 1;
-static FAutoConsoleVariableRef CVarDisplayClusterPreviewEnableReuseViewportInCluster(
-	TEXT("DC.Preview.EnableReuseViewportInCluster"),
-	GDisplayClusterPreviewEnableReuseViewportInCluster,
-	TEXT("Experimental feature (0 == disabled, 1 == enabled)"),
-	ECVF_RenderThreadSafe
-);
-
 ////////////////////////////////////////////////////////////////////////////////
 namespace UE::DisplayCluster::Viewport::ConfigurationHelpers_ICVFX
 {
-	// Initialize camera policy with camera component and settings
-	static bool ImplUpdateCameraProjectionSettings(TSharedPtr<IDisplayClusterProjectionPolicy, ESPMode::ThreadSafe>& InOutCameraProjection, ADisplayClusterRootActor& RootActor, const FDisplayClusterConfigurationICVFX_CameraSettings& CameraSettings, UCameraComponent* const CameraComponent)
-	{
-		FDisplayClusterProjectionCameraPolicySettings PolicyCameraSettings;
-		PolicyCameraSettings.FOVMultiplier = CameraSettings.GetCameraFieldOfViewMultiplier(RootActor.GetStageSettings());
-
-		// Lens correction
-		PolicyCameraSettings.FrustumRotation = CameraSettings.FrustumRotation;
-		PolicyCameraSettings.FrustumOffset = CameraSettings.FrustumOffset;
-
-		// Initialize camera policy with camera component and settings
-		return IDisplayClusterProjection::Get().CameraPolicySetCamera(InOutCameraProjection, CameraComponent, PolicyCameraSettings);
-	}
-
 	// Return unique ICVFX name
 	static FString ImplGetNameICVFX(const FString& InClusterNodeId, const FString& InViewportId, const FString& InResourceId)
 	{
@@ -81,82 +56,94 @@ namespace UE::DisplayCluster::Viewport::ConfigurationHelpers_ICVFX
 
 		return FString::Printf(TEXT("%s_%s_%s_%s"), *InClusterNodeId, DisplayClusterViewportStrings::icvfx::prefix, *InViewportId, *InResourceId);
 	}
-
-#if WITH_EDITOR
-	static bool IsPreviewEnableReuseViewportInCluster_EditorImpl(ADisplayClusterRootActor& InRootActor)
-	{
-		if (GDisplayClusterPreviewEnableReuseViewportInCluster > 0)
-		{
-			if (InRootActor.GetRenderMode() == EDisplayClusterRenderFrameMode::PreviewInScene)
-			{
-				return true;
-			}
-		}
-
-		return false;
-	}
-#endif
 };
 using namespace UE::DisplayCluster::Viewport::ConfigurationHelpers_ICVFX;
 
 ////////////////////////////////////////////////////////////////////////////////
 // FDisplayClusterViewportConfigurationHelpers_ICVFX
 ////////////////////////////////////////////////////////////////////////////////
-FDisplayClusterViewport* FDisplayClusterViewportConfigurationHelpers_ICVFX::ImplFindViewport(ADisplayClusterRootActor& RootActor, const FString& InViewportId, const FString& InResourceId)
+bool FDisplayClusterViewportConfigurationHelpers_ICVFX::CreateProjectionPolicyICVFX(FDisplayClusterViewportConfiguration& InConfiguration, const FString& InViewportId, const FString& InResourceId, bool bIsCameraProjection, TSharedPtr<IDisplayClusterProjectionPolicy, ESPMode::ThreadSafe>& OutProjPolicy)
 {
-	if (FDisplayClusterViewportManager* ViewportManager = FDisplayClusterViewportConfigurationHelpers_ICVFX::GetViewportManager(RootActor))
+	const FString& ClusterNodeId = InConfiguration.GetClusterNodeId();
+	if (ClusterNodeId.IsEmpty())
 	{
-		TSharedPtr<FDisplayClusterViewport, ESPMode::ThreadSafe> Viewport = ViewportManager->ImplFindViewport(ImplGetNameICVFX(ViewportManager->GetRenderFrameSettings().ClusterNodeId, InViewportId, InResourceId));
+		return false;
+	}
 
-		return Viewport.Get();
+	FDisplayClusterConfigurationProjection CameraProjectionPolicyConfig;
+	CameraProjectionPolicyConfig.Type = bIsCameraProjection ? DisplayClusterProjectionStrings::projection::Camera : DisplayClusterProjectionStrings::projection::Link;
+
+	// Create projection policy for viewport
+	OutProjPolicy = FDisplayClusterViewportManager::CreateProjectionPolicy(ImplGetNameICVFX(ClusterNodeId, InViewportId, InResourceId), &CameraProjectionPolicyConfig);
+
+	if (!OutProjPolicy.IsValid())
+	{
+		UE_LOG(LogDisplayClusterViewport, Error, TEXT("ICVFX Viewport '%s': projection policy for resource '%s' not created for node '%s'."), *InViewportId, *InResourceId, *ClusterNodeId);
+		return false;
 	}
 
-	return nullptr;
+	return true;
 }
 
-static bool ImplCreateProjectionPolicy(ADisplayClusterRootActor& RootActor, const FString& InViewportId, const FString& InResourceId, bool bIsCameraProjection, TSharedPtr<IDisplayClusterProjectionPolicy, ESPMode::ThreadSafe>& OutProjPolicy)
+bool FDisplayClusterViewportConfigurationHelpers_ICVFX::UpdateCameraProjectionSettingsICVFX(FDisplayClusterViewportConfiguration& InConfiguration, UDisplayClusterICVFXCameraComponent& InCameraComponent, const FDisplayClusterConfigurationICVFX_CameraSettings& InCameraSettings, const TSharedPtr<IDisplayClusterProjectionPolicy, ESPMode::ThreadSafe>& InProjectionPolicy)
 {
-	if (FDisplayClusterViewportManager* ViewportManager = FDisplayClusterViewportConfigurationHelpers_ICVFX::GetViewportManager(RootActor))
+	const FDisplayClusterConfigurationICVFX_StageSettings* StageSettings = InConfiguration.GetStageSettings();
+	if (!StageSettings)
 	{
-		const FString& ClusterNodeId = ViewportManager->GetRenderFrameSettings().ClusterNodeId;
+		return false;
+	}
 
-		FDisplayClusterConfigurationProjection CameraProjectionPolicyConfig;
-		CameraProjectionPolicyConfig.Type = bIsCameraProjection ? DisplayClusterProjectionStrings::projection::Camera : DisplayClusterProjectionStrings::projection::Link;
+	// Initialize camera policy with camera component and settings
+	FDisplayClusterProjectionCameraPolicySettings PolicyCameraSettings;
+	PolicyCameraSettings.FOVMultiplier = InCameraSettings.GetCameraFieldOfViewMultiplier(*StageSettings);
 
-		// Create projection policy for viewport
-		OutProjPolicy = FDisplayClusterViewportManager::CreateProjectionPolicy(ImplGetNameICVFX(ClusterNodeId, InViewportId, InResourceId), &CameraProjectionPolicyConfig);
+	// Lens correction
+	PolicyCameraSettings.FrustumRotation = InCameraSettings.FrustumRotation;
+	PolicyCameraSettings.FrustumOffset = InCameraSettings.FrustumOffset;
 
-		if (!OutProjPolicy.IsValid())
-		{
-			UE_LOG(LogDisplayClusterViewport, Error, TEXT("ICVFX Viewport '%s': projection policy for resource '%s' not created for node '%s'."), *InViewportId, *InResourceId, *ClusterNodeId);
-			return false;
-		}
+	static IDisplayClusterProjection& DisplayClusterProjectionAPI = IDisplayClusterProjection::Get();
 
-		return true;
-	}
+	// Initialize camera policy with camera component and settings
+	return DisplayClusterProjectionAPI.CameraPolicySetCamera(InProjectionPolicy, &InCameraComponent, PolicyCameraSettings);
+}
 
-	return false;
+bool FDisplayClusterViewportConfigurationHelpers_ICVFX::CreateProjectionPolicyCameraICVFX(FDisplayClusterViewportConfiguration& InConfiguration, UDisplayClusterICVFXCameraComponent& InCameraComponent, const FDisplayClusterConfigurationICVFX_CameraSettings& InCameraSettings, TSharedPtr<IDisplayClusterProjectionPolicy, ESPMode::ThreadSafe>& OutProjPolicy)
+{
+	return CreateProjectionPolicyICVFX(InConfiguration, InCameraComponent.GetCameraUniqueId(), DisplayClusterViewportStrings::icvfx::camera, true, OutProjPolicy)
+		&& UpdateCameraProjectionSettingsICVFX(InConfiguration, InCameraComponent, InCameraSettings, OutProjPolicy);
 }
 
-//------------------------------------------------------------------------------------------
-//                FDisplayClusterViewportConfigurationHelpers_ICVFX
-//------------------------------------------------------------------------------------------
+FDisplayClusterViewport* FDisplayClusterViewportConfigurationHelpers_ICVFX::FindViewportICVFX(FDisplayClusterViewportConfiguration& InConfiguration, const FString& InViewportId, const FString& InResourceId)
+{
+	if (FDisplayClusterViewportManager* ViewportManager = InConfiguration.GetViewportManagerImpl())
+	{
+		TSharedPtr<FDisplayClusterViewport, ESPMode::ThreadSafe> Viewport = ViewportManager->ImplFindViewport(ImplGetNameICVFX(InConfiguration.GetClusterNodeId(), InViewportId, InResourceId));
 
-FDisplayClusterViewport* FDisplayClusterViewportConfigurationHelpers_ICVFX::ImplCreateViewport(ADisplayClusterRootActor& RootActor, const FString& InViewportId, const FString& InResourceId, TSharedPtr<IDisplayClusterProjectionPolicy, ESPMode::ThreadSafe>& InProjectionPolicy)
+		return Viewport.Get();
+	}
+
+	return nullptr;
+}
+
+FDisplayClusterViewport* FDisplayClusterViewportConfigurationHelpers_ICVFX::CreateViewportICVFX(FDisplayClusterViewportConfiguration& InConfiguration, const FString& InViewportId, const FString& InResourceId, TSharedPtr<IDisplayClusterProjectionPolicy, ESPMode::ThreadSafe>& InProjectionPolicy)
 {
 	check(InProjectionPolicy.IsValid());
 
-	if (FDisplayClusterViewportManager* ViewportManager = FDisplayClusterViewportConfigurationHelpers_ICVFX::GetViewportManager(RootActor))
+	if (FDisplayClusterViewportManager* ViewportManager = InConfiguration.GetViewportManagerImpl())
 	{
 		// Create viewport for new projection policy
-		TSharedPtr<FDisplayClusterViewport, ESPMode::ThreadSafe> NewViewport = ViewportManager->ImplCreateViewport(ImplGetNameICVFX(ViewportManager->GetRenderFrameSettings().ClusterNodeId, InViewportId, InResourceId), InProjectionPolicy);
+		TSharedPtr<FDisplayClusterViewport, ESPMode::ThreadSafe> NewViewport = ViewportManager->ImplCreateViewport(ImplGetNameICVFX(InConfiguration.GetClusterNodeId(), InViewportId, InResourceId), InProjectionPolicy);
 		if (NewViewport.IsValid())
 		{
+			// Gain direct access to internal resources of the NewViewport:
+			FDisplayClusterViewport_RenderSettings&           InOutRenderSettings = NewViewport->GetRenderSettingsImpl();
+			FDisplayClusterViewport_RenderSettingsICVFX& InOutRenderSettingsICVFX = NewViewport->GetRenderSettingsICVFXImpl();
+
 			// Mark as internal resource
-			EnumAddFlags(NewViewport->RenderSettingsICVFX.RuntimeFlags, EDisplayClusterViewportRuntimeICVFXFlags::InternalResource);
+			EnumAddFlags(InOutRenderSettingsICVFX.RuntimeFlags, EDisplayClusterViewportRuntimeICVFXFlags::InternalResource);
 
 			// Dont show ICVFX composing viewports on frame target
-			NewViewport->RenderSettings.bVisible = false;
+			InOutRenderSettings.bVisible = false;
 
 			return NewViewport.Get();
 		}
@@ -165,68 +152,58 @@ FDisplayClusterViewport* FDisplayClusterViewportConfigurationHelpers_ICVFX::Impl
 	return nullptr;
 }
 
-FDisplayClusterViewportManager* FDisplayClusterViewportConfigurationHelpers_ICVFX::GetViewportManager(ADisplayClusterRootActor& RootActor)
-{
-	FDisplayClusterViewportManager* ViewportManager = static_cast<FDisplayClusterViewportManager*>(RootActor.GetViewportManager());
-	if (ViewportManager == nullptr)
-	{
-		UE_LOG(LogDisplayClusterViewport, Error, TEXT("Viewport manager not exist in root actor."));
-	}
-
-	return ViewportManager;
-}
-
-FDisplayClusterViewport* FDisplayClusterViewportConfigurationHelpers_ICVFX::FindCameraViewport(ADisplayClusterRootActor& RootActor, UDisplayClusterICVFXCameraComponent& InCameraComponent)
+FDisplayClusterViewport* FDisplayClusterViewportConfigurationHelpers_ICVFX::FindCameraViewport(FDisplayClusterViewportConfiguration& InConfiguration, UDisplayClusterICVFXCameraComponent& InCameraComponent)
 {
-	const FString CameraId = InCameraComponent.GetCameraUniqueId();
-
-	return ImplFindViewport(RootActor, CameraId, DisplayClusterViewportStrings::icvfx::camera);
+	return FindViewportICVFX(InConfiguration, InCameraComponent.GetCameraUniqueId(), DisplayClusterViewportStrings::icvfx::camera);
 }
 
-FDisplayClusterViewport* FDisplayClusterViewportConfigurationHelpers_ICVFX::GetOrCreateCameraViewport(ADisplayClusterRootActor& RootActor, UDisplayClusterICVFXCameraComponent& InCameraComponent)
+FDisplayClusterViewport* FDisplayClusterViewportConfigurationHelpers_ICVFX::GetOrCreateCameraViewport(FDisplayClusterViewportConfiguration& InConfiguration, UDisplayClusterICVFXCameraComponent& InCameraComponent, const FDisplayClusterConfigurationICVFX_CameraSettings& InCameraSettings)
 {
-	const FString CameraId = InCameraComponent.GetCameraUniqueId();
-
-	FDisplayClusterViewport* CameraViewport = ImplFindViewport(RootActor, CameraId, DisplayClusterViewportStrings::icvfx::camera);
-
-	// Create new camera viewport
+	FDisplayClusterViewport* CameraViewport = FindCameraViewport(InConfiguration, InCameraComponent);
 	if (CameraViewport == nullptr)
 	{
+		// Create new camera viewport
 		TSharedPtr<IDisplayClusterProjectionPolicy, ESPMode::ThreadSafe> CameraProjectionPolicy;
-		if (!ImplCreateProjectionPolicy(RootActor, CameraId, DisplayClusterViewportStrings::icvfx::camera, true, CameraProjectionPolicy))
+		if (!CreateProjectionPolicyCameraICVFX(InConfiguration, InCameraComponent, InCameraSettings, CameraProjectionPolicy))
 		{
 			return nullptr;
 		}
 
-		CameraViewport = ImplCreateViewport(RootActor, CameraId, DisplayClusterViewportStrings::icvfx::camera, CameraProjectionPolicy);
+		CameraViewport = CreateViewportICVFX(InConfiguration, InCameraComponent.GetCameraUniqueId(), DisplayClusterViewportStrings::icvfx::camera, CameraProjectionPolicy);
 		if (CameraViewport == nullptr)
 		{
 			return nullptr;
 		}
 	}
 
+	// Gain direct access to internal resources of the CameraViewport:
+	FDisplayClusterViewport_RenderSettingsICVFX& InOutRenderSettingsICVFX = CameraViewport->GetRenderSettingsICVFXImpl();
+
 	// Mark viewport as used
-	EnumRemoveFlags(CameraViewport->RenderSettingsICVFX.RuntimeFlags, EDisplayClusterViewportRuntimeICVFXFlags::Unused);
+	EnumRemoveFlags(InOutRenderSettingsICVFX.RuntimeFlags, EDisplayClusterViewportRuntimeICVFXFlags::Unused);
 
 	// Add viewport ICVFX usage as Incamera
-	EnumAddFlags(CameraViewport->RenderSettingsICVFX.RuntimeFlags, EDisplayClusterViewportRuntimeICVFXFlags::InCamera);
+	EnumAddFlags(InOutRenderSettingsICVFX.RuntimeFlags, EDisplayClusterViewportRuntimeICVFXFlags::InCamera);
 
 	return CameraViewport;
 }
 
 void FDisplayClusterViewportConfigurationHelpers_ICVFX::ReuseUVLightCardViewportWithinClusterNode(FDisplayClusterViewport& InUVLightCardViewport)
 {
-	if (FDisplayClusterViewportManager* ViewportManager = InUVLightCardViewport.GetViewportManagerImpl())
+	if (FDisplayClusterViewportManager* ViewportManager = InUVLightCardViewport.Configuration->GetViewportManagerImpl())
 	{
 		for (const TSharedPtr<FDisplayClusterViewport, ESPMode::ThreadSafe>& ViewportIt : ViewportManager->ImplGetEntireClusterViewports())
 		{
-			if (ViewportIt.IsValid() && ViewportIt != InUVLightCardViewport.AsShared() && !ViewportIt->RenderSettings.IsViewportOverridden()
-				&& EnumHasAnyFlags(ViewportIt->RenderSettingsICVFX.RuntimeFlags, EDisplayClusterViewportRuntimeICVFXFlags::UVLightcard))
+			if (ViewportIt.IsValid() && ViewportIt != InUVLightCardViewport.AsShared() && !ViewportIt->GetRenderSettings().IsViewportOverridden()
+				&& EnumHasAnyFlags(ViewportIt->GetRenderSettingsICVFX().RuntimeFlags, EDisplayClusterViewportRuntimeICVFXFlags::UVLightcard))
 			{
 				if (ViewportIt->IsOpenColorIOEquals(InUVLightCardViewport))
 				{
+					// Gain direct access to internal resources of the viewport:
+					FDisplayClusterViewport_RenderSettings& InOutRenderSettings = InUVLightCardViewport.GetRenderSettingsImpl();
+
 					// Reuse exist viewport:
-					InUVLightCardViewport.RenderSettings.SetViewportOverride(ViewportIt->GetId(), EDisplayClusterViewportOverrideMode::All);
+					InOutRenderSettings.SetViewportOverride(ViewportIt->GetId(), EDisplayClusterViewportOverrideMode::All);
 					break;
 				}
 			}
@@ -234,27 +211,52 @@ void FDisplayClusterViewportConfigurationHelpers_ICVFX::ReuseUVLightCardViewport
 	}
 }
 
-#if WITH_EDITOR
-TArray<TSharedPtr<FDisplayClusterViewport, ESPMode::ThreadSafe>> FDisplayClusterViewportConfigurationHelpers_ICVFX::PreviewGetRenderedInCameraViewports(ADisplayClusterRootActor& InRootActor, UDisplayClusterICVFXCameraComponent& InCameraComponent, bool bGetChromakey)
+TArray<TSharedPtr<FDisplayClusterViewport, ESPMode::ThreadSafe>> FDisplayClusterViewportConfigurationHelpers_ICVFX::GetAllVisibleInnerCameraViewports(FDisplayClusterViewportConfiguration& InConfiguration, bool bGetChromakey)
 {
 	TArray<TSharedPtr<FDisplayClusterViewport, ESPMode::ThreadSafe>> OutViewports;
 
-	// Search for rendered camera viewport on other cluster nodes
-	if (FDisplayClusterViewportManager* ViewportManager = FDisplayClusterViewportConfigurationHelpers_ICVFX::GetViewportManager(InRootActor))
+	ADisplayClusterRootActor* ConfigurationRootActor = InConfiguration.GetRootActor(EDisplayClusterRootActorType::Configuration);
+	const UDisplayClusterConfigurationData* ConfigurationData = InConfiguration.GetConfigurationData();
+	if (ConfigurationRootActor && ConfigurationData)
 	{
-		const FString ICVFXCameraId = InCameraComponent.GetCameraUniqueId();
-		const EDisplayClusterViewportRuntimeICVFXFlags InRuntimeFlagsMask = bGetChromakey ? EDisplayClusterViewportRuntimeICVFXFlags::Chromakey : EDisplayClusterViewportRuntimeICVFXFlags::InCamera;
-		const FString ViewportTypeId = bGetChromakey ? DisplayClusterViewportStrings::icvfx::chromakey : DisplayClusterViewportStrings::icvfx::camera;
+		TArray<UDisplayClusterICVFXCameraComponent*> ExistsICVFXCameraComponents;
+		ConfigurationRootActor->GetComponents<UDisplayClusterICVFXCameraComponent>(ExistsICVFXCameraComponents);
+
+		for (UDisplayClusterICVFXCameraComponent* CfgCineCameraComponent : ExistsICVFXCameraComponents)
+		{
+			if (CfgCineCameraComponent && CfgCineCameraComponent->GetCameraSettingsICVFX().IsICVFXEnabled(*ConfigurationData, InConfiguration.GetClusterNodeId()))
+			{
+				const FString& CameraName = CfgCineCameraComponent->GetCameraUniqueId();
+
+				OutViewports.Append(FDisplayClusterViewportConfigurationHelpers_ICVFX::PreviewGetRenderedInCameraViewports(InConfiguration, CameraName, bGetChromakey));
+			}
+		}
+	}
+
+	return OutViewports;
+}
+
+TArray<TSharedPtr<FDisplayClusterViewport, ESPMode::ThreadSafe>> FDisplayClusterViewportConfigurationHelpers_ICVFX::PreviewGetRenderedInCameraViewports(FDisplayClusterViewportConfiguration& InConfiguration, const FString& InICVFXCameraId, bool bGetChromakey)
+{
+	TArray<TSharedPtr<FDisplayClusterViewport, ESPMode::ThreadSafe>> OutViewports;
 
+	// Search for rendered camera viewport on other cluster nodes
+	const EDisplayClusterViewportRuntimeICVFXFlags InRuntimeFlagsMask = bGetChromakey ? EDisplayClusterViewportRuntimeICVFXFlags::Chromakey : EDisplayClusterViewportRuntimeICVFXFlags::InCamera;
+	const FString ViewportTypeId = bGetChromakey ? DisplayClusterViewportStrings::icvfx::chromakey : DisplayClusterViewportStrings::icvfx::camera;
+
+	if (FDisplayClusterViewportManager* ViewportManager = InConfiguration.GetViewportManagerImpl())
+	{
 		for (const TSharedPtr<FDisplayClusterViewport, ESPMode::ThreadSafe>& ViewportIt : ViewportManager->ImplGetEntireClusterViewports())
 		{
 			if (ViewportIt.IsValid()
-				&& EnumHasAnyFlags(ViewportIt->RenderSettingsICVFX.RuntimeFlags, InRuntimeFlagsMask)
-				&& (ViewportIt->Resources[EDisplayClusterViewportResource::InputShaderResources].Num() > 0 && ViewportIt->Resources[EDisplayClusterViewportResource::InputShaderResources][0] != nullptr && ViewportIt->Contexts.Num() > 0)
-				&& (!ViewportIt->RenderSettings.IsViewportOverridden()))
+				&& EnumHasAnyFlags(ViewportIt->GetRenderSettingsICVFX().RuntimeFlags, InRuntimeFlagsMask)
+				&& ( ViewportIt->GetViewportResources(EDisplayClusterViewportResource::InputShaderResources).Num() > 0
+					&& ViewportIt->GetViewportResources(EDisplayClusterViewportResource::InputShaderResources)[0] != nullptr
+					&& ViewportIt->GetContexts().Num() > 0)
+				&& (!ViewportIt->GetRenderSettings().IsViewportOverridden()))
 			{
 				// this is incamera viewport. Check by name
-				const FString RequiredViewportId = ImplGetNameICVFX(ViewportIt->GetClusterNodeId(), ICVFXCameraId, ViewportTypeId);
+				const FString RequiredViewportId = ImplGetNameICVFX(ViewportIt->GetClusterNodeId(), InICVFXCameraId, ViewportTypeId);
 				if (RequiredViewportId.Equals(ViewportIt->GetId()))
 				{
 					OutViewports.Add(ViewportIt);
@@ -266,166 +268,171 @@ TArray<TSharedPtr<FDisplayClusterViewport, ESPMode::ThreadSafe>> FDisplayCluster
 	return OutViewports;
 }
 
-void FDisplayClusterViewportConfigurationHelpers_ICVFX::PreviewReuseInnerFrustumViewportWithinClusterNodes_Editor(FDisplayClusterViewport& InCameraViewport, ADisplayClusterRootActor& InRootActor, UDisplayClusterICVFXCameraComponent& InCameraComponent)
+void FDisplayClusterViewportConfigurationHelpers_ICVFX::PreviewReuseInnerFrustumViewportWithinClusterNodes(FDisplayClusterViewport& InCameraViewport, UDisplayClusterICVFXCameraComponent& InCameraComponent, const FDisplayClusterConfigurationICVFX_CameraSettings& InCameraSettings)
 {
-	if (!IsPreviewEnableReuseViewportInCluster_EditorImpl(InRootActor))
-	{
-		return;
-	}
-
-	// Only for preview mode
-	if (InCameraViewport.GetRenderMode() != EDisplayClusterRenderFrameMode::PreviewInScene)
+	if (!InCameraViewport.Configuration->GetRenderFrameSettings().CanReuseViewportWithinClusterNodes())
 	{
 		return;
 	}
 
-	for (const TSharedPtr<FDisplayClusterViewport, ESPMode::ThreadSafe>& ViewportIt : PreviewGetRenderedInCameraViewports(InRootActor, InCameraComponent, false))
+	if (FDisplayClusterViewportManager* ViewportManager = InCameraViewport.Configuration->GetViewportManagerImpl())
 	{
-		if (ViewportIt.IsValid() && ViewportIt != InCameraViewport.AsShared() && ViewportIt->GetClusterNodeId() != InCameraViewport.GetClusterNodeId()
-			&& FDisplayClusterViewportConfigurationHelpers_Postprocess::IsInnerFrustumViewportSettingsEqual_Editor(*ViewportIt, InCameraViewport, InCameraComponent))
+		for (const TSharedPtr<FDisplayClusterViewport, ESPMode::ThreadSafe>& ViewportIt : PreviewGetRenderedInCameraViewports(*InCameraViewport.Configuration, InCameraComponent.GetCameraUniqueId(), false))
 		{
-			EDisplayClusterViewportOverrideMode ViewportOverrideMode = ViewportIt->IsOpenColorIOEquals(InCameraViewport) ?
-				EDisplayClusterViewportOverrideMode::All : EDisplayClusterViewportOverrideMode::InernalRTT;
+			if (ViewportIt.IsValid() && ViewportIt != InCameraViewport.AsShared() && ViewportIt->GetClusterNodeId() != InCameraViewport.GetClusterNodeId()
+				&& FDisplayClusterViewportConfigurationHelpers_Postprocess::IsInnerFrustumViewportSettingsEqual(*ViewportIt, InCameraViewport, InCameraSettings))
+			{
+				EDisplayClusterViewportOverrideMode ViewportOverrideMode = ViewportIt->IsOpenColorIOEquals(InCameraViewport) ?
+					EDisplayClusterViewportOverrideMode::All : EDisplayClusterViewportOverrideMode::InernalRTT;
 
-			// Reuse exist viewport:
-			InCameraViewport.RenderSettings.SetViewportOverride(ViewportIt->GetId(), ViewportOverrideMode);
-			return;
+				// Gain direct access to internal resources of the viewport:
+				FDisplayClusterViewport_RenderSettings& InOutRenderSettings = InCameraViewport.GetRenderSettingsImpl();
+
+				// Reuse exist viewport:
+				InOutRenderSettings.SetViewportOverride(ViewportIt->GetId(), ViewportOverrideMode);
+				return;
+			}
 		}
 	}
 }
 
-void FDisplayClusterViewportConfigurationHelpers_ICVFX::PreviewReuseChromakeyViewportWithinClusterNodes_Editor(FDisplayClusterViewport& InChromakeyViewport, ADisplayClusterRootActor& InRootActor, UDisplayClusterICVFXCameraComponent& InCameraComponent)
+void FDisplayClusterViewportConfigurationHelpers_ICVFX::PreviewReuseChromakeyViewportWithinClusterNodes(FDisplayClusterViewport& InChromakeyViewport, const FString& InICVFXCameraId)
 {
-	if (!IsPreviewEnableReuseViewportInCluster_EditorImpl(InRootActor))
+	if (!InChromakeyViewport.Configuration->GetRenderFrameSettings().CanReuseViewportWithinClusterNodes())
 	{
 		return;
 	}
 	
-	// Only for preview mode
-	if (InChromakeyViewport.GetRenderMode() != EDisplayClusterRenderFrameMode::PreviewInScene)
-	{
-		return;
-	}
-
-	for (const TSharedPtr<FDisplayClusterViewport, ESPMode::ThreadSafe>& ViewportIt : PreviewGetRenderedInCameraViewports(InRootActor, InCameraComponent, true))
+	if (FDisplayClusterViewportManager* ViewportManager = InChromakeyViewport.Configuration->GetViewportManagerImpl())
 	{
-		if (ViewportIt.IsValid() && ViewportIt != InChromakeyViewport.AsShared() && ViewportIt->GetClusterNodeId() != InChromakeyViewport.GetClusterNodeId())
+		for (const TSharedPtr<FDisplayClusterViewport, ESPMode::ThreadSafe>& ViewportIt : PreviewGetRenderedInCameraViewports(*InChromakeyViewport.Configuration, InICVFXCameraId, true))
 		{
-			// Chromakey support OCIO
-			EDisplayClusterViewportOverrideMode ViewportOverrideMode = ViewportIt->IsOpenColorIOEquals(InChromakeyViewport) ?
-				EDisplayClusterViewportOverrideMode::All : EDisplayClusterViewportOverrideMode::InernalRTT;
+			if (ViewportIt.IsValid() && ViewportIt != InChromakeyViewport.AsShared() && ViewportIt->GetClusterNodeId() != InChromakeyViewport.GetClusterNodeId())
+			{
+				// Chromakey support OCIO
+				EDisplayClusterViewportOverrideMode ViewportOverrideMode = ViewportIt->IsOpenColorIOEquals(InChromakeyViewport) ?
+					EDisplayClusterViewportOverrideMode::All : EDisplayClusterViewportOverrideMode::InernalRTT;
 
-			// Reuse exist viewport from other node
-			InChromakeyViewport.RenderSettings.SetViewportOverride(ViewportIt->GetId(), ViewportOverrideMode);
-			return;
+				// Gain direct access to internal resources of the viewport:
+				FDisplayClusterViewport_RenderSettings& InOutRenderSettings = InChromakeyViewport.GetRenderSettingsImpl();
+
+				// Reuse exist viewport from other node
+				InOutRenderSettings.SetViewportOverride(ViewportIt->GetId(), ViewportOverrideMode);
+				return;
+			}
 		}
 	}
 }
-#endif
 
-FDisplayClusterViewport* FDisplayClusterViewportConfigurationHelpers_ICVFX::GetOrCreateChromakeyViewport(ADisplayClusterRootActor& RootActor, UDisplayClusterICVFXCameraComponent& InCameraComponent)
+FDisplayClusterViewport* FDisplayClusterViewportConfigurationHelpers_ICVFX::GetOrCreateChromakeyViewport(FDisplayClusterViewportConfiguration& InConfiguration, UDisplayClusterICVFXCameraComponent& InCameraComponent, const FDisplayClusterConfigurationICVFX_CameraSettings& InCameraSettings)
 {
-	const FString CameraId = InCameraComponent.GetCameraUniqueId();
-
-	FDisplayClusterViewport* ChromakeyViewport = ImplFindViewport(RootActor, CameraId, DisplayClusterViewportStrings::icvfx::chromakey);
+	const FString InICVFXCameraId = InCameraComponent.GetCameraUniqueId();
+	FDisplayClusterViewport* ChromakeyViewport = FindViewportICVFX(InConfiguration, InICVFXCameraId, DisplayClusterViewportStrings::icvfx::chromakey);
 
 	// Create new chromakey viewport
 	if (ChromakeyViewport == nullptr)
 	{
 		TSharedPtr<IDisplayClusterProjectionPolicy, ESPMode::ThreadSafe> ChromakeyProjectionPolicy;
-		if (!ImplCreateProjectionPolicy(RootActor, CameraId, DisplayClusterViewportStrings::icvfx::chromakey, false, ChromakeyProjectionPolicy))
+		if (!CreateProjectionPolicyICVFX(InConfiguration, InICVFXCameraId, DisplayClusterViewportStrings::icvfx::chromakey, false, ChromakeyProjectionPolicy))
 		{
 			return nullptr;
 		}
 
-		ChromakeyViewport = ImplCreateViewport(RootActor, CameraId, DisplayClusterViewportStrings::icvfx::chromakey, ChromakeyProjectionPolicy);
+		ChromakeyViewport = CreateViewportICVFX(InConfiguration, InICVFXCameraId, DisplayClusterViewportStrings::icvfx::chromakey, ChromakeyProjectionPolicy);
 		if (ChromakeyViewport == nullptr)
 		{
 			return nullptr;
 		}
 	}
 
+	// Gain direct access to internal resources of the viewport:
+	FDisplayClusterViewport_RenderSettingsICVFX& InOutRenderSettingsICVFX = ChromakeyViewport->GetRenderSettingsICVFXImpl();
+
 	// Mark viewport as used
-	EnumRemoveFlags(ChromakeyViewport->RenderSettingsICVFX.RuntimeFlags, EDisplayClusterViewportRuntimeICVFXFlags::Unused);
+	EnumRemoveFlags(InOutRenderSettingsICVFX.RuntimeFlags, EDisplayClusterViewportRuntimeICVFXFlags::Unused);
 
 	// Add viewport ICVFX usage as Chromakey
-	EnumAddFlags(ChromakeyViewport->RenderSettingsICVFX.RuntimeFlags, EDisplayClusterViewportRuntimeICVFXFlags::Chromakey);
+	EnumAddFlags(InOutRenderSettingsICVFX.RuntimeFlags, EDisplayClusterViewportRuntimeICVFXFlags::Chromakey);
 
 	return ChromakeyViewport;
 }
 
-FDisplayClusterViewport* FDisplayClusterViewportConfigurationHelpers_ICVFX::GetOrCreateLightcardViewport(FDisplayClusterViewport& BaseViewport, ADisplayClusterRootActor& RootActor)
+FDisplayClusterViewport* FDisplayClusterViewportConfigurationHelpers_ICVFX::GetOrCreateLightcardViewport(FDisplayClusterViewport& BaseViewport)
 {
 	// Create new lightcard viewport
 	const FString ResourceId = DisplayClusterViewportStrings::icvfx::lightcard;
 
-	FDisplayClusterViewport* LightcardViewport = ImplFindViewport(RootActor, BaseViewport.GetId(), ResourceId);
+	FDisplayClusterViewport* LightcardViewport = FindViewportICVFX(*BaseViewport.Configuration, BaseViewport.GetId(), ResourceId);
 	if (LightcardViewport == nullptr)
 	{
 		TSharedPtr<IDisplayClusterProjectionPolicy, ESPMode::ThreadSafe> LightcardProjectionPolicy;
-		if (!ImplCreateProjectionPolicy(RootActor, BaseViewport.GetId(), ResourceId, false, LightcardProjectionPolicy))
+		if (!CreateProjectionPolicyICVFX(*BaseViewport.Configuration, BaseViewport.GetId(), ResourceId, false, LightcardProjectionPolicy))
 		{
 			return nullptr;
 		}
 
-		LightcardViewport = ImplCreateViewport(RootActor, BaseViewport.GetId(), ResourceId, LightcardProjectionPolicy);
+		LightcardViewport = CreateViewportICVFX(*BaseViewport.Configuration, BaseViewport.GetId(), ResourceId, LightcardProjectionPolicy);
 		if (LightcardViewport == nullptr)
 		{
 			return nullptr;
 		}
 	}
 
+	// Gain direct access to internal resources of the viewport:
+	FDisplayClusterViewport_RenderSettingsICVFX& InOutRenderSettingsICVFX = LightcardViewport->GetRenderSettingsICVFXImpl();
+
 	// Mark viewport as used
-	EnumRemoveFlags(LightcardViewport->RenderSettingsICVFX.RuntimeFlags, EDisplayClusterViewportRuntimeICVFXFlags::Unused);
+	EnumRemoveFlags(InOutRenderSettingsICVFX.RuntimeFlags, EDisplayClusterViewportRuntimeICVFXFlags::Unused);
 
 	// Add viewport ICVFX usage as Lightcard
-	EnumAddFlags(LightcardViewport->RenderSettingsICVFX.RuntimeFlags, EDisplayClusterViewportRuntimeICVFXFlags::Lightcard);
+	EnumAddFlags(InOutRenderSettingsICVFX.RuntimeFlags, EDisplayClusterViewportRuntimeICVFXFlags::Lightcard);
 
 	return LightcardViewport;
 }
 
-FDisplayClusterViewport* FDisplayClusterViewportConfigurationHelpers_ICVFX::GetOrCreateUVLightcardViewport(FDisplayClusterViewport& BaseViewport, ADisplayClusterRootActor& RootActor)
+FDisplayClusterViewport* FDisplayClusterViewportConfigurationHelpers_ICVFX::GetOrCreateUVLightcardViewport(FDisplayClusterViewport& BaseViewport)
 {
 	// Create new lightcard viewport
 	const FString ResourceId = DisplayClusterViewportStrings::icvfx::uv_lightcard;
 
-	FDisplayClusterViewport* UVLightcardViewport = ImplFindViewport(RootActor, BaseViewport.GetId(), ResourceId);
+	FDisplayClusterViewport* UVLightcardViewport = FindViewportICVFX(*BaseViewport.Configuration, BaseViewport.GetId(), ResourceId);
 	if (UVLightcardViewport == nullptr)
 	{
 		TSharedPtr<IDisplayClusterProjectionPolicy, ESPMode::ThreadSafe> UVLightcardProjectionPolicy;
-		if (!ImplCreateProjectionPolicy(RootActor, BaseViewport.GetId(), ResourceId, false, UVLightcardProjectionPolicy))
+		if (!CreateProjectionPolicyICVFX(*BaseViewport.Configuration, BaseViewport.GetId(), ResourceId, false, UVLightcardProjectionPolicy))
 		{
 			return nullptr;
 		}
 
-		UVLightcardViewport = ImplCreateViewport(RootActor, BaseViewport.GetId(), ResourceId, UVLightcardProjectionPolicy);
+		UVLightcardViewport = CreateViewportICVFX(*BaseViewport.Configuration, BaseViewport.GetId(), ResourceId, UVLightcardProjectionPolicy);
 		if (UVLightcardViewport == nullptr)
 		{
 			return nullptr;
 		}
 	}
 
+	// Gain direct access to internal resources of the viewport:
+	FDisplayClusterViewport_RenderSettingsICVFX& InOutRenderSettingsICVFX = UVLightcardViewport->GetRenderSettingsICVFXImpl();
+
 	// Mark viewport as used
-	EnumRemoveFlags(UVLightcardViewport->RenderSettingsICVFX.RuntimeFlags, EDisplayClusterViewportRuntimeICVFXFlags::Unused);
+	EnumRemoveFlags(InOutRenderSettingsICVFX.RuntimeFlags, EDisplayClusterViewportRuntimeICVFXFlags::Unused);
 
 	// Add viewport ICVFX usage as Lightcard
-	EnumAddFlags(UVLightcardViewport->RenderSettingsICVFX.RuntimeFlags, EDisplayClusterViewportRuntimeICVFXFlags::UVLightcard);
+	EnumAddFlags(InOutRenderSettingsICVFX.RuntimeFlags, EDisplayClusterViewportRuntimeICVFXFlags::UVLightcard);
 
 	return UVLightcardViewport;
 }
 
-bool FDisplayClusterViewportConfigurationHelpers_ICVFX::IsCameraUsed(UDisplayClusterICVFXCameraComponent& InCameraComponent)
+bool FDisplayClusterViewportConfigurationHelpers_ICVFX::IsCameraUsed(const FDisplayClusterConfigurationICVFX_CameraSettings& InCameraSettings)
 {
-	const FDisplayClusterConfigurationICVFX_CameraSettings& CameraSettings = InCameraComponent.GetCameraSettingsICVFX();
-
 	// Check rules for camera settings:
-	if (CameraSettings.bEnable == false)
+	if (InCameraSettings.bEnable == false)
 	{
 		// dont use camera if disabled
 		return false;
 	}
 
-	if (CameraSettings.RenderSettings.Replace.bAllowReplace && CameraSettings.RenderSettings.Replace.SourceTexture == nullptr)
+	if (InCameraSettings.RenderSettings.Replace.bAllowReplace && InCameraSettings.RenderSettings.Replace.SourceTexture == nullptr)
 	{
 		// RenderSettings.Override require source texture
 		return false;
@@ -434,151 +441,102 @@ bool FDisplayClusterViewportConfigurationHelpers_ICVFX::IsCameraUsed(UDisplayClu
 	return true;
 }
 
-FDisplayClusterShaderParameters_ICVFX::FCameraSettings FDisplayClusterViewportConfigurationHelpers_ICVFX::GetShaderParametersCameraSettings(const FDisplayClusterViewport& InCameraViewport, ADisplayClusterRootActor& RootActor, UDisplayClusterICVFXCameraComponent& InCameraComponent)
+FDisplayClusterShaderParameters_ICVFX::FCameraSettings FDisplayClusterViewportConfigurationHelpers_ICVFX::GetShaderParametersCameraSettings(const FDisplayClusterViewport& InCameraViewport, UDisplayClusterICVFXCameraComponent& InCameraComponent, const FDisplayClusterConfigurationICVFX_CameraSettings& InCameraSettings)
 {
-	const USceneComponent* OriginComp = RootActor.GetRootComponent();
-	const FDisplayClusterConfigurationICVFX_CameraSettings& CameraSettings = InCameraComponent.GetCameraSettingsICVFX();
-	const FDisplayClusterConfigurationICVFX_StageSettings& StageSettings = RootActor.GetStageSettings();
-
 	FDisplayClusterShaderParameters_ICVFX::FCameraSettings Result;
 
-	Result.Resource.ViewportId = InCameraViewport.GetId();
-	Result.Local2WorldTransform = OriginComp->GetComponentTransform();
+	ADisplayClusterRootActor* SceneRootActor = InCameraViewport.Configuration->GetRootActor(EDisplayClusterRootActorType::Scene);
+	ADisplayClusterRootActor* ConfigurationRootActor = InCameraViewport.Configuration->GetRootActor(EDisplayClusterRootActorType::Configuration);
+	const FDisplayClusterConfigurationICVFX_StageSettings* StageSettings = InCameraViewport.Configuration->GetStageSettings();
+
+	if (SceneRootActor && ConfigurationRootActor && StageSettings)
+	{
+		if (const USceneComponent* OriginComp = SceneRootActor->GetRootComponent())
+		{
+			Result.Resource.ViewportId = InCameraViewport.GetId();
+			Result.Local2WorldTransform = OriginComp->GetComponentTransform();
 
-	// Get camera border settings
-	CameraSettings.GetCameraBorder(StageSettings, Result.InnerCameraBorderColor, Result.InnerCameraBorderThickness);
-	Result.InnerCameraFrameAspectRatio = CameraSettings.GetCameraFrameAspectRatio(StageSettings);
+			// Get camera border settings
+			InCameraSettings.GetCameraBorder(*StageSettings, Result.InnerCameraBorderColor, Result.InnerCameraBorderThickness);
+			Result.InnerCameraFrameAspectRatio = InCameraSettings.GetCameraFrameAspectRatio(*StageSettings);
 
-	// Soft edges
-	Result.SoftEdge = CameraSettings.GetCameraSoftEdge(StageSettings);
+			// Soft edges
+			Result.SoftEdge = InCameraSettings.GetCameraSoftEdge(*StageSettings);
 
-	// Rendering order for camera overlap
-	const FString InnerFrustumID = InCameraComponent.GetCameraUniqueId();
-	const int32 CameraRenderOrder = RootActor.GetInnerFrustumPriority(InnerFrustumID);
-	Result.RenderOrder = (CameraRenderOrder<0) ? CameraSettings.RenderSettings.RenderOrder : CameraRenderOrder;
+			// Rendering order for camera overlap
+			const FString InnerFrustumID = InCameraComponent.GetCameraUniqueId();
+			const int32 CameraRenderOrder = ConfigurationRootActor->GetInnerFrustumPriority(InnerFrustumID);
+			Result.RenderOrder = (CameraRenderOrder < 0) ? InCameraSettings.RenderSettings.RenderOrder : CameraRenderOrder;
+		}
+	}
 
 	return Result;
 }
 
-bool FDisplayClusterViewportConfigurationHelpers_ICVFX::GetCameraContext(ADisplayClusterRootActor& RootActor, UDisplayClusterICVFXCameraComponent& InCameraComponent, FCameraContext_ICVFX& OutCameraContext)
+void FDisplayClusterViewportConfigurationHelpers_ICVFX::UpdateCameraViewportSettings(FDisplayClusterViewport& DstViewport, UDisplayClusterICVFXCameraComponent& InCameraComponent, const FDisplayClusterConfigurationICVFX_CameraSettings& InCameraSettings)
 {
-	const FDisplayClusterConfigurationICVFX_CameraSettings& CameraSettings = InCameraComponent.GetCameraSettingsICVFX();
-
-	// Create new camera projection policy for camera viewport
-	TSharedPtr<IDisplayClusterProjectionPolicy, ESPMode::ThreadSafe> CameraProjectionPolicy;
-	if (!ImplCreateProjectionPolicy(RootActor, InCameraComponent.GetCameraUniqueId(), DisplayClusterViewportStrings::icvfx::camera, true, CameraProjectionPolicy))
+	const FDisplayClusterConfigurationICVFX_StageSettings* StageSettings = DstViewport.Configuration->GetStageSettings();
+	if (!StageSettings)
 	{
-		return false;
-	}
-
-	// Initialize camera policy with camera component and settings
-	if (!ImplUpdateCameraProjectionSettings(CameraProjectionPolicy, RootActor, CameraSettings, &InCameraComponent))
-	{
-		return false;
-	}
-
-	// Get camera pos-rot-prj from policy
-
-	// Use default clipping planes
-	const float NCP = GNearClippingPlane;
-	const float FCP = NCP; // nDisplay does not use the far plane of the clipping
-
-	// Applying the correct sequence of steps to use the projection policy math:
-	// SetupProjectionViewPoint()->CalculateView()->GetProjectionMatrix()
-	FMinimalViewInfo CameraViewInfo;
-	float CustomNearClippingPlane = -1; // a value less than zero means ignoring.
-	CameraProjectionPolicy->SetupProjectionViewPoint(nullptr, RootActor.GetWorldDeltaSeconds(), CameraViewInfo, &CustomNearClippingPlane);
-
-	OutCameraContext.ViewLocation = CameraViewInfo.Location;
-	OutCameraContext.ViewRotation = CameraViewInfo.Rotation;
-
-	// Todo: Here we need to calculate the correct ViewOffset so that ICVFX can support stereo rendering.
-	const FVector ViewOffset = FVector::ZeroVector;
-
-	// Todo: maybe it makes sense to get the scaling value of WorldToMeters the right way?
-	// By default, this value is taken from the PlayerController in the LocalPlayer file.
-	// But ICVFX can use another data source for this value.
-	float WorldToMeters = 100.f;
-
-	// Supports custom near clipping plane
-	float ZNear = NCP;
-	float ZFar = FCP;
-	if (CustomNearClippingPlane >= 0)
-	{
-		ZNear = CustomNearClippingPlane;
-		ZFar = (NCP == FCP) ? ZNear : ZFar;
-	}
-
-	if (CameraProjectionPolicy->CalculateView(nullptr, 0, OutCameraContext.ViewLocation, OutCameraContext.ViewRotation, ViewOffset, WorldToMeters, ZNear, ZFar) &&
-		CameraProjectionPolicy->GetProjectionMatrix(nullptr, 0, OutCameraContext.PrjMatrix))
-	{
-		return true;
+		return;
 	}
 
-	return false;
-}
-
-void FDisplayClusterViewportConfigurationHelpers_ICVFX::UpdateCameraViewportSettings(FDisplayClusterViewport& DstViewport, ADisplayClusterRootActor& RootActor, UDisplayClusterICVFXCameraComponent& InCameraComponent)
-{
-	const FDisplayClusterConfigurationICVFX_StageSettings&  StageSettings = RootActor.GetStageSettings();
-	const FDisplayClusterConfigurationICVFX_CameraSettings& CameraSettings = InCameraComponent.GetCameraSettingsICVFX();
+	// Gain direct access to internal resources of the viewport:
+	FDisplayClusterViewport_RenderSettings& InOutRenderSettings = DstViewport.GetRenderSettingsImpl();
 
 	// Reset runtime flags from prev frame:
 	DstViewport.ResetRuntimeParameters();
 
 	// incamera textrure used as overlay
-	DstViewport.RenderSettings.bVisible = false;
+	InOutRenderSettings.bVisible = false;
 
 	// use framecolor instead of viewport rendering
-	if (CameraSettings.Chromakey.GetChromakeyType(StageSettings) == EDisplayClusterShaderParametersICVFX_ChromakeySource::FrameColor)
+	if (InCameraSettings.Chromakey.GetChromakeyType(*StageSettings) == EDisplayClusterShaderParametersICVFX_ChromakeySource::FrameColor)
 	{
-		DstViewport.RenderSettings.bSkipRendering = true;
+		InOutRenderSettings.bSkipRendering = true;
 	}
 
 	// Update camera viewport projection policy settings
-	ImplUpdateCameraProjectionSettings(DstViewport.ProjectionPolicy, RootActor, CameraSettings, &InCameraComponent);
+	UpdateCameraProjectionSettingsICVFX(*DstViewport.Configuration, InCameraComponent, InCameraSettings, DstViewport.GetProjectionPolicy());
 
 	// Update OCIO for Camera Viewport
-	FDisplayClusterViewportConfigurationHelpers_OpenColorIO::UpdateCameraViewport(DstViewport, RootActor, InCameraComponent);
+	FDisplayClusterViewportConfigurationHelpers_OpenColorIO::UpdateCameraViewportOCIO(DstViewport, InCameraSettings);
 
 	// Motion blur:
-	const FDisplayClusterViewport_CameraMotionBlur CameraMotionBlurParameters = InCameraComponent.GetMotionBlurParameters();
-	DstViewport.CameraMotionBlur.BlurSetup = CameraMotionBlurParameters;
+	DstViewport.UpdateConfiguration_CameraMotionBlur(InCameraComponent.GetMotionBlurParameters());
 
 	// FDisplayClusterConfigurationICVFX_CameraSettings
-	DstViewport.RenderSettings.CameraId.Empty();
+	InOutRenderSettings.CameraId.Empty();
 
 	// UDisplayClusterConfigurationICVFX_CameraRenderSettings
-	const FIntPoint DesiredSize = CameraSettings.GetCameraFrameSize(StageSettings);
+	const FIntPoint DesiredSize = InCameraSettings.GetCameraFrameSize(*StageSettings);
 
-	DstViewport.RenderSettings.Rect = FDisplayClusterViewportHelpers::GetValidViewportRect(FIntRect(FIntPoint(0, 0), DesiredSize), DstViewport.GetId(), TEXT("Configuration Camera Frame Size"));
+	InOutRenderSettings.Rect = FDisplayClusterViewportHelpers::GetValidViewportRect(FIntRect(FIntPoint(0, 0), DesiredSize), DstViewport.GetId(), TEXT("Configuration Camera Frame Size"));
 
 	// Apply postprocess for camera
-	FDisplayClusterViewportConfigurationHelpers_Postprocess::UpdateCameraPostProcessSettings(DstViewport, RootActor, InCameraComponent);
+	FDisplayClusterViewportConfigurationHelpers_Postprocess::UpdateCameraPostProcessSettings(DstViewport, InCameraComponent, InCameraSettings);
 
-	FDisplayClusterViewportConfigurationHelpers::UpdateViewportSetting_Override(DstViewport, CameraSettings.RenderSettings.Replace);
-	FDisplayClusterViewportConfigurationHelpers::UpdateViewportSetting_PostprocessBlur(DstViewport, CameraSettings.RenderSettings.PostprocessBlur);
-	FDisplayClusterViewportConfigurationHelpers::UpdateViewportSetting_GenerateMips(DstViewport, CameraSettings.RenderSettings.GenerateMips);
+	DstViewport.UpdateConfiguration_PostRenderOverride(InCameraSettings.RenderSettings.Replace);
+	DstViewport.UpdateConfiguration_PostRenderBlur(InCameraSettings.RenderSettings.PostprocessBlur);
+	DstViewport.UpdateConfiguration_PostRenderGenerateMips(InCameraSettings.RenderSettings.GenerateMips);
 
 	// UDisplayClusterConfigurationICVFX_CameraAdvancedRenderSettings
-	const FDisplayClusterConfigurationICVFX_CameraAdvancedRenderSettings& InAdvancedRS = CameraSettings.RenderSettings.AdvancedRenderSettings;
+	const FDisplayClusterConfigurationICVFX_CameraAdvancedRenderSettings& InAdvancedRS = InCameraSettings.RenderSettings.AdvancedRenderSettings;
 	{
-		DstViewport.RenderSettings.RenderTargetRatio = InAdvancedRS.RenderTargetRatio;
-		DstViewport.RenderSettings.GPUIndex = InAdvancedRS.GPUIndex;
-		DstViewport.RenderSettings.StereoGPUIndex = InAdvancedRS.StereoGPUIndex;
-		DstViewport.RenderSettings.RenderFamilyGroup = InAdvancedRS.RenderFamilyGroup;
-
-		FDisplayClusterViewportConfigurationHelpers::UpdateViewportStereoMode(DstViewport, InAdvancedRS.StereoMode);
+		InOutRenderSettings.RenderTargetRatio = InAdvancedRS.RenderTargetRatio;
+		InOutRenderSettings.GPUIndex = InAdvancedRS.GPUIndex;
+		InOutRenderSettings.StereoGPUIndex = InAdvancedRS.StereoGPUIndex;
+		InOutRenderSettings.bForceMono = FDisplayClusterViewportConfigurationHelpers::IsForceMonoscopicRendering(InAdvancedRS.StereoMode);
 	}
 
 	// Support inner camera custom frustum
-	FDisplayClusterViewportConfigurationHelpers_ICVFX::UpdateCameraCustomFrustum(CameraSettings.CustomFrustum, DstViewport.RenderSettings.CustomFrustumSettings);
+	FDisplayClusterViewportConfigurationHelpers_ICVFX::UpdateCameraCustomFrustum(InCameraSettings.CustomFrustum, InOutRenderSettings.CustomFrustumSettings);
 
 	// Set RenderTargetAdaptRatio
-	DstViewport.RenderSettings.RenderTargetAdaptRatio = CameraSettings.GetCameraAdaptResolutionRatio(StageSettings);
+	InOutRenderSettings.RenderTargetAdaptRatio = InCameraSettings.GetCameraAdaptResolutionRatio(*StageSettings);
 
 	// Set viewport buffer ratio
-	DstViewport.SetViewportBufferRatio(CameraSettings.GetCameraBufferRatio(StageSettings));
+	DstViewport.SetViewportBufferRatio(InCameraSettings.GetCameraBufferRatio(*StageSettings));
 
 	// Set media related configuration (runtime only for now)
 	if (IDisplayCluster::Get().GetOperationMode() == EDisplayClusterOperationMode::Cluster)
@@ -587,80 +545,88 @@ void FDisplayClusterViewportConfigurationHelpers_ICVFX::UpdateCameraViewportSett
 		static const TConsoleVariableData<int32>* const ICVarMediaEnabled = IConsoleManager::Get().FindTConsoleVariableDataInt(TEXT("nDisplay.media.Enabled"));
 		if (ICVarMediaEnabled && !!ICVarMediaEnabled->GetValueOnGameThread())
 		{
-			const FDisplayClusterConfigurationMediaICVFX& MediaICVFXSettings = InCameraComponent.CameraSettings.RenderSettings.Media;
+			const FDisplayClusterConfigurationMediaICVFX& MediaICVFXSettings = InCameraSettings.RenderSettings.Media;
 
 			if (MediaICVFXSettings.bEnable)
 			{
 				const FString ThisClusterNodeId = DstViewport.GetClusterNodeId();
 
 				// Don't render the viewport if media input assigned
-				DstViewport.RenderSettings.bSkipSceneRenderingButLeaveResourcesAvailable = MediaICVFXSettings.IsMediaInputAssigned(ThisClusterNodeId);
+				InOutRenderSettings.bSkipSceneRenderingButLeaveResourcesAvailable = MediaICVFXSettings.IsMediaInputAssigned(ThisClusterNodeId);
 
 				// Mark this viewport is going to be captured by a capture device
-				DstViewport.RenderSettings.bIsBeingCaptured = MediaICVFXSettings.IsMediaOutputAssigned(ThisClusterNodeId);
+				InOutRenderSettings.bIsBeingCaptured = MediaICVFXSettings.IsMediaOutputAssigned(ThisClusterNodeId);
 			}
 		}
 	}
 }
 
-void FDisplayClusterViewportConfigurationHelpers_ICVFX::UpdateChromakeyViewportSettings(FDisplayClusterViewport& DstViewport, FDisplayClusterViewport& InCameraViewport, ADisplayClusterRootActor& RootActor, UDisplayClusterICVFXCameraComponent& InCameraComponent)
+void FDisplayClusterViewportConfigurationHelpers_ICVFX::UpdateChromakeyViewportSettings(FDisplayClusterViewport& DstViewport, FDisplayClusterViewport& InCameraViewport, const FDisplayClusterConfigurationICVFX_CameraSettings& InCameraSettings)
 {
-	const FDisplayClusterConfigurationICVFX_StageSettings& StageSettings = RootActor.GetStageSettings();
-	const FDisplayClusterConfigurationICVFX_CameraSettings& CameraSettings = InCameraComponent.GetCameraSettingsICVFX();
-	const FDisplayClusterConfigurationICVFX_ChromakeySettings& ChromakeySettings = CameraSettings.Chromakey;
+	const FDisplayClusterConfigurationICVFX_StageSettings* StageSettings = DstViewport.Configuration->GetStageSettings();
+	if (!StageSettings)
+	{
+		return;
+	}
+
+	const FDisplayClusterConfigurationICVFX_ChromakeySettings& InChromakeySettings = InCameraSettings.Chromakey;
+
+	// Gain direct access to internal resources of the viewport:
+	FDisplayClusterViewport_RenderSettings&               InOutRenderSettings = DstViewport.GetRenderSettingsImpl();
+	FDisplayClusterViewport_RenderSettings& InOutCameraViewportRenderSettings = InCameraViewport.GetRenderSettingsImpl();
 
 	// Reset runtime flags from prev frame:
 	DstViewport.ResetRuntimeParameters();
 
 	// Chromakey used as overlay
-	DstViewport.RenderSettings.bVisible = false;
+	InOutRenderSettings.bVisible = false;
 
 	// Use special capture mode (this change RTT format and render flags)
-	DstViewport.RenderSettings.CaptureMode = EDisplayClusterViewportCaptureMode::Chromakey;
+	InOutRenderSettings.CaptureMode = EDisplayClusterViewportCaptureMode::Chromakey;
 
 	// UDisplayClusterConfigurationICVFX_ChromakeyRenderSettings
-	if (const FDisplayClusterConfigurationICVFX_ChromakeyRenderSettings* ChromakeyRenderSettings = ChromakeySettings.GetChromakeyRenderSettings(StageSettings))
+	if (const FDisplayClusterConfigurationICVFX_ChromakeyRenderSettings* InChromakeyRenderSettings = InChromakeySettings.GetChromakeyRenderSettings(*StageSettings))
 	{
-		FDisplayClusterViewportConfigurationHelpers::UpdateViewportSetting_Override(DstViewport, ChromakeyRenderSettings->Replace);
-		FDisplayClusterViewportConfigurationHelpers::UpdateViewportSetting_PostprocessBlur(DstViewport, ChromakeyRenderSettings->PostprocessBlur);
-		FDisplayClusterViewportConfigurationHelpers::UpdateViewportSetting_GenerateMips(DstViewport, ChromakeyRenderSettings->GenerateMips);
+		DstViewport.UpdateConfiguration_PostRenderOverride(InChromakeyRenderSettings->Replace);
+		DstViewport.UpdateConfiguration_PostRenderBlur(InChromakeyRenderSettings->PostprocessBlur);
+		DstViewport.UpdateConfiguration_PostRenderGenerateMips(InChromakeyRenderSettings->GenerateMips);
 
 		// Update visibility settings only for rendered viewports
-		if (!DstViewport.PostRenderSettings.Replace.IsEnabled())
+		if (!DstViewport.GetPostRenderSettings().Replace.IsEnabled())
 		{
-			check(ChromakeyRenderSettings->ShowOnlyList.IsVisibilityListValid());
+			check(InChromakeyRenderSettings->ShowOnlyList.IsVisibilityListValid());
 
-			FDisplayClusterViewportConfigurationHelpers_Visibility::UpdateShowOnlyList(DstViewport, RootActor, ChromakeyRenderSettings->ShowOnlyList);
+			FDisplayClusterViewportConfigurationHelpers_Visibility::UpdateShowOnlyList(DstViewport, InChromakeyRenderSettings->ShowOnlyList);
 		}
 
-		FDisplayClusterViewportConfigurationHelpers::UpdateViewportSetting_OverlayRenderSettings(DstViewport, ChromakeyRenderSettings->AdvancedRenderSettings);
+		DstViewport.UpdateConfiguration_OverlayRenderSettings(InChromakeyRenderSettings->AdvancedRenderSettings);
 
 		// Support custom overlay size
-		if (ChromakeyRenderSettings->CustomSize.bUseCustomSize)
+		if (InChromakeyRenderSettings->CustomSize.bUseCustomSize)
 		{
 			FIntPoint DesiredSize;
-			DesiredSize.X = ChromakeyRenderSettings->CustomSize.CustomWidth;
-			DesiredSize.Y = ChromakeyRenderSettings->CustomSize.CustomHeight;
+			DesiredSize.X = InChromakeyRenderSettings->CustomSize.CustomWidth;
+			DesiredSize.Y = InChromakeyRenderSettings->CustomSize.CustomHeight;
 
-			DstViewport.RenderSettings.Rect = FDisplayClusterViewportHelpers::GetValidViewportRect(FIntRect(FIntPoint(0, 0), DesiredSize), DstViewport.GetId(), TEXT("Configuration custom chromakey Frame Size"));
+			InOutRenderSettings.Rect = FDisplayClusterViewportHelpers::GetValidViewportRect(FIntRect(FIntPoint(0, 0), DesiredSize), DstViewport.GetId(), TEXT("Configuration custom chromakey Frame Size"));
 		}
 
 		// Debug: override the texture of the target viewport from this chromakeyRTT
-		if (ChromakeyRenderSettings->bReplaceCameraViewport)
+		if (InChromakeyRenderSettings->bReplaceCameraViewport)
 		{
-			InCameraViewport.RenderSettings.SetViewportOverride(DstViewport.GetId());
+			InOutCameraViewportRenderSettings.SetViewportOverride(DstViewport.GetId());
 		}
 	}
 
 	// Update OCIO for Chromakey Viewport
 	// Note: Chromakey OCIO is temporarily disabled
-	// FDisplayClusterViewportConfigurationHelpers_OpenColorIO::UpdateChromakeyViewport(DstViewport, RootActor, InCameraComponent);
+	// FDisplayClusterViewportConfigurationHelpers_OpenColorIO::UpdateChromakeyViewportOCIO(DstViewport, RootActor, InCameraComponent);
 
 	// Support inner camera custom frustum
-	FDisplayClusterViewportConfigurationHelpers_ICVFX::UpdateCameraCustomFrustum(CameraSettings.CustomFrustum, DstViewport.RenderSettings.CustomFrustumSettings);
+	FDisplayClusterViewportConfigurationHelpers_ICVFX::UpdateCameraCustomFrustum(InCameraSettings.CustomFrustum, InOutRenderSettings.CustomFrustumSettings);
 
 	// Attach to parent viewport
-	DstViewport.RenderSettings.AssignParentViewport(InCameraViewport.GetId(), InCameraViewport.RenderSettings);
+	InOutRenderSettings.AssignParentViewport(InCameraViewport.GetId(), InOutCameraViewportRenderSettings);
 }
 
 void FDisplayClusterViewportConfigurationHelpers_ICVFX::UpdateCameraSettings_Chromakey(FDisplayClusterShaderParameters_ICVFX::FCameraSettings& InOutCameraSettings, const FDisplayClusterConfigurationICVFX_StageSettings& InStageSettings, const FDisplayClusterConfigurationICVFX_CameraSettings& InCameraSettings, bool bEnableChromakeyMarkers, const FString& InChromakeyViewportId)
@@ -745,74 +711,87 @@ void FDisplayClusterViewportConfigurationHelpers_ICVFX::UpdateCameraSettings_Ove
 	}
 }
 
-void FDisplayClusterViewportConfigurationHelpers_ICVFX::UpdateLightcardViewportSetting(FDisplayClusterViewport& DstViewport, FDisplayClusterViewport& BaseViewport, ADisplayClusterRootActor& RootActor)
+void FDisplayClusterViewportConfigurationHelpers_ICVFX::UpdateLightcardViewportSetting(FDisplayClusterViewport& DstViewport, FDisplayClusterViewport& BaseViewport)
 {
-	const FDisplayClusterConfigurationICVFX_StageSettings& StageSettings = RootActor.GetStageSettings();
-	const FDisplayClusterConfigurationICVFX_LightcardSettings& LightcardSettings = StageSettings.Lightcard;
+	const FDisplayClusterConfigurationICVFX_StageSettings* InStageSettings = DstViewport.Configuration->GetStageSettings();
+	if (!InStageSettings)
+	{
+		return;
+	}
+
+	const FDisplayClusterConfigurationICVFX_LightcardSettings& InLightcardSettings = InStageSettings->Lightcard;
+
+	// Gain direct access to internal settings of the viewport:
+	FDisplayClusterViewport_RenderSettings&                       InOutRenderSettings = DstViewport.GetRenderSettingsImpl();
+	FDisplayClusterViewport_RenderSettings&           InOutBaseViewportRenderSettings = BaseViewport.GetRenderSettingsImpl();
+	FDisplayClusterViewport_RenderSettingsICVFX& InOutBaseViewportRenderSettingsICVFX = BaseViewport.GetRenderSettingsICVFXImpl();
 
 	// Reset runtime flags from prev frame:
 	DstViewport.ResetRuntimeParameters();
 
 	// LIghtcard texture used as overlay
-	DstViewport.RenderSettings.bVisible = false;
+	InOutRenderSettings.bVisible = false;
 
-	if (!LightcardSettings.bEnable)
+	if (!InLightcardSettings.bEnable)
 	{
 		// Disable this viewport
-		DstViewport.RenderSettings.bEnable = false;
+		InOutRenderSettings.bEnable = false;
 		return;
 	}
 
-	FDisplayClusterViewportConfigurationHelpers_Postprocess::UpdateLightcardPostProcessSettings(DstViewport, BaseViewport, RootActor);
+	FDisplayClusterViewportConfigurationHelpers_Postprocess::UpdateLightcardPostProcessSettings(DstViewport, BaseViewport);
 
 	// Update OCIO for Lightcard Viewport
-	FDisplayClusterViewportConfigurationHelpers_OpenColorIO::UpdateLightcardViewport(DstViewport, BaseViewport, RootActor);
+	FDisplayClusterViewportConfigurationHelpers_OpenColorIO::UpdateLightcardViewportOCIO(DstViewport, BaseViewport);
 
-	DstViewport.RenderSettings.CaptureMode = EDisplayClusterViewportCaptureMode::Lightcard;
+	InOutRenderSettings.CaptureMode = EDisplayClusterViewportCaptureMode::Lightcard;
 
-	const FDisplayClusterConfigurationICVFX_LightcardRenderSettings& InRenderSettings = LightcardSettings.RenderSettings;
+	const FDisplayClusterConfigurationICVFX_LightcardRenderSettings& InRenderSettings = InLightcardSettings.RenderSettings;
 	{
-		FDisplayClusterViewportConfigurationHelpers::UpdateViewportSetting_Override(DstViewport, InRenderSettings.Replace);
-		FDisplayClusterViewportConfigurationHelpers::UpdateViewportSetting_PostprocessBlur(DstViewport, InRenderSettings.PostprocessBlur);
-		FDisplayClusterViewportConfigurationHelpers::UpdateViewportSetting_GenerateMips(DstViewport, InRenderSettings.GenerateMips);
+		DstViewport.UpdateConfiguration_PostRenderOverride(InRenderSettings.Replace);
+		DstViewport.UpdateConfiguration_PostRenderBlur(InRenderSettings.PostprocessBlur);
+		DstViewport.UpdateConfiguration_PostRenderGenerateMips(InRenderSettings.GenerateMips);
 
 		// Update visibility settings only for rendered viewports
-		if (!DstViewport.PostRenderSettings.Replace.IsEnabled())
+		if (!DstViewport.GetPostRenderSettings().Replace.IsEnabled())
 		{
-			check(LightcardSettings.ShowOnlyList.IsVisibilityListValid());
+			check(InLightcardSettings.ShowOnlyList.IsVisibilityListValid());
 
-			FDisplayClusterViewportConfigurationHelpers_Visibility::UpdateShowOnlyList(DstViewport, RootActor, LightcardSettings.ShowOnlyList);
+			FDisplayClusterViewportConfigurationHelpers_Visibility::UpdateShowOnlyList(DstViewport, InLightcardSettings.ShowOnlyList);
 		}
 
-		FDisplayClusterViewportConfigurationHelpers::UpdateViewportSetting_OverlayRenderSettings(DstViewport, InRenderSettings.AdvancedRenderSettings);
+		DstViewport.UpdateConfiguration_OverlayRenderSettings(InRenderSettings.AdvancedRenderSettings);
 	}
 
 	// Attach to parent viewport
-	DstViewport.RenderSettings.AssignParentViewport(BaseViewport.GetId(), BaseViewport.RenderSettings);
+	InOutRenderSettings.AssignParentViewport(BaseViewport.GetId(), BaseViewport.GetRenderSettings());
 
-	if (EnumHasAnyFlags(DstViewport.RenderSettingsICVFX.RuntimeFlags, EDisplayClusterViewportRuntimeICVFXFlags::Lightcard))
+	if (EnumHasAnyFlags(DstViewport.GetRenderSettingsICVFX().RuntimeFlags, EDisplayClusterViewportRuntimeICVFXFlags::Lightcard))
 	{
 		// Debug: override the texture of the target viewport from this lightcard RTT
 		if (InRenderSettings.bReplaceViewport)
 		{
-			BaseViewport.RenderSettings.SetViewportOverride(DstViewport.GetId());
+			InOutBaseViewportRenderSettings.SetViewportOverride(DstViewport.GetId());
 		}
 		else
 		{
-			BaseViewport.RenderSettingsICVFX.ICVFX.LightCard.ViewportId = DstViewport.GetId();
+			InOutBaseViewportRenderSettingsICVFX.ICVFX.LightCard.ViewportId = DstViewport.GetId();
 		}
 	}
 
-	if (EnumHasAnyFlags(DstViewport.RenderSettingsICVFX.RuntimeFlags, EDisplayClusterViewportRuntimeICVFXFlags::UVLightcard))
+	if (EnumHasAnyFlags(DstViewport.GetRenderSettingsICVFX().RuntimeFlags, EDisplayClusterViewportRuntimeICVFXFlags::UVLightcard))
 	{
-		BaseViewport.RenderSettingsICVFX.ICVFX.UVLightCard.ViewportId = DstViewport.GetId();
+		InOutBaseViewportRenderSettingsICVFX.ICVFX.UVLightCard.ViewportId = DstViewport.GetId();
 	}
 
-	if (EnumHasAnyFlags(DstViewport.RenderSettingsICVFX.RuntimeFlags, EDisplayClusterViewportRuntimeICVFXFlags::Lightcard | EDisplayClusterViewportRuntimeICVFXFlags::UVLightcard))
+	if (EnumHasAnyFlags(DstViewport.GetRenderSettingsICVFX().RuntimeFlags, EDisplayClusterViewportRuntimeICVFXFlags::Lightcard | EDisplayClusterViewportRuntimeICVFXFlags::UVLightcard))
 	{
+		// Gain direct access to internal settings of the viewport:
+		FDisplayClusterViewport_RenderSettingsICVFX& InOutRenderSettingsICVFX = DstViewport.GetRenderSettingsICVFXImpl();
+
 		// Set the light card gamma used to linearize the light card textures before blending during final composite
 		// The OCIO pass will have already linearized the light card renders, so set the gamma to 1 if using OCIO on the light cards
-		BaseViewport.RenderSettingsICVFX.ICVFX.LightCardGamma = DstViewport.OpenColorIO.IsValid() ? 1.0 : 2.2;
+		InOutRenderSettingsICVFX.ICVFX.LightCardGamma = DstViewport.GetOpenColorIO().IsValid() ? 1.0 : 2.2;
 	}
 }
 
diff --git a/Engine/Plugins/Runtime/nDisplay/Source/DisplayCluster/Private/Render/Viewport/Configuration/DisplayClusterViewportConfigurationHelpers_ICVFX.h b/Engine/Plugins/Runtime/nDisplay/Source/DisplayCluster/Private/Render/Viewport/Configuration/DisplayClusterViewportConfigurationHelpers_ICVFX.h
--- a/Engine/Plugins/Runtime/nDisplay/Source/DisplayCluster/Private/Render/Viewport/Configuration/DisplayClusterViewportConfigurationHelpers_ICVFX.h
+++ b/Engine/Plugins/Runtime/nDisplay/Source/DisplayCluster/Private/Render/Viewport/Configuration/DisplayClusterViewportConfigurationHelpers_ICVFX.h
@@ -8,47 +8,35 @@
 #include "IDisplayClusterProjection.h"
 #include "Render/Projection/IDisplayClusterProjectionPolicy.h"
 
+class FDisplayClusterViewportConfiguration;
 class FDisplayClusterViewport;
-class FDisplayClusterViewportManager;
-class ADisplayClusterRootActor;
 class UDisplayClusterICVFXCameraComponent;
 
 struct FDisplayClusterConfigurationICVFX_StageSettings;
-struct FDisplayClusterConfigurationICVFX_ChromakeySettings;
-struct FDisplayClusterConfigurationICVFX_ChromakeyMarkers;
-struct FDisplayClusterConfigurationICVFX_LightcardSettings;
 struct FDisplayClusterConfigurationICVFX_CameraCustomFrustum;
 struct FDisplayClusterConfigurationICVFX_CameraSettings;
 struct FDisplayClusterViewport_CustomFrustumSettings;
 
-struct FCameraContext_ICVFX
-{
-	//@todo: add stereo context support
-	FRotator ViewRotation;
-	FVector  ViewLocation;
-	FMatrix  PrjMatrix;
-};
-
+/**
+* ICVFX configuration helper class.
+*/
 class FDisplayClusterViewportConfigurationHelpers_ICVFX
 {
 public:
-	/* Get the ViewportManager object from DCRA. */
-	static FDisplayClusterViewportManager* GetViewportManager(ADisplayClusterRootActor& RootActor);
-
 	/* Find existing InnerFrustum viewport. */
-	static FDisplayClusterViewport* FindCameraViewport(ADisplayClusterRootActor& RootActor, UDisplayClusterICVFXCameraComponent& InCameraComponent);
+	static FDisplayClusterViewport* FindCameraViewport(FDisplayClusterViewportConfiguration& InConfiguration, UDisplayClusterICVFXCameraComponent& InCameraComponent);
 
 	/* Return the new or existing InnerFrustum viewport. */
-	static FDisplayClusterViewport* GetOrCreateCameraViewport(ADisplayClusterRootActor& RootActor, UDisplayClusterICVFXCameraComponent& InCameraComponent);
+	static FDisplayClusterViewport* GetOrCreateCameraViewport(FDisplayClusterViewportConfiguration& InConfiguration, UDisplayClusterICVFXCameraComponent& InCameraComponent, const FDisplayClusterConfigurationICVFX_CameraSettings& InCameraSettings);
 
 	/* Return the new or existing Chromakey viewport. */
-	static FDisplayClusterViewport* GetOrCreateChromakeyViewport(ADisplayClusterRootActor& RootActor, UDisplayClusterICVFXCameraComponent& InCameraComponent);
+	static FDisplayClusterViewport* GetOrCreateChromakeyViewport(FDisplayClusterViewportConfiguration& InConfiguration, UDisplayClusterICVFXCameraComponent& InCameraComponent, const FDisplayClusterConfigurationICVFX_CameraSettings& InCameraSettings);
 	
 	/* Return the new or existing LightCard viewport. */
-	static FDisplayClusterViewport* GetOrCreateLightcardViewport(FDisplayClusterViewport& BaseViewport, ADisplayClusterRootActor& RootActor);
+	static FDisplayClusterViewport* GetOrCreateLightcardViewport(FDisplayClusterViewport& BaseViewport);
 
 	/* Return the new or existing UVLightCard viewport. */
-	static FDisplayClusterViewport* GetOrCreateUVLightcardViewport(FDisplayClusterViewport& BaseViewport, ADisplayClusterRootActor& RootActor);
+	static FDisplayClusterViewport* GetOrCreateUVLightcardViewport(FDisplayClusterViewport& BaseViewport);
 
 	/* Because UVLightCard viewports share the same outer texture (from LightCardManager), they can clone each other.
 	 * Except when the OCIO settings are not equal.
@@ -56,20 +44,24 @@ public:
 	 */
 	static void ReuseUVLightCardViewportWithinClusterNode(FDisplayClusterViewport& InUVLightCardViewport);
 
-#if WITH_EDITOR
-	static TArray<TSharedPtr<FDisplayClusterViewport, ESPMode::ThreadSafe>> PreviewGetRenderedInCameraViewports(ADisplayClusterRootActor& RootActor, UDisplayClusterICVFXCameraComponent& InCameraComponent, bool bGetChromakey = false);
-	static void PreviewReuseInnerFrustumViewportWithinClusterNodes_Editor(FDisplayClusterViewport& InCameraViewport, ADisplayClusterRootActor& RootActor, UDisplayClusterICVFXCameraComponent& InCameraComponent);
-	static void PreviewReuseChromakeyViewportWithinClusterNodes_Editor(FDisplayClusterViewport& InChromakeyViewport, ADisplayClusterRootActor& RootActor, UDisplayClusterICVFXCameraComponent& InCameraComponent);
-#endif
+	/** Returns all visible InnerCamera viewports. */
+	static TArray<TSharedPtr<FDisplayClusterViewport, ESPMode::ThreadSafe>> GetAllVisibleInnerCameraViewports(FDisplayClusterViewportConfiguration& InConfiguration, bool bGetChromakey = false);
+
+	static TArray<TSharedPtr<FDisplayClusterViewport, ESPMode::ThreadSafe>> PreviewGetRenderedInCameraViewports(FDisplayClusterViewportConfiguration& InConfiguration, const FString& InICVFXCameraId, bool bGetChromakey = false);
+	static void PreviewReuseInnerFrustumViewportWithinClusterNodes(FDisplayClusterViewport& InCameraViewport, UDisplayClusterICVFXCameraComponent& InCameraComponent, const FDisplayClusterConfigurationICVFX_CameraSettings& InCameraSettings);
+	static void PreviewReuseChromakeyViewportWithinClusterNodes(FDisplayClusterViewport& InChromakeyViewport, const FString& InICVFXCameraId);
+
+	static void UpdateCameraViewportSettings(FDisplayClusterViewport& DstViewport, UDisplayClusterICVFXCameraComponent& InCameraComponent, const FDisplayClusterConfigurationICVFX_CameraSettings& InCameraSettings);
+	static void UpdateChromakeyViewportSettings(FDisplayClusterViewport& DstViewport, FDisplayClusterViewport& InCameraViewport, const FDisplayClusterConfigurationICVFX_CameraSettings& InCameraSettings);
+	static void UpdateLightcardViewportSetting(FDisplayClusterViewport& DstViewport, FDisplayClusterViewport& BaseViewport);
 
-	static void UpdateCameraViewportSettings(FDisplayClusterViewport& DstViewport, ADisplayClusterRootActor& RootActor, UDisplayClusterICVFXCameraComponent& InCameraComponent);
-	static void UpdateChromakeyViewportSettings(FDisplayClusterViewport& DstViewport, FDisplayClusterViewport& InCameraViewport, ADisplayClusterRootActor& RootActor, UDisplayClusterICVFXCameraComponent& InCameraComponent);
-	static void UpdateLightcardViewportSetting(FDisplayClusterViewport& DstViewport, FDisplayClusterViewport& BaseViewport, ADisplayClusterRootActor& RootActor);
+	static bool IsCameraUsed(const FDisplayClusterConfigurationICVFX_CameraSettings& InCameraSettings);
 
-	static bool IsCameraUsed(UDisplayClusterICVFXCameraComponent& InCameraComponent);
-	static bool GetCameraContext(ADisplayClusterRootActor& RootActor, UDisplayClusterICVFXCameraComponent& InCameraComponent, FCameraContext_ICVFX& OutCameraContext);
+	static bool CreateProjectionPolicyICVFX(FDisplayClusterViewportConfiguration& InConfiguration, const FString& InViewportId, const FString& InResourceId, bool bIsCameraProjection, TSharedPtr<IDisplayClusterProjectionPolicy, ESPMode::ThreadSafe>& OutProjPolicy);
+	static bool CreateProjectionPolicyCameraICVFX(FDisplayClusterViewportConfiguration& InConfiguration, UDisplayClusterICVFXCameraComponent& CameraComponent, const FDisplayClusterConfigurationICVFX_CameraSettings& InCameraSettings, TSharedPtr<IDisplayClusterProjectionPolicy, ESPMode::ThreadSafe>& OutProjPolicy);
+	static bool UpdateCameraProjectionSettingsICVFX(FDisplayClusterViewportConfiguration& InConfiguration, UDisplayClusterICVFXCameraComponent& CameraComponent, const FDisplayClusterConfigurationICVFX_CameraSettings& InCameraSettings, const TSharedPtr<IDisplayClusterProjectionPolicy, ESPMode::ThreadSafe>& InProjectionPolicy);
 
-	static FDisplayClusterShaderParameters_ICVFX::FCameraSettings GetShaderParametersCameraSettings(const FDisplayClusterViewport& InCameraViewport, ADisplayClusterRootActor& RootActor, UDisplayClusterICVFXCameraComponent& InCameraComponent);
+	static FDisplayClusterShaderParameters_ICVFX::FCameraSettings GetShaderParametersCameraSettings(const FDisplayClusterViewport& InCameraViewport, UDisplayClusterICVFXCameraComponent& InCameraComponent, const FDisplayClusterConfigurationICVFX_CameraSettings& InCameraSettings);
 	
 	/** Configures the camera's chromakey render settings to match the specified configuration settings  */
 	static void UpdateCameraSettings_Chromakey(FDisplayClusterShaderParameters_ICVFX::FCameraSettings& InOutCameraSettings, const FDisplayClusterConfigurationICVFX_StageSettings& InStageSettings, const FDisplayClusterConfigurationICVFX_CameraSettings& InCameraSettings, bool bEnableChromakeyMarkers, const FString& InChromakeyViewportId);
@@ -77,7 +69,6 @@ public:
 
 	static void UpdateCameraCustomFrustum(const FDisplayClusterConfigurationICVFX_CameraCustomFrustum& InCameraCustomFrustumConfiguration, FDisplayClusterViewport_CustomFrustumSettings& OutCustomFrustumSettings);
 
-	static FDisplayClusterViewport* ImplFindViewport(ADisplayClusterRootActor& RootActor, const FString& InViewportId, const FString& InResourceId);
-private:
-	static FDisplayClusterViewport* ImplCreateViewport(ADisplayClusterRootActor& RootActor, const FString& InViewportId, const FString& InResourceId, TSharedPtr<IDisplayClusterProjectionPolicy, ESPMode::ThreadSafe>& InProjectionPolicy);
+	static FDisplayClusterViewport* FindViewportICVFX(FDisplayClusterViewportConfiguration& InConfiguration, const FString& InViewportId, const FString& InResourceId);
+	static FDisplayClusterViewport* CreateViewportICVFX(FDisplayClusterViewportConfiguration& InConfiguration, const FString& InViewportId, const FString& InResourceId, TSharedPtr<IDisplayClusterProjectionPolicy, ESPMode::ThreadSafe>& InProjectionPolicy);
 };
diff --git a/Engine/Plugins/Runtime/nDisplay/Source/DisplayCluster/Private/Render/Viewport/Configuration/DisplayClusterViewportConfigurationHelpers_OpenColorIO.cpp b/Engine/Plugins/Runtime/nDisplay/Source/DisplayCluster/Private/Render/Viewport/Configuration/DisplayClusterViewportConfigurationHelpers_OpenColorIO.cpp
--- a/Engine/Plugins/Runtime/nDisplay/Source/DisplayCluster/Private/Render/Viewport/Configuration/DisplayClusterViewportConfigurationHelpers_OpenColorIO.cpp
+++ b/Engine/Plugins/Runtime/nDisplay/Source/DisplayCluster/Private/Render/Viewport/Configuration/DisplayClusterViewportConfigurationHelpers_OpenColorIO.cpp
@@ -11,16 +11,16 @@
 #include "Render/Viewport/DisplayClusterViewportManager.h"
 
 #include "Components/DisplayClusterICVFXCameraComponent.h"
-#include "OpenColorIODisplayExtension.h"
 
 ///////////////////////////////////////////////////////////////////////////////////////
 // FDisplayClusterViewportConfigurationHelpers_OpenColorIO
 ///////////////////////////////////////////////////////////////////////////////////////
-bool FDisplayClusterViewportConfigurationHelpers_OpenColorIO::UpdateBaseViewport(FDisplayClusterViewport& DstViewport, ADisplayClusterRootActor& RootActor, const UDisplayClusterConfigurationViewport& InViewportConfiguration)
+bool FDisplayClusterViewportConfigurationHelpers_OpenColorIO::UpdateBaseViewportOCIO(FDisplayClusterViewport& DstViewport, const UDisplayClusterConfigurationViewport& InViewportConfiguration)
 {
-	if (!EnumHasAnyFlags(DstViewport.RenderSettingsICVFX.RuntimeFlags, EDisplayClusterViewportRuntimeICVFXFlags::InCamera | EDisplayClusterViewportRuntimeICVFXFlags::Chromakey | EDisplayClusterViewportRuntimeICVFXFlags::Lightcard | EDisplayClusterViewportRuntimeICVFXFlags::UVLightcard))
+	if (!EnumHasAnyFlags(DstViewport.GetRenderSettingsICVFX().RuntimeFlags, EDisplayClusterViewportRuntimeICVFXFlags::InCamera | EDisplayClusterViewportRuntimeICVFXFlags::Chromakey | EDisplayClusterViewportRuntimeICVFXFlags::Lightcard | EDisplayClusterViewportRuntimeICVFXFlags::UVLightcard))
 	{
-		if (const FOpenColorIOColorConversionSettings* OCIOConfiguration = RootActor.GetStageSettings().FindViewportOCIOConfiguration(DstViewport.GetId()))
+		const FDisplayClusterConfigurationICVFX_StageSettings* StageSettings =DstViewport.Configuration->GetStageSettings();
+		if (const FOpenColorIOColorConversionSettings* OCIOConfiguration = StageSettings ? StageSettings->FindViewportOCIOConfiguration(DstViewport.GetId()) : nullptr)
 		{
 			ApplyOCIOConfiguration(DstViewport, *OCIOConfiguration);
 
@@ -33,11 +33,12 @@ bool FDisplayClusterViewportConfigurationHelpers_OpenColorIO::UpdateBaseViewport
 	return false;
 }
 
-bool FDisplayClusterViewportConfigurationHelpers_OpenColorIO::UpdateLightcardViewport(FDisplayClusterViewport& DstViewport, FDisplayClusterViewport& BaseViewport, ADisplayClusterRootActor& RootActor)
+bool FDisplayClusterViewportConfigurationHelpers_OpenColorIO::UpdateLightcardViewportOCIO(FDisplayClusterViewport& DstViewport, FDisplayClusterViewport& BaseViewport)
 {
-	if (EnumHasAnyFlags(DstViewport.RenderSettingsICVFX.RuntimeFlags, EDisplayClusterViewportRuntimeICVFXFlags::Lightcard | EDisplayClusterViewportRuntimeICVFXFlags::UVLightcard))
+	if (EnumHasAnyFlags(DstViewport.GetRenderSettingsICVFX().RuntimeFlags, EDisplayClusterViewportRuntimeICVFXFlags::Lightcard | EDisplayClusterViewportRuntimeICVFXFlags::UVLightcard))
 	{
-		if (const FOpenColorIOColorConversionSettings* OCIOConfiguration = RootActor.GetStageSettings().FindLightcardOCIOConfiguration(BaseViewport.GetId()))
+		const FDisplayClusterConfigurationICVFX_StageSettings* StageSettings = DstViewport.Configuration->GetStageSettings();
+		if (const FOpenColorIOColorConversionSettings* OCIOConfiguration = StageSettings ? StageSettings->FindLightcardOCIOConfiguration(BaseViewport.GetId()): nullptr)
 		{
 			ApplyOCIOConfiguration(DstViewport, *OCIOConfiguration);
 
@@ -51,11 +52,11 @@ bool FDisplayClusterViewportConfigurationHelpers_OpenColorIO::UpdateLightcardVie
 	return false;
 }
 
-bool FDisplayClusterViewportConfigurationHelpers_OpenColorIO::UpdateCameraViewport(FDisplayClusterViewport& DstViewport, ADisplayClusterRootActor& RootActor, UDisplayClusterICVFXCameraComponent& InCameraComponent)
+bool FDisplayClusterViewportConfigurationHelpers_OpenColorIO::UpdateCameraViewportOCIO(FDisplayClusterViewport& DstViewport, const FDisplayClusterConfigurationICVFX_CameraSettings& InCameraSettings)
 {
-	if (EnumHasAllFlags(DstViewport.RenderSettingsICVFX.RuntimeFlags, EDisplayClusterViewportRuntimeICVFXFlags::InCamera))
+	if (EnumHasAllFlags(DstViewport.GetRenderSettingsICVFX().RuntimeFlags, EDisplayClusterViewportRuntimeICVFXFlags::InCamera))
 	{
-		if (const FOpenColorIOColorConversionSettings* OCIOConfiguration = InCameraComponent.GetCameraSettingsICVFX().FindInnerFrustumOCIOConfiguration(DstViewport.GetClusterNodeId()))
+		if (const FOpenColorIOColorConversionSettings* OCIOConfiguration = InCameraSettings.FindInnerFrustumOCIOConfiguration(DstViewport.GetClusterNodeId()))
 		{
 			ApplyOCIOConfiguration(DstViewport, *OCIOConfiguration);
 
@@ -68,11 +69,11 @@ bool FDisplayClusterViewportConfigurationHelpers_OpenColorIO::UpdateCameraViewpo
 	return false;
 }
 
-bool FDisplayClusterViewportConfigurationHelpers_OpenColorIO::UpdateChromakeyViewport(FDisplayClusterViewport& DstViewport, ADisplayClusterRootActor& RootActor, UDisplayClusterICVFXCameraComponent& InCameraComponent)
+bool FDisplayClusterViewportConfigurationHelpers_OpenColorIO::UpdateChromakeyViewportOCIO(FDisplayClusterViewport& DstViewport, const FDisplayClusterConfigurationICVFX_CameraSettings& InCameraSettings)
 {
-	if (EnumHasAllFlags(DstViewport.RenderSettingsICVFX.RuntimeFlags, EDisplayClusterViewportRuntimeICVFXFlags::Chromakey))
+	if (EnumHasAllFlags(DstViewport.GetRenderSettingsICVFX().RuntimeFlags, EDisplayClusterViewportRuntimeICVFXFlags::Chromakey))
 	{
-		if (const FOpenColorIOColorConversionSettings* OCIOConfiguration = InCameraComponent.GetCameraSettingsICVFX().FindChromakeyOCIOConfiguration(DstViewport.GetClusterNodeId()))
+		if (const FOpenColorIOColorConversionSettings* OCIOConfiguration = InCameraSettings.FindChromakeyOCIOConfiguration(DstViewport.GetClusterNodeId()))
 		{
 			ApplyOCIOConfiguration(DstViewport, *OCIOConfiguration);
 
@@ -87,30 +88,27 @@ bool FDisplayClusterViewportConfigurationHelpers_OpenColorIO::UpdateChromakeyVie
 
 void FDisplayClusterViewportConfigurationHelpers_OpenColorIO::ApplyOCIOConfiguration(FDisplayClusterViewport& DstViewport, const FOpenColorIOColorConversionSettings& InConversionSettings)
 {
-	if (DstViewport.OpenColorIO.IsValid() && DstViewport.OpenColorIO->IsConversionSettingsEqual(InConversionSettings))
+	if (DstViewport.GetOpenColorIO().IsValid() && DstViewport.GetOpenColorIO()->IsConversionSettingsEqual(InConversionSettings))
 	{
 		// Already assigned
 		return;
 	}
 
-	DstViewport.OpenColorIO.Reset();
-	DstViewport.OpenColorIO = MakeShared<FDisplayClusterViewport_OpenColorIO>(InConversionSettings);
+	DstViewport.SetOpenColorIO(MakeShared<FDisplayClusterViewport_OpenColorIO>(InConversionSettings));
 }
 
 void FDisplayClusterViewportConfigurationHelpers_OpenColorIO::DisableOCIOConfiguration(FDisplayClusterViewport& DstViewport)
 {
 	// Remove OICO ref
-	DstViewport.OpenColorIO.Reset();
+	DstViewport.SetOpenColorIO(nullptr);
 }
 
-#if WITH_EDITOR
-bool FDisplayClusterViewportConfigurationHelpers_OpenColorIO::IsInnerFrustumViewportSettingsEqual_Editor(const FDisplayClusterViewport& InViewport1, const FDisplayClusterViewport& InViewport2, UDisplayClusterICVFXCameraComponent& InCameraComponent)
+bool FDisplayClusterViewportConfigurationHelpers_OpenColorIO::IsInnerFrustumViewportOCIOSettingsEqual(const FDisplayClusterViewport& InViewport1, const FDisplayClusterViewport& InViewport2, const FDisplayClusterConfigurationICVFX_CameraSettings& InCameraSettings)
 {
-	return InCameraComponent.GetCameraSettingsICVFX().IsInnerFrustumViewportSettingsEqual_Editor(InViewport1.GetClusterNodeId(), InViewport2.GetClusterNodeId());
+	return InCameraSettings.IsInnerFrustumViewportSettingsEqual(InViewport1.GetClusterNodeId(), InViewport2.GetClusterNodeId());
 }
 
-bool FDisplayClusterViewportConfigurationHelpers_OpenColorIO::IsChromakeyViewportSettingsEqual_Editor(const FDisplayClusterViewport& InViewport1, const FDisplayClusterViewport& InViewport2, UDisplayClusterICVFXCameraComponent& InCameraComponent)
+bool FDisplayClusterViewportConfigurationHelpers_OpenColorIO::IsChromakeyViewportOCIOSettingsEqual(const FDisplayClusterViewport& InViewport1, const FDisplayClusterViewport& InViewport2, const FDisplayClusterConfigurationICVFX_CameraSettings& InCameraSettings)
 {
-	return InCameraComponent.GetCameraSettingsICVFX().IsChromakeyViewportSettingsEqual_Editor(InViewport1.GetClusterNodeId(), InViewport2.GetClusterNodeId());
+	return InCameraSettings.IsChromakeyViewportSettingsEqual(InViewport1.GetClusterNodeId(), InViewport2.GetClusterNodeId());
 }
-#endif
diff --git a/Engine/Plugins/Runtime/nDisplay/Source/DisplayCluster/Private/Render/Viewport/Configuration/DisplayClusterViewportConfigurationHelpers_OpenColorIO.h b/Engine/Plugins/Runtime/nDisplay/Source/DisplayCluster/Private/Render/Viewport/Configuration/DisplayClusterViewportConfigurationHelpers_OpenColorIO.h
--- a/Engine/Plugins/Runtime/nDisplay/Source/DisplayCluster/Private/Render/Viewport/Configuration/DisplayClusterViewportConfigurationHelpers_OpenColorIO.h
+++ b/Engine/Plugins/Runtime/nDisplay/Source/DisplayCluster/Private/Render/Viewport/Configuration/DisplayClusterViewportConfigurationHelpers_OpenColorIO.h
@@ -5,13 +5,14 @@
 #include "CoreMinimal.h"
 
 class FDisplayClusterViewport;
-class ADisplayClusterRootActor;
-
 class UDisplayClusterConfigurationViewport;
-class UDisplayClusterICVFXCameraComponent;
 
+struct FDisplayClusterConfigurationICVFX_CameraSettings;
 struct FOpenColorIOColorConversionSettings;
 
+/**
+* OCIO configuration helper class.
+*/
 class FDisplayClusterViewportConfigurationHelpers_OpenColorIO
 {
 public:
@@ -23,7 +24,7 @@ public:
 	 *
 	 * @return - true if success.
 	 */
-	static bool UpdateBaseViewport(FDisplayClusterViewport& DstViewport, ADisplayClusterRootActor& RootActor, const UDisplayClusterConfigurationViewport& InViewportConfiguration);
+	static bool UpdateBaseViewportOCIO(FDisplayClusterViewport& DstViewport, const UDisplayClusterConfigurationViewport& InViewportConfiguration);
 
 	/** Update OCIO for LightCard viewport.
 	 *
@@ -33,7 +34,7 @@ public:
 	 *
 	 * @return - true if success.
 	 */
-	static bool UpdateLightcardViewport(FDisplayClusterViewport& DstViewport, FDisplayClusterViewport& BaseViewport, ADisplayClusterRootActor& RootActor);
+	static bool UpdateLightcardViewportOCIO(FDisplayClusterViewport& DstViewport, FDisplayClusterViewport& BaseViewport);
 
 	/** Update OCIO for InnerFrustum(InCamera) viewport.
 	 *
@@ -43,7 +44,7 @@ public:
 	 *
 	 * @return - true if success
 	 */
-	static bool UpdateCameraViewport(FDisplayClusterViewport& DstViewport, ADisplayClusterRootActor& RootActor, UDisplayClusterICVFXCameraComponent& InCameraComponent);
+	static bool UpdateCameraViewportOCIO(FDisplayClusterViewport& DstViewport, const FDisplayClusterConfigurationICVFX_CameraSettings& InCameraSettings);
 
 	/** Update OCIO for Chromakey viewport.
 	 *
@@ -53,9 +54,8 @@ public:
 	 *
 	 * @return - true if success
 	 */
-	static bool UpdateChromakeyViewport(FDisplayClusterViewport& DstViewport, ADisplayClusterRootActor& RootActor, UDisplayClusterICVFXCameraComponent& InCameraComponent);
+	static bool UpdateChromakeyViewportOCIO(FDisplayClusterViewport& DstViewport, const FDisplayClusterConfigurationICVFX_CameraSettings& InCameraSettings);
 
-#if WITH_EDITOR
 	/** Compare InnerFrustum OCIO configuration for input viewports
 	 *
 	 * @param InViewport1       - Viewport #1
@@ -64,7 +64,7 @@ public:
 	 *
 	 * @return - true if OCIO is equals
 	 */
-	static bool IsInnerFrustumViewportSettingsEqual_Editor(const FDisplayClusterViewport& InViewport1, const FDisplayClusterViewport& InViewport2, UDisplayClusterICVFXCameraComponent& InCameraComponent);
+	static bool IsInnerFrustumViewportOCIOSettingsEqual(const FDisplayClusterViewport& InViewport1, const FDisplayClusterViewport& InViewport2, const FDisplayClusterConfigurationICVFX_CameraSettings& InCameraSettings);
 
 	/** Compare Chromakey OCIO configuration for input viewports
 	 *
@@ -74,8 +74,7 @@ public:
 	 *
 	 * @return - true if OCIO is equals
 	 */
-	static bool IsChromakeyViewportSettingsEqual_Editor(const FDisplayClusterViewport& InViewport1, const FDisplayClusterViewport& InViewport2, UDisplayClusterICVFXCameraComponent& InCameraComponent);
-#endif
+	static bool IsChromakeyViewportOCIOSettingsEqual(const FDisplayClusterViewport& InViewport1, const FDisplayClusterViewport& InViewport2, const FDisplayClusterConfigurationICVFX_CameraSettings& InCameraSettings);
 
 private:
 	/** Apply OCIO to viewport.
