nDisplay: Added new interface IDisplayClusterViewportConfiguration

diff --git a/Engine/Plugins/Runtime/nDisplay/Source/DisplayCluster/Private/Render/Viewport/Configuration/DisplayClusterViewportConfigurationHelpers_Postprocess.cpp b/Engine/Plugins/Runtime/nDisplay/Source/DisplayCluster/Private/Render/Viewport/Configuration/DisplayClusterViewportConfigurationHelpers_Postprocess.cpp
--- a/Engine/Plugins/Runtime/nDisplay/Source/DisplayCluster/Private/Render/Viewport/Configuration/DisplayClusterViewportConfigurationHelpers_Postprocess.cpp
+++ b/Engine/Plugins/Runtime/nDisplay/Source/DisplayCluster/Private/Render/Viewport/Configuration/DisplayClusterViewportConfigurationHelpers_Postprocess.cpp
@@ -9,169 +9,346 @@
 #include "DisplayClusterRootActor.h"
 #include "Components/DisplayClusterICVFXCameraComponent.h"
 
-static void ImplUpdateCustomPostprocess(FDisplayClusterViewport& DstViewport, bool bEnabled, const FDisplayClusterConfigurationViewport_CustomPostprocessSettings& InCustomPostprocess, IDisplayClusterViewport_CustomPostProcessSettings::ERenderPass RenderPass)
+namespace UE::DisplayCluster::Configuration::PostprocessHelpers
 {
-	if (bEnabled)
+	static inline void ImplUpdateCustomPostprocess(FDisplayClusterViewport& DstViewport, bool bEnabled, const FDisplayClusterConfigurationViewport_CustomPostprocessSettings& InCustomPostprocess, IDisplayClusterViewport_CustomPostProcessSettings::ERenderPass RenderPass)
 	{
-		DstViewport.CustomPostProcessSettings.AddCustomPostProcess(RenderPass, InCustomPostprocess.PostProcessSettings, InCustomPostprocess.BlendWeight, InCustomPostprocess.bIsOneFrame);
+		if (bEnabled)
+		{
+			DstViewport.GetViewport_CustomPostProcessSettings().AddCustomPostProcess(RenderPass, InCustomPostprocess.PostProcessSettings, InCustomPostprocess.BlendWeight, InCustomPostprocess.bIsOneFrame);
+		}
+		else
+		{
+			DstViewport.GetViewport_CustomPostProcessSettings().RemoveCustomPostProcess(RenderPass);
+		}
 	}
-	else
+
+	static inline void ImplRemoveCustomPostprocess(FDisplayClusterViewport& DstViewport, IDisplayClusterViewport_CustomPostProcessSettings::ERenderPass RenderPass)
 	{
-		DstViewport.CustomPostProcessSettings.RemoveCustomPostProcess(RenderPass);
+		DstViewport.GetViewport_CustomPostProcessSettings().RemoveCustomPostProcess(RenderPass);
 	}
-}
 
-static void ImplRemoveCustomPostprocess(FDisplayClusterViewport& DstViewport, IDisplayClusterViewport_CustomPostProcessSettings::ERenderPass RenderPass)
-{
-	DstViewport.CustomPostProcessSettings.RemoveCustomPostProcess(RenderPass);
-}
+	static inline bool ImplUpdatePerViewportColorGrading(FDisplayClusterViewport& DstViewport, const FDisplayClusterConfigurationViewport_PerViewportColorGrading PerViewportColorGrading)
+	{
+		const FDisplayClusterConfigurationICVFX_StageSettings* StageSettings = DstViewport.Configuration->GetStageSettings();
+		if (!StageSettings)
+		{
+			return false;
+		}
 
+		// enable entire cluster only when global settings is on
+		const bool bUseEntireClusterPostProcess = StageSettings->EntireClusterColorGrading.bEnableEntireClusterColorGrading && PerViewportColorGrading.bIsEntireClusterEnabled;
 
-static bool ImplUpdatePerViewportColorGrading(FDisplayClusterViewport& DstViewport, ADisplayClusterRootActor& RootActor, const FDisplayClusterConfigurationViewport_PerViewportColorGrading PerViewportColorGrading)
-{
-	// Check and apply the overall cluster post process settings
-	const FDisplayClusterConfigurationICVFX_StageSettings& StageSettings = RootActor.GetStageSettings();
+		FDisplayClusterConfigurationViewport_CustomPostprocessSettings FinalPerViewportColorGrading;
+		FinalPerViewportColorGrading.bIsEnabled = true;
+		FinalPerViewportColorGrading.bIsOneFrame = true;
+		FinalPerViewportColorGrading.BlendWeight = 1;
 
-	// enable entire cluster only when global settings is on
-	const bool bUseEntireClusterPostProcess = StageSettings.EntireClusterColorGrading.bEnableEntireClusterColorGrading && PerViewportColorGrading.bIsEntireClusterEnabled;
-	
-	FDisplayClusterConfigurationViewport_CustomPostprocessSettings FinalPerViewportColorGrading;
-	FinalPerViewportColorGrading.bIsEnabled = true;
-	FinalPerViewportColorGrading.bIsOneFrame = true;
-	FinalPerViewportColorGrading.BlendWeight = 1;
+		// blend with entire cluster
+		if (bUseEntireClusterPostProcess)
+		{
+			FinalPerViewportColorGrading.BlendWeight = StageSettings->EntireClusterColorGrading.ColorGradingSettings.BlendWeight;
 
-	// blend with entire cluster
-	if (bUseEntireClusterPostProcess)
-	{
-		FinalPerViewportColorGrading.BlendWeight = StageSettings.EntireClusterColorGrading.ColorGradingSettings.BlendWeight;
+			FDisplayClusterViewportConfigurationHelpers_Postprocess::BlendPostProcessSettings(FinalPerViewportColorGrading.PostProcessSettings, StageSettings->EntireClusterColorGrading.ColorGradingSettings, PerViewportColorGrading.ColorGradingSettings);
+			FinalPerViewportColorGrading.BlendWeight *= PerViewportColorGrading.ColorGradingSettings.BlendWeight;
 
-		FDisplayClusterViewportConfigurationHelpers_Postprocess::BlendPostProcessSettings(FinalPerViewportColorGrading.PostProcessSettings, StageSettings.EntireClusterColorGrading.ColorGradingSettings, PerViewportColorGrading.ColorGradingSettings);
-		FinalPerViewportColorGrading.BlendWeight *= PerViewportColorGrading.ColorGradingSettings.BlendWeight;
+			ImplUpdateCustomPostprocess(DstViewport, FinalPerViewportColorGrading.bIsEnabled, FinalPerViewportColorGrading, IDisplayClusterViewport_CustomPostProcessSettings::ERenderPass::FinalPerViewport);
+			return true;
+		}
+		else
+		{
+			// pass color grading without blending with entire cluster
+			FDisplayClusterViewportConfigurationHelpers_Postprocess::CopyBlendPostProcessSettings(FinalPerViewportColorGrading.PostProcessSettings, PerViewportColorGrading.ColorGradingSettings);
+			FinalPerViewportColorGrading.BlendWeight *= PerViewportColorGrading.ColorGradingSettings.BlendWeight;
 
-		ImplUpdateCustomPostprocess(DstViewport, FinalPerViewportColorGrading.bIsEnabled, FinalPerViewportColorGrading, IDisplayClusterViewport_CustomPostProcessSettings::ERenderPass::FinalPerViewport);
-		return true;
-	}
-	else
-	{
-		// pass color grading without blending with entire cluster
-		FDisplayClusterViewportConfigurationHelpers_Postprocess::CopyBlendPostProcessSettings(FinalPerViewportColorGrading.PostProcessSettings, PerViewportColorGrading.ColorGradingSettings);
-		FinalPerViewportColorGrading.BlendWeight *= PerViewportColorGrading.ColorGradingSettings.BlendWeight;
+			ImplUpdateCustomPostprocess(DstViewport, FinalPerViewportColorGrading.bIsEnabled, FinalPerViewportColorGrading, IDisplayClusterViewport_CustomPostProcessSettings::ERenderPass::FinalPerViewport);
+			return true;
+		}
 
-		ImplUpdateCustomPostprocess(DstViewport, FinalPerViewportColorGrading.bIsEnabled, FinalPerViewportColorGrading, IDisplayClusterViewport_CustomPostProcessSettings::ERenderPass::FinalPerViewport);
-		return true;
+		return false;
 	}
 
-	return false;
-}
-
-static bool ImplUpdateEntireClusterColorGrading(FDisplayClusterViewport& DstViewport, ADisplayClusterRootActor& RootActor)
-{
-	const FDisplayClusterConfigurationICVFX_StageSettings& StageSettings = RootActor.GetStageSettings();
-	if (StageSettings.EntireClusterColorGrading.bEnableEntireClusterColorGrading)
+	static inline bool ImplUpdateEntireClusterColorGrading(FDisplayClusterViewport& DstViewport)
 	{
-		FDisplayClusterConfigurationViewport_CustomPostprocessSettings FinalEntireClusterColorGrading;
-		FinalEntireClusterColorGrading.bIsEnabled = true;
-		FinalEntireClusterColorGrading.bIsOneFrame = true;
-		FinalEntireClusterColorGrading.BlendWeight = 1;
+		const FDisplayClusterConfigurationICVFX_StageSettings* StageSettings = DstViewport.Configuration->GetStageSettings();
+		if (StageSettings && StageSettings->EntireClusterColorGrading.bEnableEntireClusterColorGrading)
+		{
+			FDisplayClusterConfigurationViewport_CustomPostprocessSettings FinalEntireClusterColorGrading;
+			FinalEntireClusterColorGrading.bIsEnabled = true;
+			FinalEntireClusterColorGrading.bIsOneFrame = true;
+			FinalEntireClusterColorGrading.BlendWeight = 1;
 
-		FDisplayClusterViewportConfigurationHelpers_Postprocess::CopyBlendPostProcessSettings(FinalEntireClusterColorGrading.PostProcessSettings, StageSettings.EntireClusterColorGrading.ColorGradingSettings);
-		ImplUpdateCustomPostprocess(DstViewport, FinalEntireClusterColorGrading.bIsEnabled, FinalEntireClusterColorGrading, IDisplayClusterViewport_CustomPostProcessSettings::ERenderPass::FinalPerViewport);
-		return true;
+			FDisplayClusterViewportConfigurationHelpers_Postprocess::CopyBlendPostProcessSettings(FinalEntireClusterColorGrading.PostProcessSettings, StageSettings->EntireClusterColorGrading.ColorGradingSettings);
+			ImplUpdateCustomPostprocess(DstViewport, FinalEntireClusterColorGrading.bIsEnabled, FinalEntireClusterColorGrading, IDisplayClusterViewport_CustomPostProcessSettings::ERenderPass::FinalPerViewport);
+			return true;
+		}
+
+		return false;
 	}
-	return false;
-}
 
-static bool ImplUpdateIncameraPerNodeColorGrading(FDisplayClusterViewport& DstViewport, ADisplayClusterRootActor& RootActor, const FDisplayClusterConfigurationViewport_AllNodesColorGrading AllNodesColorGrading, const FDisplayClusterConfigurationViewport_PerNodeColorGrading PerNodeColorGrading)
-{
-	const FDisplayClusterConfigurationICVFX_StageSettings& StageSettings = RootActor.GetStageSettings();
+	static inline bool ImplUpdateIncameraPerNodeColorGrading(FDisplayClusterViewport& DstViewport, const FDisplayClusterConfigurationViewport_AllNodesColorGrading AllNodesColorGrading, const FDisplayClusterConfigurationViewport_PerNodeColorGrading PerNodeColorGrading)
+	{
+		const FDisplayClusterConfigurationICVFX_StageSettings* StageSettings = DstViewport.Configuration->GetStageSettings();
+		if (!StageSettings)
+		{
+			return false;
+		}
 
-	FDisplayClusterConfigurationViewport_CustomPostprocessSettings FinalPerNodeColorGrading;
-	FinalPerNodeColorGrading.bIsEnabled = true;
-	FinalPerNodeColorGrading.bIsOneFrame = true;
-	FinalPerNodeColorGrading.BlendWeight = 1;
+		FDisplayClusterConfigurationViewport_CustomPostprocessSettings FinalPerNodeColorGrading;
+		FinalPerNodeColorGrading.bIsEnabled = true;
+		FinalPerNodeColorGrading.bIsOneFrame = true;
+		FinalPerNodeColorGrading.BlendWeight = 1;
 
-	const bool bIncludeUseEntireClusterPostProcess = StageSettings.EntireClusterColorGrading.bEnableEntireClusterColorGrading && PerNodeColorGrading.bEntireClusterColorGrading;
-	const bool bIncludeAllNodesColorGrading = AllNodesColorGrading.bEnableInnerFrustumAllNodesColorGrading && PerNodeColorGrading.bAllNodesColorGrading;
+		const bool bIncludeUseEntireClusterPostProcess = StageSettings->EntireClusterColorGrading.bEnableEntireClusterColorGrading && PerNodeColorGrading.bEntireClusterColorGrading;
+		const bool bIncludeAllNodesColorGrading = AllNodesColorGrading.bEnableInnerFrustumAllNodesColorGrading && PerNodeColorGrading.bAllNodesColorGrading;
 
-	if (bIncludeUseEntireClusterPostProcess)
-	{
-		if (bIncludeAllNodesColorGrading)
+		if (bIncludeUseEntireClusterPostProcess)
 		{
-			// all three options are enabled - cluster + all + node
-			FDisplayClusterViewportConfigurationHelpers_Postprocess::PerNodeBlendPostProcessSettings(FinalPerNodeColorGrading.PostProcessSettings, StageSettings.EntireClusterColorGrading.ColorGradingSettings, AllNodesColorGrading.ColorGradingSettings, PerNodeColorGrading.ColorGradingSettings);
-			FinalPerNodeColorGrading.BlendWeight *= AllNodesColorGrading.ColorGradingSettings.BlendWeight * PerNodeColorGrading.ColorGradingSettings.BlendWeight;
+			if (bIncludeAllNodesColorGrading)
+			{
+				// all three options are enabled - cluster + all + node
+				FDisplayClusterViewportConfigurationHelpers_Postprocess::PerNodeBlendPostProcessSettings(FinalPerNodeColorGrading.PostProcessSettings, StageSettings->EntireClusterColorGrading.ColorGradingSettings, AllNodesColorGrading.ColorGradingSettings, PerNodeColorGrading.ColorGradingSettings);
+				FinalPerNodeColorGrading.BlendWeight *= AllNodesColorGrading.ColorGradingSettings.BlendWeight * PerNodeColorGrading.ColorGradingSettings.BlendWeight;
 
-			ImplUpdateCustomPostprocess(DstViewport, FinalPerNodeColorGrading.bIsEnabled, FinalPerNodeColorGrading, IDisplayClusterViewport_CustomPostProcessSettings::ERenderPass::FinalPerViewport);
+				ImplUpdateCustomPostprocess(DstViewport, FinalPerNodeColorGrading.bIsEnabled, FinalPerNodeColorGrading, IDisplayClusterViewport_CustomPostProcessSettings::ERenderPass::FinalPerViewport);
+			}
+			else
+			{
+				// only cluster + node
+				FDisplayClusterViewportConfigurationHelpers_Postprocess::BlendPostProcessSettings(FinalPerNodeColorGrading.PostProcessSettings, StageSettings->EntireClusterColorGrading.ColorGradingSettings, PerNodeColorGrading.ColorGradingSettings);
+				FinalPerNodeColorGrading.BlendWeight *= PerNodeColorGrading.ColorGradingSettings.BlendWeight;
+				ImplUpdateCustomPostprocess(DstViewport, FinalPerNodeColorGrading.bIsEnabled, FinalPerNodeColorGrading, IDisplayClusterViewport_CustomPostProcessSettings::ERenderPass::FinalPerViewport);
+			}
+
+			return true;
 		}
 		else
 		{
-			// only cluster + node
-			FDisplayClusterViewportConfigurationHelpers_Postprocess::BlendPostProcessSettings(FinalPerNodeColorGrading.PostProcessSettings, StageSettings.EntireClusterColorGrading.ColorGradingSettings, PerNodeColorGrading.ColorGradingSettings);
-			FinalPerNodeColorGrading.BlendWeight *= PerNodeColorGrading.ColorGradingSettings.BlendWeight;
+			// entire cluster settings disabled, only all nodes cases
+			FinalPerNodeColorGrading.BlendWeight = AllNodesColorGrading.ColorGradingSettings.BlendWeight;
+
+			if (bIncludeAllNodesColorGrading)
+			{
+				// all nodes + node
+				FDisplayClusterViewportConfigurationHelpers_Postprocess::BlendPostProcessSettings(FinalPerNodeColorGrading.PostProcessSettings, AllNodesColorGrading.ColorGradingSettings, PerNodeColorGrading.ColorGradingSettings);
+				FinalPerNodeColorGrading.BlendWeight *= PerNodeColorGrading.ColorGradingSettings.BlendWeight;
+			}
+			else
+			{
+				// node only
+				FDisplayClusterViewportConfigurationHelpers_Postprocess::CopyBlendPostProcessSettings(FinalPerNodeColorGrading.PostProcessSettings, PerNodeColorGrading.ColorGradingSettings);
+			}
+
 			ImplUpdateCustomPostprocess(DstViewport, FinalPerNodeColorGrading.bIsEnabled, FinalPerNodeColorGrading, IDisplayClusterViewport_CustomPostProcessSettings::ERenderPass::FinalPerViewport);
+			return true;
 		}
 
-		return true;
+		return false;
 	}
-	else
+
+	static inline bool ImplUpdateIncameraAllNodesColorGrading(FDisplayClusterViewport& DstViewport, const FDisplayClusterConfigurationViewport_AllNodesColorGrading AllNodesColorGrading)
 	{
-		// entire cluster settings disabled, only all nodes cases
-		FinalPerNodeColorGrading.BlendWeight = AllNodesColorGrading.ColorGradingSettings.BlendWeight;
+		const FDisplayClusterConfigurationICVFX_StageSettings* StageSettings = DstViewport.Configuration->GetStageSettings();
+		if (!StageSettings)
+		{
+			return false;
+		}
 
-		if (bIncludeAllNodesColorGrading)
+		FDisplayClusterConfigurationViewport_CustomPostprocessSettings FinalAllNodesColorGrading;
+		FinalAllNodesColorGrading.bIsEnabled = true;
+		FinalAllNodesColorGrading.bIsOneFrame = true;
+		FinalAllNodesColorGrading.BlendWeight = 1;
+
+
+		// enable entire cluster only when global settings is on
+		const bool bEnableEntireClusterColorGrading = StageSettings->EntireClusterColorGrading.bEnableEntireClusterColorGrading && AllNodesColorGrading.bEnableEntireClusterColorGrading;
+
+		if (bEnableEntireClusterColorGrading)
 		{
-			// all nodes + node
-			FDisplayClusterViewportConfigurationHelpers_Postprocess::BlendPostProcessSettings(FinalPerNodeColorGrading.PostProcessSettings, AllNodesColorGrading.ColorGradingSettings, PerNodeColorGrading.ColorGradingSettings);
-			FinalPerNodeColorGrading.BlendWeight *= PerNodeColorGrading.ColorGradingSettings.BlendWeight;
+			FinalAllNodesColorGrading.BlendWeight = StageSettings->EntireClusterColorGrading.ColorGradingSettings.BlendWeight;
+
+			FDisplayClusterViewportConfigurationHelpers_Postprocess::BlendPostProcessSettings(FinalAllNodesColorGrading.PostProcessSettings, StageSettings->EntireClusterColorGrading.ColorGradingSettings, AllNodesColorGrading.ColorGradingSettings);
+			FinalAllNodesColorGrading.BlendWeight *= AllNodesColorGrading.ColorGradingSettings.BlendWeight;
 		}
 		else
 		{
-			// node only
-			FDisplayClusterViewportConfigurationHelpers_Postprocess::CopyBlendPostProcessSettings(FinalPerNodeColorGrading.PostProcessSettings, PerNodeColorGrading.ColorGradingSettings);
+			FDisplayClusterViewportConfigurationHelpers_Postprocess::CopyBlendPostProcessSettings(FinalAllNodesColorGrading.PostProcessSettings, AllNodesColorGrading.ColorGradingSettings);
 		}
 
-		ImplUpdateCustomPostprocess(DstViewport, FinalPerNodeColorGrading.bIsEnabled, FinalPerNodeColorGrading, IDisplayClusterViewport_CustomPostProcessSettings::ERenderPass::FinalPerViewport);
+		ImplUpdateCustomPostprocess(DstViewport, FinalAllNodesColorGrading.bIsEnabled, FinalAllNodesColorGrading, IDisplayClusterViewport_CustomPostProcessSettings::ERenderPass::FinalPerViewport);
+
 		return true;
 	}
 
-	return false;
-}
-
-static bool ImplUpdateIncameraAllNodesColorGrading(FDisplayClusterViewport& DstViewport, ADisplayClusterRootActor& RootActor, const FDisplayClusterConfigurationViewport_AllNodesColorGrading AllNodesColorGrading)
-{
-	FDisplayClusterConfigurationViewport_CustomPostprocessSettings FinalAllNodesColorGrading;
-	FinalAllNodesColorGrading.bIsEnabled = true;
-	FinalAllNodesColorGrading.bIsOneFrame = true;
-	FinalAllNodesColorGrading.BlendWeight = 1;
-
-	const FDisplayClusterConfigurationICVFX_StageSettings& StageSettings = RootActor.GetStageSettings();
+// Note that skipped parameters in macro definitions will just evaluate to nothing
+// This is intentional to get around the inconsistent naming in the color grading fields in FPostProcessSettings
+#define PP_CONDITIONAL_BLEND(BLENDOP, COLOR, OUTGROUP, INGROUP, NAME, OFFSETOP, OFFSETVALUE) \
+	{ \
+		bool bOverridePPSettings0 = PPSettings0.INGROUP bOverride_##NAME; \
+		bool bOverridePPSettings1 = (PPSettings1 != nullptr) && PPSettings1->INGROUP bOverride_##NAME; \
+		bool bOverridePPSettings2 = (PPSettings2 != nullptr) && PPSettings2->INGROUP bOverride_##NAME; \
+		 \
+		if (bOverridePPSettings0 && bOverridePPSettings1 && bOverridePPSettings2) \
+		{ \
+			OutputPP.COLOR##NAME##OUTGROUP = PPSettings0.INGROUP NAME BLENDOP PPSettings1->INGROUP NAME BLENDOP PPSettings2->INGROUP NAME OFFSETOP OFFSETVALUE OFFSETOP OFFSETVALUE; \
+			OutputPP.bOverride_##COLOR##NAME##OUTGROUP = true; \
+		} \
+		else if (bOverridePPSettings0 && bOverridePPSettings1) \
+		{ \
+			OutputPP.COLOR##NAME##OUTGROUP = PPSettings0.INGROUP NAME BLENDOP PPSettings1->INGROUP NAME OFFSETOP OFFSETVALUE; \
+			OutputPP.bOverride_##COLOR##NAME##OUTGROUP = true; \
+		} \
+		else if (bOverridePPSettings0 && bOverridePPSettings2) \
+		{ \
+			OutputPP.COLOR##NAME##OUTGROUP = PPSettings0.INGROUP NAME BLENDOP PPSettings2->INGROUP NAME OFFSETOP OFFSETVALUE; \
+			OutputPP.bOverride_##COLOR##NAME##OUTGROUP = true; \
+		} \
+		else if (bOverridePPSettings1 && bOverridePPSettings2) \
+		{ \
+			OutputPP.COLOR##NAME##OUTGROUP = PPSettings1->INGROUP NAME BLENDOP PPSettings2->INGROUP NAME OFFSETOP OFFSETVALUE; \
+			OutputPP.bOverride_##COLOR##NAME##OUTGROUP = true; \
+		} \
+		else if (bOverridePPSettings2) \
+		{ \
+			OutputPP.COLOR##NAME##OUTGROUP = PPSettings2->INGROUP NAME; \
+			OutputPP.bOverride_##COLOR##NAME##OUTGROUP = true; \
+		} \
+		else if (bOverridePPSettings1) \
+		{ \
+			OutputPP.COLOR##NAME##OUTGROUP = PPSettings1->INGROUP NAME; \
+			OutputPP.bOverride_##COLOR##NAME##OUTGROUP = true; \
+		} \
+		else if (bOverridePPSettings0) \
+		{ \
+			OutputPP.COLOR##NAME##OUTGROUP = PPSettings0.INGROUP NAME; \
+			OutputPP.bOverride_##COLOR##NAME##OUTGROUP = true; \
+		} \
+	} \
 
-	// enable entire cluster only when global settings is on
-	const bool bEnableEntireClusterColorGrading = StageSettings.EntireClusterColorGrading.bEnableEntireClusterColorGrading && AllNodesColorGrading.bEnableEntireClusterColorGrading;
+/*
+* This will override the settings using the priority.
+* bOverridePPSettings2 (any additional settings specified by the user) will be of highest priority.
+* bOverridePPSettings1 (which is nDisplay override settings) will be of highest priority.
+* following by Cumulative settings (bOverridePPSettings0).
+*/
+#define PP_CONDITIONAL_OVERRIDE(COLOR, OUTGROUP, INGROUP, NAME) \
+	{ \
+		bool bOverridePPSettings0 = PPSettings0.INGROUP bOverride_##NAME; \
+		bool bOverridePPSettings1 = PPSettings1 && PPSettings1->INGROUP bOverride_##NAME; \
+		bool bOverridePPSettings2 = PPSettings2 && PPSettings2->INGROUP bOverride_##NAME; \
+		if (bOverridePPSettings0) \
+		{ \
+			OutputPP.COLOR##NAME##OUTGROUP = PPSettings0.INGROUP NAME; \
+			OutputPP.bOverride_##COLOR##NAME##OUTGROUP = true; \
+		} \
+		if (bOverridePPSettings1) \
+		{ \
+			OutputPP.COLOR##NAME##OUTGROUP = PPSettings1->INGROUP NAME; \
+			OutputPP.bOverride_##COLOR##NAME##OUTGROUP = true; \
+		} \
+		if (bOverridePPSettings2) \
+		{ \
+			OutputPP.COLOR##NAME##OUTGROUP = PPSettings2->INGROUP NAME; \
+			OutputPP.bOverride_##COLOR##NAME##OUTGROUP = true; \
+		} \
+	} \
 
-	if (bEnableEntireClusterColorGrading)
+	static inline void ImplBlendPostProcessSettings(FPostProcessSettings& OutputPP, const FDisplayClusterConfigurationViewport_ColorGradingRenderingSettings& PPSettings0, const FDisplayClusterConfigurationViewport_ColorGradingRenderingSettings* PPSettings1, const FDisplayClusterConfigurationViewport_ColorGradingRenderingSettings* PPSettings2)
 	{
-		FinalAllNodesColorGrading.BlendWeight = StageSettings.EntireClusterColorGrading.ColorGradingSettings.BlendWeight;
-		
-		FDisplayClusterViewportConfigurationHelpers_Postprocess::BlendPostProcessSettings(FinalAllNodesColorGrading.PostProcessSettings, StageSettings.EntireClusterColorGrading.ColorGradingSettings, AllNodesColorGrading.ColorGradingSettings);
-		FinalAllNodesColorGrading.BlendWeight *= AllNodesColorGrading.ColorGradingSettings.BlendWeight;
-	}
-	else
-	{
-		FDisplayClusterViewportConfigurationHelpers_Postprocess::CopyBlendPostProcessSettings(FinalAllNodesColorGrading.PostProcessSettings, AllNodesColorGrading.ColorGradingSettings);
+		PP_CONDITIONAL_BLEND(+, , , , AutoExposureBias, , );
+		PP_CONDITIONAL_BLEND(+, , , , ColorCorrectionHighlightsMin, , );
+		PP_CONDITIONAL_BLEND(+, , , , ColorCorrectionHighlightsMax, , );
+		PP_CONDITIONAL_BLEND(+, , , , ColorCorrectionShadowsMax, , );
+
+		PP_CONDITIONAL_OVERRIDE(, , WhiteBalance., TemperatureType);
+		PP_CONDITIONAL_BLEND(+, , , WhiteBalance., WhiteTemp, +, -6500.0f);
+		PP_CONDITIONAL_BLEND(+, , , WhiteBalance., WhiteTint, , );
+
+		PP_CONDITIONAL_BLEND(*, Color, , Global., Saturation, , );
+		PP_CONDITIONAL_BLEND(*, Color, , Global., Contrast, , );
+		PP_CONDITIONAL_BLEND(*, Color, , Global., Gamma, , );
+		PP_CONDITIONAL_BLEND(*, Color, , Global., Gain, , );
+		PP_CONDITIONAL_BLEND(+, Color, , Global., Offset, , );
+
+		PP_CONDITIONAL_BLEND(*, Color, Shadows, Shadows., Saturation, , );
+		PP_CONDITIONAL_BLEND(*, Color, Shadows, Shadows., Contrast, , );
+		PP_CONDITIONAL_BLEND(*, Color, Shadows, Shadows., Gamma, , );
+		PP_CONDITIONAL_BLEND(*, Color, Shadows, Shadows., Gain, , );
+		PP_CONDITIONAL_BLEND(+, Color, Shadows, Shadows., Offset, , );
+
+		PP_CONDITIONAL_BLEND(*, Color, Midtones, Midtones., Saturation, , );
+		PP_CONDITIONAL_BLEND(*, Color, Midtones, Midtones., Contrast, , );
+		PP_CONDITIONAL_BLEND(*, Color, Midtones, Midtones., Gamma, , );
+		PP_CONDITIONAL_BLEND(*, Color, Midtones, Midtones., Gain, , );
+		PP_CONDITIONAL_BLEND(+, Color, Midtones, Midtones., Offset, , );
+
+		PP_CONDITIONAL_BLEND(*, Color, Highlights, Highlights., Saturation, , );
+		PP_CONDITIONAL_BLEND(*, Color, Highlights, Highlights., Contrast, , );
+		PP_CONDITIONAL_BLEND(*, Color, Highlights, Highlights., Gamma, , );
+		PP_CONDITIONAL_BLEND(*, Color, Highlights, Highlights., Gain, , );
+		PP_CONDITIONAL_BLEND(+, Color, Highlights, Highlights., Offset, , );
+
+		PP_CONDITIONAL_BLEND(+, , , Misc., BlueCorrection, , );
+		PP_CONDITIONAL_BLEND(+, , , Misc., ExpandGamut, , );
+		PP_CONDITIONAL_BLEND(+, , , Misc., SceneColorTint, , );
 	}
 
-	ImplUpdateCustomPostprocess(DstViewport, FinalAllNodesColorGrading.bIsEnabled, FinalAllNodesColorGrading, IDisplayClusterViewport_CustomPostProcessSettings::ERenderPass::FinalPerViewport);
+#define PP_CONDITIONAL_COPY(COLOR, OUTGROUP, INGROUP, NAME) \
+		if (!bIsConditionalCopy || InPPS->bOverride_##COLOR##NAME##INGROUP) \
+		{ \
+			OutViewportPPSettings->OUTGROUP NAME = InPPS->COLOR##NAME##INGROUP; \
+			OutViewportPPSettings->OUTGROUP bOverride_##NAME = true; \
+		}
 
-	return true;
-}
+	static inline void ImplCopyPPSStruct(bool bIsConditionalCopy, FDisplayClusterConfigurationViewport_ColorGradingRenderingSettings* OutViewportPPSettings, FPostProcessSettings* InPPS)
+	{
+		if ((OutViewportPPSettings != nullptr) && (InPPS != nullptr))
+		{
+			PP_CONDITIONAL_COPY(, , , AutoExposureBias);
+			PP_CONDITIONAL_COPY(, , , ColorCorrectionHighlightsMin);
+			PP_CONDITIONAL_COPY(, , , ColorCorrectionHighlightsMax);
+			PP_CONDITIONAL_COPY(, , , ColorCorrectionShadowsMax);
+
+			PP_CONDITIONAL_COPY(, WhiteBalance., , TemperatureType);
+			PP_CONDITIONAL_COPY(, WhiteBalance., , WhiteTemp);
+			PP_CONDITIONAL_COPY(, WhiteBalance., , WhiteTint);
+
+			PP_CONDITIONAL_COPY(Color, Global., , Saturation);
+			PP_CONDITIONAL_COPY(Color, Global., , Contrast);
+			PP_CONDITIONAL_COPY(Color, Global., , Gamma);
+			PP_CONDITIONAL_COPY(Color, Global., , Gain);
+			PP_CONDITIONAL_COPY(Color, Global., , Offset);
+
+			PP_CONDITIONAL_COPY(Color, Shadows., Shadows, Saturation);
+			PP_CONDITIONAL_COPY(Color, Shadows., Shadows, Contrast);
+			PP_CONDITIONAL_COPY(Color, Shadows., Shadows, Gamma);
+			PP_CONDITIONAL_COPY(Color, Shadows., Shadows, Gain);
+			PP_CONDITIONAL_COPY(Color, Shadows., Shadows, Offset);
+
+			PP_CONDITIONAL_COPY(Color, Midtones., Midtones, Saturation);
+			PP_CONDITIONAL_COPY(Color, Midtones., Midtones, Contrast);
+			PP_CONDITIONAL_COPY(Color, Midtones., Midtones, Gamma);
+			PP_CONDITIONAL_COPY(Color, Midtones., Midtones, Gain);
+			PP_CONDITIONAL_COPY(Color, Midtones., Midtones, Offset);
+
+			PP_CONDITIONAL_COPY(Color, Highlights., Highlights, Saturation);
+			PP_CONDITIONAL_COPY(Color, Highlights., Highlights, Contrast);
+			PP_CONDITIONAL_COPY(Color, Highlights., Highlights, Gamma);
+			PP_CONDITIONAL_COPY(Color, Highlights., Highlights, Gain);
+			PP_CONDITIONAL_COPY(Color, Highlights., Highlights, Offset);
+
+			PP_CONDITIONAL_COPY(, Misc., , BlueCorrection);
+			PP_CONDITIONAL_COPY(, Misc., , ExpandGamut);
+			PP_CONDITIONAL_COPY(, Misc., , SceneColorTint);
+		}
+	}
+};
+using namespace UE::DisplayCluster::Configuration;
 
-#if WITH_EDITOR
 // return true when same settings used for both viewports
-bool FDisplayClusterViewportConfigurationHelpers_Postprocess::IsInnerFrustumViewportSettingsEqual_Editor(const FDisplayClusterViewport& InViewport1, const FDisplayClusterViewport& InViewport2, UDisplayClusterICVFXCameraComponent& InCameraComponent)
+bool FDisplayClusterViewportConfigurationHelpers_Postprocess::IsInnerFrustumViewportSettingsEqual(const FDisplayClusterViewport& InViewport1, const FDisplayClusterViewport& InViewport2, const FDisplayClusterConfigurationICVFX_CameraSettings& InCameraSettings)
 {
-	const FDisplayClusterConfigurationICVFX_CameraSettings& CameraSettings = InCameraComponent.GetCameraSettingsICVFX();
-	for (const FDisplayClusterConfigurationViewport_PerNodeColorGrading& ColorGradingProfileIt : CameraSettings.PerNodeColorGrading)
+	for (const FDisplayClusterConfigurationViewport_PerNodeColorGrading& ColorGradingProfileIt : InCameraSettings.PerNodeColorGrading)
 	{
 		if (ColorGradingProfileIt.bIsEnabled)
 		{
@@ -201,17 +378,14 @@ bool FDisplayClusterViewportConfigurationHelpers_Postprocess::IsInnerFrustumView
 
 	return true;
 }
-#endif
 
-bool FDisplayClusterViewportConfigurationHelpers_Postprocess::ImplUpdateInnerFrustumColorGrading(FDisplayClusterViewport& DstViewport, ADisplayClusterRootActor& RootActor, UDisplayClusterICVFXCameraComponent& InCameraComponent)
+bool FDisplayClusterViewportConfigurationHelpers_Postprocess::ImplUpdateInnerFrustumColorGrading(FDisplayClusterViewport& DstViewport, const FDisplayClusterConfigurationICVFX_CameraSettings& InCameraSettings)
 {
-	const FDisplayClusterConfigurationICVFX_CameraSettings& CameraSettings = InCameraComponent.GetCameraSettingsICVFX();
-
 	// per node color grading first (it includes all nodes blending too)
 	const FString& ClusterNodeId = DstViewport.GetClusterNodeId();
 	check(!ClusterNodeId.IsEmpty());
 
-	for (const FDisplayClusterConfigurationViewport_PerNodeColorGrading& ColorGradingProfileIt : CameraSettings.PerNodeColorGrading)
+	for (const FDisplayClusterConfigurationViewport_PerNodeColorGrading& ColorGradingProfileIt : InCameraSettings.PerNodeColorGrading)
 	{
 		// Only allowed profiles
 		if (ColorGradingProfileIt.bIsEnabled)
@@ -221,51 +395,56 @@ bool FDisplayClusterViewportConfigurationHelpers_Postprocess::ImplUpdateInnerFru
 				if (ClusterNodeId.Compare(ClusterNodeIt, ESearchCase::IgnoreCase) == 0)
 				{
 					// Use cluster node PP
-					return ImplUpdateIncameraPerNodeColorGrading(DstViewport, RootActor, CameraSettings.AllNodesColorGrading, ColorGradingProfileIt);
+					return PostprocessHelpers::ImplUpdateIncameraPerNodeColorGrading(DstViewport, InCameraSettings.AllNodesColorGrading, ColorGradingProfileIt);
 				}
 			}
 		}
 	}
 
 	// run through dedicated all nodes pass only when per node list is empty
-	if (CameraSettings.AllNodesColorGrading.bEnableInnerFrustumAllNodesColorGrading)
+	if (InCameraSettings.AllNodesColorGrading.bEnableInnerFrustumAllNodesColorGrading)
 	{
 		// all nodes only color grading
-		return ImplUpdateIncameraAllNodesColorGrading(DstViewport, RootActor, CameraSettings.AllNodesColorGrading);
+		return PostprocessHelpers::ImplUpdateIncameraAllNodesColorGrading(DstViewport, InCameraSettings.AllNodesColorGrading);
 	}
 
 	return false;
 }
 
-bool FDisplayClusterViewportConfigurationHelpers_Postprocess::UpdateLightcardPostProcessSettings(FDisplayClusterViewport& DstViewport, FDisplayClusterViewport& BaseViewport, ADisplayClusterRootActor& RootActor)
+bool FDisplayClusterViewportConfigurationHelpers_Postprocess::UpdateLightcardPostProcessSettings(FDisplayClusterViewport& DstViewport, FDisplayClusterViewport& BaseViewport)
 {
-	const FDisplayClusterConfigurationICVFX_StageSettings& StageSettings = RootActor.GetStageSettings();
-	const FDisplayClusterConfigurationICVFX_LightcardSettings& LightcardSettings = StageSettings.Lightcard;
+	const FDisplayClusterConfigurationICVFX_StageSettings* StageSettings = DstViewport.Configuration->GetStageSettings();
+	if (!StageSettings)
+	{
+		return false;
+	}
+
+	const FDisplayClusterConfigurationICVFX_LightcardSettings& LightcardSettings = StageSettings->Lightcard;
 
 	// First try use global OCIO from stage settings
 	if (LightcardSettings.bEnableOuterViewportColorGrading)
 	{
-		if (ImplUpdateViewportColorGrading(DstViewport, RootActor, BaseViewport.GetId()))
+		if (ImplUpdateViewportColorGrading(DstViewport, BaseViewport.GetId()))
 		{
 			return true;
 		}
 	}
 
 	// This viewport doesn't use PP
-	ImplRemoveCustomPostprocess(DstViewport, IDisplayClusterViewport_CustomPostProcessSettings::ERenderPass::FinalPerViewport);
+	PostprocessHelpers::ImplRemoveCustomPostprocess(DstViewport, IDisplayClusterViewport_CustomPostProcessSettings::ERenderPass::FinalPerViewport);
 
 	return false;
 }
 
-bool FDisplayClusterViewportConfigurationHelpers_Postprocess::ImplUpdateViewportColorGrading(FDisplayClusterViewport& DstViewport, ADisplayClusterRootActor& RootActor, const FString& InClusterViewportId)
+bool FDisplayClusterViewportConfigurationHelpers_Postprocess::ImplUpdateViewportColorGrading(FDisplayClusterViewport& DstViewport, const FString& InClusterViewportId)
 {
-	const FDisplayClusterConfigurationICVFX_StageSettings& StageSettings = RootActor.GetStageSettings();
-	if (StageSettings.EnableColorGrading == false)
+	const FDisplayClusterConfigurationICVFX_StageSettings* StageSettings = DstViewport.Configuration->GetStageSettings();
+	if (!StageSettings || StageSettings->EnableColorGrading == false)
 	{
 		return false;
 	}
 	
-	for (const FDisplayClusterConfigurationViewport_PerViewportColorGrading& ColorGradingProfileIt : StageSettings.PerViewportColorGrading)
+	for (const FDisplayClusterConfigurationViewport_PerViewportColorGrading& ColorGradingProfileIt : StageSettings->PerViewportColorGrading)
 	{
 		if (ColorGradingProfileIt.bIsEnabled)
 		{
@@ -274,22 +453,27 @@ bool FDisplayClusterViewportConfigurationHelpers_Postprocess::ImplUpdateViewport
 				if (InClusterViewportId.Compare(ViewportNameIt, ESearchCase::IgnoreCase) == 0)
 				{
 					// Use per viewport blending
-					return ImplUpdatePerViewportColorGrading(DstViewport, RootActor, ColorGradingProfileIt);
+					return PostprocessHelpers::ImplUpdatePerViewportColorGrading(DstViewport, ColorGradingProfileIt);
 				}
 			}
 		}
 	}
 
 	// per viewport color grading is empty, entire cluster only
-	return ImplUpdateEntireClusterColorGrading(DstViewport, RootActor);
+	return PostprocessHelpers::ImplUpdateEntireClusterColorGrading(DstViewport);
 }
 
-void FDisplayClusterViewportConfigurationHelpers_Postprocess::UpdateCameraPostProcessSettings(FDisplayClusterViewport& DstViewport, ADisplayClusterRootActor& RootActor, UDisplayClusterICVFXCameraComponent& InCameraComponent)
+void FDisplayClusterViewportConfigurationHelpers_Postprocess::UpdateCameraPostProcessSettings(FDisplayClusterViewport& DstViewport, UDisplayClusterICVFXCameraComponent& InCameraComponent, const FDisplayClusterConfigurationICVFX_CameraSettings& InCameraSettings)
 {
-	const FDisplayClusterConfigurationICVFX_CameraSettings& CameraSettings = InCameraComponent.GetCameraSettingsICVFX();
-	const FDisplayClusterConfigurationICVFX_StageSettings& StageSettings = RootActor.GetStageSettings();
+	const FDisplayClusterConfigurationICVFX_StageSettings* StageSettings = DstViewport.Configuration->GetStageSettings();
+	ADisplayClusterRootActor* SceneRootActor = DstViewport.Configuration->GetRootActor(EDisplayClusterRootActorType::Scene);
+	if (!(StageSettings && SceneRootActor))
+	{
+		return;
+	}
+
 	// check if frustum color grading is enabled	
-	if (StageSettings.EnableColorGrading && CameraSettings.EnableInnerFrustumColorGrading)
+	if (StageSettings->EnableColorGrading && InCameraSettings.EnableInnerFrustumColorGrading)
 	{
 		FDisplayClusterConfigurationViewport_CustomPostprocessSettings CameraPPS;
 		CameraPPS.bIsOneFrame = true;
@@ -299,7 +483,7 @@ void FDisplayClusterViewportConfigurationHelpers_Postprocess::UpdateCameraPostPr
 
 		// All logic was moved to the UDisplayClusterICVFXCameraComponent::GetCameraView() virtual function.
 		FMinimalViewInfo DesiredView;
-		IDisplayClusterViewport::GetCameraComponentView(&InCameraComponent, RootActor.GetWorldDeltaSeconds(), bUseCameraPostprocess, DesiredView);
+		IDisplayClusterViewport::GetCameraComponentView(&InCameraComponent, SceneRootActor->GetWorldDeltaSeconds(), bUseCameraPostprocess, DesiredView);
 		if (DesiredView.PostProcessBlendWeight > 0)
 		{
 			// Send camera postprocess to override
@@ -307,216 +491,54 @@ void FDisplayClusterViewportConfigurationHelpers_Postprocess::UpdateCameraPostPr
 			CameraPPS.PostProcessSettings = DesiredView.PostProcessSettings;
 		}
 
-		ImplUpdateCustomPostprocess(DstViewport, CameraPPS.bIsEnabled, CameraPPS, IDisplayClusterViewport_CustomPostProcessSettings::ERenderPass::Override);
+		PostprocessHelpers::ImplUpdateCustomPostprocess(DstViewport, CameraPPS.bIsEnabled, CameraPPS, IDisplayClusterViewport_CustomPostProcessSettings::ERenderPass::Override);
 
-		if (!ImplUpdateInnerFrustumColorGrading(DstViewport, RootActor, InCameraComponent))
+		if (!ImplUpdateInnerFrustumColorGrading(DstViewport, InCameraSettings))
 		{
 			// This viewport doesn't use per-viewport PP
-			ImplRemoveCustomPostprocess(DstViewport, IDisplayClusterViewport_CustomPostProcessSettings::ERenderPass::FinalPerViewport);
+			PostprocessHelpers::ImplRemoveCustomPostprocess(DstViewport, IDisplayClusterViewport_CustomPostProcessSettings::ERenderPass::FinalPerViewport);
 		}
 	}
 }
 
-void FDisplayClusterViewportConfigurationHelpers_Postprocess::UpdateCustomPostProcessSettings(FDisplayClusterViewport& DstViewport, ADisplayClusterRootActor& RootActor, const FDisplayClusterConfigurationViewport_CustomPostprocess& InCustomPostprocessConfiguration)
+void FDisplayClusterViewportConfigurationHelpers_Postprocess::UpdateCustomPostProcessSettings(FDisplayClusterViewport& DstViewport, const FDisplayClusterConfigurationViewport_CustomPostprocess& InCustomPostprocessConfiguration)
 {
 	// update postprocess settings (Start, Override, Final)
-	ImplUpdateCustomPostprocess(DstViewport, InCustomPostprocessConfiguration.Start.bIsEnabled, InCustomPostprocessConfiguration.Start, IDisplayClusterViewport_CustomPostProcessSettings::ERenderPass::Start);
-	ImplUpdateCustomPostprocess(DstViewport, InCustomPostprocessConfiguration.Override.bIsEnabled, InCustomPostprocessConfiguration.Override, IDisplayClusterViewport_CustomPostProcessSettings::ERenderPass::Override);
-	ImplUpdateCustomPostprocess(DstViewport, InCustomPostprocessConfiguration.Final.bIsEnabled, InCustomPostprocessConfiguration.Final, IDisplayClusterViewport_CustomPostProcessSettings::ERenderPass::Final);
+	PostprocessHelpers::ImplUpdateCustomPostprocess(DstViewport, InCustomPostprocessConfiguration.Start.bIsEnabled, InCustomPostprocessConfiguration.Start, IDisplayClusterViewport_CustomPostProcessSettings::ERenderPass::Start);
+	PostprocessHelpers::ImplUpdateCustomPostprocess(DstViewport, InCustomPostprocessConfiguration.Override.bIsEnabled, InCustomPostprocessConfiguration.Override, IDisplayClusterViewport_CustomPostProcessSettings::ERenderPass::Override);
+	PostprocessHelpers::ImplUpdateCustomPostprocess(DstViewport, InCustomPostprocessConfiguration.Final.bIsEnabled, InCustomPostprocessConfiguration.Final, IDisplayClusterViewport_CustomPostProcessSettings::ERenderPass::Final);
 }
 
-void FDisplayClusterViewportConfigurationHelpers_Postprocess::UpdatePerViewportPostProcessSettings(FDisplayClusterViewport& DstViewport, ADisplayClusterRootActor& RootActor)
+void FDisplayClusterViewportConfigurationHelpers_Postprocess::UpdatePerViewportPostProcessSettings(FDisplayClusterViewport& DstViewport)
 {
-	if (!ImplUpdateViewportColorGrading(DstViewport, RootActor, DstViewport.GetId()))
+	if (!ImplUpdateViewportColorGrading(DstViewport, DstViewport.GetId()))
 	{
 		// This viewport doesn't use PP
-		ImplRemoveCustomPostprocess(DstViewport, IDisplayClusterViewport_CustomPostProcessSettings::ERenderPass::FinalPerViewport);
+		PostprocessHelpers::ImplRemoveCustomPostprocess(DstViewport, IDisplayClusterViewport_CustomPostProcessSettings::ERenderPass::FinalPerViewport);
 	}
 }
 
-// Note that skipped parameters in macro definitions will just evaluate to nothing
-// This is intentional to get around the inconsistent naming in the color grading fields in FPostProcessSettings
-#define PP_CONDITIONAL_BLEND(BLENDOP, COLOR, OUTGROUP, INGROUP, NAME, OFFSETOP, OFFSETVALUE) \
-	{ \
-		bool bOverridePPSettings0 = PPSettings0.INGROUP bOverride_##NAME; \
-		bool bOverridePPSettings1 = (PPSettings1 != nullptr) && PPSettings1->INGROUP bOverride_##NAME; \
-		bool bOverridePPSettings2 = (PPSettings2 != nullptr) && PPSettings2->INGROUP bOverride_##NAME; \
-		 \
-		if (bOverridePPSettings0 && bOverridePPSettings1 && bOverridePPSettings2) \
-		{ \
-			OutputPP.COLOR##NAME##OUTGROUP = PPSettings0.INGROUP NAME BLENDOP PPSettings1->INGROUP NAME BLENDOP PPSettings2->INGROUP NAME OFFSETOP OFFSETVALUE OFFSETOP OFFSETVALUE; \
-			OutputPP.bOverride_##COLOR##NAME##OUTGROUP = true; \
-		} \
-		else if (bOverridePPSettings0 && bOverridePPSettings1) \
-		{ \
-			OutputPP.COLOR##NAME##OUTGROUP = PPSettings0.INGROUP NAME BLENDOP PPSettings1->INGROUP NAME OFFSETOP OFFSETVALUE; \
-			OutputPP.bOverride_##COLOR##NAME##OUTGROUP = true; \
-		} \
-		else if (bOverridePPSettings0 && bOverridePPSettings2) \
-		{ \
-			OutputPP.COLOR##NAME##OUTGROUP = PPSettings0.INGROUP NAME BLENDOP PPSettings2->INGROUP NAME OFFSETOP OFFSETVALUE; \
-			OutputPP.bOverride_##COLOR##NAME##OUTGROUP = true; \
-		} \
-		else if (bOverridePPSettings1 && bOverridePPSettings2) \
-		{ \
-			OutputPP.COLOR##NAME##OUTGROUP = PPSettings1->INGROUP NAME BLENDOP PPSettings2->INGROUP NAME OFFSETOP OFFSETVALUE; \
-			OutputPP.bOverride_##COLOR##NAME##OUTGROUP = true; \
-		} \
-		else if (bOverridePPSettings2) \
-		{ \
-			OutputPP.COLOR##NAME##OUTGROUP = PPSettings2->INGROUP NAME; \
-			OutputPP.bOverride_##COLOR##NAME##OUTGROUP = true; \
-		} \
-		else if (bOverridePPSettings1) \
-		{ \
-			OutputPP.COLOR##NAME##OUTGROUP = PPSettings1->INGROUP NAME; \
-			OutputPP.bOverride_##COLOR##NAME##OUTGROUP = true; \
-		} \
-		else if (bOverridePPSettings0) \
-		{ \
-			OutputPP.COLOR##NAME##OUTGROUP = PPSettings0.INGROUP NAME; \
-			OutputPP.bOverride_##COLOR##NAME##OUTGROUP = true; \
-		} \
-	} \
-
-/* 
-* This will override the settings using the priority. 
-* bOverridePPSettings2 (any additional settings specified by the user) will be of highest priority.
-* bOverridePPSettings1 (which is nDisplay override settings) will be of highest priority.
-* following by Cumulative settings (bOverridePPSettings0).
-*/
-#define PP_CONDITIONAL_OVERRIDE(COLOR, OUTGROUP, INGROUP, NAME) \
-	{ \
-		bool bOverridePPSettings0 = PPSettings0.INGROUP bOverride_##NAME; \
-		bool bOverridePPSettings1 = PPSettings1 && PPSettings1->INGROUP bOverride_##NAME; \
-		bool bOverridePPSettings2 = PPSettings2 && PPSettings2->INGROUP bOverride_##NAME; \
-		if (bOverridePPSettings0) \
-		{ \
-			OutputPP.COLOR##NAME##OUTGROUP = PPSettings0.INGROUP NAME; \
-			OutputPP.bOverride_##COLOR##NAME##OUTGROUP = true; \
-		} \
-		if (bOverridePPSettings1) \
-		{ \
-			OutputPP.COLOR##NAME##OUTGROUP = PPSettings1->INGROUP NAME; \
-			OutputPP.bOverride_##COLOR##NAME##OUTGROUP = true; \
-		} \
-		if (bOverridePPSettings2) \
-		{ \
-			OutputPP.COLOR##NAME##OUTGROUP = PPSettings2->INGROUP NAME; \
-			OutputPP.bOverride_##COLOR##NAME##OUTGROUP = true; \
-		} \
-	} \
-
-static void ImplBlendPostProcessSettings(FPostProcessSettings& OutputPP, const FDisplayClusterConfigurationViewport_ColorGradingRenderingSettings& PPSettings0, const FDisplayClusterConfigurationViewport_ColorGradingRenderingSettings* PPSettings1, const FDisplayClusterConfigurationViewport_ColorGradingRenderingSettings* PPSettings2)
-{
-	PP_CONDITIONAL_BLEND(+, , , , AutoExposureBias, , );
-	PP_CONDITIONAL_BLEND(+, , , , ColorCorrectionHighlightsMin, , );
-	PP_CONDITIONAL_BLEND(+, , , , ColorCorrectionHighlightsMax, , );
-	PP_CONDITIONAL_BLEND(+, , , , ColorCorrectionShadowsMax, , );
-
-	PP_CONDITIONAL_OVERRIDE(, , WhiteBalance., TemperatureType);
-	PP_CONDITIONAL_BLEND(+, , , WhiteBalance., WhiteTemp, +, -6500.0f);
-	PP_CONDITIONAL_BLEND(+, , , WhiteBalance., WhiteTint, , );
-
-	PP_CONDITIONAL_BLEND(*, Color, , Global., Saturation, , );
-	PP_CONDITIONAL_BLEND(*, Color, , Global., Contrast, , );
-	PP_CONDITIONAL_BLEND(*, Color, , Global., Gamma, , );
-	PP_CONDITIONAL_BLEND(*, Color, , Global., Gain, , );
-	PP_CONDITIONAL_BLEND(+, Color, , Global., Offset, , );
-
-	PP_CONDITIONAL_BLEND(*, Color, Shadows, Shadows., Saturation, , );
-	PP_CONDITIONAL_BLEND(*, Color, Shadows, Shadows., Contrast, , );
-	PP_CONDITIONAL_BLEND(*, Color, Shadows, Shadows., Gamma, , );
-	PP_CONDITIONAL_BLEND(*, Color, Shadows, Shadows., Gain, , );
-	PP_CONDITIONAL_BLEND(+, Color, Shadows, Shadows., Offset, , );
-
-	PP_CONDITIONAL_BLEND(*, Color, Midtones, Midtones., Saturation, , );
-	PP_CONDITIONAL_BLEND(*, Color, Midtones, Midtones., Contrast, , );
-	PP_CONDITIONAL_BLEND(*, Color, Midtones, Midtones., Gamma, , );
-	PP_CONDITIONAL_BLEND(*, Color, Midtones, Midtones., Gain, , );
-	PP_CONDITIONAL_BLEND(+, Color, Midtones, Midtones., Offset, , );
-
-	PP_CONDITIONAL_BLEND(*, Color, Highlights, Highlights., Saturation, , );
-	PP_CONDITIONAL_BLEND(*, Color, Highlights, Highlights., Contrast, , );
-	PP_CONDITIONAL_BLEND(*, Color, Highlights, Highlights., Gamma, , );
-	PP_CONDITIONAL_BLEND(*, Color, Highlights, Highlights., Gain, , );
-	PP_CONDITIONAL_BLEND(+, Color, Highlights, Highlights., Offset, , );
-
-	PP_CONDITIONAL_BLEND(+, , , Misc., BlueCorrection, , );
-	PP_CONDITIONAL_BLEND(+, , , Misc., ExpandGamut, , );
-	PP_CONDITIONAL_BLEND(+, , , Misc., SceneColorTint, , );
-}
 void FDisplayClusterViewportConfigurationHelpers_Postprocess::CopyBlendPostProcessSettings(FPostProcessSettings& OutputPP, const FDisplayClusterConfigurationViewport_ColorGradingRenderingSettings& InPPSettings)
 {
-	ImplBlendPostProcessSettings(OutputPP, InPPSettings, nullptr, nullptr);
+	PostprocessHelpers::ImplBlendPostProcessSettings(OutputPP, InPPSettings, nullptr, nullptr);
 }
 
 void FDisplayClusterViewportConfigurationHelpers_Postprocess::PerNodeBlendPostProcessSettings(FPostProcessSettings& OutputPP, const FDisplayClusterConfigurationViewport_ColorGradingRenderingSettings& ClusterPPSettings, const FDisplayClusterConfigurationViewport_ColorGradingRenderingSettings& ViewportPPSettings, const FDisplayClusterConfigurationViewport_ColorGradingRenderingSettings& PerNodePPSettings)
 {
-	ImplBlendPostProcessSettings(OutputPP, ClusterPPSettings, &ViewportPPSettings, &PerNodePPSettings);
+	PostprocessHelpers::ImplBlendPostProcessSettings(OutputPP, ClusterPPSettings, &ViewportPPSettings, &PerNodePPSettings);
 }
 
 void FDisplayClusterViewportConfigurationHelpers_Postprocess::BlendPostProcessSettings(FPostProcessSettings& OutputPP, const FDisplayClusterConfigurationViewport_ColorGradingRenderingSettings& ClusterPPSettings, const FDisplayClusterConfigurationViewport_ColorGradingRenderingSettings& ViewportPPSettings)
 {
-	ImplBlendPostProcessSettings(OutputPP, ClusterPPSettings, &ViewportPPSettings, nullptr);
-}
-
-#define PP_CONDITIONAL_COPY(COLOR, OUTGROUP, INGROUP, NAME) \
-		if (!bIsConditionalCopy || InPPS->bOverride_##COLOR##NAME##INGROUP) \
-		{ \
-			OutViewportPPSettings->OUTGROUP NAME = InPPS->COLOR##NAME##INGROUP; \
-			OutViewportPPSettings->OUTGROUP bOverride_##NAME = true; \
-		}
-
-static void ImplCopyPPSStruct(bool bIsConditionalCopy, FDisplayClusterConfigurationViewport_ColorGradingRenderingSettings* OutViewportPPSettings, FPostProcessSettings* InPPS)
-{
-	if ((OutViewportPPSettings != nullptr) && (InPPS != nullptr))
-	{
-		PP_CONDITIONAL_COPY(, , , AutoExposureBias);
-		PP_CONDITIONAL_COPY(, , , ColorCorrectionHighlightsMin);
-		PP_CONDITIONAL_COPY(, , , ColorCorrectionHighlightsMax);
-		PP_CONDITIONAL_COPY(, , , ColorCorrectionShadowsMax);
-
-		PP_CONDITIONAL_COPY(, WhiteBalance., , TemperatureType);
-		PP_CONDITIONAL_COPY(, WhiteBalance., , WhiteTemp);
-		PP_CONDITIONAL_COPY(, WhiteBalance., , WhiteTint);
-
-		PP_CONDITIONAL_COPY(Color, Global., , Saturation);
-		PP_CONDITIONAL_COPY(Color, Global., , Contrast);
-		PP_CONDITIONAL_COPY(Color, Global., , Gamma);
-		PP_CONDITIONAL_COPY(Color, Global., , Gain);
-		PP_CONDITIONAL_COPY(Color, Global., , Offset);
-
-		PP_CONDITIONAL_COPY(Color, Shadows., Shadows, Saturation);
-		PP_CONDITIONAL_COPY(Color, Shadows., Shadows, Contrast);
-		PP_CONDITIONAL_COPY(Color, Shadows., Shadows, Gamma);
-		PP_CONDITIONAL_COPY(Color, Shadows., Shadows, Gain);
-		PP_CONDITIONAL_COPY(Color, Shadows., Shadows, Offset);
-
-		PP_CONDITIONAL_COPY(Color, Midtones., Midtones, Saturation);
-		PP_CONDITIONAL_COPY(Color, Midtones., Midtones, Contrast);
-		PP_CONDITIONAL_COPY(Color, Midtones., Midtones, Gamma);
-		PP_CONDITIONAL_COPY(Color, Midtones., Midtones, Gain);
-		PP_CONDITIONAL_COPY(Color, Midtones., Midtones, Offset);
-
-		PP_CONDITIONAL_COPY(Color, Highlights., Highlights, Saturation);
-		PP_CONDITIONAL_COPY(Color, Highlights., Highlights, Contrast);
-		PP_CONDITIONAL_COPY(Color, Highlights., Highlights, Gamma);
-		PP_CONDITIONAL_COPY(Color, Highlights., Highlights, Gain);
-		PP_CONDITIONAL_COPY(Color, Highlights., Highlights, Offset);
-
-		PP_CONDITIONAL_COPY(, Misc., , BlueCorrection);
-		PP_CONDITIONAL_COPY(, Misc., , ExpandGamut);
-		PP_CONDITIONAL_COPY(, Misc., , SceneColorTint);
-	}
+	PostprocessHelpers::ImplBlendPostProcessSettings(OutputPP, ClusterPPSettings, &ViewportPPSettings, nullptr);
 }
 
 void FDisplayClusterViewportConfigurationHelpers_Postprocess::CopyPPSStructConditional(FDisplayClusterConfigurationViewport_ColorGradingRenderingSettings* OutViewportPPSettings, FPostProcessSettings* InPPS)
 {
-	ImplCopyPPSStruct(true, OutViewportPPSettings, InPPS);
+	PostprocessHelpers::ImplCopyPPSStruct(true, OutViewportPPSettings, InPPS);
 }
 
 void FDisplayClusterViewportConfigurationHelpers_Postprocess::CopyPPSStruct(FDisplayClusterConfigurationViewport_ColorGradingRenderingSettings* OutViewportPPSettings, FPostProcessSettings* InPPS)
 {
-	ImplCopyPPSStruct(false, OutViewportPPSettings, InPPS);
-	}
+	PostprocessHelpers::ImplCopyPPSStruct(false, OutViewportPPSettings, InPPS);
+}
diff --git a/Engine/Plugins/Runtime/nDisplay/Source/DisplayCluster/Private/Render/Viewport/Configuration/DisplayClusterViewportConfigurationHelpers_Postprocess.h b/Engine/Plugins/Runtime/nDisplay/Source/DisplayCluster/Private/Render/Viewport/Configuration/DisplayClusterViewportConfigurationHelpers_Postprocess.h
--- a/Engine/Plugins/Runtime/nDisplay/Source/DisplayCluster/Private/Render/Viewport/Configuration/DisplayClusterViewportConfigurationHelpers_Postprocess.h
+++ b/Engine/Plugins/Runtime/nDisplay/Source/DisplayCluster/Private/Render/Viewport/Configuration/DisplayClusterViewportConfigurationHelpers_Postprocess.h
@@ -5,28 +5,27 @@
 #include "CoreMinimal.h"
 
 class FDisplayClusterViewport;
-class ADisplayClusterRootActor;
 class UDisplayClusterICVFXCameraComponent;
-class UDisplayClusterConfigurationViewport;
 
 struct FPostProcessSettings;
 struct FDisplayClusterConfigurationViewport_ColorGradingRenderingSettings;
 struct FDisplayClusterConfigurationViewport_CustomPostprocess;
-struct FDisplayClusterConfigurationViewport_PerNodeColorGradingParameters;
+struct FDisplayClusterConfigurationICVFX_CameraSettings;
 
+/**
+* Postprocess configuration helper class.
+*/
 class FDisplayClusterViewportConfigurationHelpers_Postprocess
 {
 public:
-	static void UpdateCustomPostProcessSettings(FDisplayClusterViewport& DstViewport, ADisplayClusterRootActor& RootActor, const FDisplayClusterConfigurationViewport_CustomPostprocess& InCustomPostprocessConfiguration);
-	static void UpdatePerViewportPostProcessSettings(FDisplayClusterViewport& DstViewport, ADisplayClusterRootActor& RootActor);
-	static void UpdateCameraPostProcessSettings(FDisplayClusterViewport& DstViewport, ADisplayClusterRootActor& RootActor, UDisplayClusterICVFXCameraComponent& InCameraComponent);
+	static void UpdateCustomPostProcessSettings(FDisplayClusterViewport& DstViewport, const FDisplayClusterConfigurationViewport_CustomPostprocess& InCustomPostprocessConfiguration);
+	static void UpdatePerViewportPostProcessSettings(FDisplayClusterViewport& DstViewport);
+	static void UpdateCameraPostProcessSettings(FDisplayClusterViewport& DstViewport, UDisplayClusterICVFXCameraComponent& InCameraComponent, const FDisplayClusterConfigurationICVFX_CameraSettings& InCameraSettings);
 
-	static bool UpdateLightcardPostProcessSettings(FDisplayClusterViewport& DstViewport, FDisplayClusterViewport& BaseViewport, ADisplayClusterRootActor& RootActor);
+	static bool UpdateLightcardPostProcessSettings(FDisplayClusterViewport& DstViewport, FDisplayClusterViewport& BaseViewport);
 
-#if WITH_EDITOR
 	// return true when same settings used for both viewports
-	static bool IsInnerFrustumViewportSettingsEqual_Editor(const FDisplayClusterViewport& InViewport1, const FDisplayClusterViewport& InViewport2, UDisplayClusterICVFXCameraComponent& InCameraComponent);
-#endif
+	static bool IsInnerFrustumViewportSettingsEqual(const FDisplayClusterViewport& InViewport1, const FDisplayClusterViewport& InViewport2, const FDisplayClusterConfigurationICVFX_CameraSettings& InCameraSettings);
 
 public:
 	static void BlendPostProcessSettings(FPostProcessSettings& OutputPP, const FDisplayClusterConfigurationViewport_ColorGradingRenderingSettings& ClusterPPSettings, const FDisplayClusterConfigurationViewport_ColorGradingRenderingSettings& ViewportPPSettings);
@@ -37,6 +36,6 @@ public:
 	static void PerNodeBlendPostProcessSettings(FPostProcessSettings& OutputPP, const FDisplayClusterConfigurationViewport_ColorGradingRenderingSettings& ClusterPPSettings, const FDisplayClusterConfigurationViewport_ColorGradingRenderingSettings& ViewportPPSettings, const FDisplayClusterConfigurationViewport_ColorGradingRenderingSettings& PerNodePPSettings);
 
 private:
-	static bool ImplUpdateInnerFrustumColorGrading(FDisplayClusterViewport& DstViewport, ADisplayClusterRootActor& RootActor, UDisplayClusterICVFXCameraComponent& InCameraComponent);
-	static bool ImplUpdateViewportColorGrading(FDisplayClusterViewport& DstViewport, ADisplayClusterRootActor& RootActor, const FString& InClusterViewportId);
+	static bool ImplUpdateInnerFrustumColorGrading(FDisplayClusterViewport& DstViewport, const FDisplayClusterConfigurationICVFX_CameraSettings& InCameraSettings);
+	static bool ImplUpdateViewportColorGrading(FDisplayClusterViewport& DstViewport, const FString& InClusterViewportId);
 };
diff --git a/Engine/Plugins/Runtime/nDisplay/Source/DisplayCluster/Private/Render/Viewport/Configuration/DisplayClusterViewportConfigurationHelpers_Visibility.cpp b/Engine/Plugins/Runtime/nDisplay/Source/DisplayCluster/Private/Render/Viewport/Configuration/DisplayClusterViewportConfigurationHelpers_Visibility.cpp
--- a/Engine/Plugins/Runtime/nDisplay/Source/DisplayCluster/Private/Render/Viewport/Configuration/DisplayClusterViewportConfigurationHelpers_Visibility.cpp
+++ b/Engine/Plugins/Runtime/nDisplay/Source/DisplayCluster/Private/Render/Viewport/Configuration/DisplayClusterViewportConfigurationHelpers_Visibility.cpp
@@ -1,6 +1,7 @@
 // Copyright Epic Games, Inc. All Rights Reserved.
 
-#include "DisplayClusterViewportConfigurationHelpers_Visibility.h"
+#include "Render/Viewport/Configuration/DisplayClusterViewportConfigurationHelpers_Visibility.h"
+#include "Render/Viewport/Configuration/DisplayClusterViewportConfiguration.h"
 
 #include "DisplayClusterConfigurationTypes.h"
 #include "DisplayClusterConfigurationTypes_Viewport.h"
@@ -10,110 +11,133 @@
 #include "Components/DisplayClusterICVFXCameraComponent.h"
 
 #include "Render/Viewport/DisplayClusterViewport.h"
+#include "Render/Viewport/DisplayClusterViewportManager.h"
 #include "DisplayClusterConfigurationTypes_Viewport.h"
 
-static void ImplCollectVisibility(ADisplayClusterRootActor& RootActor, const FDisplayClusterConfigurationICVFX_VisibilityList& InVisibilityList, TArray<FName>& OutActorLayerNames, TSet<FPrimitiveComponentId>& OutAdditionalComponentsList)
+namespace UE::DisplayCluster::Configuration::VisibilityHelpers
 {
-	// Collect RootActor components
-	if (InVisibilityList.RootActorComponentNames.Num() > 0)
+	static inline void ImplCollectActorComponents(AActor& InActor, TSet<FPrimitiveComponentId>& OutAdditionalComponentsList)
 	{
-		RootActor.FindPrimitivesByName(InVisibilityList.RootActorComponentNames, OutAdditionalComponentsList);
+		for (const UActorComponent* Component : InActor.GetComponents())
+		{
+			if (const UPrimitiveComponent* PrimComp = Cast<UPrimitiveComponent>(Component))
+			{
+				OutAdditionalComponentsList.Add(PrimComp->GetPrimitiveSceneId());
+			}
+		}
 	}
 
-	auto CollectActorRefs = [&](const TArray<TSoftObjectPtr<AActor>>& InActorRefs)
- 
+	static inline void ImplCollectVisibility(FDisplayClusterViewportConfiguration& InConfiguration, const FDisplayClusterConfigurationICVFX_VisibilityList& InVisibilityList, TArray<FName>& OutActorLayerNames, TSet<FPrimitiveComponentId>& OutAdditionalComponentsList)
 	{
-		for (const TSoftObjectPtr<AActor>& ActorSOPtrIt : InActorRefs)
+		// Collect RootActor components
+		if (InVisibilityList.RootActorComponentNames.Num() > 0)
 		{
-			if (ActorSOPtrIt.IsValid())
+			if (ADisplayClusterRootActor* SceneRootActor = InConfiguration.GetRootActor(EDisplayClusterRootActorType::Scene))
 			{
-				for (const UActorComponent* Component : ActorSOPtrIt->GetComponents())
+				SceneRootActor->FindPrimitivesByName(InVisibilityList.RootActorComponentNames, OutAdditionalComponentsList);
+			}
+		}
+
+		auto CollectActorRefs = [&](UWorld* CurrentWorld, const TArray<TSoftObjectPtr<AActor>>& InActorRefs)
+			{
+				for (const TSoftObjectPtr<AActor>& ActorSOPtrIt : InActorRefs)
 				{
-					if (const UPrimitiveComponent* PrimComp = Cast<UPrimitiveComponent>(Component))
+					if (ActorSOPtrIt.IsValid())
 					{
-						OutAdditionalComponentsList.Add(PrimComp->GetPrimitiveSceneId());
+						if (ActorSOPtrIt->GetWorld() == CurrentWorld)
+						{
+							ImplCollectActorComponents(*ActorSOPtrIt.Get(), OutAdditionalComponentsList);
+						}
+						else if(CurrentWorld)
+						{
+							// re-reference to the current world
+							// Not implemented.
+							//!
+						}
 					}
 				}
-			}
-		}
-	};
-	
-	// Collect Actors refs
-	CollectActorRefs(InVisibilityList.Actors);
-	CollectActorRefs(InVisibilityList.AutoAddedActors);
-
-	// Collect ActorLayers
-	for (const FActorLayer& ActorLayerIt : InVisibilityList.ActorLayers)
-	{
-		if (!ActorLayerIt.Name.IsNone())
+			};
+
+		// Collect Actors refs
+		UWorld* CurrentWorld = InConfiguration.GetCurrentWorld();
+		CollectActorRefs(CurrentWorld, InVisibilityList.Actors);
+		CollectActorRefs(CurrentWorld, InVisibilityList.AutoAddedActors);
+
+		// Collect ActorLayers
+		for (const FActorLayer& ActorLayerIt : InVisibilityList.ActorLayers)
 		{
-			OutActorLayerNames.AddUnique(ActorLayerIt.Name);
+			if (!ActorLayerIt.Name.IsNone())
+			{
+				OutActorLayerNames.AddUnique(ActorLayerIt.Name);
+			}
 		}
 	}
-}
+};
+using namespace UE::DisplayCluster::Configuration;
 
-void FDisplayClusterViewportConfigurationHelpers_Visibility::UpdateShowOnlyList(FDisplayClusterViewport& DstViewport, ADisplayClusterRootActor& RootActor, const FDisplayClusterConfigurationICVFX_VisibilityList& InVisibilityList)
+void FDisplayClusterViewportConfigurationHelpers_Visibility::UpdateShowOnlyList(FDisplayClusterViewport& DstViewport, const FDisplayClusterConfigurationICVFX_VisibilityList& InVisibilityList)
 {
 	TSet<FPrimitiveComponentId> AdditionalComponentsList;
 	TArray<FName> ActorLayers;
-	ImplCollectVisibility(RootActor, InVisibilityList, ActorLayers, AdditionalComponentsList);
+	VisibilityHelpers::ImplCollectVisibility(*DstViewport.Configuration, InVisibilityList, ActorLayers, AdditionalComponentsList);
 
-	DstViewport.VisibilitySettings.UpdateConfiguration(EDisplayClusterViewport_VisibilityMode::ShowOnly, ActorLayers, AdditionalComponentsList);
+	DstViewport.GetVisibilitySettingsImpl().UpdateVisibilitySettings(EDisplayClusterViewport_VisibilityMode::ShowOnly, ActorLayers, AdditionalComponentsList);
 }
 
-void FDisplayClusterViewportConfigurationHelpers_Visibility::AppendHideList_ICVFX(FDisplayClusterViewport& DstViewport, ADisplayClusterRootActor& InRootActor, const FDisplayClusterConfigurationICVFX_VisibilityList& InHideList)
+void FDisplayClusterViewportConfigurationHelpers_Visibility::AppendHideList_ICVFX(FDisplayClusterViewport& DstViewport, const FDisplayClusterConfigurationICVFX_VisibilityList& InHideList)
 {
 	TSet<FPrimitiveComponentId> AdditionalComponentsList;
 	TArray<FName> ActorLayers;
-	ImplCollectVisibility(InRootActor, InHideList, ActorLayers, AdditionalComponentsList);
+	VisibilityHelpers::ImplCollectVisibility(*DstViewport.Configuration, InHideList, ActorLayers, AdditionalComponentsList);
 
-	DstViewport.VisibilitySettings.AppendHideList(ActorLayers, AdditionalComponentsList);
+	DstViewport.GetVisibilitySettingsImpl().AppendHideList(ActorLayers, AdditionalComponentsList);
 }
 
-void FDisplayClusterViewportConfigurationHelpers_Visibility::UpdateHideList_ICVFX(TArray<TSharedPtr<FDisplayClusterViewport, ESPMode::ThreadSafe>>& DstViewports, ADisplayClusterRootActor& InRootActor)
+void FDisplayClusterViewportConfigurationHelpers_Visibility::UpdateHideList_ICVFX(FDisplayClusterViewportConfiguration& InConfiguration, TArray<TSharedPtr<FDisplayClusterViewport, ESPMode::ThreadSafe>>& DstViewports)
 {
-	if (DstViewports.Num() > 0)
+	ADisplayClusterRootActor* ConfigurationRootActor = InConfiguration.GetRootActor(EDisplayClusterRootActorType::Configuration);
+	const FDisplayClusterConfigurationICVFX_StageSettings* StageSettings = InConfiguration.GetStageSettings();
+
+	if (DstViewports.Num() > 0 && ConfigurationRootActor && StageSettings)
 	{
 		TArray<FName> ActorLayerNames;
 		TSet<FPrimitiveComponentId> AdditionalComponentsList;
 
-		const FDisplayClusterConfigurationICVFX_StageSettings& StageSettings = InRootActor.GetStageSettings();
-
-		ImplCollectVisibility(InRootActor, StageSettings.HideList, ActorLayerNames, AdditionalComponentsList);
+		VisibilityHelpers::ImplCollectVisibility(InConfiguration, StageSettings->HideList, ActorLayerNames, AdditionalComponentsList);
 
 		// Hide lightcard
-		ImplCollectVisibility(InRootActor, StageSettings.Lightcard.ShowOnlyList, ActorLayerNames, AdditionalComponentsList);
+		VisibilityHelpers::ImplCollectVisibility(InConfiguration, StageSettings->Lightcard.ShowOnlyList, ActorLayerNames, AdditionalComponentsList);
 
 		// Also hide chromakeys for all cameras
-		TArray<UDisplayClusterICVFXCameraComponent*> ICVFXCamerasComps;
-		InRootActor.GetComponents(ICVFXCamerasComps);
+		TArray<UDisplayClusterICVFXCameraComponent*> ConfigurationRootActorCameras;
+		ConfigurationRootActor->GetComponents(ConfigurationRootActorCameras);
 
-		for (const UDisplayClusterICVFXCameraComponent* ICVFXCameraIt : ICVFXCamerasComps)
+		for (const UDisplayClusterICVFXCameraComponent* ConfigurationCameraIt : ConfigurationRootActorCameras)
 		{
-			if (ICVFXCameraIt)
+			if (ConfigurationCameraIt)
 			{
-				if (const FDisplayClusterConfigurationICVFX_ChromakeyRenderSettings* ChromakeyRenderSettings = ICVFXCameraIt->GetCameraSettingsICVFX().Chromakey.GetChromakeyRenderSettings(StageSettings))
+				if (const FDisplayClusterConfigurationICVFX_ChromakeyRenderSettings* ChromakeyRenderSettings = ConfigurationCameraIt->GetCameraSettingsICVFX().Chromakey.GetChromakeyRenderSettings(*StageSettings))
 				{
-					ImplCollectVisibility(InRootActor, ChromakeyRenderSettings->ShowOnlyList, ActorLayerNames, AdditionalComponentsList);
+					VisibilityHelpers::ImplCollectVisibility(InConfiguration, ChromakeyRenderSettings->ShowOnlyList, ActorLayerNames, AdditionalComponentsList);
 				}
 			}
 		}
 
 		TArray<FName> OuterActorLayerNames;
 		TSet<FPrimitiveComponentId> OuterAdditionalComponentsList;
-		ImplCollectVisibility(InRootActor, StageSettings.OuterViewportHideList, OuterActorLayerNames, OuterAdditionalComponentsList);
+		VisibilityHelpers::ImplCollectVisibility(InConfiguration, StageSettings->OuterViewportHideList, OuterActorLayerNames, OuterAdditionalComponentsList);
 
 		// Update hide list for all desired viewports:
 		for (TSharedPtr<FDisplayClusterViewport, ESPMode::ThreadSafe>& ViewportIt : DstViewports)
 		{
 			if (ViewportIt.IsValid())
 			{
-				ViewportIt->VisibilitySettings.UpdateConfiguration(EDisplayClusterViewport_VisibilityMode::Hide, ActorLayerNames, AdditionalComponentsList);
+				ViewportIt->GetVisibilitySettingsImpl().UpdateVisibilitySettings(EDisplayClusterViewport_VisibilityMode::Hide, ActorLayerNames, AdditionalComponentsList);
 
 				// Support additional hide list for outer viewports
 				if (EnumHasAllFlags(ViewportIt->GetRenderSettingsICVFX().RuntimeFlags, EDisplayClusterViewportRuntimeICVFXFlags::Target))
 				{
-					ViewportIt->VisibilitySettings.AppendHideList(OuterActorLayerNames, OuterAdditionalComponentsList);
+					ViewportIt->GetVisibilitySettingsImpl().AppendHideList(OuterActorLayerNames, OuterAdditionalComponentsList);
 				}
 			}
 		}
diff --git a/Engine/Plugins/Runtime/nDisplay/Source/DisplayCluster/Private/Render/Viewport/Configuration/DisplayClusterViewportConfigurationHelpers_Visibility.h b/Engine/Plugins/Runtime/nDisplay/Source/DisplayCluster/Private/Render/Viewport/Configuration/DisplayClusterViewportConfigurationHelpers_Visibility.h
--- a/Engine/Plugins/Runtime/nDisplay/Source/DisplayCluster/Private/Render/Viewport/Configuration/DisplayClusterViewportConfigurationHelpers_Visibility.h
+++ b/Engine/Plugins/Runtime/nDisplay/Source/DisplayCluster/Private/Render/Viewport/Configuration/DisplayClusterViewportConfigurationHelpers_Visibility.h
@@ -4,21 +4,23 @@
 
 #include "CoreMinimal.h"
 
-class ADisplayClusterRootActor;
+class FDisplayClusterViewportConfiguration;
 class FDisplayClusterViewport;
-
 struct FDisplayClusterConfigurationICVFX_VisibilityList;
 
+/**
+* Visibility configuration helper class.
+*/
 class FDisplayClusterViewportConfigurationHelpers_Visibility
 {
 public:
 	// Update ShowOnly list for DstViewport
-	static void UpdateShowOnlyList(FDisplayClusterViewport& DstViewport, ADisplayClusterRootActor& RootActor, const FDisplayClusterConfigurationICVFX_VisibilityList& InVisibilityList);
+	static void UpdateShowOnlyList(FDisplayClusterViewport& DstViewport, const FDisplayClusterConfigurationICVFX_VisibilityList& InVisibilityList);
 
 	// Update hide lists for DstViewports (lightcards, chromakey, stage settings hide list, outer viewports hide list, etc)
-	static void UpdateHideList_ICVFX(TArray<TSharedPtr<FDisplayClusterViewport, ESPMode::ThreadSafe>>& DstViewports, ADisplayClusterRootActor& InRootActor);
+	static void UpdateHideList_ICVFX(FDisplayClusterViewportConfiguration& InConfiguration, TArray<TSharedPtr<FDisplayClusterViewport, ESPMode::ThreadSafe>>& DstViewports);
 	
 	// Append to exist hide list (must be call after UpdateHideList_ICVFX)
-	static void AppendHideList_ICVFX(FDisplayClusterViewport& DstViewport, ADisplayClusterRootActor& InRootActor, const FDisplayClusterConfigurationICVFX_VisibilityList& InHideList);
+	static void AppendHideList_ICVFX(FDisplayClusterViewport& DstViewport, const FDisplayClusterConfigurationICVFX_VisibilityList& InHideList);
 };
 
diff --git a/Engine/Plugins/Runtime/nDisplay/Source/DisplayCluster/Private/Render/Viewport/Configuration/DisplayClusterViewportConfiguration_ICVFXCamera.cpp b/Engine/Plugins/Runtime/nDisplay/Source/DisplayCluster/Private/Render/Viewport/Configuration/DisplayClusterViewportConfiguration_ICVFXCamera.cpp
new file mode 100644
--- /dev/null
+++ b/Engine/Plugins/Runtime/nDisplay/Source/DisplayCluster/Private/Render/Viewport/Configuration/DisplayClusterViewportConfiguration_ICVFXCamera.cpp
@@ -0,0 +1,223 @@
+// Copyright Epic Games, Inc. All Rights Reserved.
+
+#include "DisplayClusterViewportConfiguration_ICVFXCamera.h"
+
+#include "DisplayClusterViewportConfiguration.h"
+#include "DisplayClusterViewportConfigurationHelpers.h"
+#include "DisplayClusterViewportConfigurationHelpers_ICVFX.h"
+#include "DisplayClusterViewportConfigurationHelpers_Visibility.h"
+
+#include "DisplayClusterRootActor.h"
+#include "DisplayClusterConfigurationTypes_Viewport.h"
+#include "DisplayClusterConfigurationTypes_ICVFX.h"
+
+#include "IDisplayClusterProjection.h"
+#include "DisplayClusterProjectionStrings.h"
+
+#include "Containers/DisplayClusterProjectionCameraPolicySettings.h"
+#include "Render/Projection/IDisplayClusterProjectionPolicy.h"
+
+#include "Render/Viewport/DisplayClusterViewport.h"
+#include "Render/Viewport/DisplayClusterViewportManager.h"
+#include "Render/Viewport/DisplayClusterViewportStrings.h"
+#include "Render/Viewport/RenderFrame/DisplayClusterRenderFrameSettings.h"
+#include "Render/Viewport/LightCard/DisplayClusterViewportLightCardManager.h"
+
+#include "Components/DisplayClusterICVFXCameraComponent.h"
+
+#include "Misc/DisplayClusterLog.h"
+#include "Misc/Parse.h"
+
+////////////////////////////////////////////////////////////////////////////////
+// FDisplayClusterViewportConfiguration_ICVFXCamera
+////////////////////////////////////////////////////////////////////////////////
+bool FDisplayClusterViewportConfiguration_ICVFXCamera::CreateAndSetupInnerCameraViewport()
+{
+	if (FDisplayClusterViewport* NewCameraViewport = FDisplayClusterViewportConfigurationHelpers_ICVFX::GetOrCreateCameraViewport(Configuration, CameraComponent, GetCameraSettings()))
+	{
+		CameraViewport = NewCameraViewport->AsShared();
+
+		// overlay rendered only for enabled incamera
+		check(GetCameraSettings().bEnable);
+
+		// Update camera viewport settings
+		FDisplayClusterViewportConfigurationHelpers_ICVFX::UpdateCameraViewportSettings(*CameraViewport, CameraComponent, GetCameraSettings());
+
+		// Support projection policy update
+		CameraViewport->UpdateConfiguration_ProjectionPolicy();
+
+		// Reuse for EditorPreview
+		FDisplayClusterViewportConfigurationHelpers_ICVFX::PreviewReuseInnerFrustumViewportWithinClusterNodes(*CameraViewport, CameraComponent, GetCameraSettings());
+
+		return true;
+	}
+
+	return false;
+}
+
+bool FDisplayClusterViewportConfiguration_ICVFXCamera::IsCameraProjectionVisibleOnViewport(FDisplayClusterViewport* TargetViewport)
+{
+	if (TargetViewport && TargetViewport->GetProjectionPolicy().IsValid())
+	{
+		// Currently, only mono context is supported to check the visibility of the inner camera.
+		if (TargetViewport->GetProjectionPolicy()->IsCameraProjectionVisible(CameraContext.ViewRotation, CameraContext.ViewLocation, CameraContext.PrjMatrix))
+		{
+			return true;
+		}
+	}
+
+	// do not use camera for this viewport
+	return false;
+}
+
+void FDisplayClusterViewportConfiguration_ICVFXCamera::Update()
+{
+	if (CreateAndSetupInnerCameraViewport())
+	{
+		if (CameraViewport.IsValid())
+		{
+			FDisplayClusterShaderParameters_ICVFX::FCameraSettings ShaderParametersCameraSettings = FDisplayClusterViewportConfigurationHelpers_ICVFX::GetShaderParametersCameraSettings(*CameraViewport, CameraComponent, GetCameraSettings());
+
+			// Add this camera data to all visible targets:
+			for (TSharedPtr<FDisplayClusterViewport, ESPMode::ThreadSafe>& ViewportIt : VisibleTargets)
+			{
+				if (ViewportIt.IsValid())
+				{
+					// Gain direct access to internal settings of the viewport:
+					FDisplayClusterViewport_RenderSettingsICVFX& InOutRenderSettingsICVFX = ViewportIt->GetRenderSettingsICVFXImpl();
+
+					InOutRenderSettingsICVFX.ICVFX.Cameras.Add(ShaderParametersCameraSettings);
+				}
+			}
+		}
+
+		// Create and assign chromakey for all targets for this camera
+		CreateAndSetupInnerCameraChromakey();
+	}
+}
+
+bool FDisplayClusterViewportConfiguration_ICVFXCamera::Initialize()
+{
+	// Create new camera projection policy for camera viewport
+	TSharedPtr<IDisplayClusterProjectionPolicy, ESPMode::ThreadSafe> CameraProjectionPolicy;
+	if (!FDisplayClusterViewportConfigurationHelpers_ICVFX::CreateProjectionPolicyCameraICVFX(Configuration, CameraComponent, GetCameraSettings(), CameraProjectionPolicy))
+	{
+		return false;
+	}
+
+	ADisplayClusterRootActor* SceneRootActor = Configuration.GetRootActor(EDisplayClusterRootActorType::Scene);
+	if (!SceneRootActor)
+	{
+		return false;
+	}
+
+	// Applying the correct sequence of steps to use the projection policy math:
+	// SetupProjectionViewPoint()->CalculateView()->GetProjectionMatrix()
+	FMinimalViewInfo CameraViewInfo;
+	float CustomNearClippingPlane = -1; // a value less than zero means ignoring.
+	CameraProjectionPolicy->SetupProjectionViewPoint(nullptr, SceneRootActor->GetWorldDeltaSeconds(), CameraViewInfo, &CustomNearClippingPlane);
+
+	CameraContext.ViewLocation = CameraViewInfo.Location;
+	CameraContext.ViewRotation = CameraViewInfo.Rotation;
+
+	// Todo: Here we need to calculate the correct ViewOffset so that ICVFX can support stereo rendering.
+	const FVector ViewOffset = FVector::ZeroVector;
+
+	// Get world scale
+	float WorldToMeters = 100.f;
+	if (UWorld* World = SceneRootActor->GetWorld())
+	{
+		if (AWorldSettings* WorldSettings = World->GetWorldSettings())
+		{
+			WorldToMeters = WorldSettings->WorldToMeters;
+		}
+	}
+
+	// Supports custom near clipping plane
+	const float NCP = (CustomNearClippingPlane >= 0) ? CustomNearClippingPlane : GNearClippingPlane;
+
+	if(CameraProjectionPolicy->CalculateView(nullptr, 0, CameraContext.ViewLocation, CameraContext.ViewRotation, ViewOffset, WorldToMeters, NCP, NCP)
+	&& CameraProjectionPolicy->GetProjectionMatrix(nullptr, 0, CameraContext.PrjMatrix))
+	{
+		return true;
+	}
+
+	return false;
+}
+
+const FDisplayClusterConfigurationICVFX_CameraSettings& FDisplayClusterViewportConfiguration_ICVFXCamera::GetCameraSettings() const
+{
+	return ConfigurationCameraComponent.GetCameraSettingsICVFX();
+}
+
+bool FDisplayClusterViewportConfiguration_ICVFXCamera::ImplCreateChromakeyViewport()
+{
+	check(CameraViewport.IsValid());
+
+	const FString ICVFXCameraId = CameraComponent.GetCameraUniqueId();
+
+	// Create new chromakey viewport
+	if (FDisplayClusterViewport* NewChromakeyViewport = FDisplayClusterViewportConfigurationHelpers_ICVFX::GetOrCreateChromakeyViewport(Configuration, CameraComponent, GetCameraSettings()))
+	{
+		ChromakeyViewport = NewChromakeyViewport->AsShared();
+
+		// Update chromakey viewport settings
+		FDisplayClusterViewportConfigurationHelpers_ICVFX::UpdateChromakeyViewportSettings(*ChromakeyViewport, *CameraViewport, GetCameraSettings());
+
+		// Support projection policy update
+		ChromakeyViewport->UpdateConfiguration_ProjectionPolicy();
+
+		// reuse for EditorPreview
+		FDisplayClusterViewportConfigurationHelpers_ICVFX::PreviewReuseChromakeyViewportWithinClusterNodes(*ChromakeyViewport, ICVFXCameraId);
+
+		return true;
+	}
+
+	return false;
+}
+
+bool FDisplayClusterViewportConfiguration_ICVFXCamera::CreateAndSetupInnerCameraChromakey()
+{
+	const FDisplayClusterConfigurationICVFX_StageSettings* StageSettings = Configuration.GetStageSettings();
+	if (!StageSettings)
+	{
+		return false;
+	}
+
+	// Try create chromakey render on demand
+	if (const FDisplayClusterConfigurationICVFX_ChromakeyRenderSettings* ChromakeyRenderSettings = GetCameraSettings().Chromakey.GetChromakeyRenderSettings(*StageSettings))
+	{
+		if (ChromakeyRenderSettings->ShouldUseChromakeyViewport(*StageSettings))
+		{
+			ImplCreateChromakeyViewport();
+		}
+	}
+
+	// Chromakey viewport name with alpha channel
+	const FString ChromakeyViewportId(ChromakeyViewport.IsValid() ? ChromakeyViewport->GetId() : TEXT(""));
+
+	// Assign this chromakey to all supported targets
+	for (const TSharedPtr<FDisplayClusterViewport, ESPMode::ThreadSafe>& OuterViewportIt : VisibleTargets)
+	{
+		if (OuterViewportIt.IsValid())
+		{
+			const bool bEnableChromakey = !EnumHasAnyFlags(OuterViewportIt->GetRenderSettingsICVFX().Flags, EDisplayClusterViewportICVFXFlags::DisableChromakey);
+			const bool bEnableChromakeyMarkers = !EnumHasAnyFlags(OuterViewportIt->GetRenderSettingsICVFX().Flags, EDisplayClusterViewportICVFXFlags::DisableChromakeyMarkers);
+
+			// Gain direct access to internal settings of the viewport:
+			FDisplayClusterViewport_RenderSettingsICVFX& InOutOuterViewportRenderSettingsICVFX = OuterViewportIt->GetRenderSettingsICVFXImpl();
+			FDisplayClusterShaderParameters_ICVFX::FCameraSettings& DstCameraData = InOutOuterViewportRenderSettingsICVFX.ICVFX.Cameras.Last();
+
+			if (bEnableChromakey)
+			{
+				// Setup chromakey with markers
+				FDisplayClusterViewportConfigurationHelpers_ICVFX::UpdateCameraSettings_Chromakey(DstCameraData, *StageSettings, GetCameraSettings(), bEnableChromakeyMarkers, ChromakeyViewportId);
+			}
+
+			// Setup overlap chromakey with markers
+			FDisplayClusterViewportConfigurationHelpers_ICVFX::UpdateCameraSettings_OverlapChromakey(DstCameraData, *StageSettings, GetCameraSettings(), bEnableChromakeyMarkers);
+		}
+	}
+
+	return true;
+}
diff --git a/Engine/Plugins/Runtime/nDisplay/Source/DisplayCluster/Private/Render/Viewport/Configuration/DisplayClusterViewportConfiguration_ProjectionPolicy.cpp b/Engine/Plugins/Runtime/nDisplay/Source/DisplayCluster/Private/Render/Viewport/Configuration/DisplayClusterViewportConfiguration_ProjectionPolicy.cpp
new file mode 100644
--- /dev/null
+++ b/Engine/Plugins/Runtime/nDisplay/Source/DisplayCluster/Private/Render/Viewport/Configuration/DisplayClusterViewportConfiguration_ProjectionPolicy.cpp
@@ -0,0 +1,178 @@
+// Copyright Epic Games, Inc. All Rights Reserved.
+
+#include "DisplayClusterViewportConfiguration_ProjectionPolicy.h"
+#include "DisplayClusterViewportConfigurationHelpers_ICVFX.h"
+
+#include "Render/Viewport/DisplayClusterViewport.h"
+#include "Render/Viewport/DisplayClusterViewportManager.h"
+
+#include "DisplayClusterViewportConfigurationHelpers.h"
+#include "DisplayClusterViewportConfigurationHelpers_Postprocess.h"
+#include "DisplayClusterConfigurationTypes.h"
+
+#include "IDisplayClusterProjection.h"
+#include "Render/Projection/IDisplayClusterProjectionPolicy.h"
+#include "DisplayClusterProjectionStrings.h"
+
+#include "DisplayClusterRootActor.h"
+#include "Misc/DisplayClusterHelpers.h"
+#include "Misc/DisplayClusterLog.h"
+
+#include "Components/DisplayClusterICVFXCameraComponent.h"
+
+#include "Containers/DisplayClusterProjectionCameraPolicySettings.h"
+
+namespace UE::DisplayCluster::Configuration::ProjectionPolicyHelpers
+{
+	template<class T>
+	T* FindComponentByName(ADisplayClusterRootActor& InRootActor, const FString& InComponentName)
+	{
+		TArray<T*> ActorComps;
+		InRootActor.GetComponents(ActorComps);
+		for (T* CompIt : ActorComps)
+		{
+			if (CompIt->GetName() == InComponentName)
+			{
+				return CompIt;
+			}
+		}
+
+		return nullptr;
+	}
+};
+using namespace UE::DisplayCluster::Configuration;
+
+///////////////////////////////////////////////////////////////////
+// FDisplayClusterViewportConfiguration_ProjectionPolicy
+///////////////////////////////////////////////////////////////////
+void FDisplayClusterViewportConfiguration_ProjectionPolicy::Update()
+{
+	if (FDisplayClusterViewportManager* ViewportManager = Configuration.GetViewportManagerImpl())
+	{
+		for (const TSharedPtr<FDisplayClusterViewport, ESPMode::ThreadSafe>& ViewportIt : ViewportManager->ImplGetCurrentRenderFrameViewports())
+		{
+			// ignore ICVFX internal resources
+			if (ViewportIt.IsValid() && !EnumHasAnyFlags(ViewportIt->GetRenderSettingsICVFX().RuntimeFlags, EDisplayClusterViewportRuntimeICVFXFlags::InternalResource))
+			{
+				// Support advanced logic for 'camera' projection policy
+				if (ViewportIt->GetProjectionPolicy().IsValid())
+				{
+					if (ViewportIt->GetProjectionPolicy()->GetType().Compare(DisplayClusterProjectionStrings::projection::Camera) == 0)
+					{
+						UpdateCameraPolicy(*ViewportIt);
+					}
+
+					// Projection policies can override postprocess settings
+					// The camera PP override code has been moved from FDisplayClusterViewportConfiguration_ProjectionPolicy::UpdateCameraPolicy_Base() to the projection policy new function:
+					ViewportIt->GetProjectionPolicy()->UpdatePostProcessSettings(ViewportIt.Get());
+				}
+			}
+		}
+	}
+}
+
+bool FDisplayClusterViewportConfiguration_ProjectionPolicy::UpdateCameraPolicy(FDisplayClusterViewport& DstViewport)
+{
+	const TMap<FString, FString>& CameraPolicyParameters = DstViewport.GetProjectionPolicy()->GetParameters();
+
+	FString CameraComponentId;
+	// Get assigned camera ID
+	if (!DisplayClusterHelpers::map::template ExtractValue(CameraPolicyParameters, DisplayClusterProjectionStrings::cfg::camera::Component, CameraComponentId))
+	{
+		// use default cameras
+		return true;
+	}
+
+	if (CameraComponentId.IsEmpty())
+	{
+		UE_LOG(LogDisplayClusterViewport, Verbose, TEXT("Viewport '%s': referenced camera '' (empty name)."), *DstViewport.GetId());
+		return false;
+	}
+
+	// Get ICVFX camera component
+	if (UpdateCameraPolicy_ICVFX(DstViewport, CameraComponentId))
+	{
+		return true;
+	}
+
+	// Get camera component
+	if (UpdateCameraPolicy_Base(DstViewport, CameraComponentId))
+	{
+		return true;
+	}
+
+	UE_LOG(LogDisplayClusterViewport, Error, TEXT("Viewport '%s': referenced camera '%s' not found."), *DstViewport.GetId(), *CameraComponentId);
+
+	return false;
+}
+
+bool FDisplayClusterViewportConfiguration_ProjectionPolicy::UpdateCameraPolicy_Base(FDisplayClusterViewport& DstViewport, const FString& CameraComponentId)
+{
+	ADisplayClusterRootActor* ConfigurationRootActor = Configuration.GetRootActor(EDisplayClusterRootActorType::Configuration);
+	ADisplayClusterRootActor*         SceneRootActor = Configuration.GetRootActor(EDisplayClusterRootActorType::Scene);
+	if (ConfigurationRootActor && SceneRootActor)
+	{
+		UCameraComponent*         SceneCameraComp = ProjectionPolicyHelpers::FindComponentByName<UCameraComponent>(*SceneRootActor, CameraComponentId);
+		UCameraComponent* ConfigurationCameraComp = ProjectionPolicyHelpers::FindComponentByName<UCameraComponent>(*ConfigurationRootActor, CameraComponentId);
+		if (SceneCameraComp && ConfigurationCameraComp)
+		{
+			// Now this parameter is always 1 for the standard camera component.
+			const float FOVMultiplier = 1.f;
+
+			// add camera's post processing materials
+			// Moved to IDisplayClusterProjectionPolicy::UpdatePostProcessSettings()
+
+			FDisplayClusterProjectionCameraPolicySettings PolicyCameraSettings;
+			PolicyCameraSettings.FOVMultiplier = FOVMultiplier;
+
+			// Initialize camera policy with camera component and settings
+			return IDisplayClusterProjection::Get().CameraPolicySetCamera(DstViewport.GetProjectionPolicy(), SceneCameraComp, PolicyCameraSettings);
+		}
+	}
+
+	return false;
+}
+
+bool FDisplayClusterViewportConfiguration_ProjectionPolicy::UpdateCameraPolicy_ICVFX(FDisplayClusterViewport& DstViewport, const FString& CameraComponentId)
+{
+	FDisplayClusterViewportManager*  ViewportManager = Configuration.GetViewportManagerImpl();
+	ADisplayClusterRootActor* ConfigurationRootActor = Configuration.GetRootActor(EDisplayClusterRootActorType::Configuration);
+	ADisplayClusterRootActor*         SceneRootActor = Configuration.GetRootActor(EDisplayClusterRootActorType::Scene);
+	if (ConfigurationRootActor && SceneRootActor && ViewportManager)
+	{
+		UDisplayClusterICVFXCameraComponent*         SceneCameraComp = ProjectionPolicyHelpers::FindComponentByName<UDisplayClusterICVFXCameraComponent>(*SceneRootActor, CameraComponentId);
+		UDisplayClusterICVFXCameraComponent* ConfigurationCameraComp = ProjectionPolicyHelpers::FindComponentByName<UDisplayClusterICVFXCameraComponent>(*ConfigurationRootActor, CameraComponentId);
+		if (SceneCameraComp && ConfigurationCameraComp)
+		{
+			// If icvfx viewport already rendered, just copy to visible
+			if (FDisplayClusterViewport* ExistCameraViewport = FDisplayClusterViewportConfigurationHelpers_ICVFX::FindCameraViewport(Configuration, *SceneCameraComp))
+			{
+				// Re-use camera RTT, and support OCIO:
+				const EDisplayClusterViewportOverrideMode ViewportOverrideMode = DstViewport.IsOpenColorIOEquals(*ExistCameraViewport) ? EDisplayClusterViewportOverrideMode::All : EDisplayClusterViewportOverrideMode::InernalRTT;
+
+				// Gain direct access to internal settings of the viewport:
+				FDisplayClusterViewport_RenderSettings& InOutRenderSettings = DstViewport.GetRenderSettingsImpl();
+
+				InOutRenderSettings.SetViewportOverride(ExistCameraViewport->GetId(), ViewportOverrideMode);
+			}
+			else
+			{
+				// Request render:
+				// Save base settings from camera viewport
+				const FDisplayClusterViewport_RenderSettings           SavedRenderSettings = DstViewport.GetRenderSettings();
+				const FDisplayClusterViewport_RenderSettingsICVFX SavedRenderSettingsICVFX = DstViewport.GetRenderSettingsICVFX();
+
+				// Update render settings for icvfx camera:
+				FDisplayClusterViewportConfigurationHelpers_ICVFX::UpdateCameraViewportSettings(DstViewport, *SceneCameraComp, ConfigurationCameraComp->GetCameraSettingsICVFX());
+
+				// Restore some base settings:
+				DstViewport.GetRenderSettingsImpl() = SavedRenderSettings;
+				DstViewport.GetRenderSettingsICVFXImpl() = SavedRenderSettingsICVFX;
+			}
+
+			return true;
+		}
+	}
+
+	return false;
+}
diff --git a/Engine/Plugins/Runtime/nDisplay/Source/DisplayCluster/Private/Render/Viewport/Containers/DisplayClusterViewportProxyData.h b/Engine/Plugins/Runtime/nDisplay/Source/DisplayCluster/Private/Render/Viewport/Containers/DisplayClusterViewportProxyData.h
--- a/Engine/Plugins/Runtime/nDisplay/Source/DisplayCluster/Private/Render/Viewport/Containers/DisplayClusterViewportProxyData.h
+++ b/Engine/Plugins/Runtime/nDisplay/Source/DisplayCluster/Private/Render/Viewport/Containers/DisplayClusterViewportProxyData.h
@@ -17,20 +17,23 @@ class IDisplayClusterRender_MeshComponent;
 class FDisplayClusterViewport_OpenColorIO;
 class FSceneViewStateReference;
 
-//
-// Container for data exchange game->render threads
-//
+/**
+ *  Container for data exchange game->render threads
+ */
 class FDisplayClusterViewportProxyData
 {
 public:
-	FDisplayClusterViewportProxyData(const TSharedPtr<FDisplayClusterViewport, ESPMode::ThreadSafe>& SrcViewport);
+	FDisplayClusterViewportProxyData(const TSharedRef<FDisplayClusterViewportProxy, ESPMode::ThreadSafe>& InDestViewportProxy)
+		: DestViewportProxy(InDestViewportProxy)
+	{ }
 	~FDisplayClusterViewportProxyData() = default;
 
-	void UpdateProxy_RenderThread() const;
-
-private:
-	TSharedPtr<FDisplayClusterViewportProxy, ESPMode::ThreadSafe> DstViewportProxy;
+public:
+	// Dest viewport proxy
+	const TSharedRef<FDisplayClusterViewportProxy, ESPMode::ThreadSafe> DestViewportProxy;
 
+	// This data is stored and copied from the game thread to the rendering thread:
+public:
 	TSharedPtr<IDisplayClusterRender_MeshComponent, ESPMode::ThreadSafe> RemapMesh;
 
 	FDisplayClusterViewport_OverscanRuntimeSettings OverscanRuntimeSettings;
diff --git a/Engine/Plugins/Runtime/nDisplay/Source/DisplayCluster/Private/Render/Viewport/Containers/DisplayClusterViewportRemap.cpp b/Engine/Plugins/Runtime/nDisplay/Source/DisplayCluster/Private/Render/Viewport/Containers/DisplayClusterViewportRemap.cpp
--- a/Engine/Plugins/Runtime/nDisplay/Source/DisplayCluster/Private/Render/Viewport/Containers/DisplayClusterViewportRemap.cpp
+++ b/Engine/Plugins/Runtime/nDisplay/Source/DisplayCluster/Private/Render/Viewport/Containers/DisplayClusterViewportRemap.cpp
@@ -11,7 +11,7 @@
 #include "Render/Viewport/RenderFrame/DisplayClusterRenderFrameSettings.h"
 #include "Render/Containers/DisplayClusterRender_MeshGeometry.h"
 
-bool FDisplayClusterViewportRemap::UpdateConfiguration(const FDisplayClusterViewport& InViewport, const FDisplayClusterConfigurationViewport_Remap& InRemapConfiguration)
+bool FDisplayClusterViewportRemap::UpdateViewportRemap(const FDisplayClusterViewport& InViewport, const FDisplayClusterConfigurationViewport_Remap& InRemapConfiguration)
 {
 	if (InRemapConfiguration.bEnable)
 	{
@@ -65,7 +65,7 @@ bool FDisplayClusterViewportRemap::UpdateConfiguration(const FDisplayClusterView
 void FDisplayClusterViewportRemap::Update(const FDisplayClusterViewport& InViewport, const FIntPoint& InRenderFrameSize)
 {
 	// Preview in editor not support this feature
-	if (InViewport.GetRenderMode() == EDisplayClusterRenderFrameMode::PreviewInScene)
+	if (InViewport.Configuration->IsPreviewRendering())
 	{
 		return;
 	}
diff --git a/Engine/Plugins/Runtime/nDisplay/Source/DisplayCluster/Private/Render/Viewport/Containers/DisplayClusterViewportRemap.h b/Engine/Plugins/Runtime/nDisplay/Source/DisplayCluster/Private/Render/Viewport/Containers/DisplayClusterViewportRemap.h
--- a/Engine/Plugins/Runtime/nDisplay/Source/DisplayCluster/Private/Render/Viewport/Containers/DisplayClusterViewportRemap.h
+++ b/Engine/Plugins/Runtime/nDisplay/Source/DisplayCluster/Private/Render/Viewport/Containers/DisplayClusterViewportRemap.h
@@ -21,7 +21,7 @@ public:
 	 * @param InRemapConfiguration - The remapping configuration to update to
 	 * @returns true if the render geometry needs updating after the configuration has been updated; otherwise, false
 	 */
-	bool UpdateConfiguration(const FDisplayClusterViewport& InViewport, const FDisplayClusterConfigurationViewport_Remap& InRemapConfiguration);
+	bool UpdateViewportRemap(const FDisplayClusterViewport& InViewport, const FDisplayClusterConfigurationViewport_Remap& InRemapConfiguration);
 
 	/** Gets a pointer to the mesh geometry to use for rendering the remapped regions */
 	TSharedPtr<IDisplayClusterRender_MeshComponent, ESPMode::ThreadSafe> GetRemapMesh() const;
diff --git a/Engine/Plugins/Runtime/nDisplay/Source/DisplayCluster/Private/Render/Viewport/Containers/ImplDisplayClusterViewport_CameraMotionBlur.h b/Engine/Plugins/Runtime/nDisplay/Source/DisplayCluster/Private/Render/Viewport/Containers/ImplDisplayClusterViewport_CameraMotionBlur.h
--- a/Engine/Plugins/Runtime/nDisplay/Source/DisplayCluster/Private/Render/Viewport/Containers/ImplDisplayClusterViewport_CameraMotionBlur.h
+++ b/Engine/Plugins/Runtime/nDisplay/Source/DisplayCluster/Private/Render/Viewport/Containers/ImplDisplayClusterViewport_CameraMotionBlur.h
@@ -73,7 +73,7 @@ public:
 		}
 	}
 
-	void ResetConfiguration()
+	void BeginUpdateSettings()
 	{
 		BlurSetup = FDisplayClusterViewport_CameraMotionBlur();
 	}
