nDisplay: Added new interface IDisplayClusterViewportConfiguration

diff --git a/Engine/Plugins/Runtime/nDisplay/Source/DisplayCluster/Private/Render/Viewport/DisplayClusterViewport.cpp b/Engine/Plugins/Runtime/nDisplay/Source/DisplayCluster/Private/Render/Viewport/DisplayClusterViewport.cpp
--- a/Engine/Plugins/Runtime/nDisplay/Source/DisplayCluster/Private/Render/Viewport/DisplayClusterViewport.cpp
+++ b/Engine/Plugins/Runtime/nDisplay/Source/DisplayCluster/Private/Render/Viewport/DisplayClusterViewport.cpp
@@ -9,6 +9,7 @@
 #include "Render/Viewport/DisplayClusterViewportManagerViewExtension.h"
 
 #include "Render/Viewport/Configuration/DisplayClusterViewportConfiguration.h"
+#include "Render/Viewport/Configuration/DisplayClusterViewportConfigurationProxy.h"
 #include "Render/Viewport/LightCard/DisplayClusterViewportLightCardManager.h"
 
 #include "Render/Projection/IDisplayClusterProjectionPolicy.h"
@@ -19,6 +20,7 @@
 #include "Render/Viewport/RenderTarget/DisplayClusterRenderTargetResource.h"
 
 #include "Render/Viewport/Containers/DisplayClusterViewport_PostRenderSettings.h"
+#include "Render/Viewport/Containers/DisplayClusterViewportProxyData.h"
 
 #include "EngineUtils.h"
 #include "SceneManagement.h"
@@ -42,25 +44,24 @@ static FAutoConsoleVariableRef CVarDisplayClusterMultiGPUEnable(
 ///////////////////////////////////////////////////////////////////////////////////////
 //          FDisplayClusterViewport
 ///////////////////////////////////////////////////////////////////////////////////////
-FDisplayClusterViewport::FDisplayClusterViewport(FDisplayClusterViewportManager& InOwner, const FString& InClusterNodeId, const FString& InViewportId, const TSharedPtr<IDisplayClusterProjectionPolicy, ESPMode::ThreadSafe>& InProjectionPolicy)
-	: UninitializedProjectionPolicy(InProjectionPolicy)
+FDisplayClusterViewport::FDisplayClusterViewport(const TSharedRef<FDisplayClusterViewportConfiguration, ESPMode::ThreadSafe>& InConfiguration, const FString& InViewportId, const TSharedPtr<IDisplayClusterProjectionPolicy, ESPMode::ThreadSafe>& InProjectionPolicy)
+	: Configuration(InConfiguration)
+	, ViewportProxy(MakeShared<FDisplayClusterViewportProxy, ESPMode::ThreadSafe>(InConfiguration, InViewportId, InProjectionPolicy))
 	, ViewportId(InViewportId)
-	, ClusterNodeId(InClusterNodeId)
-	, ViewportManagerWeakPtr(InOwner.AsShared())
-	, ViewportManagerProxyWeakPtr(InOwner.GetViewportManagerProxy())
+	, ClusterNodeId(InConfiguration->GetClusterNodeId())
+	, UninitializedProjectionPolicy(InProjectionPolicy)
 {
 	check(!ClusterNodeId.IsEmpty());
 	check(!ViewportId.IsEmpty());
 	check(UninitializedProjectionPolicy.IsValid());
 
-	if (ViewportManagerProxyWeakPtr.IsValid())
+	if (FDisplayClusterViewportManagerProxy* ViewportManagerProxy = Configuration->Proxy->GetViewportManagerProxyImpl())
 	{
-		// Create scene proxy pair with on game thread. Outside, in ViewportManager added to proxy array on render thread
-		ViewportProxy = MakeShared<FDisplayClusterViewportProxy, ESPMode::ThreadSafe>(*this);
-
 		// Add viewport proxy on renderthread
 		ENQUEUE_RENDER_COMMAND(CreateDisplayClusterViewportProxy)(
-			[ViewportManagerProxy = ViewportManagerProxyWeakPtr.Pin(), ViewportProxy = ViewportProxy](FRHICommandListImmediate& RHICmdList)
+			[ ViewportManagerProxy = ViewportManagerProxy->AsShared()
+			, ViewportProxy = ViewportProxy
+			](FRHICommandListImmediate& RHICmdList)
 			{
 				ViewportManagerProxy->CreateViewport_RenderThread(ViewportProxy);
 			}
@@ -70,18 +71,19 @@ FDisplayClusterViewport::FDisplayClusterViewport(FDisplayClusterViewportManager&
 
 FDisplayClusterViewport::~FDisplayClusterViewport()
 {
-	if (ViewportManagerProxyWeakPtr.IsValid())
+	if (FDisplayClusterViewportManagerProxy* ViewportManagerProxy = Configuration->Proxy->GetViewportManagerProxyImpl())
 	{
 		// Remove viewport proxy on render_thread
 		ENQUEUE_RENDER_COMMAND(DeleteDisplayClusterViewportProxy)(
-			[ViewportManagerProxy = ViewportManagerProxyWeakPtr.Pin(), ViewportProxy = ViewportProxy](FRHICommandListImmediate& RHICmdList)
+			[ ViewportManagerProxy = ViewportManagerProxy->AsShared()
+			, ViewportProxy = ViewportProxy
+			](FRHICommandListImmediate& RHICmdList)
 			{
 				ViewportManagerProxy->DeleteViewport_RenderThread(ViewportProxy);
 			}
 		);
 	}
 
-	ViewportProxy.Reset();
 	OpenColorIO.Reset();
 
 	// Handle projection policy EndScene event
@@ -91,98 +93,59 @@ FDisplayClusterViewport::~FDisplayClusterViewport()
 	ProjectionPolicy.Reset();
 	UninitializedProjectionPolicy.Reset();
 
-	if (FDisplayClusterViewportManager* ViewportManager = GetViewportManagerImpl())
+	if (FDisplayClusterViewportManager* ViewportManager = Configuration->GetViewportManagerImpl())
 	{
 		// Reset RTT size after viewport delete
 		ViewportManager->ResetSceneRenderTargetSize();
 	}
 }
 
-IDisplayClusterViewportManager* FDisplayClusterViewport::GetViewportManager() const
+FDisplayClusterViewportProxyData* FDisplayClusterViewport::CreateViewportProxyData()
 {
-	return GetViewportManagerImpl();
-}
-
-FDisplayClusterViewportManager* FDisplayClusterViewport::GetViewportManagerImpl() const
-{
-	return GetViewportManagerRefImpl().Get();
-}
+	FDisplayClusterViewportProxyData* OutViewportProxyData = new FDisplayClusterViewportProxyData(ViewportProxy);
 
-FDisplayClusterViewportManagerProxy* FDisplayClusterViewport::GetViewportManagerProxyImpl() const
-{
-	return GetViewportManagerProxyRefImpl().Get();
-}
+	OutViewportProxyData->OpenColorIO = OpenColorIO;
 
-TSharedPtr<FDisplayClusterViewportManager, ESPMode::ThreadSafe> FDisplayClusterViewport::GetViewportManagerRefImpl() const
-{
-	return ViewportManagerWeakPtr.IsValid() ? ViewportManagerWeakPtr.Pin(): nullptr;
-}
+	OutViewportProxyData->RenderSettings = RenderSettings;
+	OutViewportProxyData->RenderSettingsICVFX.SetParameters(RenderSettingsICVFX);
+	OutViewportProxyData->PostRenderSettings.SetParameters(PostRenderSettings);
 
-TSharedPtr<FDisplayClusterViewportManagerProxy, ESPMode::ThreadSafe> FDisplayClusterViewport::GetViewportManagerProxyRefImpl() const
-{
-	return ViewportManagerProxyWeakPtr.IsValid() ? ViewportManagerProxyWeakPtr.Pin() : nullptr;
-}
+	// Additional parameters
+	OutViewportProxyData->OverscanRuntimeSettings = OverscanRuntimeSettings;
 
-ADisplayClusterRootActor* FDisplayClusterViewport::GetRootActor() const
-{
-	if (FDisplayClusterViewportManager* ViewportManager = GetViewportManagerImpl())
-	{
-		return ViewportManager->GetRootActor();
-	}
+	OutViewportProxyData->RemapMesh = ViewportRemap.GetRemapMesh();
 
-	return nullptr;
-}
+	OutViewportProxyData->ProjectionPolicy = ProjectionPolicy;
+	OutViewportProxyData->Contexts = Contexts;
 
-UWorld* FDisplayClusterViewport::GetCurrentWorld() const
-{
-	if (FDisplayClusterViewportManager* ViewportManager = GetViewportManagerImpl())
-	{
-		return ViewportManager->GetCurrentWorld();
-	}
+	OutViewportProxyData->Resources = Resources;
+	OutViewportProxyData->ViewStates = ViewStates;
 
-	return nullptr;
+	return OutViewportProxyData;
 }
 
-bool FDisplayClusterViewport::IsCurrentWorldHasAnyType(const EWorldType::Type InWorldType1, const EWorldType::Type InWorldType2, const EWorldType::Type InWorldType3) const
+void FDisplayClusterViewport::BeginNewFrame(const FIntPoint& InRenderFrameSize)
 {
-	if (UWorld* CurrentWorld = GetCurrentWorld())
-	{
-		return (CurrentWorld->WorldType == InWorldType1 && InWorldType1 != EWorldType::None)
-			|| (CurrentWorld->WorldType == InWorldType2 && InWorldType2 != EWorldType::None)
-			|| (CurrentWorld->WorldType == InWorldType3 && InWorldType3 != EWorldType::None);
-	}
-
-	return false;
-}
-
-const FDisplayClusterRenderFrameSettings* FDisplayClusterViewport::GetRenderFrameSettings() const
-{
-	if (FDisplayClusterViewportManager* ViewportManager = GetViewportManagerImpl())
-	{
-		return &ViewportManager->GetRenderFrameSettings();
-	}
+	check(IsInGameThread());
 
-	return nullptr;
+	// Update ViewportRemap geometry
+	ViewportRemap.Update(*this, InRenderFrameSize);
 }
 
-EDisplayClusterRenderFrameMode FDisplayClusterViewport::GetRenderMode() const
+void FDisplayClusterViewport::FinalizeNewFrame()
 {
-	if (FDisplayClusterViewportManager* ViewportManager = GetViewportManagerImpl())
-	{
-		return ViewportManager->GetRenderMode();
-	}
+	check(IsInGameThread());
 
-	return EDisplayClusterRenderFrameMode::Unknown;
-}
+	// When all viewports processed, we remove all single frame custom postprocess
+	CustomPostProcessSettings.FinalizeFrame();
 
-bool FDisplayClusterViewport::IsSceneOpened() const
-{
-	if (FDisplayClusterViewportManager* ViewportManager = GetViewportManagerImpl())
+	// Update projection policy proxy data
+	if (ProjectionPolicy.IsValid())
 	{
-		return ViewportManager->IsSceneOpened();
+		ProjectionPolicy->UpdateProxyData(this);
 	}
 
-	return false;
+	RenderSettings.FinishUpdateSettings();
 }
 
 bool FDisplayClusterViewport::IsOpenColorIOEquals(const FDisplayClusterViewport& InViewport) const
@@ -209,13 +172,6 @@ bool FDisplayClusterViewport::IsOpenColorIOEquals(const FDisplayClusterViewport&
 
 const TArray<FSceneViewExtensionRef> FDisplayClusterViewport::GatherActiveExtensions(FViewport* InViewport) const
 {
-	TSharedPtr<FDisplayClusterViewportManager, ESPMode::ThreadSafe> ViewportManager = GetViewportManagerRefImpl();
-	if (!ViewportManager.IsValid())
-	{
-		// No viewport manager = no extension found.
-		return TArray<FSceneViewExtensionRef>();
-	}
-
 	// Use VE from engine for default render and MRQ:
 	switch (RenderSettings.CaptureMode)
 	{
@@ -223,14 +179,14 @@ const TArray<FSceneViewExtensionRef> FDisplayClusterViewport::GatherActiveExtens
 	case EDisplayClusterViewportCaptureMode::MoviePipeline:
 		if (InViewport)
 		{
-			FDisplayClusterSceneViewExtensionContext ViewExtensionContext(InViewport, ViewportManager, GetId());
+			FDisplayClusterSceneViewExtensionContext ViewExtensionContext(InViewport, Configuration, GetId());
 			return GEngine->ViewExtensions->GatherActiveExtensions(ViewExtensionContext);
 		}
 		else
 		{
-			if (UWorld* CurrentWorld = ViewportManager->GetCurrentWorld())
+			if (UWorld* CurrentWorld = Configuration->GetCurrentWorld())
 			{
-				FDisplayClusterSceneViewExtensionContext ViewExtensionContext(CurrentWorld->Scene, ViewportManager, GetId());
+				FDisplayClusterSceneViewExtensionContext ViewExtensionContext(CurrentWorld->Scene, Configuration, GetId());
 				return GEngine->ViewExtensions->GatherActiveExtensions(ViewExtensionContext);
 			}
 		}
@@ -252,10 +208,13 @@ const TArray<FSceneViewExtensionRef> FDisplayClusterViewport::GatherActiveExtens
 	case EDisplayClusterViewportCaptureMode::Lightcard:
 	{
 		// Chromakey and LightCard use only nDisplay VE for callback purposes (preserve alpha channel, etc).
-		TSharedPtr<FDisplayClusterViewportManagerViewExtension, ESPMode::ThreadSafe> ViewportManagerViewExtension = ViewportManager->GetViewportManagerViewExtension();
-		if (ViewportManagerViewExtension.IsValid())
+		if (FDisplayClusterViewportManager* ViewportManager = Configuration->GetViewportManagerImpl())
 		{
-			OutCustomExtensions.Add(ViewportManagerViewExtension->AsShared());
+			TSharedPtr<FDisplayClusterViewportManagerViewExtension, ESPMode::ThreadSafe> ViewportManagerViewExtension = ViewportManager->GetViewportManagerViewExtension();
+			if (ViewportManagerViewExtension.IsValid())
+			{
+				OutCustomExtensions.Add(ViewportManagerViewExtension->AsShared());
+			}
 		}
 	}
 	break;
@@ -270,32 +229,25 @@ const TArray<FSceneViewExtensionRef> FDisplayClusterViewport::GatherActiveExtens
 	return OutCustomExtensions;
 }
 
-bool FDisplayClusterViewport::HandleStartScene()
+void FDisplayClusterViewport::HandleStartScene()
 {
-	bool bResult = true;
-	if (IsSceneOpened())
+	if (UninitializedProjectionPolicy.IsValid())
 	{
-		if (UninitializedProjectionPolicy.IsValid())
+		if (UninitializedProjectionPolicy->HandleStartScene(this))
 		{
-			bResult = UninitializedProjectionPolicy->HandleStartScene(this);
-			if (bResult)
-			{
-				ProjectionPolicy = UninitializedProjectionPolicy;
-				UninitializedProjectionPolicy.Reset();
-			}
+			ProjectionPolicy = UninitializedProjectionPolicy;
+			UninitializedProjectionPolicy.Reset();
 		}
-		else 
+	}
+	else 
+	{
+		// Already Initialized
+		if (!ProjectionPolicy.IsValid())
 		{
-			// Already Initialized
-			if (!ProjectionPolicy.IsValid())
-			{
-				// No projection policy for this viewport
-				UE_LOG(LogDisplayClusterViewport, Error, TEXT("No projection policy assigned for Viewports '%s'."), *GetId());
-			}
+			// No projection policy for this viewport
+			UE_LOG(LogDisplayClusterViewport, Error, TEXT("No projection policy assigned for Viewports '%s'."), *GetId());
 		}
 	}
-
-	return bResult;
 }
 
 void FDisplayClusterViewport::HandleEndScene()
@@ -359,7 +311,7 @@ void FDisplayClusterViewport::SetViewportBufferRatio(const float InBufferRatio)
 	const float BufferRatio = FLegacyScreenPercentageDriver::GetCVarResolutionFraction() * InBufferRatio;
 	if (RenderSettings.BufferRatio > BufferRatio)
 	{
-		if (FDisplayClusterViewportManager* ViewportManager = GetViewportManagerImpl())
+		if (FDisplayClusterViewportManager* ViewportManager = Configuration->GetViewportManagerImpl())
 		{
 			// Reset scene RTT when buffer ratio changed down
 			ViewportManager->ResetSceneRenderTargetSize();
@@ -492,17 +444,12 @@ void FDisplayClusterViewport::ResetFrameContexts()
 	Resources.Release();
 }
 
-bool FDisplayClusterViewport::UpdateFrameContexts(const uint32 InStereoViewIndex, const FDisplayClusterRenderFrameSettings& InFrameSettings)
+bool FDisplayClusterViewport::UpdateFrameContexts(const uint32 InStereoViewIndex)
 {
 	check(IsInGameThread());
 
-	FDisplayClusterViewportManager* ViewportManager = GetViewportManagerImpl();
-	if (!ViewportManager)
-	{
-		return false;
-	}
-
-	const uint32 FrameTargetsAmount = ViewportManager->GetViewPerViewportAmount();
+	const FDisplayClusterRenderFrameSettings& InFrameSettings = Configuration->GetRenderFrameSettings();
+	const uint32 FrameTargetsAmount = InFrameSettings.GetViewPerViewportAmount();
 	if (FrameTargetsAmount == 0)
 	{
 		ResetFrameContexts();
@@ -510,30 +457,17 @@ bool FDisplayClusterViewport::UpdateFrameContexts(const uint32 InStereoViewIndex
 		return false;
 	}
 
-	FIntRect     DesiredFrameTargetRect = RenderSettings.Rect;
-	{
-		switch (InFrameSettings.RenderMode)
-		{
-		case EDisplayClusterRenderFrameMode::SideBySide:
-			AdjustRect(DesiredFrameTargetRect, 0.5f, 1.f);
-			break;
-		case EDisplayClusterRenderFrameMode::TopBottom:
-			AdjustRect(DesiredFrameTargetRect, 1.f, 0.5f);
-			break;
-		case EDisplayClusterRenderFrameMode::PreviewInScene:
-		{
-			// Preview downscale in range 0..1
-			float MultXY = FMath::Clamp(InFrameSettings.PreviewRenderTargetRatioMult, 0.f, 1.f);
-			AdjustRect(DesiredFrameTargetRect, MultXY, MultXY);
+	FIntRect DesiredFrameTargetRect = RenderSettings.Rect;
 
-			// Align each frame to zero
-			DesiredFrameTargetRect = FIntRect(FIntPoint(0, 0), DesiredFrameTargetRect.Size());
+	// Apply desired frame mult
+	const FVector2D DesiredFrameMult = InFrameSettings.GetDesiredFrameMult();
+	AdjustRect(DesiredFrameTargetRect, DesiredFrameMult.X, DesiredFrameMult.Y);
 
-			break;
-		}
-		default:
-			break;
-		}
+	// Support preview in scene rendering
+	if (InFrameSettings.IsPreviewRendering())
+	{
+		// Preview renders each viewport into a separate texture, so each frame is zero-aligned
+		DesiredFrameTargetRect = FIntRect(FIntPoint(0, 0), DesiredFrameTargetRect.Size());
 	}
 
 	// Special case mono->stereo
@@ -600,18 +534,12 @@ bool FDisplayClusterViewport::UpdateFrameContexts(const uint32 InStereoViewIndex
 	// Scale context for rendering
 	FIntPoint DesiredContextSize = GetDesiredContextSize(FrameTargetRect.Size(), InFrameSettings);
 
-#if WITH_EDITOR
-	switch (InFrameSettings.RenderMode)
+	// Apply restrictions on the maximum size of the viewport texture.
+	const int32 ViewportTextureMaxSize = InFrameSettings.GetViewportTextureMaxSize();
+	if(ViewportTextureMaxSize > 0)
 	{
-	case EDisplayClusterRenderFrameMode::PreviewInScene:
-		{
-			DesiredContextSize = FDisplayClusterViewportHelpers::GetTextureSizeLessThanMax(DesiredContextSize, InFrameSettings.PreviewMaxTextureDimension);
-		}
-		break;
-	default:
-		break;
+		DesiredContextSize = FDisplayClusterViewportHelpers::GetTextureSizeLessThanMax(DesiredContextSize, ViewportTextureMaxSize);
 	}
-#endif
 
 	// Exclude zero-size viewports from render
 	if (DesiredContextSize.GetMin() <= 0)
@@ -683,19 +611,21 @@ bool FDisplayClusterViewport::UpdateFrameContexts(const uint32 InStereoViewIndex
 		// Use the UVLightCard viewport only when this type of lightcards has been defined
 		bool bUseUVLightCardViewport = false;
 
-		TSharedPtr<FDisplayClusterViewportLightCardManager, ESPMode::ThreadSafe> LightCardManager = ViewportManager->GetLightCardManager();
-		if (LightCardManager.IsValid() && LightCardManager->IsUVLightCardEnabled())
+		if (FDisplayClusterViewportManager* ViewportManager = Configuration->GetViewportManagerImpl())
 		{
-			// Custom viewport size from LC Manager
-			ContextSize = LightCardManager->GetUVLightCardResourceSize();
-
-			// Size must be not null
-			if (ContextSize.GetMin() > 1)
+			if (ViewportManager->LightCardManager->IsUVLightCardEnabled())
 			{
-				FrameTargetRect = RenderTargetRect = FIntRect(FIntPoint(0, 0), ContextSize);
+				// Custom viewport size from LC Manager
+				ContextSize = ViewportManager->LightCardManager->GetUVLightCardResourceSize();
+
+				// Size must be not null
+				if (ContextSize.GetMin() > 1)
+				{
+					FrameTargetRect = RenderTargetRect = FIntRect(FIntPoint(0, 0), ContextSize);
 
-				// Allow to use this viewport
-				bUseUVLightCardViewport = true;
+					// Allow to use this viewport
+					bUseUVLightCardViewport = true;
+				}
 			}
 		}
 
@@ -723,25 +653,21 @@ bool FDisplayClusterViewport::UpdateFrameContexts(const uint32 InStereoViewIndex
 		}
 
 		const int32 MaxExplicitGPUIndex = GDisplayClusterMultiGPUEnable ? GNumExplicitGPUsForRendering - 1 : 0;
-
-		if (MaxExplicitGPUIndex > 0)
-		{
-			if (InFrameSettings.bIsRenderingInEditor)
+		if (MaxExplicitGPUIndex > 0 && !bDisableRender)
 			{
 				// Experimental: allow mGPU for preview rendering:
-				if (InFrameSettings.bAllowMultiGPURenderingInEditor && !bDisableRender)
+			if (const FIntPoint* GPURange = InFrameSettings.GetPreviewMultiGPURendering())
 				{
-					int32 MinGPUIndex = FMath::Min(InFrameSettings.PreviewMinGPUIndex, MaxExplicitGPUIndex);
-					int32 MaxGPUIndex = FMath::Min(InFrameSettings.PreviewMaxGPUIndex, MaxExplicitGPUIndex);
+				int32 MinGPUIndex = FMath::Min(GPURange->X, MaxExplicitGPUIndex);
+				int32 MaxGPUIndex = FMath::Min(GPURange->Y, MaxExplicitGPUIndex);
 
 					static int32 PreviewGPUIndex = MinGPUIndex;
-					Context.GPUIndex = PreviewGPUIndex++;
-
 					if (PreviewGPUIndex > MaxGPUIndex)
 					{
 						PreviewGPUIndex = MinGPUIndex;
 					}
-				}
+
+				Context.GPUIndex = PreviewGPUIndex++;
 			}
 			else
 			{
@@ -774,7 +700,7 @@ bool FDisplayClusterViewport::UpdateFrameContexts(const uint32 InStereoViewIndex
 	}
 
 	// Reserve for resources
-	if (InFrameSettings.RenderMode == EDisplayClusterRenderFrameMode::PreviewInScene)
+	if (InFrameSettings.IsPreviewRendering())
 	{
 		// reserve preview texture resource
 		Resources[EDisplayClusterViewportResource::OutputPreviewTargetableResources].AddZeroed(FrameTargetsAmount);
diff --git a/Engine/Plugins/Runtime/nDisplay/Source/DisplayCluster/Private/Render/Viewport/DisplayClusterViewport.h b/Engine/Plugins/Runtime/nDisplay/Source/DisplayCluster/Private/Render/Viewport/DisplayClusterViewport.h
--- a/Engine/Plugins/Runtime/nDisplay/Source/DisplayCluster/Private/Render/Viewport/DisplayClusterViewport.h
+++ b/Engine/Plugins/Runtime/nDisplay/Source/DisplayCluster/Private/Render/Viewport/DisplayClusterViewport.h
@@ -3,6 +3,7 @@
 #pragma once
 
 #include "Render/Viewport/IDisplayClusterViewport.h"
+#include "Render/Viewport/Configuration/DisplayClusterViewportConfiguration.h"
 
 #include "Render/Viewport/Containers/ImplDisplayClusterViewport_CameraMotionBlur.h"
 #include "Render/Viewport/Containers/DisplayClusterViewport_CustomFrustumRuntimeSettings.h"
@@ -22,33 +23,19 @@
 
 class FDisplayClusterViewportManager;
 class FDisplayClusterViewportManagerProxy;
-class FDisplayClusterRenderTargetManager;
-class FDisplayClusterRenderFrameManager;
-class FDisplayClusterViewportProxyData;
 class FDisplayClusterViewportProxy;
-struct FDisplayClusterRenderFrameSettings;
-class FDisplayClusterViewportConfigurationCameraViewport;
-class FDisplayClusterViewportConfigurationCameraICVFX;
-class FDisplayClusterViewportConfigurationICVFX;
-
-class FDisplayClusterViewportConfigurationHelpers;
-class FDisplayClusterViewportConfigurationHelpers_ICVFX;
-class FDisplayClusterViewportConfigurationHelpers_OpenColorIO;
-class FDisplayClusterViewportConfigurationHelpers_Postprocess;
 
-struct FDisplayClusterViewportConfigurationProjectionPolicy;
+struct FDisplayClusterRenderFrameSettings;
 
 /**
  * Rendering viewport (sub-region of the main viewport)
  */
-
 class FDisplayClusterViewport
 	: public IDisplayClusterViewport
 	, public TSharedFromThis<FDisplayClusterViewport, ESPMode::ThreadSafe>
 {
 public:
-	FDisplayClusterViewport(FDisplayClusterViewportManager& Owner, const FString& ClusterNodeId, const FString& ViewportId, const TSharedPtr<IDisplayClusterProjectionPolicy, ESPMode::ThreadSafe>& InProjectionPolicy);
-	
+	FDisplayClusterViewport(const TSharedRef<FDisplayClusterViewportConfiguration, ESPMode::ThreadSafe>& InConfiguration, const FString& ViewportId, const TSharedPtr<IDisplayClusterProjectionPolicy, ESPMode::ThreadSafe>& InProjectionPolicy);
 	virtual ~FDisplayClusterViewport();
 
 public:
@@ -65,7 +52,15 @@ public:
 		return AsShared();
 	}
 
-	virtual EDisplayClusterRenderFrameMode GetRenderMode() const override;
+	virtual IDisplayClusterViewportConfiguration& GetConfiguration() override
+	{
+		return Configuration.Get();
+	}
+
+	virtual const IDisplayClusterViewportConfiguration& GetConfiguration() const override
+	{
+		return Configuration.Get();
+	}
 
 	virtual FString GetId() const override
 	{ 
@@ -107,7 +102,7 @@ public:
 
 	virtual bool SetupViewPoint(FMinimalViewInfo& InOutViewInfo) override;
 	virtual float GetStereoEyeOffsetDistance(const uint32 InContextNum) override;
-	virtual class UDisplayClusterCameraComponent* GetViewPointCameraComponent() const override;
+	virtual class UDisplayClusterCameraComponent* GetViewPointCameraComponent(const EDisplayClusterRootActorType InRootActorType) const override;
 	virtual bool GetViewPointCameraEye(const uint32 InContextNum, FVector& OutViewLocation, FRotator& OutViewRotation, FVector& OutViewOffset) override;
 
 	virtual const FDisplayClusterViewport_RenderSettingsICVFX& GetRenderSettingsICVFX() const override
@@ -149,33 +144,14 @@ public:
 	// Setup scene view for rendering specified Context
 	virtual void SetupSceneView(uint32 ContextNum, class UWorld* World, class FSceneViewFamily& InViewFamily, FSceneView& InView) const override;
 
-	virtual class IDisplayClusterViewportManager* GetViewportManager() const override;
-	virtual class ADisplayClusterRootActor* GetRootActor() const override;
-	virtual class UWorld* GetCurrentWorld() const override;
-	virtual bool IsSceneOpened() const override;
-
-	// Return true, if current world type equal to InWorldType
-	virtual bool IsCurrentWorldHasAnyType(const EWorldType::Type InWorldType1, const EWorldType::Type InWorldType2 = EWorldType::None, const EWorldType::Type InWorldType3 = EWorldType::None) const override;
-
 	//////////////////////////////////////////////////////
 	/// ~IDisplayClusterViewport
 	//////////////////////////////////////////////////////
 
-	TSharedPtr<FDisplayClusterViewportManager, ESPMode::ThreadSafe> GetViewportManagerRefImpl() const;
-	TSharedPtr<FDisplayClusterViewportManagerProxy, ESPMode::ThreadSafe> GetViewportManagerProxyRefImpl() const;
-
-	FDisplayClusterViewportManager* GetViewportManagerImpl() const;
-	FDisplayClusterViewportManagerProxy* GetViewportManagerProxyImpl() const;
-
 	FSceneView* ImplCalcScenePreview(class FSceneViewFamilyContext& InOutViewFamily, uint32 ContextNum);
 	bool    ImplPreview_CalculateStereoViewOffset(const uint32 InContextNum, FRotator& ViewRotation, const float WorldToMeters, FVector& ViewLocation);
 	FMatrix ImplPreview_GetStereoProjectionMatrix(const uint32 InContextNum);
 
-#if WITH_EDITOR
-	bool GetPreviewPixels(TSharedPtr<class FDisplayClusterViewportReadPixelsData, ESPMode::ThreadSafe>& OutPixelsData) const;
-#endif //WITH_EDITOR
-
-
 	// Get from logic request for additional targetable resource
 	bool ShouldUseAdditionalTargetableResource() const;
 	bool ShouldUseAdditionalFrameTargetableResource() const;
@@ -183,8 +159,6 @@ public:
 
 	void SetViewportBufferRatio(const float InBufferRatio);
 
-	const FDisplayClusterRenderFrameSettings* GetRenderFrameSettings() const;
-
 	inline bool FindContext(const int32 ViewIndex, uint32* OutContextNum)
 	{
 		check(IsInGameThread());
@@ -205,7 +179,7 @@ public:
 		return false;
 	}
 
-	bool HandleStartScene();
+	void HandleStartScene();
 	void HandleEndScene();
 
 	void AddReferencedObjects(FReferenceCollector& Collector);
@@ -216,8 +190,8 @@ public:
 		RenderSettings.BeginUpdateSettings();
 		RenderSettingsICVFX.BeginUpdateSettings();
 		PostRenderSettings.BeginUpdateSettings();
-		VisibilitySettings.ResetConfiguration();
-		CameraMotionBlur.ResetConfiguration();
+		VisibilitySettings.BeginUpdateSettings();
+		CameraMotionBlur.BeginUpdateSettings();
 
 		OverscanRuntimeSettings = FDisplayClusterViewport_OverscanRuntimeSettings();
 		CustomFrustumRuntimeSettings = FDisplayClusterViewport_CustomFrustumRuntimeSettings();
@@ -229,11 +203,10 @@ public:
 	/** Initialize viewport contexts and resources for new frame
 	 *
 	 * @param InStereoViewIndex - initial StereoViewIndex for this viewport
-	 * @param InFrameSettings   - New settings for current frame
 	 *
 	 * @return - true, if success
 	 */
-	bool UpdateFrameContexts(const uint32 InStereoViewIndex, const FDisplayClusterRenderFrameSettings& InFrameSettings);
+	bool UpdateFrameContexts(const uint32 InStereoViewIndex);
 
 	/** Reset viewport contexts and resources. */
 	void ResetFrameContexts();
@@ -241,71 +214,157 @@ public:
 	/** Compare OCIO with another viewport, return true if they are equal. */
 	bool IsOpenColorIOEquals(const FDisplayClusterViewport& InViewport) const;
 
+	/** Get viewport OCIO instance. */
+	const TSharedPtr<FDisplayClusterViewport_OpenColorIO, ESPMode::ThreadSafe>& GetOpenColorIO() const
+	{
+		check(IsInGameThread());
+		return OpenColorIO;
+	}
+
+	/** Set viewport OCIO instance. */
+	void SetOpenColorIO(const TSharedPtr<FDisplayClusterViewport_OpenColorIO, ESPMode::ThreadSafe>& InOpenColorIO)
+	{
+		check(IsInGameThread());
+
+		OpenColorIO = InOpenColorIO;
+	}
+
+	/** Get viewport const resources for all contexts by type. */
+	const TArray<TSharedPtr<FDisplayClusterViewportResource, ESPMode::ThreadSafe>>& GetViewportResources(const EDisplayClusterViewportResource InResourceType) const
+	{
+		check(IsInGameThread());
+		return Resources[InResourceType];
+	}
+
+	/** Get viewport const resources for all contexts by type. */
+	TArray<TSharedPtr<FDisplayClusterViewportResource, ESPMode::ThreadSafe>>& GetViewportResourcesImpl(const EDisplayClusterViewportResource InResourceType)
+	{
+		check(IsInGameThread());
+		return Resources[InResourceType];
+	}
+
+	/** Create proxy data from this viewport internals.*/
+	class FDisplayClusterViewportProxyData* CreateViewportProxyData();
+
+
+	/** Gain direct access to internal data of the viewport. */
+	FDisplayClusterViewport_RenderSettings& GetRenderSettingsImpl()
+	{
+		check(IsInGameThread());
+		return RenderSettings;
+	}
+
+	/** Gain direct access to internal data of the viewport. */
+	FDisplayClusterViewport_RenderSettingsICVFX& GetRenderSettingsICVFXImpl()
+	{
+		check(IsInGameThread());
+		return RenderSettingsICVFX;
+	}
+
+	/** Gain direct access to internal data of the viewport. */
+	FDisplayClusterViewport_PostRenderSettings& GetPostRenderSettingsImpl()
+	{
+		check(IsInGameThread());
+		return PostRenderSettings;
+	}
+
+	/** Gain direct access to internal data of the viewport. */
+	FDisplayClusterViewport_VisibilitySettings& GetVisibilitySettingsImpl()
+	{
+		check(IsInGameThread());
+		return VisibilitySettings;
+	}
+
+	/** Start a new frame with the specified size. The associated objects and geometries will be updated accordingly. */
+	void BeginNewFrame(const FIntPoint& InRenderFrameSize);
+
+	/** Finalize new frame . */
+	void FinalizeNewFrame();
+
+	/** Release the projection policy assigned to this viewport. */
+	void ReleaseProjectionPolicy()
+	{
+		ProjectionPolicy.Reset();
+		UninitializedProjectionPolicy.Reset();
+	}
+
+	/** Update projection policy from configuration. */
+	void UpdateConfiguration_ProjectionPolicy(const struct FDisplayClusterConfigurationProjection* InConfigurationProjectionPolicy = nullptr);
+
+	/** setup overlay configuration for this viewport. */
+	void UpdateConfiguration_OverlayRenderSettings(const struct FDisplayClusterConfigurationICVFX_OverlayAdvancedRenderSettings& InOverlaySettings);
+
+	/** setup overscan configuration for this viewport. */
+	void UpdateConfiguration_Overscan(const struct FDisplayClusterConfigurationViewport_Overscan& InOverscan);
+
+	/** setup CameraMotionBlur configuration for this viewport. */
+	void UpdateConfiguration_CameraMotionBlur(const struct FDisplayClusterViewport_CameraMotionBlur& InCameraMotionBlur);
+
+	/** setup PostRender mips configuration for this viewport. */
+	void UpdateConfiguration_PostRenderGenerateMips(const struct FDisplayClusterConfigurationPostRender_GenerateMips& InGenerateMips);
+
+	/** setup PostRender override configuration for this viewport. */
+	void UpdateConfiguration_PostRenderOverride(const struct FDisplayClusterConfigurationPostRender_Override& InOverride);
+
+	/** setup PostRender blur configuration for this viewport. */
+	void UpdateConfiguration_PostRenderBlur(const struct FDisplayClusterConfigurationPostRender_BlurPostprocess& InBlurPostprocess);
+
+	/** setup viewport remap configuration for this viewport. */
+	bool UpdateConfiguration_ViewportRemap(const struct FDisplayClusterConfigurationViewport_Remap& InRemapConfiguration);
+
+	/** Support view states for preview. */
+	FSceneViewStateInterface* GetViewState(uint32 ViewIndex);
+
+	/** Cleanup view states. */
+	void CleanupViewState();
+
 private:
 	float GetClusterRenderTargetRatioMult(const FDisplayClusterRenderFrameSettings& InFrameSettings) const;
 	FIntPoint GetDesiredContextSize(const FIntPoint& InSize, const FDisplayClusterRenderFrameSettings& InFrameSettings) const;
 	float GetCustomBufferRatio(const FDisplayClusterRenderFrameSettings& InFrameSettings) const;
 
-private:
-	// Support view states for preview
-	FSceneViewStateInterface* GetViewState(uint32 ViewIndex);
-
 public:
-	void CleanupViewState();
+	// Configuration of the current cluster node
+	const TSharedRef<FDisplayClusterViewportConfiguration, ESPMode::ThreadSafe> Configuration;
+
+	// viewport proxy (render thread data)
+	const TSharedRef<FDisplayClusterViewportProxy, ESPMode::ThreadSafe> ViewportProxy;
+
+	// Unique viewport name
+	const FString ViewportId;
+
+	// Owner cluster node name
+	const FString ClusterNodeId;
+
+private:
+	// Unified repository of viewport resources
+	FDisplayClusterViewportResources Resources;
 
-public:
 	/** nDisplay OpenColorIO object. */
 	TSharedPtr<FDisplayClusterViewport_OpenColorIO, ESPMode::ThreadSafe> OpenColorIO;
 
-public:
 	// Projection policy instance that serves this viewport
 	TSharedPtr<IDisplayClusterProjectionPolicy, ESPMode::ThreadSafe> ProjectionPolicy;
 	TSharedPtr<IDisplayClusterProjectionPolicy, ESPMode::ThreadSafe> UninitializedProjectionPolicy;
 
-	// Game thread only settings:
+	// Postprocess
 	FDisplayClusterViewport_CustomPostProcessSettings CustomPostProcessSettings;
+
+	// Visibility settings
 	FDisplayClusterViewport_VisibilitySettings        VisibilitySettings;
 
-	// Additional features:
+	// Additional features
 	FImplDisplayClusterViewport_CameraMotionBlur CameraMotionBlur;
 
+	// Overscan rendering feature
 	FDisplayClusterViewport_OverscanRuntimeSettings      OverscanRuntimeSettings;
+
+	// Custom frustum rendering feature
 	FDisplayClusterViewport_CustomFrustumRuntimeSettings CustomFrustumRuntimeSettings;
 
 	// viewport OutputRemap feature
 	FDisplayClusterViewportRemap ViewportRemap;
 
-	// Unified repository of viewport resources
-	FDisplayClusterViewportResources Resources;
-
-protected:
-	friend FDisplayClusterViewportProxy;
-	friend FDisplayClusterViewportProxyData;
-	friend FDisplayClusterViewportManager;
-	friend FDisplayClusterRenderTargetManager;
-	friend FDisplayClusterRenderFrameManager;
-	friend FDisplayClusterViewportConfigurationCameraViewport;
-	friend FDisplayClusterViewportConfigurationCameraICVFX;
-
-	friend FDisplayClusterViewportConfigurationICVFX;
-
-	friend FDisplayClusterViewportConfigurationHelpers;
-	friend FDisplayClusterViewportConfigurationHelpers_ICVFX;
-	friend FDisplayClusterViewportConfigurationHelpers_OpenColorIO;
-	friend FDisplayClusterViewportConfigurationHelpers_Postprocess;
-	friend FDisplayClusterViewportConfigurationProjectionPolicy;
-
-	friend FDisplayClusterViewportRemap;
-
-	// viewport render thread data
-	TSharedPtr<FDisplayClusterViewportProxy, ESPMode::ThreadSafe> ViewportProxy;
-
-	// Unique viewport name
-	const FString ViewportId;
-
-	// Owner cluster node name
-	const FString ClusterNodeId;
-
 	// Viewport render params
 	FDisplayClusterViewport_RenderSettings       RenderSettings;
 	FDisplayClusterViewport_RenderSettingsICVFX  RenderSettingsICVFX;
@@ -314,11 +373,6 @@ protected:
 	// Viewport contexts (left/center/right eyes)
 	TArray<FDisplayClusterViewport_Context> Contexts;
 
-	// viewport owners
-	TWeakPtr<FDisplayClusterViewportManager, ESPMode::ThreadSafe> ViewportManagerWeakPtr;
-	TWeakPtr<FDisplayClusterViewportManagerProxy, ESPMode::ThreadSafe> ViewportManagerProxyWeakPtr;
-
-private:
 	// View states (preview only)
 	TArray<TSharedPtr<FSceneViewStateReference, ESPMode::ThreadSafe>> ViewStates;
 
diff --git a/Engine/Plugins/Runtime/nDisplay/Source/DisplayCluster/Private/Render/Viewport/DisplayClusterViewportFrameStatsViewExtension.cpp b/Engine/Plugins/Runtime/nDisplay/Source/DisplayCluster/Private/Render/Viewport/DisplayClusterViewportFrameStatsViewExtension.cpp
--- a/Engine/Plugins/Runtime/nDisplay/Source/DisplayCluster/Private/Render/Viewport/DisplayClusterViewportFrameStatsViewExtension.cpp
+++ b/Engine/Plugins/Runtime/nDisplay/Source/DisplayCluster/Private/Render/Viewport/DisplayClusterViewportFrameStatsViewExtension.cpp
@@ -78,19 +78,14 @@ namespace UE::DisplayClusterViewExtension
 ///////////////////////////////////////////////////////////////////////////////////////
 // FDisplayClusterViewportFrameStatsViewExtension
 ///////////////////////////////////////////////////////////////////////////////////////
-FDisplayClusterViewportFrameStatsViewExtension::FDisplayClusterViewportFrameStatsViewExtension(const FAutoRegister& AutoRegister, const FDisplayClusterViewportManager* InViewportManager)
+FDisplayClusterViewportFrameStatsViewExtension::FDisplayClusterViewportFrameStatsViewExtension(const FAutoRegister& AutoRegister, const TSharedRef<FDisplayClusterViewportConfiguration, ESPMode::ThreadSafe>& InConfiguration)
 	: FSceneViewExtensionBase(AutoRegister)
-	, ViewportManagerWeakPtr(InViewportManager->AsShared())
+	, Configuration(InConfiguration)
 { }
 
-FDisplayClusterViewportFrameStatsViewExtension::~FDisplayClusterViewportFrameStatsViewExtension()
-{
-	ViewportManagerWeakPtr.Reset();
-}
-
 void FDisplayClusterViewportFrameStatsViewExtension::SubscribeToPostProcessingPass(EPostProcessingPass PassId, FAfterPassCallbackDelegateArray& InOutPassCallbacks, bool bIsPassEnabled)
 {
-	if (PassId == EPostProcessingPass::Tonemap)
+	if (IsActive() && PassId == EPostProcessingPass::Tonemap)
 	{
 		InOutPassCallbacks.Add(FAfterPassCallbackDelegate::CreateRaw(this, &FDisplayClusterViewportFrameStatsViewExtension::PostProcessPassAfterTonemap_RenderThread));
 	}
@@ -98,10 +93,13 @@ void FDisplayClusterViewportFrameStatsViewExtension::SubscribeToPostProcessingPa
 
 void FDisplayClusterViewportFrameStatsViewExtension::BeginRenderViewFamily(FSceneViewFamily& InViewFamily)
 {
-	const FTimecode CurrentTimecode = FApp::GetTimecode();
+	if (IsActive())
+	{
+		const FTimecode CurrentTimecode = FApp::GetTimecode();
 
-	EncodedTimecode = (static_cast<uint8>(CurrentTimecode.Hours) << 24u) | (static_cast<uint8>(CurrentTimecode.Minutes) << 16u) | (static_cast<uint8>(CurrentTimecode.Seconds) << 8u) | static_cast<uint8>(CurrentTimecode.Frames);
-	FrameCount = GFrameCounter;
+		EncodedTimecode = (static_cast<uint8>(CurrentTimecode.Hours) << 24u) | (static_cast<uint8>(CurrentTimecode.Minutes) << 16u) | (static_cast<uint8>(CurrentTimecode.Seconds) << 8u) | static_cast<uint8>(CurrentTimecode.Frames);
+		FrameCount = GFrameCounter;
+	}
 }
 
 FScreenPassTexture FDisplayClusterViewportFrameStatsViewExtension::PostProcessPassAfterTonemap_RenderThread(FRDGBuilder& GraphBuilder, const FSceneView& View, const FPostProcessMaterialInputs& Inputs)
@@ -165,7 +163,7 @@ bool FDisplayClusterViewportFrameStatsViewExtension::IsActiveThisFrame_Internal(
 		if (Context.IsA(MoveTempIfPossible(DCViewExtensionContext)))
 		{
 			const FDisplayClusterSceneViewExtensionContext& DisplayContext = static_cast<const FDisplayClusterSceneViewExtensionContext&>(Context);
-			if (DisplayContext.GetViewportManager() == GetViewportManager())
+			if (DisplayContext.Configuration == Configuration)
 			{
 				// Apply only for DC viewports
 				return true;
@@ -178,5 +176,6 @@ bool FDisplayClusterViewportFrameStatsViewExtension::IsActiveThisFrame_Internal(
 
 bool FDisplayClusterViewportFrameStatsViewExtension::IsActive() const
 {
-	return GetViewportManager() != nullptr && GDisplayClusterShowFrameStats > 0;
+	// VE is active as long as the viewport manager proxy still exists.
+	return Configuration->Proxy->GetViewportManagerProxyImpl() != nullptr && GDisplayClusterShowFrameStats > 0;
 }
diff --git a/Engine/Plugins/Runtime/nDisplay/Source/DisplayCluster/Private/Render/Viewport/DisplayClusterViewportFrameStatsViewExtension.h b/Engine/Plugins/Runtime/nDisplay/Source/DisplayCluster/Private/Render/Viewport/DisplayClusterViewportFrameStatsViewExtension.h
--- a/Engine/Plugins/Runtime/nDisplay/Source/DisplayCluster/Private/Render/Viewport/DisplayClusterViewportFrameStatsViewExtension.h
+++ b/Engine/Plugins/Runtime/nDisplay/Source/DisplayCluster/Private/Render/Viewport/DisplayClusterViewportFrameStatsViewExtension.h
@@ -2,8 +2,7 @@
 
 #pragma once
 #include "SceneViewExtension.h"
-
-class FDisplayClusterViewportManager;
+#include "Render/Viewport/Configuration/DisplayClusterViewportConfiguration.h"
 
 #define DISPLAYCLUSTER_SCENE_DEBUG_VIEW_EXTENSION_PRIORITY 999
 
@@ -13,8 +12,8 @@ class FDisplayClusterViewportManager;
 class FDisplayClusterViewportFrameStatsViewExtension : public FSceneViewExtensionBase
 {
 public:
-	FDisplayClusterViewportFrameStatsViewExtension(const FAutoRegister& AutoRegister, const FDisplayClusterViewportManager* InViewportManager);
-	virtual ~FDisplayClusterViewportFrameStatsViewExtension();
+	FDisplayClusterViewportFrameStatsViewExtension(const FAutoRegister& AutoRegister, const TSharedRef<FDisplayClusterViewportConfiguration, ESPMode::ThreadSafe>& InConfiguration);
+	virtual ~FDisplayClusterViewportFrameStatsViewExtension() = default;
 
 public:
 	//~ Begin ISceneViewExtension interface
@@ -30,19 +29,15 @@ public:
 protected:
 	virtual bool IsActiveThisFrame_Internal(const FSceneViewExtensionContext& Context) const override;
 
-private:
-	/** Get viewport manager ptr. */
-	inline const FDisplayClusterViewportManager* GetViewportManager() const
-	{
-		return ViewportManagerWeakPtr.IsValid() ? ViewportManagerWeakPtr.Pin().Get() : nullptr;
-	}
+public:
+	// Configuration of the current cluster node
+	const TSharedRef<FDisplayClusterViewportConfiguration, ESPMode::ThreadSafe> Configuration;
 
+private:
 	/** True, if VE can be used at the moment. */
 	bool IsActive() const;
 
 private:
 	std::atomic<uint32> FrameCount;
 	std::atomic<uint32> EncodedTimecode;
-
-	TWeakPtr<const FDisplayClusterViewportManager, ESPMode::ThreadSafe> ViewportManagerWeakPtr;
 };
diff --git a/Engine/Plugins/Runtime/nDisplay/Source/DisplayCluster/Private/Render/Viewport/DisplayClusterViewportHelpers.cpp b/Engine/Plugins/Runtime/nDisplay/Source/DisplayCluster/Private/Render/Viewport/DisplayClusterViewportHelpers.cpp
--- a/Engine/Plugins/Runtime/nDisplay/Source/DisplayCluster/Private/Render/Viewport/DisplayClusterViewportHelpers.cpp
+++ b/Engine/Plugins/Runtime/nDisplay/Source/DisplayCluster/Private/Render/Viewport/DisplayClusterViewportHelpers.cpp
@@ -68,22 +68,13 @@ int32 FDisplayClusterViewportHelpers::GetMaxTextureNumMips(const FDisplayCluster
 {
 	int32 NumMips = InNumMips;
 
-#if WITH_EDITOR
-	switch (InRenderFrameSettings.RenderMode)
+	if(GDisplayClusterPreviewDefaultPixelFormat != 0 && InRenderFrameSettings.IsPreviewRendering())
 	{
-	case EDisplayClusterRenderFrameMode::PreviewInScene:
-		if (GDisplayClusterPreviewDefaultPixelFormat != 0)
-		{
-			//@todo: now UE support mips generation only for fixed point textures (8bit RGBA)
-			// Remove this hack latter
-			// Disable preview mips generation in case of unsupported RTT texture format.
-			NumMips = 0;
-		}
-		break;
-	default:
-		break;
+		//@todo: now UE support mips generation only for fixed point textures (8bit RGBA)
+		// Remove this hack latter
+		// Disable preview mips generation in case of unsupported RTT texture format.
+		NumMips = 0;
 	}
-#endif
 
 	if (GDisplayClusterMaxNumMips >= 0)
 	{
diff --git a/Engine/Plugins/Runtime/nDisplay/Source/DisplayCluster/Private/Render/Viewport/DisplayClusterViewportManager.cpp b/Engine/Plugins/Runtime/nDisplay/Source/DisplayCluster/Private/Render/Viewport/DisplayClusterViewportManager.cpp
--- a/Engine/Plugins/Runtime/nDisplay/Source/DisplayCluster/Private/Render/Viewport/DisplayClusterViewportManager.cpp
+++ b/Engine/Plugins/Runtime/nDisplay/Source/DisplayCluster/Private/Render/Viewport/DisplayClusterViewportManager.cpp
@@ -22,7 +22,8 @@
 #include "Render/Viewport/Postprocess/DisplayClusterViewportPostProcessManager.h"
 #include "Render/Viewport/Postprocess/DisplayClusterViewportPostProcessOutputRemap.h"
 #include "Render/Viewport/Configuration/DisplayClusterViewportConfiguration.h"
-#include "Render/Viewport/Configuration/DisplayClusterViewportConfigurationInstanceData.h"
+#include "Render/Viewport/Configuration/DisplayClusterViewportConfigurationProxy.h"
+#include "Render/Viewport/Configuration/DisplayClusterViewportConfiguration_Viewport.h"
 #include "Render/Viewport/DisplayClusterViewportStrings.h"
 
 #include "IDisplayClusterWarpBlend.h"
@@ -132,9 +133,12 @@ using namespace UE::DisplayCluster::ViewportManager;
 ///////////////////////////////////////////////////////////////////////////////////////
 
 FDisplayClusterViewportManager::FDisplayClusterViewportManager()
+	: Configuration(MakeShared<FDisplayClusterViewportConfiguration, ESPMode::ThreadSafe>())
+	, RenderTargetManager(MakeShared<FDisplayClusterRenderTargetManager, ESPMode::ThreadSafe>(Configuration))
+	, PostProcessManager(MakeShared<FDisplayClusterViewportPostProcessManager, ESPMode::ThreadSafe>(Configuration))
+	, LightCardManager(MakeShared<FDisplayClusterViewportLightCardManager, ESPMode::ThreadSafe>(Configuration))
+	, RenderFrameManager(MakeShared<FDisplayClusterRenderFrameManager>(Configuration))
 {
-	ViewportManagerProxy = MakeShared<FDisplayClusterViewportManagerProxy, ESPMode::ThreadSafe>();
-
 	// Always reset RTT when root actor re-created
 	ResetSceneRenderTargetSize();
 
@@ -143,16 +147,11 @@ FDisplayClusterViewportManager::FDisplayClusterViewportManager()
 
 void FDisplayClusterViewportManager::Initialize()
 {
-	Configuration = MakeUnique<FDisplayClusterViewportConfiguration>(*this);
-	PostProcessManager = MakeShared<FDisplayClusterViewportPostProcessManager, ESPMode::ThreadSafe>(*this);
-
-	RenderFrameManager = MakeUnique<FDisplayClusterRenderFrameManager>();
-	RenderTargetManager = MakeShared<FDisplayClusterRenderTargetManager, ESPMode::ThreadSafe>(ViewportManagerProxy.Get());
-
-	LightCardManager = MakeShared<FDisplayClusterViewportLightCardManager, ESPMode::ThreadSafe>(*this);
+	// Create ViewportManager proxy object
+	ViewportManagerProxy = MakeShared<FDisplayClusterViewportManagerProxy, ESPMode::ThreadSafe>(*this);
 
-	// initialize proxy
-	ViewportManagerProxy->Initialize(*this);
+	// Initialize configuration for exists viewport managers/proxy objects
+	Configuration->Initialize(*this);
 }
 
 FDisplayClusterViewportManager::~FDisplayClusterViewportManager()
@@ -163,26 +162,19 @@ FDisplayClusterViewportManager::~FDisplayClusterViewportManager()
 	EntireClusterViewports.Reset();
 	CurrentRenderFrameViewports.Reset();
 
-	RenderTargetManager.Reset();
-	PostProcessManager.Reset();
-
-	if (LightCardManager.IsValid())
-	{
-		LightCardManager->Release();
-		LightCardManager.Reset();
-	}
-
+	// Release all DC VE
 	ViewportManagerViewExtension.Reset();
 	ViewportManagerViewPointExtension.Reset();
+	FrameStatsViewExtension.Reset();
 
-	Configuration.Reset();
+	LightCardManager->Release();
 
 	if (ViewportManagerProxy.IsValid())
 	{
 		// Remove viewport manager proxy on render_thread
 		ENQUEUE_RENDER_COMMAND(DeleteDisplayClusterViewportManagerProxy)(
 			[ViewportManagerProxy = ViewportManagerProxy](FRHICommandListImmediate& RHICmdList)
-	{
+			{
 				ViewportManagerProxy->Release_RenderThread();
 			});
 
@@ -190,45 +182,36 @@ FDisplayClusterViewportManager::~FDisplayClusterViewportManager()
 	}
 }
 
-const IDisplayClusterViewportManagerProxy* FDisplayClusterViewportManager::GetProxy() const
+IDisplayClusterViewportConfiguration& FDisplayClusterViewportManager::GetConfiguration()
 {
-	return ViewportManagerProxy.Get();
+	return Configuration.Get();
 }
 
-IDisplayClusterViewportManagerProxy* FDisplayClusterViewportManager::GetProxy()
+const IDisplayClusterViewportConfiguration& FDisplayClusterViewportManager::GetConfiguration() const
 {
-	return ViewportManagerProxy.Get();
+	return Configuration.Get();
 }
 
-UWorld* FDisplayClusterViewportManager::GetCurrentWorld() const
-{
-	check(IsInGameThread());
-
-	if (!CurrentWorldRef.IsValid() || CurrentWorldRef.IsStale())
-	{
-		return nullptr;
-	}
-
-	return CurrentWorldRef.Get();
-}
-
-ADisplayClusterRootActor* FDisplayClusterViewportManager::GetRootActor() const
+const IDisplayClusterViewportManagerProxy* FDisplayClusterViewportManager::GetProxy() const
 {
-	return Configuration->GetRootActor();
+	return ViewportManagerProxy.Get();
 }
 
-bool FDisplayClusterViewportManager::IsSceneOpened() const
+IDisplayClusterViewportManagerProxy* FDisplayClusterViewportManager::GetProxy()
 {
-	check(IsInGameThread());
-
-	return CurrentWorldRef.IsValid() && !CurrentWorldRef.IsStale();
+	return ViewportManagerProxy.Get();
 }
 
-void FDisplayClusterViewportManager::StartScene(UWorld* InWorld)
+void FDisplayClusterViewportManager::HandleStartScene()
 {
 	check(IsInGameThread());
+	check(!Configuration->bCurrentSceneActive);
 
-	CurrentWorldRef = TWeakObjectPtr<UWorld>(InWorld);
+	if (Configuration->GetCurrentWorld() == nullptr)
+	{
+		// Do not start scene, if world not defined
+		return;
+	}
 
 	for (const TSharedPtr<FDisplayClusterViewport, ESPMode::ThreadSafe>& ViewportIt : ImplGetEntireClusterViewports())
 	{
@@ -238,18 +221,14 @@ void FDisplayClusterViewportManager::StartScene(UWorld* InWorld)
 		}
 	}
 
-	if (PostProcessManager.IsValid())
-	{
-		PostProcessManager->HandleStartScene();
-	}
+	PostProcessManager->HandleStartScene();
+	LightCardManager->HandleStartScene();
 
-	if (LightCardManager.IsValid())
-	{
-		LightCardManager->HandleStartScene();
-	}
+	Configuration->bCurrentSceneActive = true;
+	Configuration->bCurrentSceneNeedsToBeUpdated = false;
 }
 
-void FDisplayClusterViewportManager::EndScene()
+void FDisplayClusterViewportManager::HandleEndScene()
 {
 	check(IsInGameThread());
 
@@ -261,52 +240,25 @@ void FDisplayClusterViewportManager::EndScene()
 		}
 	}
 
-	if (PostProcessManager.IsValid())
-	{
-		PostProcessManager->HandleEndScene();
-	}
-
-	if (LightCardManager.IsValid())
-	{
-		LightCardManager->HandleEndScene();
-	}
+	PostProcessManager->HandleEndScene();
+	LightCardManager->HandleEndScene();
 
-	CurrentWorldRef.Reset();
+	Configuration->bCurrentSceneActive = false;
 }
 
-void FDisplayClusterViewportManager::ResetScene()
-{
-	check(IsInGameThread());
-
-	if (LightCardManager.IsValid())
-	{
-		LightCardManager->HandleEndScene();
-		LightCardManager->HandleStartScene();
-	}
-
-	for (const TSharedPtr<FDisplayClusterViewport, ESPMode::ThreadSafe>& ViewportIt : ImplGetEntireClusterViewports())
-	{
-		if (ViewportIt.IsValid())
-		{
-			ViewportIt->HandleEndScene();
-			ViewportIt->HandleStartScene();
-		}
-	}
-}
-
-void FDisplayClusterViewportManager::HandleViewportRTTChanges(const TArray<FDisplayClusterViewport_Context>& PrevContexts, const TArray<FDisplayClusterViewport_Context>& Contexts)
+void FDisplayClusterViewportManager::HandleViewportRTTChanges(const TArray<FDisplayClusterViewport_Context>& InPrevContexts, const TArray<FDisplayClusterViewport_Context>& InContexts)
 {
 	// Support for resetting RTT size when viewport size is changed
-	if (PrevContexts.Num() != Contexts.Num())
+	if (InPrevContexts.Num() != InContexts.Num())
 	{
 		// Reset scene RTT size when viewport disabled
 		ResetSceneRenderTargetSize();
 	}
 	else
 	{
-		for (int32 ContextIt = 0; ContextIt < Contexts.Num(); ContextIt++)
+		for (int32 ContextIt = 0; ContextIt < InContexts.Num(); ContextIt++)
 		{
-			if (Contexts[ContextIt].RenderTargetRect.Size() != PrevContexts[ContextIt].RenderTargetRect.Size())
+			if (InContexts[ContextIt].RenderTargetRect.Size() != InPrevContexts[ContextIt].RenderTargetRect.Size())
 			{
 				ResetSceneRenderTargetSize();
 				break;
@@ -363,7 +315,7 @@ bool FDisplayClusterViewportManager::ShouldUseFullSizeFrameTargetableResource()
 {
 	check(IsInGameThread());
 
-	if (PostProcessManager.IsValid() && PostProcessManager->ShouldUseFullSizeFrameTargetableResource())
+	if (PostProcessManager->ShouldUseFullSizeFrameTargetableResource())
 	{
 		return true;
 	}
@@ -383,7 +335,7 @@ bool FDisplayClusterViewportManager::ShouldUseAdditionalFrameTargetableResource(
 {
 	check(IsInGameThread());
 
-	if (PostProcessManager.IsValid() && PostProcessManager->ShouldUseAdditionalFrameTargetableResource())
+	if (PostProcessManager->ShouldUseAdditionalFrameTargetableResource())
 	{
 		return true;
 	}
@@ -399,110 +351,38 @@ bool FDisplayClusterViewportManager::ShouldUseAdditionalFrameTargetableResource(
 	return false;
 }
 
-bool FDisplayClusterViewportManager::UpdateCustomConfiguration(EDisplayClusterRenderFrameMode InRenderMode, const TArray<FString>& InViewportNames, class ADisplayClusterRootActor* InRootActorPtr)
-{
-	ImplUpdateClusterNodeViewports(InRenderMode, TEXT(""));
-
-	if (InRootActorPtr)
-	{
-		const bool bIsRootActorChanged = Configuration->SetRootActor(InRootActorPtr);
-
-		// When the root actor changes, we have to ResetScene() to reinitialize the internal references of the projection policy.
-		if (bIsRootActorChanged)
-		{
-			ResetScene();
-		}
-
-		return Configuration->UpdateCustomConfiguration(InRenderMode, InViewportNames);
-	}
-
-	return false;
-}
-
-bool FDisplayClusterViewportManager::UpdateConfiguration(EDisplayClusterRenderFrameMode InRenderMode, const FString& InClusterNodeId, ADisplayClusterRootActor* InRootActorPtr, const FDisplayClusterPreviewSettings* InPreviewSettings)
-{
-	ImplUpdateClusterNodeViewports(InRenderMode, InClusterNodeId);
-
-	if (InRootActorPtr)
-	{
-		const bool bIsRootActorChanged = Configuration->SetRootActor(InRootActorPtr);
-
-		// When the root actor changes, we have to ResetScene() to reinitialize the internal references of the projection policy.
-		if (bIsRootActorChanged)
-		{
-			ResetScene();
-		}
-
-		if (LightCardManager.IsValid())
-		{
-			LightCardManager->UpdateConfiguration();
-		}
-
-		if (InPreviewSettings == nullptr)
-		{
-			if (Configuration->UpdateConfiguration(InRenderMode, InClusterNodeId))
-			{
-				return true;
-			}
-		}
-#if WITH_EDITOR
-		else
-		{
-			if (Configuration->UpdatePreviewConfiguration(InRenderMode, InClusterNodeId, *InPreviewSettings))
-			{
-				return true;
-			}
-		}
-#endif
-	}
-
-	return false;
-}
-
-const FDisplayClusterRenderFrameSettings& FDisplayClusterViewportManager::GetRenderFrameSettings() const
-{
-	check(IsInGameThread());
-
-	return Configuration->GetRenderFrameSettings();
-}
-
-EDisplayClusterRenderFrameMode FDisplayClusterViewportManager::GetRenderMode() const
+void FDisplayClusterViewportManager::UpdateCurrentRenderFrameViewports() const
 {
-	check(IsInGameThread());
-
-	return Configuration->GetRenderFrameSettings().RenderMode;
-}
+	const FString& InClusterNodeId = Configuration->GetClusterNodeId();
+	TArray<TSharedPtr<FDisplayClusterViewport, ESPMode::ThreadSafe>> ActiveViewports;
 
-void FDisplayClusterViewportManager::ImplUpdateClusterNodeViewports(const EDisplayClusterRenderFrameMode InRenderMode, const FString& InClusterNodeId)
-{
 	// Get viewports used to rendering of the current frame
 	if (InClusterNodeId.IsEmpty())
 	{
 		// When a cluster node name is empty, we render without the cluster nodes
-		CurrentRenderFrameViewports = ImplGetEntireClusterViewports();
+		ActiveViewports = ImplGetEntireClusterViewports();
 	}
 	else
 	{
 		// Get viewport for current cluster nodes
-		CurrentRenderFrameViewports.Reset();
 		for (const TSharedPtr<FDisplayClusterViewport, ESPMode::ThreadSafe>& Viewport : ImplGetEntireClusterViewports())
 		{
 			if (Viewport.IsValid() && (Viewport->GetClusterNodeId() == InClusterNodeId))
 			{
-				CurrentRenderFrameViewports.Add(Viewport);
+				ActiveViewports.Add(Viewport);
 			}
 		}
 	}
 
 	// Per-viewport preview render required special render order for viewports
 	// Sort viewports: move linked and overridden viewports to last
-	if (InRenderMode == EDisplayClusterRenderFrameMode::PreviewInScene)
+	if (Configuration->IsPreviewRendering())
 	{
 		TArray<TSharedPtr<FDisplayClusterViewport, ESPMode::ThreadSafe>> LinkedViewports;
 		TArray<TSharedPtr<FDisplayClusterViewport, ESPMode::ThreadSafe>> OverriddenViewports;
 
-		TArray<TSharedPtr<FDisplayClusterViewport, ESPMode::ThreadSafe>> UsedViewports(CurrentRenderFrameViewports);
-		CurrentRenderFrameViewports.Reset();
+		TArray<TSharedPtr<FDisplayClusterViewport, ESPMode::ThreadSafe>> UsedViewports(ActiveViewports);
+		ActiveViewports.Reset();
 
 		for (const TSharedPtr<FDisplayClusterViewport, ESPMode::ThreadSafe>& Viewport : UsedViewports)
 		{
@@ -511,23 +391,26 @@ void FDisplayClusterViewportManager::ImplUpdateClusterNodeViewports(const EDispl
 				continue;
 			}
 
-			if (Viewport->RenderSettings.IsViewportOverridden())
+			if (Viewport->GetRenderSettings().IsViewportOverridden())
 			{
 				OverriddenViewports.Add(Viewport);
 			}
-			else if (Viewport->RenderSettings.IsViewportHasParent())
+			else if (Viewport->GetRenderSettings().IsViewportHasParent())
 			{
 				LinkedViewports.Add(Viewport);
 			}
 			else
 			{
-				CurrentRenderFrameViewports.Add(Viewport);
+				ActiveViewports.Add(Viewport);
 			}
 		}
 
-		CurrentRenderFrameViewports.Append(LinkedViewports);
-		CurrentRenderFrameViewports.Append(OverriddenViewports);
+		ActiveViewports.Append(LinkedViewports);
+		ActiveViewports.Append(OverriddenViewports);
 	}
+
+	// This is a mutable class variable intended only for this function
+	(TArray<TSharedPtr<FDisplayClusterViewport, ESPMode::ThreadSafe>>&)CurrentRenderFrameViewports = ActiveViewports;
 }
 
 void FDisplayClusterViewportManager::RegisterCallbacks()
@@ -544,29 +427,9 @@ void FDisplayClusterViewportManager::UnregisterCallbacks()
 #endif
 }
 
-int32 FDisplayClusterViewportManager::GetViewPerViewportAmount() const
-{
-	switch (GetRenderMode())
-	{
-	case EDisplayClusterRenderFrameMode::Mono:
-	case EDisplayClusterRenderFrameMode::PreviewInScene:
-		return 1;
-
-	case EDisplayClusterRenderFrameMode::Stereo:
-	case EDisplayClusterRenderFrameMode::SideBySide:
-	case EDisplayClusterRenderFrameMode::TopBottom:
-		return 2;
-
-	default:
-		break;
-	}
-
-	return 0;
-}
-
 int32 FDisplayClusterViewportManager::FindFirstViewportStereoViewIndex(const TSharedPtr<FDisplayClusterViewport, ESPMode::ThreadSafe>& InViewport) const
 {
-	const int32 ViewPerViewportAmount = GetViewPerViewportAmount();
+	const int32 ViewPerViewportAmount = Configuration->GetRenderFrameSettings().GetViewPerViewportAmount();
 	if (InViewport.IsValid() && ViewPerViewportAmount > 0)
 	{
 		const int32 ViewportIndex = ImplGetEntireClusterViewports().Find(InViewport);
@@ -582,58 +445,37 @@ int32 FDisplayClusterViewportManager::FindFirstViewportStereoViewIndex(const TSh
 	return INDEX_NONE;
 }
 
-bool FDisplayClusterViewportManager::BeginNewFrame(FViewport* InViewport, UWorld* InWorld, FDisplayClusterRenderFrame& OutRenderFrame)
+bool FDisplayClusterViewportManager::BeginNewFrame(FViewport* InViewport, FDisplayClusterRenderFrame& OutRenderFrame)
 {
 	check(IsInGameThread());
-
-	OutRenderFrame.ViewportManagerWeakPtr = this->AsShared();
-
-	// Handle world runtime update
-	UWorld* CurrentWorld = GetCurrentWorld();
-	if (CurrentWorld != InWorld)
+	if(!Configuration->IsSceneOpened())
 	{
-		// Handle end current scene
-		if (CurrentWorld)
-		{
-			EndScene();
-		}
-
-		// Handle begin new scene
-		if (InWorld)
-		{
-			StartScene(InWorld);
-		}
-		else
-		{
-			// no world for render
-			return false;
-		}
+		// no world for render
+		return false;
 	}
 
+	OutRenderFrame.ViewportManagerWeakPtr = this->AsShared();
+
 	// Create DC ViewExtension to handle special features
 	if (!ViewportManagerViewExtension.IsValid())
 	{
-		ViewportManagerViewExtension = FSceneViewExtensions::NewExtension<FDisplayClusterViewportManagerViewExtension>(this);
+		ViewportManagerViewExtension = FSceneViewExtensions::NewExtension<FDisplayClusterViewportManagerViewExtension>(Configuration);
 	}
 
 	// Create DC ViewPointExtension to handle special features
 	if (!ViewportManagerViewPointExtension.IsValid())
 	{
-		ViewportManagerViewPointExtension = FSceneViewExtensions::NewExtension<FDisplayClusterViewportManagerViewPointExtension>(this);
+		ViewportManagerViewPointExtension = FSceneViewExtensions::NewExtension<FDisplayClusterViewportManagerViewPointExtension>(Configuration);
 	}
 
+	// Create DC FrameStatsViewExtension to handle special features
 	if (!FrameStatsViewExtension.IsValid())
 	{
-		FrameStatsViewExtension = FSceneViewExtensions::NewExtension<FDisplayClusterViewportFrameStatsViewExtension>(this);
+		FrameStatsViewExtension = FSceneViewExtensions::NewExtension<FDisplayClusterViewportFrameStatsViewExtension>(Configuration);
 	}
 
 	// Before new frame
-	if (PostProcessManager.IsValid())
-	{
-		PostProcessManager->HandleSetupNewFrame();
-	}
-
-	const FDisplayClusterRenderFrameSettings& RenderFrameSettings = GetRenderFrameSettings();
+	PostProcessManager->HandleSetupNewFrame();
 
 	// Initialize viewports from new render settings, and create new contexts, reset prev frame resources
 	for (const TSharedPtr<FDisplayClusterViewport, ESPMode::ThreadSafe>& ViewportIt : ImplGetCurrentRenderFrameViewports())
@@ -650,16 +492,16 @@ bool FDisplayClusterViewportManager::BeginNewFrame(FViewport* InViewport, UWorld
 			TArray<FDisplayClusterViewport_Context> PrevContexts;
 			PrevContexts.Append(ViewportIt->GetContexts());
 
-			if (ViewportIt->ProjectionPolicy.IsValid())
+			if (ViewportIt->GetProjectionPolicy().IsValid())
 			{
-				ViewportIt->ProjectionPolicy->BeginUpdateFrameContexts(ViewportIt.Get());
+				ViewportIt->GetProjectionPolicy()->BeginUpdateFrameContexts(ViewportIt.Get());
 			}
 
-			ViewportIt->UpdateFrameContexts(FirstViewportStereoViewIndex, RenderFrameSettings);
+			ViewportIt->UpdateFrameContexts(FirstViewportStereoViewIndex);
 
-			if (ViewportIt->ProjectionPolicy.IsValid())
+			if (ViewportIt->GetProjectionPolicy().IsValid())
 			{
-				ViewportIt->ProjectionPolicy->EndUpdateFrameContexts(ViewportIt.Get());
+				ViewportIt->GetProjectionPolicy()->EndUpdateFrameContexts(ViewportIt.Get());
 			}
 
 			HandleViewportRTTChanges(PrevContexts, ViewportIt->GetContexts());
@@ -683,22 +525,28 @@ bool FDisplayClusterViewportManager::BeginNewFrame(FViewport* InViewport, UWorld
 	{
 		if (ViewportIt.IsValid())
 		{
-			for (FDisplayClusterViewport_Context& ContextIt : ViewportIt->Contexts)
+			TArray<FDisplayClusterViewport_Context> ViewportContexts = ViewportIt->GetContexts();
+
+			for (FDisplayClusterViewport_Context& ContextIt : ViewportContexts)
 			{
 				// Get Context Display gamma
-				const TSharedPtr<FDisplayClusterViewportResource, ESPMode::ThreadSafe> ContextRTT = ViewportIt->Resources[EDisplayClusterViewportResource::RenderTargets].IsValidIndex(ContextIt.ContextNum) ? ViewportIt->Resources[EDisplayClusterViewportResource::RenderTargets][ContextIt.ContextNum] : nullptr;
+				const TSharedPtr<FDisplayClusterViewportResource, ESPMode::ThreadSafe> ContextRTT = ViewportIt->GetViewportResources(EDisplayClusterViewportResource::RenderTargets).IsValidIndex(ContextIt.ContextNum)
+					? ViewportIt->GetViewportResources(EDisplayClusterViewportResource::RenderTargets)[ContextIt.ContextNum] : nullptr;
 				const float ViewportDisplayGamma = ContextRTT.IsValid() ? ContextRTT->GetResourceSettings().GetDisplayGamma() : DefaultDisplayGamma;
 
 				ContextIt.RenderThreadData.EngineDisplayGamma = ViewportDisplayGamma;
 
 
-				if (DefaultEngineShowFlags && !EnumHasAnyFlags(ViewportIt->RenderSettingsICVFX.RuntimeFlags, EDisplayClusterViewportRuntimeICVFXFlags::UVLightcard))
+				if (DefaultEngineShowFlags && !EnumHasAnyFlags(ViewportIt->GetRenderSettingsICVFX().RuntimeFlags, EDisplayClusterViewportRuntimeICVFXFlags::UVLightcard))
 				{
 					// Todo: DefaultEngineShowFlags is not the correct value, since each viewport uses its own Engine flags.
 					// UV Light cards should not be set to the default engine flags, since they are rendered using a custom render pass that does not utilize the flags
 					ContextIt.RenderThreadData.EngineShowFlags = *DefaultEngineShowFlags;
 				}
 			}
+
+			// Save changes
+			ViewportIt->SetContexts(ViewportContexts);
 		}
 	}
 
@@ -706,40 +554,35 @@ bool FDisplayClusterViewportManager::BeginNewFrame(FViewport* InViewport, UWorld
 	UpdateSceneRenderTargetSize();
 
 	// Build new frame structure
-	if (!RenderFrameManager->BuildRenderFrame(InViewport, RenderFrameSettings, ImplGetCurrentRenderFrameViewports(), OutRenderFrame))
+	if (!RenderFrameManager->BuildRenderFrame(InViewport, ImplGetCurrentRenderFrameViewports(), OutRenderFrame))
 	{
 		return false;
 	}
 
 	// Allocate resources for frame
-	if (!RenderTargetManager->AllocateRenderFrameResources(InViewport, RenderFrameSettings, ImplGetCurrentRenderFrameViewports(), OutRenderFrame))
+	if (!RenderTargetManager->AllocateRenderFrameResources(InViewport, ImplGetCurrentRenderFrameViewports(), OutRenderFrame))
 	{
 		return false;
 	}
 
 	const FIntPoint RenderFrameSize = OutRenderFrame.FrameRect.Size();
-
 	for (const TSharedPtr<FDisplayClusterViewport, ESPMode::ThreadSafe>& Viewport : ImplGetCurrentRenderFrameViewports())
 	{
 		if (Viewport.IsValid())
 		{
-			// Update ViewportRemap geometry
-			Viewport->ViewportRemap.Update(*Viewport, RenderFrameSize);
+			Viewport->BeginNewFrame(RenderFrameSize);
 		}
 	}
 
 	// Update desired views number
-	OutRenderFrame.DesiredNumberOfViews = (ImplGetEntireClusterViewports().Num() * GetViewPerViewportAmount()) + 1;
+	OutRenderFrame.DesiredNumberOfViews = (ImplGetEntireClusterViewports().Num() * Configuration->GetRenderFrameSettings().GetViewPerViewportAmount()) + 1;
 
 #if WITH_EDITOR
 	// Get preview resources from root actor
 	ImplUpdatePreviewRTTResources();
 #endif /*WITH_EDITOR*/
 
-	if (PostProcessManager.IsValid())
-	{
-		PostProcessManager->HandleBeginNewFrame(OutRenderFrame);
-	}
+	PostProcessManager->HandleBeginNewFrame(OutRenderFrame);
 
 	return true;
 }
@@ -754,9 +597,9 @@ void FDisplayClusterViewportManager::InitializeNewFrame()
 		if (Viewport.IsValid())
 		{
 			// Initialize OCIO resources
-			if (Viewport->OpenColorIO.IsValid())
+			if (Viewport->GetOpenColorIO().IsValid())
 			{
-				Viewport->OpenColorIO->UpdateOpenColorIORenderPassResources();
+				Viewport->GetOpenColorIO()->UpdateOpenColorIORenderPassResources();
 			}
 		}
 	}
@@ -765,8 +608,8 @@ void FDisplayClusterViewportManager::InitializeNewFrame()
 	FDisplayClusterWarpPolicyManager WarpPolicyManager;
 	WarpPolicyManager.HandleNewFrame(*this);
 
-	// Send render frame settings to rendering thread
-	ViewportManagerProxy->ImplUpdateRenderFrameSettings(GetRenderFrameSettings(), ViewportManagerViewExtension);
+	// Send viewport manager data to rendering thread
+	ViewportManagerProxy->ImplUpdateViewportManagerProxy_GameThread(*this);
 
 	/**
 	 * [1] Send new viewport render resources to proxy:
@@ -778,14 +621,11 @@ void FDisplayClusterViewportManager::InitializeNewFrame()
 	ViewportManagerProxy->ImplUpdateViewportProxies_GameThread(ImplGetCurrentRenderFrameViewports());
 
 	// Send postprocess data to render thread
-	if (PostProcessManager.IsValid())
-	{
-		// Update postprocess data from game thread
-		PostProcessManager->Tick();
+	// Update postprocess data from game thread
+	PostProcessManager->Tick();
 
-		// Send updated postprocess data to rendering thread
-		PostProcessManager->FinalizeNewFrame();
-	}
+	// Send updated postprocess data to rendering thread
+	PostProcessManager->FinalizeNewFrame();
 }
 
 void FDisplayClusterViewportManager::FinalizeNewFrame()
@@ -799,24 +639,7 @@ void FDisplayClusterViewportManager::FinalizeNewFrame()
 	{
 		if (Viewport.IsValid())
 		{
-			// When all viewports processed, we remove all single frame custom postprocess
-			Viewport->CustomPostProcessSettings.FinalizeFrame();
-
-			// Update projection policy proxy data
-			if (Viewport->ProjectionPolicy.IsValid())
-			{
-				Viewport->ProjectionPolicy->UpdateProxyData(Viewport.Get());
-			}
-
-		}
-	}
-
-	// Finish update settings
-	for (const TSharedPtr<FDisplayClusterViewport, ESPMode::ThreadSafe>& Viewport : ImplGetCurrentRenderFrameViewports())
-	{
-		if (Viewport.IsValid())
-		{
-			Viewport->RenderSettings.FinishUpdateSettings();
+			Viewport->FinalizeNewFrame();
 		}
 	}
 }
@@ -850,26 +673,18 @@ FSceneViewFamily::ConstructionValues FDisplayClusterViewportManager::CreateViewF
 		break;
 	}
 
-	const FDisplayClusterRenderFrameSettings& RenderFrameSettings = GetRenderFrameSettings();
-	switch (RenderFrameSettings.RenderMode)
-	{
-	case EDisplayClusterRenderFrameMode::PreviewInScene:
-
-		if (RenderFrameSettings.bPreviewEnablePostProcess == false)
+	const FDisplayClusterRenderFrameSettings& RenderFrameSettings = Configuration->GetRenderFrameSettings();
+	if(RenderFrameSettings.IsPostProcessDisabled())
 		{
 			// Disable postprocess for preview
 			InEngineShowFlags.PostProcessing = 0;
 		}
-		break;
-	default:
-		break;
-	}
 
 	switch (InFrameTarget.CaptureMode)
 	{
 	case EDisplayClusterViewportCaptureMode::Chromakey:
 	case EDisplayClusterViewportCaptureMode::Lightcard:
-		switch (GetRenderFrameSettings().AlphaChannelCaptureMode)
+		switch (RenderFrameSettings.AlphaChannelCaptureMode)
 		{
 			case EDisplayClusterRenderFrameAlphaChannelCaptureMode::Copy:
 				// Disable AA
@@ -928,20 +743,8 @@ FSceneViewFamily::ConstructionValues FDisplayClusterViewportManager::CreateViewF
 
 bool FDisplayClusterViewportManager::ShouldRenderFinalColor() const
 {
-#if WITH_EDITOR
-	const FDisplayClusterRenderFrameSettings& RenderFrameSettings = GetRenderFrameSettings();
-	switch (RenderFrameSettings.RenderMode)
-	{
-	case EDisplayClusterRenderFrameMode::PreviewInScene:
-		// Disable postprocess for preview
-		return RenderFrameSettings.bPreviewEnablePostProcess;
-
-	default:
-		break;
-	}
-#endif
-
-	return true;
+	// Do not render final color, if postprocess is disabled
+	return !Configuration->GetRenderFrameSettings().IsPostProcessDisabled();
 }
 
 void FDisplayClusterViewportManager::ConfigureViewFamily(const FDisplayClusterRenderFrameTarget& InFrameTarget, const FDisplayClusterRenderFrameTargetViewFamily& InFrameViewFamily, FSceneViewFamilyContext& ViewFamily)
@@ -966,10 +769,7 @@ void FDisplayClusterViewportManager::ConfigureViewFamily(const FDisplayClusterRe
 
 void FDisplayClusterViewportManager::RenderFrame(FViewport* InViewport)
 {
-	if (LightCardManager.IsValid())
-	{
-		LightCardManager->RenderFrame();
-	}
+	LightCardManager->RenderFrame();
 
 	ViewportManagerProxy->ImplRenderFrame_GameThread(InViewport);
 }
@@ -993,14 +793,6 @@ FDisplayClusterViewport* FDisplayClusterViewportManager::CreateViewport(const FS
 		return nullptr;
 	}
 
-	ADisplayClusterRootActor* RootActorPtr = GetRootActor();
-	if (!RootActorPtr)
-	{
-		UE_LOG(LogDisplayClusterViewport, Error, TEXT("Can't create viewport '%s': DCRA not defined"), *InViewportId);
-
-		return nullptr;
-	}
-
 	// Create projection policy for viewport
 	const TSharedPtr<IDisplayClusterProjectionPolicy, ESPMode::ThreadSafe> NewProjectionPolicy = CreateProjectionPolicy(InViewportId, &ConfigurationViewport.ProjectionPolicy);
 	if (NewProjectionPolicy.IsValid())
@@ -1009,7 +801,7 @@ FDisplayClusterViewport* FDisplayClusterViewportManager::CreateViewport(const FS
 		TSharedPtr<FDisplayClusterViewport, ESPMode::ThreadSafe> NewViewport = ImplCreateViewport(InViewportId, NewProjectionPolicy);
 		if (NewViewport.IsValid())
 		{
-			FDisplayClusterViewportConfigurationInstanceData::UpdateViewportConfiguration(*NewViewport , *this, *RootActorPtr, ConfigurationViewport);
+			FDisplayClusterViewportConfiguration_Viewport::UpdateViewportConfiguration(*NewViewport , ConfigurationViewport);
 
 			return NewViewport.Get();
 		}
@@ -1073,20 +865,15 @@ bool FDisplayClusterViewportManager::DeleteViewport(const FString& ViewportId)
 TSharedPtr<FDisplayClusterViewport, ESPMode::ThreadSafe> FDisplayClusterViewportManager::ImplCreateViewport(const FString& ViewportId, const TSharedPtr<IDisplayClusterProjectionPolicy, ESPMode::ThreadSafe>& InProjectionPolicy)
 {
 	check(IsInGameThread());
-
 	check(InProjectionPolicy.IsValid());
 
-	// Create viewport for cluster node used for rendering
-	const FString& ClusterNodeId = GetRenderFrameSettings().ClusterNodeId;
-	check(!ClusterNodeId.IsEmpty());
-
-	// Create a new viewport
-	TSharedPtr<FDisplayClusterViewport, ESPMode::ThreadSafe> NewViewport = MakeShared<FDisplayClusterViewport>(*this, ClusterNodeId, ViewportId, InProjectionPolicy);
+	// Create a new viewport for cluster node used for rendering
+	TSharedPtr<FDisplayClusterViewport, ESPMode::ThreadSafe> NewViewport = MakeShared<FDisplayClusterViewport>(Configuration, ViewportId, InProjectionPolicy);
 	if (NewViewport.IsValid())
 	{
 		// Add viewport on gamethread
 		EntireClusterViewports.Add(NewViewport);
-		CurrentRenderFrameViewports.Add(NewViewport);
+		Configuration->bCurrentRenderFrameViewportsNeedsToBeUpdated = true;
 
 		// Handle start scene for viewport
 		NewViewport->HandleStartScene();
@@ -1100,28 +887,17 @@ void FDisplayClusterViewportManager::ImplDeleteViewport(const TSharedPtr<FDispla
 	check(ExistViewport.IsValid());
 
 	// Handle projection policy event
-	ExistViewport->ProjectionPolicy.Reset();
-	ExistViewport->UninitializedProjectionPolicy.Reset();
+	ExistViewport->ReleaseProjectionPolicy();
 
+	// Remove viewport from the entire viewports list
+	int32 ViewportIndex = EntireClusterViewports.Find(ExistViewport);
+	if (ViewportIndex != INDEX_NONE)
 	{
-		// Remove viewport from the entire viewports list
-		int32 ViewportIndex = EntireClusterViewports.Find(ExistViewport);
-		if (ViewportIndex != INDEX_NONE)
-		{
-			EntireClusterViewports[ViewportIndex] = nullptr;
-			EntireClusterViewports.RemoveAt(ViewportIndex);
-		}
+		EntireClusterViewports[ViewportIndex] = nullptr;
+		EntireClusterViewports.RemoveAt(ViewportIndex);
 	}
 
-	{
-		// Remove viewport from the cluster viewports list
-		int32 ViewportIndex = CurrentRenderFrameViewports.Find(ExistViewport);
-		if (ViewportIndex != INDEX_NONE)
-		{
-			CurrentRenderFrameViewports[ViewportIndex] = nullptr;
-			CurrentRenderFrameViewports.RemoveAt(ViewportIndex);
-		}
-	}
+	Configuration->bCurrentRenderFrameViewportsNeedsToBeUpdated = true;
 
 	// Reset RTT size after viewport delete
 	ResetSceneRenderTargetSize();
@@ -1183,7 +959,7 @@ TSharedPtr<IDisplayClusterProjectionPolicy, ESPMode::ThreadSafe> FDisplayCluster
 	return nullptr;
 }
 
-TArray<TSharedPtr<IDisplayClusterViewport, ESPMode::ThreadSafe>> FDisplayClusterViewportManager::GetViewportsForWarpPolicy(const TSharedPtr<IDisplayClusterWarpPolicy>& InWarpPolicy) const
+TArray<TSharedPtr<IDisplayClusterViewport, ESPMode::ThreadSafe>> FDisplayClusterViewportManager::GetEntireClusterViewportsForWarpPolicy(const TSharedPtr<IDisplayClusterWarpPolicy>& InWarpPolicy) const
 {
 	TArray<TSharedPtr<IDisplayClusterViewport, ESPMode::ThreadSafe>> WarpPolicyViewports;
 
@@ -1226,10 +1002,7 @@ void FDisplayClusterViewportManager::MarkComponentGeometryDirty(const FName InCo
 	}
 
 	// 2. Update all ProceduralMeshComponent references for OutputRemap
-	if (PostProcessManager.IsValid())
-	{
-		PostProcessManager->GetOutputRemap()->MarkProceduralMeshComponentGeometryDirty(InComponentName);
-	}
+	PostProcessManager->GetOutputRemap()->MarkProceduralMeshComponentGeometryDirty(InComponentName);
 }
 
 void FDisplayClusterViewportManager::AddReferencedObjects(FReferenceCollector& Collector)
diff --git a/Engine/Plugins/Runtime/nDisplay/Source/DisplayCluster/Private/Render/Viewport/DisplayClusterViewportManager.h b/Engine/Plugins/Runtime/nDisplay/Source/DisplayCluster/Private/Render/Viewport/DisplayClusterViewportManager.h
--- a/Engine/Plugins/Runtime/nDisplay/Source/DisplayCluster/Private/Render/Viewport/DisplayClusterViewportManager.h
+++ b/Engine/Plugins/Runtime/nDisplay/Source/DisplayCluster/Private/Render/Viewport/DisplayClusterViewportManager.h
@@ -13,20 +13,15 @@
 #include "Render/Viewport/DisplayClusterViewport.h"
 #include "Render/Viewport/Containers/DisplayClusterViewport_Enums.h"
 
-class FDisplayClusterViewportConfiguration;
-class FDisplayClusterRenderTargetManager;
-class FDisplayClusterViewportPostProcessManager;
-class FDisplayClusterRenderFrameManager;
 class FDisplayClusterViewportManagerProxy;
-class FDisplayClusterViewportLightCardManager;
 class IDisplayClusterProjectionPolicy;
-
 class  UDisplayClusterConfigurationViewport;
 struct FDisplayClusterConfigurationProjection;
-struct FDisplayClusterRenderFrameSettings;
-
 class FViewport;
 
+/**
+ * Implementation of the nDisplay IDisplayClusterViewportManager
+ */
 class FDisplayClusterViewportManager
 	: public IDisplayClusterViewportManager
 	, public TSharedFromThis<FDisplayClusterViewportManager, ESPMode::ThreadSafe>
@@ -37,9 +32,8 @@ public:
 
 public:
 	/** Game thread funcs */
-	void StartScene(UWorld* World);
-	void EndScene();
-	void ResetScene();
+	void HandleStartScene();
+	void HandleEndScene();
 
 	void Initialize();
 
@@ -54,20 +48,13 @@ public:
 		return AsShared();
 	}
 
-	virtual EDisplayClusterRenderFrameMode GetRenderMode() const override;
-
 	virtual const IDisplayClusterViewportManagerProxy* GetProxy() const override;
 	virtual       IDisplayClusterViewportManagerProxy* GetProxy() override;
 
-	virtual UWorld*                   GetCurrentWorld() const override;
-	virtual ADisplayClusterRootActor* GetRootActor() const override;
-
-	virtual bool IsSceneOpened() const override;
+	virtual IDisplayClusterViewportConfiguration& GetConfiguration() override;
+	virtual const IDisplayClusterViewportConfiguration& GetConfiguration() const override;
 
-	virtual bool UpdateConfiguration(EDisplayClusterRenderFrameMode InRenderMode, const FString& InClusterNodeId, class ADisplayClusterRootActor* InRootActorPtr, const FDisplayClusterPreviewSettings* InPreviewSettings = nullptr) override;
-	virtual bool UpdateCustomConfiguration(EDisplayClusterRenderFrameMode InRenderMode, const TArray<FString>& InViewportNames, class ADisplayClusterRootActor* InRootActorPtr) override;
-
-	virtual bool BeginNewFrame(FViewport* InViewport, UWorld* InWorld, FDisplayClusterRenderFrame& OutRenderFrame) override;
+	virtual bool BeginNewFrame(FViewport* InViewport, FDisplayClusterRenderFrame& OutRenderFrame) override;
 	virtual void InitializeNewFrame() override;
 	virtual void FinalizeNewFrame() override;
 
@@ -91,13 +78,13 @@ private:
 	void OnPreGarbageCollect();
 
 public:
-
 	virtual IDisplayClusterViewport* FindViewport(const FString& InViewportId) const override;
 	virtual IDisplayClusterViewport* FindViewport(const int32 ViewIndex, uint32* OutContextNum = nullptr) const override;
 
 	virtual const TArrayView<TSharedPtr<IDisplayClusterViewport, ESPMode::ThreadSafe>> GetCurrentRenderFrameViewports() const override
 	{
-		return TArrayView<TSharedPtr<IDisplayClusterViewport, ESPMode::ThreadSafe>>((TSharedPtr<IDisplayClusterViewport, ESPMode::ThreadSafe>*)(CurrentRenderFrameViewports.GetData()), CurrentRenderFrameViewports.Num());
+		const TArray<TSharedPtr<FDisplayClusterViewport, ESPMode::ThreadSafe>>& OutViewports = ImplGetCurrentRenderFrameViewports();
+		return TArrayView<TSharedPtr<IDisplayClusterViewport, ESPMode::ThreadSafe>>((TSharedPtr<IDisplayClusterViewport, ESPMode::ThreadSafe>*)(OutViewports.GetData()), OutViewports.Num());
 	}
 
 	virtual const TArrayView<TSharedPtr<IDisplayClusterViewport, ESPMode::ThreadSafe>> GetEntireClusterViewports() const override
@@ -105,7 +92,7 @@ public:
 		return TArrayView<TSharedPtr<IDisplayClusterViewport, ESPMode::ThreadSafe>>((TSharedPtr<IDisplayClusterViewport, ESPMode::ThreadSafe>*)(EntireClusterViewports.GetData()), EntireClusterViewports.Num());
 	}
 
-	virtual TArray<TSharedPtr<IDisplayClusterViewport, ESPMode::ThreadSafe>> GetViewportsForWarpPolicy(const TSharedPtr<IDisplayClusterWarpPolicy>& InWarpPolicy) const override;
+	virtual TArray<TSharedPtr<IDisplayClusterViewport, ESPMode::ThreadSafe>> GetEntireClusterViewportsForWarpPolicy(const TSharedPtr<IDisplayClusterWarpPolicy>& InWarpPolicy) const override;
 
 	virtual void MarkComponentGeometryDirty(const FName InComponentName = NAME_None) override;
 	virtual void AddReferencedObjects(FReferenceCollector& Collector) override;
@@ -117,6 +104,12 @@ public:
 	/** Getting the viewports of the current rendering frame (viewports from the current node or from a special named list). */
 	inline const TArray<TSharedPtr<FDisplayClusterViewport, ESPMode::ThreadSafe>>& ImplGetCurrentRenderFrameViewports() const
 	{
+		if (Configuration->bCurrentRenderFrameViewportsNeedsToBeUpdated)
+		{
+			Configuration->bCurrentRenderFrameViewportsNeedsToBeUpdated = false;
+			UpdateCurrentRenderFrameViewports();
+		}
+
 		return CurrentRenderFrameViewports;
 	}
 
@@ -136,6 +129,7 @@ public:
 
 	static TSharedPtr<IDisplayClusterProjectionPolicy, ESPMode::ThreadSafe> CreateProjectionPolicy(const FString& InViewportId, const FDisplayClusterConfigurationProjection* InConfigurationProjectionPolicy);
 
+	/** Get DC viewport manager proxy. */
 	TSharedPtr<FDisplayClusterViewportManagerProxy, ESPMode::ThreadSafe> GetViewportManagerProxy() const
 	{
 		check(IsInGameThread());
@@ -143,10 +137,6 @@ public:
 		return ViewportManagerProxy;
 	}
 
-	const FDisplayClusterRenderFrameSettings& GetRenderFrameSettings() const;
-
-	/** Return number of contexts per viewport. */
-	int32 GetViewPerViewportAmount() const;
 
 	/** Return initial StereoViewIndex for the input viewport. */
 	int32 FindFirstViewportStereoViewIndex(const TSharedPtr<FDisplayClusterViewport, ESPMode::ThreadSafe>& InViewport) const;
@@ -158,45 +148,43 @@ public:
 	TSharedPtr<class FDisplayClusterViewportManagerViewExtension, ESPMode::ThreadSafe> GetViewportManagerViewExtension() const
 	{ return ViewportManagerViewExtension; }
 
-	/** Return LightCardManager object. */
-	TSharedPtr<FDisplayClusterViewportLightCardManager, ESPMode::ThreadSafe> GetLightCardManager() const
-	{ return LightCardManager; }
-
-	/** Return PostProcessManager object. */
-	TSharedPtr<FDisplayClusterViewportPostProcessManager, ESPMode::ThreadSafe> GetPostProcessManager() const
-	{ return PostProcessManager; }
-
 	bool ShouldUseAdditionalFrameTargetableResource() const;
 	bool ShouldUseFullSizeFrameTargetableResource() const;
 
 	void ResetSceneRenderTargetSize();
 
 private:
+	/** This function updates CurrentRenderFrameViewports according to the current rendering settings. */
+	void UpdateCurrentRenderFrameViewports() const;
+
 	void UpdateSceneRenderTargetSize();
 	void HandleViewportRTTChanges(const TArray<FDisplayClusterViewport_Context>& PrevContexts, const TArray<FDisplayClusterViewport_Context>& Contexts);
-	void ImplUpdateClusterNodeViewports(const EDisplayClusterRenderFrameMode InRenderMode, const FString& InClusterNodeId);
 
 	/** Register any callbacks */
 	void RegisterCallbacks();
 	
 	/** Unregister any callbacks used */
 	void UnregisterCallbacks();
-	
-protected:
-	friend FDisplayClusterViewportManagerProxy;
-	friend FDisplayClusterViewportConfiguration;
-
-	TSharedPtr<FDisplayClusterRenderTargetManager, ESPMode::ThreadSafe>        RenderTargetManager;
-	TSharedPtr<FDisplayClusterViewportPostProcessManager, ESPMode::ThreadSafe> PostProcessManager;
-	TSharedPtr<FDisplayClusterViewportLightCardManager, ESPMode::ThreadSafe>   LightCardManager;
 
 public:
 	// Configuration of the current cluster node
-	TUniquePtr<FDisplayClusterViewportConfiguration> Configuration;
+	const TSharedRef<class FDisplayClusterViewportConfiguration, ESPMode::ThreadSafe> Configuration;
+
+	// Resource manager
+	const TSharedRef<class FDisplayClusterRenderTargetManager, ESPMode::ThreadSafe> RenderTargetManager;
+
+	// Postprocess manager
+	const TSharedRef<class FDisplayClusterViewportPostProcessManager, ESPMode::ThreadSafe> PostProcessManager;
+
+	// LC manager
+	const TSharedRef<class FDisplayClusterViewportLightCardManager, ESPMode::ThreadSafe> LightCardManager;
 
-private:
 	// Manager for creating the render frame stucture
-	TUniquePtr<FDisplayClusterRenderFrameManager>  RenderFrameManager;
+	const TSharedRef<class FDisplayClusterRenderFrameManager> RenderFrameManager;
+
+private:
+	// Render thread proxy manager. Deleted on render thread
+	TSharedPtr<FDisplayClusterViewportManagerProxy, ESPMode::ThreadSafe> ViewportManagerProxy;
 
 	// Viewports of the entire cluster
 	TArray<TSharedPtr<FDisplayClusterViewport, ESPMode::ThreadSafe>> EntireClusterViewports;
@@ -204,12 +192,7 @@ private:
 	// Viewports of the current render frame (viewports from the current node or from a special named list)
 	TArray<TSharedPtr<FDisplayClusterViewport, ESPMode::ThreadSafe>> CurrentRenderFrameViewports;
 
-	// Render thread proxy manager. Deleted on render thread
-	TSharedPtr<FDisplayClusterViewportManagerProxy, ESPMode::ThreadSafe> ViewportManagerProxy;
-
-	// Pointer to the current scene
-	TWeakObjectPtr<UWorld> CurrentWorldRef;
-
+	/** Scene RTT resize method.*/
 	enum class ESceneRenderTargetResizeMethod : uint8
 	{
 		None = 0,
