nDisplay: Added new interface IDisplayClusterViewportConfiguration

diff --git a/Engine/Plugins/Runtime/nDisplay/Source/DisplayCluster/Private/Render/Viewport/DisplayClusterViewportManagerProxy.cpp b/Engine/Plugins/Runtime/nDisplay/Source/DisplayCluster/Private/Render/Viewport/DisplayClusterViewportManagerProxy.cpp
--- a/Engine/Plugins/Runtime/nDisplay/Source/DisplayCluster/Private/Render/Viewport/DisplayClusterViewportManagerProxy.cpp
+++ b/Engine/Plugins/Runtime/nDisplay/Source/DisplayCluster/Private/Render/Viewport/DisplayClusterViewportManagerProxy.cpp
@@ -2,6 +2,10 @@
 
 #include "Render/Viewport/DisplayClusterViewportManagerProxy.h"
 
+#include "Render/Viewport/DisplayClusterViewportManager.h"
+#include "Render/Viewport/Configuration/DisplayClusterViewportConfiguration.h"
+#include "Render/Viewport/Configuration/DisplayClusterViewportConfigurationProxy.h"
+
 #include "ClearQuad.h"
 #include "Misc/DisplayClusterLog.h"
 
@@ -40,17 +44,17 @@ static TAutoConsoleVariable<int32> CVarClearFrameRTTEnabled(
 ///////////////////////////////////////////////////////////////////////////////////////
 //          FDisplayClusterViewportManagerProxy
 ///////////////////////////////////////////////////////////////////////////////////////
-FDisplayClusterViewportManagerProxy::FDisplayClusterViewportManagerProxy()
+FDisplayClusterViewportManagerProxy::FDisplayClusterViewportManagerProxy(const FDisplayClusterViewportManager& InViewportManager)
+	: ConfigurationProxy(InViewportManager.Configuration->Proxy)
+	, RenderTargetManager(InViewportManager.RenderTargetManager)
+	, PostProcessManager(InViewportManager.PostProcessManager)
+	, LightCardManagerProxy(InViewportManager.LightCardManager->LightCardManagerProxy)
+	, ViewportManagerViewExtension(InViewportManager.GetViewportManagerViewExtension())
 { }
 
 FDisplayClusterViewportManagerProxy::~FDisplayClusterViewportManagerProxy()
 { }
 
-EDisplayClusterRenderFrameMode FDisplayClusterViewportManagerProxy::GetRenderMode() const
-{
-	return RenderFrameSettings.RenderMode;
-}
-
 void FDisplayClusterViewportManagerProxy::Release_RenderThread()
 {
 	check(IsInRenderingThread());
@@ -59,19 +63,8 @@ void FDisplayClusterViewportManagerProxy::Release_RenderThread()
 	EntireClusterViewportProxies.Empty();
 	CurrentRenderFrameViewportProxies.Empty();
 
-	if (RenderTargetManager.IsValid())
-	{
-		RenderTargetManager->Release();
-		RenderTargetManager.Reset();
-	}
-
-	if (PostProcessManager.IsValid())
-	{
-		PostProcessManager->Release();
-		PostProcessManager.Reset();
-	}
-
-	LightCardManagerProxy.Reset();
+	RenderTargetManager->Release();
+	PostProcessManager->Release();
 
 	if (ViewportManagerViewExtension.IsValid())
 	{
@@ -82,18 +75,10 @@ void FDisplayClusterViewportManagerProxy::Release_RenderThread()
 	}
 }
 
-void FDisplayClusterViewportManagerProxy::Initialize(FDisplayClusterViewportManager& InViewportManager)
-{
-	RenderTargetManager = InViewportManager.RenderTargetManager;
-	PostProcessManager = InViewportManager.PostProcessManager;
-	LightCardManagerProxy = InViewportManager.LightCardManager->GetLightCardManagerProxy();
-	ViewportManagerViewExtension = InViewportManager.ViewportManagerViewExtension;
-}
-
 void FDisplayClusterViewportManagerProxy::ImplUpdateClusterNodeViewportProxies_RenderThread()
 {
 	// Get viewports used to rendering of the current frame
-	if (RenderFrameSettings.ClusterNodeId.IsEmpty())
+	if (ConfigurationProxy->GetClusterNodeId_RenderThread().IsEmpty())
 	{
 		// When a cluster node name is empty, we render without the cluster nodes
 		CurrentRenderFrameViewportProxies = ImplGetEntireClusterViewportProxies_RenderThread();
@@ -104,7 +89,7 @@ void FDisplayClusterViewportManagerProxy::ImplUpdateClusterNodeViewportProxies_R
 		CurrentRenderFrameViewportProxies.Reset();
 		for (const TSharedPtr<FDisplayClusterViewportProxy, ESPMode::ThreadSafe>& ViewportProxyIt : ImplGetEntireClusterViewportProxies_RenderThread())
 		{
-			if (ViewportProxyIt.IsValid() && (ViewportProxyIt->GetClusterNodeId() == RenderFrameSettings.ClusterNodeId))
+			if (ViewportProxyIt.IsValid() && (ViewportProxyIt->GetClusterNodeId() == ConfigurationProxy->GetClusterNodeId_RenderThread()))
 			{
 				CurrentRenderFrameViewportProxies.Add(ViewportProxyIt);
 			}
@@ -139,25 +124,25 @@ void FDisplayClusterViewportManagerProxy::DeleteViewport_RenderThread(const TSha
 	}
 }
 
-DECLARE_GPU_STAT_NAMED(nDisplay_ViewportManager_UpdateRenderFrameSettings, TEXT("nDisplay ViewportManager::UpdateRenderFrameSettings"));
+DECLARE_GPU_STAT_NAMED(nDisplay_ViewportManager_UpdateViewportManagerProxy, TEXT("nDisplay ViewportManager::UpdateProxy"));
 
-void FDisplayClusterViewportManagerProxy::ImplUpdateRenderFrameSettings(const FDisplayClusterRenderFrameSettings& InRenderFrameSettings,
-	const TSharedPtr<FDisplayClusterViewportManagerViewExtension, ESPMode::ThreadSafe>& InViewportManagerViewExtension)
+void FDisplayClusterViewportManagerProxy::ImplUpdateViewportManagerProxy_GameThread(const FDisplayClusterViewportManager& InViewportManager)
 {
 	check(IsInGameThread());
 
-	FDisplayClusterRenderFrameSettings* Settings = new FDisplayClusterRenderFrameSettings(InRenderFrameSettings);
+	// Update configuration proxy:
+	ConfigurationProxy->UpdateConfigurationProxy_GameThread(*InViewportManager.Configuration);
 
-	// Send frame settings to renderthread
-	ENQUEUE_RENDER_COMMAND(DisplayClusterUpdateRenderFrameSettings)(
-		[ViewportManagerProxy = SharedThis(this), Settings, InViewportManagerViewExtension](FRHICommandListImmediate& RHICmdList)
+	// Send viewport manager data to proxy
+	ENQUEUE_RENDER_COMMAND(DisplayClusterUpdateViewportManagerProxy)(
+		[ ViewportManagerProxy = SharedThis(this)
+		, ViewExtension = InViewportManager.GetViewportManagerViewExtension()
+		](FRHICommandListImmediate& RHICmdList)
 	{
-		SCOPED_GPU_STAT(RHICmdList, nDisplay_ViewportManager_UpdateRenderFrameSettings);
-		SCOPED_DRAW_EVENT(RHICmdList, nDisplay_ViewportManager_UpdateRenderFrameSettings);
+		SCOPED_GPU_STAT(RHICmdList, nDisplay_ViewportManager_UpdateViewportManagerProxy);
+		SCOPED_DRAW_EVENT(RHICmdList, nDisplay_ViewportManager_UpdateViewportManagerProxy);
 
-		ViewportManagerProxy->RenderFrameSettings = *Settings;
-		ViewportManagerProxy->ViewportManagerViewExtension = InViewportManagerViewExtension;
-		delete Settings;
+		ViewportManagerProxy->ViewportManagerViewExtension = ViewExtension;
 
 		// After updated settings we need update cluster node viewports
 		ViewportManagerProxy->ImplUpdateClusterNodeViewportProxies_RenderThread();
@@ -173,9 +158,9 @@ void FDisplayClusterViewportManagerProxy::ImplUpdateViewportProxies_GameThread(c
 	TArray<FDisplayClusterViewportProxyData*> ViewportProxiesData;
 	for (const TSharedPtr<FDisplayClusterViewport, ESPMode::ThreadSafe>& ViewportIt : InCurrentRenderFrameViewports)
 	{
-		if (ViewportIt.IsValid())
+		if(FDisplayClusterViewportProxyData* NewData = ViewportIt.IsValid() ? ViewportIt->CreateViewportProxyData() : nullptr)
 		{
-			ViewportProxiesData.Add(new FDisplayClusterViewportProxyData(ViewportIt));
+			ViewportProxiesData.Add(NewData);
 		}
 	}
 
@@ -186,15 +171,15 @@ void FDisplayClusterViewportManagerProxy::ImplUpdateViewportProxies_GameThread(c
 		SCOPED_GPU_STAT(RHICmdList, nDisplay_ViewportManager_UpdateViewports);
 		SCOPED_DRAW_EVENT(RHICmdList, nDisplay_ViewportManager_UpdateViewports);
 
-		for (FDisplayClusterViewportProxyData* It : ProxiesData)
+		// Update game on rendering thread:
+		for (FDisplayClusterViewportProxyData* ProxyDataIt : ProxiesData)
 		{
-			It->UpdateProxy_RenderThread();
-			delete It;
+			ProxyDataIt->DestViewportProxy->UpdateViewportProxyData_RenderThread(*ProxyDataIt);
+			delete ProxyDataIt;
 		}
 	});
 }
 
-
 DECLARE_GPU_STAT_NAMED(nDisplay_ViewportManager_RenderFrame, TEXT("nDisplay ViewportManager::RenderFrame"));
 DECLARE_GPU_STAT_NAMED(nDisplay_ViewportManager_CrossGPUTransfer, TEXT("nDisplay ViewportManager::CrossGPUTransfer"));
 DECLARE_GPU_STAT_NAMED(nDisplay_ViewportManager_UpdateDeferredResources, TEXT("nDisplay ViewportManager::UpdateDeferredResources"));
@@ -211,10 +196,7 @@ void FDisplayClusterViewportManagerProxy::ImplRenderFrame_GameThread(FViewport*
 		const FDisplayClusterViewportManagerProxy* ViewportManagerProxy = &InViewportManagerProxy.Get();
 
 		// Handle render setup
-		if (ViewportManagerProxy->PostProcessManager.IsValid())
-		{
-			ViewportManagerProxy->PostProcessManager->HandleRenderFrameSetup_RenderThread(RHICmdList, ViewportManagerProxy);
-		}
+		ViewportManagerProxy->PostProcessManager->HandleRenderFrameSetup_RenderThread(RHICmdList, ViewportManagerProxy);
 	});
 
 	ENQUEUE_RENDER_COMMAND(DisplayClusterRenderFrame_CrossGPUTransfer)(
@@ -226,7 +208,7 @@ void FDisplayClusterViewportManagerProxy::ImplRenderFrame_GameThread(FViewport*
 		const FDisplayClusterViewportManagerProxy* ViewportManagerProxy = &InViewportManagerProxy.Get();
 
 		// mGPU not used for in-editor rendering
-		if(ViewportManagerProxy->RenderFrameSettings.bIsRenderingInEditor == false || ViewportManagerProxy->RenderFrameSettings.bAllowMultiGPURenderingInEditor)
+		if(ViewportManagerProxy->ConfigurationProxy->GetRenderFrameSettings().ShouldUseCrossGPUTransfers())
 		{
 			// Move all render target cross gpu
 			ViewportManagerProxy->DoCrossGPUTransfers_RenderThread(RHICmdList);
@@ -257,18 +239,12 @@ void FDisplayClusterViewportManagerProxy::ImplRenderFrame_GameThread(FViewport*
 
 		const FDisplayClusterViewportManagerProxy* ViewportManagerProxy = &InViewportManagerProxy.Get();
 
-		if (ViewportManagerProxy->PostProcessManager.IsValid())
-		{
-			ViewportManagerProxy->PostProcessManager->HandleBeginUpdateFrameResources_RenderThread(RHICmdList, ViewportManagerProxy);
-		}
+		ViewportManagerProxy->PostProcessManager->HandleBeginUpdateFrameResources_RenderThread(RHICmdList, ViewportManagerProxy);
 
 		// Update the frame resources: post-processing, warping, and finally resolving everything to the frame resource
 		ViewportManagerProxy->UpdateFrameResources_RenderThread(RHICmdList);
 
-		if (ViewportManagerProxy->PostProcessManager.IsValid())
-		{
-			ViewportManagerProxy->PostProcessManager->HandleEndUpdateFrameResources_RenderThread(RHICmdList, ViewportManagerProxy);
-		}
+		ViewportManagerProxy->PostProcessManager->HandleEndUpdateFrameResources_RenderThread(RHICmdList, ViewportManagerProxy);
 
 		// Postrender notification before copying final image to the backbuffer
 		IDisplayCluster::Get().GetCallbacks().OnDisplayClusterPostFrameRender_RenderThread().Broadcast(RHICmdList, ViewportManagerProxy, InViewport);
@@ -280,7 +256,7 @@ void FDisplayClusterViewportManagerProxy::ImplRenderFrame_GameThread(FViewport*
 				// For quadbuf stereo copy only left eye, right copy from OutputFrameTarget
 				//@todo Copy QuadBuf_LeftEye/(mono,sbs,tp) to separate rtt, before UI and debug rendering
 				//@todo QuadBuf_LeftEye copied latter, before present
-				switch (ViewportManagerProxy->GetRenderMode())
+				switch (ViewportManagerProxy->GetConfigurationProxy().GetRenderMode_RenderThread())
 				{
 				case EDisplayClusterRenderFrameMode::SideBySide:
 				case EDisplayClusterRenderFrameMode::TopBottom:
@@ -311,7 +287,7 @@ void FDisplayClusterViewportManagerProxy::UpdateDeferredResources_RenderThread(F
 
 	for (const TSharedPtr<FDisplayClusterViewportProxy, ESPMode::ThreadSafe>& ViewportProxy : ImplGetCurrentRenderFrameViewportProxies_RenderThread())
 	{
-		if (!ViewportProxy->RenderSettings.IsViewportOverridden())
+		if (!ViewportProxy->GetRenderSettings_RenderThread().IsViewportOverridden())
 		{
 			ViewportProxy->UpdateDeferredResources(RHICmdList);
 		}
@@ -373,10 +349,7 @@ void FDisplayClusterViewportManagerProxy::UpdateFrameResources_RenderThread(FRHI
 	check(IsInRenderingThread());
 
 	// Do postprocess before warp&blend
-	if (PostProcessManager.IsValid())
-	{
-		PostProcessManager->PerformPostProcessViewBeforeWarpBlend_RenderThread(RHICmdList, this);
-	}
+	PostProcessManager->PerformPostProcessViewBeforeWarpBlend_RenderThread(RHICmdList, this);
 
 	// Support viewport overlap order sorting:
 	TArray<TSharedPtr<FDisplayClusterViewportProxy, ESPMode::ThreadSafe>> SortedViewportProxy = ImplGetCurrentRenderFrameViewportProxies_RenderThread();
@@ -432,14 +405,11 @@ void FDisplayClusterViewportManagerProxy::UpdateFrameResources_RenderThread(FRHI
 		}
 	}
 
-	if (PostProcessManager.IsValid())
-	{
-		// per-frame handle
-		PostProcessManager->HandleUpdateFrameResourcesAfterWarpBlend_RenderThread(RHICmdList, this);
+	// per-frame handle
+	PostProcessManager->HandleUpdateFrameResourcesAfterWarpBlend_RenderThread(RHICmdList, this);
 
-		// Per-view postprocess
-		PostProcessManager->PerformPostProcessViewAfterWarpBlend_RenderThread(RHICmdList, this);
-	}
+	// Per-view postprocess
+	PostProcessManager->PerformPostProcessViewAfterWarpBlend_RenderThread(RHICmdList, this);
 
 	// Post resolve to Frame RTT
 	// All warp&blend results are now inside AdditionalTargetableResource. Viewport images of other projection policies are still stored in the InputShaderResource.
@@ -458,10 +428,7 @@ void FDisplayClusterViewportManagerProxy::UpdateFrameResources_RenderThread(FRHI
 
 	IDisplayCluster::Get().GetCallbacks().OnDisplayClusterPostWarp_RenderThread().Broadcast(RHICmdList, this);
 
-	if (PostProcessManager.IsValid())
-	{
-		PostProcessManager->PerformPostProcessFrameAfterWarpBlend_RenderThread(RHICmdList, this);
-	}
+	PostProcessManager->PerformPostProcessFrameAfterWarpBlend_RenderThread(RHICmdList, this);
 }
 
 void FDisplayClusterViewportManagerProxy::DoCrossGPUTransfers_RenderThread(FRHICommandListImmediate& RHICmdList) const
@@ -485,7 +452,7 @@ void FDisplayClusterViewportManagerProxy::DoCrossGPUTransfers_RenderThread(FRHIC
 			continue;
 		}
 
-		if (!ViewportProxyIt->RenderSettings.bEnableCrossGPUTransfer || ViewportProxyIt->Resources[EDisplayClusterViewportResource::RenderTargets].IsEmpty())
+		if (!ViewportProxyIt->GetRenderSettings_RenderThread().bEnableCrossGPUTransfer || ViewportProxyIt->GetResources_RenderThread()[EDisplayClusterViewportResource::RenderTargets].IsEmpty())
 		{
 			// Skip a frozen viewport that has already been transferred between GPUs
 			// The first time freezing should do the transfer (RenderTargets must be assigned on the first pass)
@@ -493,7 +460,7 @@ void FDisplayClusterViewportManagerProxy::DoCrossGPUTransfers_RenderThread(FRHIC
 		}
 
 		// Do cross GPU trasfers for all viewport contexts
-		for (const FDisplayClusterViewport_Context& ViewportContext : ViewportProxyIt->Contexts)
+		for (const FDisplayClusterViewport_Context& ViewportContext : ViewportProxyIt->GetContexts_RenderThread())
 		{
 			if (!ViewportContext.bOverrideCrossGPUTransfer)
 			{
@@ -508,7 +475,7 @@ void FDisplayClusterViewportManagerProxy::DoCrossGPUTransfers_RenderThread(FRHIC
 				continue;
 			}
 
-			if (!ViewportProxyIt->Resources[EDisplayClusterViewportResource::RenderTargets].IsValidIndex(ViewportContext.ContextNum))
+			if (!ViewportProxyIt->GetResources_RenderThread()[EDisplayClusterViewportResource::RenderTargets].IsValidIndex(ViewportContext.ContextNum))
 			{
 				// RTT does not exist for this context
 				continue;
@@ -522,8 +489,10 @@ void FDisplayClusterViewportManagerProxy::DoCrossGPUTransfers_RenderThread(FRHIC
 				continue;
 			}
 
+			const FDisplayClusterRenderFrameSettings& RenderFrameSettings = ConfigurationProxy->GetRenderFrameSettings();
+
 			// Perform an MGPU transfer for a specific viewport context:
-			TSharedPtr<FDisplayClusterViewportResource, ESPMode::ThreadSafe> ViewportRenderTargetResource = ViewportProxyIt->Resources[EDisplayClusterViewportResource::RenderTargets][ViewportContext.ContextNum];
+			TSharedPtr<FDisplayClusterViewportResource, ESPMode::ThreadSafe> ViewportRenderTargetResource = ViewportProxyIt->GetResources_RenderThread()[EDisplayClusterViewportResource::RenderTargets][ViewportContext.ContextNum];
 			{
 				if (FRenderTarget* RenderTarget = ViewportRenderTargetResource.IsValid() ? ViewportRenderTargetResource->GetViewportResourceRenderTarget() : nullptr)
 				{
@@ -565,10 +534,10 @@ bool FDisplayClusterViewportManagerProxy::GetFrameTargets_RenderThread(TArray<FR
 	for (const TSharedPtr<FDisplayClusterViewportProxy, ESPMode::ThreadSafe>& ViewportProxyIt : ImplGetCurrentRenderFrameViewportProxies_RenderThread())
 	{
 		// Process only valid viewports with output frame resources
-		if (ViewportProxyIt.IsValid() && !ViewportProxyIt->Resources[EDisplayClusterViewportResource::OutputFrameTargetableResources].IsEmpty())
+		if (ViewportProxyIt.IsValid() && !ViewportProxyIt->GetResources_RenderThread()[EDisplayClusterViewportResource::OutputFrameTargetableResources].IsEmpty())
 		{
-			const TArray<TSharedPtr<FDisplayClusterViewportResource, ESPMode::ThreadSafe>>& Frames = ViewportProxyIt->Resources[EDisplayClusterViewportResource::OutputFrameTargetableResources];
-			const TArray<TSharedPtr<FDisplayClusterViewportResource, ESPMode::ThreadSafe>>& AdditionalFrames = ViewportProxyIt->Resources[EDisplayClusterViewportResource::AdditionalFrameTargetableResources];
+			const TArray<TSharedPtr<FDisplayClusterViewportResource, ESPMode::ThreadSafe>>& Frames = ViewportProxyIt->GetResources_RenderThread()[EDisplayClusterViewportResource::OutputFrameTargetableResources];
+			const TArray<TSharedPtr<FDisplayClusterViewportResource, ESPMode::ThreadSafe>>& AdditionalFrames = ViewportProxyIt->GetResources_RenderThread()[EDisplayClusterViewportResource::AdditionalFrameTargetableResources];
 
 			for (int32 FrameIt = 0; FrameIt < Frames.Num(); FrameIt++)
 			{
diff --git a/Engine/Plugins/Runtime/nDisplay/Source/DisplayCluster/Private/Render/Viewport/DisplayClusterViewportManagerProxy.h b/Engine/Plugins/Runtime/nDisplay/Source/DisplayCluster/Private/Render/Viewport/DisplayClusterViewportManagerProxy.h
--- a/Engine/Plugins/Runtime/nDisplay/Source/DisplayCluster/Private/Render/Viewport/DisplayClusterViewportManagerProxy.h
+++ b/Engine/Plugins/Runtime/nDisplay/Source/DisplayCluster/Private/Render/Viewport/DisplayClusterViewportManagerProxy.h
@@ -28,7 +28,7 @@ class FDisplayClusterViewportManagerProxy
 	, public TSharedFromThis<FDisplayClusterViewportManagerProxy, ESPMode::ThreadSafe>
 {
 public:
-	FDisplayClusterViewportManagerProxy();
+	FDisplayClusterViewportManagerProxy(const FDisplayClusterViewportManager& InViewportManager);
 	virtual ~FDisplayClusterViewportManagerProxy();
 
 public:
@@ -43,7 +43,11 @@ public:
 		return AsShared();
 	}
 
-	virtual EDisplayClusterRenderFrameMode GetRenderMode() const override;
+	/** Get viewport manager proxy configuration interface. */
+	virtual const IDisplayClusterViewportConfigurationProxy& GetConfigurationProxy() const override
+	{
+		return *ConfigurationProxy;
+	}
 
 	virtual IDisplayClusterViewportProxy* FindViewport_RenderThread(const FString& InViewportId) const override
 	{
@@ -82,17 +86,10 @@ public:
 	/** Unregister exist viewport proxy. */
 	void DeleteViewport_RenderThread(const TSharedPtr<FDisplayClusterViewportProxy, ESPMode::ThreadSafe>& InViewportProxy);
 
-	/** Initialize ViewportManagerProxy from ViewportManager. */
-	void Initialize(FDisplayClusterViewportManager& InViewportManager);
-
 	/** Get LightCardManager proxy object. */
 	inline TSharedPtr<FDisplayClusterViewportLightCardManagerProxy, ESPMode::ThreadSafe> GetLightCardManagerProxy_RenderThread() const
 	{ return LightCardManagerProxy; }
 
-	/** Get RenderFrameSettings on render thread. */
-	inline  const FDisplayClusterRenderFrameSettings& GetRenderFrameSettings_RenderThread() const
-	{ return RenderFrameSettings; }
-
 	/** Geting the viewport proxies of the current rendering frame. (viewports from the current node or from a special named list) */
 	inline const TArray<TSharedPtr<FDisplayClusterViewportProxy, ESPMode::ThreadSafe>>& ImplGetCurrentRenderFrameViewportProxies_RenderThread() const
 	{ return CurrentRenderFrameViewportProxies; }
@@ -107,13 +104,12 @@ public:
 	/** Find viewport and context index by StereoViewIndex. */
 	FDisplayClusterViewportProxy* ImplFindViewportProxy_RenderThread(const int32 StereoViewIndex, uint32* OutContextNum = nullptr) const;
 
-	/** Copy RenderFrameSettings from game to render thread. */
-	void ImplUpdateRenderFrameSettings(const FDisplayClusterRenderFrameSettings& InRenderFrameSettings,
-		const TSharedPtr<FDisplayClusterViewportManagerViewExtension, ESPMode::ThreadSafe>& InViewportManagerViewExtension);
-
 	/** Sending viewport data from the game stream to the viewport proxy that is in the rendering stream. */
 	void ImplUpdateViewportProxies_GameThread(const TArray<TSharedPtr<FDisplayClusterViewport, ESPMode::ThreadSafe>>& InCurrentRenderFrameViewports);
 
+	/** Sending viewport manager data from the game stream to the viewport manager proxy that is in the rendering stream. */
+	void ImplUpdateViewportManagerProxy_GameThread(const FDisplayClusterViewportManager& InViewportManager);
+
 	/** Rendering the final frame of nDisplay. Called after all viewports have been rendered in RTT. */
 	void ImplRenderFrame_GameThread(FViewport* InViewport);
 
@@ -124,19 +120,20 @@ private:
 	/** Update ClusterNodeViewportProxies from ViewportProxies for the current cluster node. */
 	void ImplUpdateClusterNodeViewportProxies_RenderThread();
 
-private:
+public:
+	// Configuration for proxy
+	const TSharedRef<FDisplayClusterViewportConfigurationProxy, ESPMode::ThreadSafe> ConfigurationProxy;
+
 	// RenderTarget manager object. Used to remove related objects in order.
-	TSharedPtr<FDisplayClusterRenderTargetManager, ESPMode::ThreadSafe>        RenderTargetManager;
+	const TSharedRef<FDisplayClusterRenderTargetManager, ESPMode::ThreadSafe> RenderTargetManager;
 
 	// PostProcess manager object. Used to remove related objects in order.
-	TSharedPtr<FDisplayClusterViewportPostProcessManager, ESPMode::ThreadSafe> PostProcessManager;
+	const TSharedRef<FDisplayClusterViewportPostProcessManager, ESPMode::ThreadSafe> PostProcessManager;
 
 	// LightCard manager proxy object.
-	TSharedPtr<FDisplayClusterViewportLightCardManagerProxy, ESPMode::ThreadSafe> LightCardManagerProxy;
-
-	// Copy of RenderFrameSettings for the render thread.
-	FDisplayClusterRenderFrameSettings RenderFrameSettings;
+	const TSharedRef<FDisplayClusterViewportLightCardManagerProxy, ESPMode::ThreadSafe> LightCardManagerProxy;
 
+private:
 	// Viewport proxies of the entire cluster.
 	TArray<TSharedPtr<FDisplayClusterViewportProxy, ESPMode::ThreadSafe>> EntireClusterViewportProxies;
 
diff --git a/Engine/Plugins/Runtime/nDisplay/Source/DisplayCluster/Private/Render/Viewport/DisplayClusterViewportManagerViewExtension.cpp b/Engine/Plugins/Runtime/nDisplay/Source/DisplayCluster/Private/Render/Viewport/DisplayClusterViewportManagerViewExtension.cpp
--- a/Engine/Plugins/Runtime/nDisplay/Source/DisplayCluster/Private/Render/Viewport/DisplayClusterViewportManagerViewExtension.cpp
+++ b/Engine/Plugins/Runtime/nDisplay/Source/DisplayCluster/Private/Render/Viewport/DisplayClusterViewportManagerViewExtension.cpp
@@ -23,36 +23,30 @@ using namespace UE::DisplayCluster::ViewportManagerViewExtension;
 ///////////////////////////////////////////////////////////////////////////////////////
 // FDisplayClusterViewportManagerViewExtension
 ///////////////////////////////////////////////////////////////////////////////////////
-FDisplayClusterViewportManagerViewExtension::FDisplayClusterViewportManagerViewExtension(const FAutoRegister& AutoRegister, const FDisplayClusterViewportManager* InViewportManager)
+FDisplayClusterViewportManagerViewExtension::FDisplayClusterViewportManagerViewExtension(const FAutoRegister& AutoRegister, const TSharedRef<FDisplayClusterViewportConfiguration, ESPMode::ThreadSafe>& InConfiguration)
 	: FSceneViewExtensionBase(AutoRegister)
-	, ViewportManagerWeakPtr(InViewportManager->AsWeak())
-	, ViewportManagerProxyWeakPtr(InViewportManager->GetViewportManagerProxy()->AsWeak())
+	, Configuration(InConfiguration)
 {
 	RegisterCallbacks();
 }
 
 FDisplayClusterViewportManagerViewExtension::~FDisplayClusterViewportManagerViewExtension()
 {
-	ViewportManagerWeakPtr.Reset();
-	ViewportManagerProxyWeakPtr.Reset();
-
 	UnregisterCallbacks();
 }
 
 bool FDisplayClusterViewportManagerViewExtension::IsActive() const
 {
-	return ViewportManagerProxyWeakPtr.IsValid() && ViewportManagerProxyWeakPtr.Pin().IsValid();
+	return !bReleased && (Configuration->Proxy->GetViewportManagerProxyImpl() != nullptr);
 }
 
 void FDisplayClusterViewportManagerViewExtension::Release_RenderThread()
 {
 	check(IsInRenderingThread());
 
-	ViewportManagerWeakPtr.Reset();
-	ViewportManagerProxyWeakPtr.Reset();
+	bReleased = true;
 
 	ViewportProxies.Empty();
-
 	UnregisterCallbacks();
 }
 
@@ -128,7 +122,7 @@ FScreenPassTexture FDisplayClusterViewportManagerViewExtension::PostProcessPassA
 		return ReturnUntouchedSceneColorForPostProcessing(Inputs);
 	}
 
-	if (const FDisplayClusterViewportManagerProxy* ViewportManagerProxy = GetViewportManagerProxy())
+	if (const FDisplayClusterViewportManagerProxy* ViewportManagerProxy = Configuration->Proxy->GetViewportManagerProxyImpl())
 	{
 		uint32 ContextNum = 0;
 		if (FDisplayClusterViewportProxy* ViewportProxyPtr = ViewportManagerProxy->ImplFindViewportProxy_RenderThread(View.StereoViewIndex, &ContextNum))
@@ -150,7 +144,7 @@ FScreenPassTexture FDisplayClusterViewportManagerViewExtension::PostProcessPassA
 		return ReturnUntouchedSceneColorForPostProcessing(Inputs);
 	}
 
-	if (const FDisplayClusterViewportManagerProxy* ViewportManagerProxy = GetViewportManagerProxy())
+	if (const FDisplayClusterViewportManagerProxy* ViewportManagerProxy = Configuration->Proxy->GetViewportManagerProxyImpl())
 	{
 		uint32 ContextNum = 0;
 		if (FDisplayClusterViewportProxy* ViewportProxyPtr = ViewportManagerProxy->ImplFindViewportProxy_RenderThread(View.StereoViewIndex, &ContextNum))
@@ -172,7 +166,7 @@ FScreenPassTexture FDisplayClusterViewportManagerViewExtension::PostProcessPassA
 		return ReturnUntouchedSceneColorForPostProcessing(Inputs);
 	}
 
-	if (const FDisplayClusterViewportManagerProxy* ViewportManagerProxy = GetViewportManagerProxy())
+	if (const FDisplayClusterViewportManagerProxy* ViewportManagerProxy = Configuration->Proxy->GetViewportManagerProxyImpl())
 	{
 		uint32 ContextNum = 0;
 		if (FDisplayClusterViewportProxy* ViewportProxyPtr = ViewportManagerProxy->ImplFindViewportProxy_RenderThread(View.StereoViewIndex, &ContextNum))
@@ -215,7 +209,7 @@ bool FDisplayClusterViewportManagerViewExtension::IsActiveThisFrame_Internal(con
 	if (Context.IsA(MoveTempIfPossible(DCViewExtensionContext)))
 	{
 		const FDisplayClusterSceneViewExtensionContext& DisplayContext = static_cast<const FDisplayClusterSceneViewExtensionContext&>(Context);
-		if (DisplayContext.GetViewportManager() == GetViewportManager())
+		if (DisplayContext.Configuration == Configuration)
 		{
 			// Apply only for DC viewports
 			return true;
@@ -241,7 +235,7 @@ void FDisplayClusterViewportManagerViewExtension::PreRenderViewFamily_RenderThre
 		{
 			FViewportProxy NewViewportProxy;
 
-			if (const FDisplayClusterViewportManagerProxy* ViewportManagerProxy = GetViewportManagerProxy())
+			if (const FDisplayClusterViewportManagerProxy* ViewportManagerProxy = Configuration->Proxy->GetViewportManagerProxyImpl())
 			{
 				if (FDisplayClusterViewportProxy* ViewportProxyPtr = static_cast<FDisplayClusterViewportProxy*>(ViewportManagerProxy->FindViewport_RenderThread(SceneView->StereoViewIndex, &NewViewportProxy.ViewportProxyContext.ContextNum)))
 				{
diff --git a/Engine/Plugins/Runtime/nDisplay/Source/DisplayCluster/Private/Render/Viewport/DisplayClusterViewportManagerViewExtension.h b/Engine/Plugins/Runtime/nDisplay/Source/DisplayCluster/Private/Render/Viewport/DisplayClusterViewportManagerViewExtension.h
--- a/Engine/Plugins/Runtime/nDisplay/Source/DisplayCluster/Private/Render/Viewport/DisplayClusterViewportManagerViewExtension.h
+++ b/Engine/Plugins/Runtime/nDisplay/Source/DisplayCluster/Private/Render/Viewport/DisplayClusterViewportManagerViewExtension.h
@@ -2,10 +2,9 @@
 
 #pragma once
 #include "SceneViewExtension.h"
+#include "Render/Viewport/Configuration/DisplayClusterViewportConfiguration.h"
 #include "Render/Viewport/Containers/DisplayClusterViewportProxy_Context.h"
 
-class FDisplayClusterViewportManager;
-class FDisplayClusterViewportManagerProxy;
 class FDisplayClusterViewportProxy;
 
 #define DISPLAYCLUSTER_SCENE_VIEW_EXTENSION_PRIORITY -1
@@ -16,7 +15,7 @@ class FDisplayClusterViewportProxy;
 class FDisplayClusterViewportManagerViewExtension : public FSceneViewExtensionBase
 {
 public:
-	FDisplayClusterViewportManagerViewExtension(const FAutoRegister& AutoRegister, const FDisplayClusterViewportManager* InViewportManager);
+	FDisplayClusterViewportManagerViewExtension(const FAutoRegister& AutoRegister, const TSharedRef<FDisplayClusterViewportConfiguration, ESPMode::ThreadSafe>& InConfiguration);
 	virtual ~FDisplayClusterViewportManagerViewExtension();
 
 public:
@@ -105,24 +104,13 @@ private:
 	/** Return true if the VE is safe to use. */
 	bool IsActive() const;
 
-protected:
-	inline const FDisplayClusterViewportManager* GetViewportManager() const
-	{
-		return ViewportManagerWeakPtr.IsValid() ? ViewportManagerWeakPtr.Pin().Get() : nullptr;
-	}
-
-	inline const FDisplayClusterViewportManagerProxy* GetViewportManagerProxy() const
-	{
-		return ViewportManagerProxyWeakPtr.IsValid() ? ViewportManagerProxyWeakPtr.Pin().Get() : nullptr;
-	}
+public:
+	// Configuration of the current cluster node
+	const TSharedRef<FDisplayClusterViewportConfiguration, ESPMode::ThreadSafe> Configuration;
 
 private:
 	FDelegateHandle ResolvedSceneColorCallbackHandle;
 
-private:
-	TWeakPtr<const FDisplayClusterViewportManager, ESPMode::ThreadSafe> ViewportManagerWeakPtr;
-	TWeakPtr<const FDisplayClusterViewportManagerProxy, ESPMode::ThreadSafe> ViewportManagerProxyWeakPtr;
-
 	struct FViewportProxy
 	{
 		inline FDisplayClusterViewportProxy* GetViewportProxy() const
@@ -142,4 +130,7 @@ private:
 
 	/** RenderThreadData:. Viewport proxies from rendered view family. */
 	TArray<FViewportProxy> ViewportProxies;
+
+	// Is this VE released
+	bool bReleased = false;
 };
diff --git a/Engine/Plugins/Runtime/nDisplay/Source/DisplayCluster/Private/Render/Viewport/DisplayClusterViewportManagerViewPointExtension.cpp b/Engine/Plugins/Runtime/nDisplay/Source/DisplayCluster/Private/Render/Viewport/DisplayClusterViewportManagerViewPointExtension.cpp
--- a/Engine/Plugins/Runtime/nDisplay/Source/DisplayCluster/Private/Render/Viewport/DisplayClusterViewportManagerViewPointExtension.cpp
+++ b/Engine/Plugins/Runtime/nDisplay/Source/DisplayCluster/Private/Render/Viewport/DisplayClusterViewportManagerViewPointExtension.cpp
@@ -8,24 +8,25 @@
 ///////////////////////////////////////////////////////////////////////////////////////
 // FDisplayClusterViewportManagerViewPointExtension
 ///////////////////////////////////////////////////////////////////////////////////////
-FDisplayClusterViewportManagerViewPointExtension::FDisplayClusterViewportManagerViewPointExtension(const FAutoRegister& AutoRegister, const FDisplayClusterViewportManager* InViewportManager)
+FDisplayClusterViewportManagerViewPointExtension::FDisplayClusterViewportManagerViewPointExtension(const FAutoRegister& AutoRegister, const TSharedRef<FDisplayClusterViewportConfiguration, ESPMode::ThreadSafe>& InConfiguration)
 	: FSceneViewExtensionBase(AutoRegister)
-	, ViewportManagerWeakPtr(InViewportManager->AsShared())
+	, Configuration(InConfiguration)
 { }
 
-FDisplayClusterViewportManagerViewPointExtension::~FDisplayClusterViewportManagerViewPointExtension()
+bool FDisplayClusterViewportManagerViewPointExtension::IsActiveThisFrame_Internal(const FSceneViewExtensionContext& Context) const
 {
-	ViewportManagerWeakPtr.Reset();
+	return IsActive() && Context.IsStereoSupported();
 }
 
-bool FDisplayClusterViewportManagerViewPointExtension::IsActiveThisFrame_Internal(const FSceneViewExtensionContext& Context) const
+/** True, if VE can be used at the moment. */
+bool FDisplayClusterViewportManagerViewPointExtension::IsActive() const
 {
-	return IsActive() && Context.IsStereoSupported();
+	return Configuration->GetViewportManager() != nullptr && CurrentStereoViewIndex != INDEX_NONE;
 }
 
 void FDisplayClusterViewportManagerViewPointExtension::SetupViewPoint(APlayerController* Player, FMinimalViewInfo& InOutViewInfo)
 {
-	if (IDisplayClusterViewport* DCViewport = IsActive() ? GetViewportManager()->FindViewport(CurrentStereoViewIndex) : nullptr)
+	if (IDisplayClusterViewport* DCViewport = IsActive() ? Configuration->GetViewportManager()->FindViewport(CurrentStereoViewIndex) : nullptr)
 	{
 		DCViewport->SetupViewPoint(InOutViewInfo);
 	}
diff --git a/Engine/Plugins/Runtime/nDisplay/Source/DisplayCluster/Private/Render/Viewport/DisplayClusterViewportManagerViewPointExtension.h b/Engine/Plugins/Runtime/nDisplay/Source/DisplayCluster/Private/Render/Viewport/DisplayClusterViewportManagerViewPointExtension.h
--- a/Engine/Plugins/Runtime/nDisplay/Source/DisplayCluster/Private/Render/Viewport/DisplayClusterViewportManagerViewPointExtension.h
+++ b/Engine/Plugins/Runtime/nDisplay/Source/DisplayCluster/Private/Render/Viewport/DisplayClusterViewportManagerViewPointExtension.h
@@ -2,8 +2,7 @@
 
 #pragma once
 #include "SceneViewExtension.h"
-
-class FDisplayClusterViewportManager;
+#include "Render/Viewport/Configuration/DisplayClusterViewportConfiguration.h"
 
 #define DISPLAYCLUSTER_SCENE_VIEWPOINT_EXTENSION_PRIORITY -1
 
@@ -13,8 +12,8 @@ class FDisplayClusterViewportManager;
 class FDisplayClusterViewportManagerViewPointExtension : public FSceneViewExtensionBase
 {
 public:
-	FDisplayClusterViewportManagerViewPointExtension(const FAutoRegister& AutoRegister, const FDisplayClusterViewportManager* InViewportManager);
-	virtual ~FDisplayClusterViewportManagerViewPointExtension();
+	FDisplayClusterViewportManagerViewPointExtension(const FAutoRegister& AutoRegister, const TSharedRef<FDisplayClusterViewportConfiguration, ESPMode::ThreadSafe>& InConfiguration);
+	virtual ~FDisplayClusterViewportManagerViewPointExtension() = default;
 
 public:
 	//~ Begin ISceneViewExtension interface
@@ -40,21 +39,14 @@ protected:
 	virtual bool IsActiveThisFrame_Internal(const FSceneViewExtensionContext& Context) const override;
 
 private:
-	/** Get viewport manager ptr. */
-	inline const FDisplayClusterViewportManager* GetViewportManager() const
-	{
-		return ViewportManagerWeakPtr.IsValid() ? ViewportManagerWeakPtr.Pin().Get() : nullptr;
-	}
-
 	/** True, if VE can be used at the moment. */
-	bool IsActive() const
-	{
-		return GetViewportManager() != nullptr && CurrentStereoViewIndex != INDEX_NONE;
-	}
+	bool IsActive() const;
 
-private:
-	TWeakPtr<const FDisplayClusterViewportManager, ESPMode::ThreadSafe> ViewportManagerWeakPtr;
+public:
+	// Configuration of the current cluster node
+	const TSharedRef<FDisplayClusterViewportConfiguration, ESPMode::ThreadSafe> Configuration;
 
+private:
 	// Current StereoViewIndex for rendered viewport
 	int32 CurrentStereoViewIndex = INDEX_NONE;
 };
diff --git a/Engine/Plugins/Runtime/nDisplay/Source/DisplayCluster/Private/Render/Viewport/DisplayClusterViewportManager_EditorPreview.cpp b/Engine/Plugins/Runtime/nDisplay/Source/DisplayCluster/Private/Render/Viewport/DisplayClusterViewportManager_EditorPreview.cpp
--- a/Engine/Plugins/Runtime/nDisplay/Source/DisplayCluster/Private/Render/Viewport/DisplayClusterViewportManager_EditorPreview.cpp
+++ b/Engine/Plugins/Runtime/nDisplay/Source/DisplayCluster/Private/Render/Viewport/DisplayClusterViewportManager_EditorPreview.cpp
@@ -37,19 +37,19 @@ void FDisplayClusterViewportManager::ImplUpdatePreviewRTTResources()
 	check(IsInGameThread());
 
 	// Only for preview modes:
-	if(GetRenderMode() != EDisplayClusterRenderFrameMode::PreviewInScene)
+	if(!Configuration->IsPreviewRendering())
 	{
 		return;
 	}
 
 	// The preview RTT created inside root actor
-	ADisplayClusterRootActor* RootActor = GetRootActor();
-	if (RootActor == nullptr)
+	ADisplayClusterRootActor* PreviewRootActor = Configuration->GetRootActor(EDisplayClusterRootActorType::Preview);
+	if (PreviewRootActor == nullptr)
 	{
 		return;
 	}
 
-	const FString& ClusterNodeId = GetRenderFrameSettings().ClusterNodeId;
+	const FString& ClusterNodeId = Configuration->GetRenderFrameSettings().ClusterNodeId;
 
 	TArray<FString>        PreviewViewportNames;
 	TArray<FTextureRHIRef> PreviewRenderTargetableTextures;
@@ -63,7 +63,7 @@ void FDisplayClusterViewportManager::ImplUpdatePreviewRTTResources()
 		// update only current cluster node
 		if (ViewportIt.IsValid() && (ClusterNodeId.IsEmpty() || ViewportIt->GetClusterNodeId() == ClusterNodeId))
 		{
-			if (ViewportIt->RenderSettings.bEnable && ViewportIt->RenderSettings.bVisible)
+			if (ViewportIt->GetRenderSettings().bEnable && ViewportIt->GetRenderSettings().bVisible)
 			{
 				PreviewViewportNames.Add(ViewportIt->GetId());
 				PreviewRenderTargetableTextures.AddDefaulted();
@@ -72,7 +72,7 @@ void FDisplayClusterViewportManager::ImplUpdatePreviewRTTResources()
 	}
 
 	// Get all supported preview rtt resources from root actor
-	RootActor->GetPreviewRenderTargetableTextures(PreviewViewportNames, PreviewRenderTargetableTextures);
+	PreviewRootActor->GetPreviewRenderTargetableTextures(PreviewViewportNames, PreviewRenderTargetableTextures);
 
 	// Configure preview RTT to viwports:
 	for (int32 ViewportIndex = 0; ViewportIndex < PreviewViewportNames.Num(); ViewportIndex++)
@@ -82,16 +82,19 @@ void FDisplayClusterViewportManager::ImplUpdatePreviewRTTResources()
 		{
 			FTextureRHIRef& PreviewRTT = PreviewRenderTargetableTextures[ViewportIndex];
 			if (PreviewRTT.IsValid()
-				&& !DesiredViewport->Resources[EDisplayClusterViewportResource::OutputPreviewTargetableResources].IsEmpty()
-				&& DesiredViewport->Resources[EDisplayClusterViewportResource::OutputPreviewTargetableResources][0].IsValid())
+				&& !DesiredViewport->GetViewportResources(EDisplayClusterViewportResource::OutputPreviewTargetableResources).IsEmpty()
+				&& DesiredViewport->GetViewportResources(EDisplayClusterViewportResource::OutputPreviewTargetableResources)[0].IsValid())
 			{
 				// Use mapped preview viewport
-				DesiredViewport->Resources[EDisplayClusterViewportResource::OutputPreviewTargetableResources][0]->SetExternalViewportResourceRHI(PreviewRTT);
+				DesiredViewport->GetViewportResources(EDisplayClusterViewportResource::OutputPreviewTargetableResources)[0]->SetExternalViewportResourceRHI(PreviewRTT);
 			}
 			else
 			{
+				// Gain direct access to internal settings of the viewport:
+				FDisplayClusterViewport_RenderSettings& InOutDesiredViewportRenderSettings = DesiredViewport->GetRenderSettingsImpl();
+
 				// disable visible, but unused by root actor viewports
-				DesiredViewport->RenderSettings.bEnable = false;
+				InOutDesiredViewportRenderSettings.bEnable = false;
 			}
 		}
 	}
@@ -115,16 +118,15 @@ bool FDisplayClusterViewportManager::RenderInEditor(FDisplayClusterRenderFrame&
 	bOutFrameRendered = false;
 	OutViewportsAmount = 0;
 
-	UWorld* CurrentWorld = GetCurrentWorld();
-	if (CurrentWorld == nullptr)
+	UWorld* CurrentWorld = Configuration->GetCurrentWorld();
+	const ADisplayClusterRootActor* PreviewRootActor = Configuration->GetRootActor(EDisplayClusterRootActorType::Preview);
+	if (!CurrentWorld || !PreviewRootActor)
 	{
 		return false;
 	}
 
-	const ADisplayClusterRootActor* RootActor = GetRootActor();
-
 #if WITH_EDITOR
-	if (RootActor && !RootActor->IsEditorRenderEnabled())
+	if (!PreviewRootActor->IsEditorRenderEnabled())
 	{
 		bOutFrameRendered = true;
 		return true;
@@ -175,11 +177,8 @@ bool FDisplayClusterViewportManager::RenderInEditor(FDisplayClusterRenderFrame&
 				ConfigureViewFamily(RenderTargetIt, ViewFamiliesIt, ViewFamily);
 
 				if (RenderTargetIt.CaptureMode == EDisplayClusterViewportCaptureMode::Default
-#if WITH_EDITOR
-					&& RootActor
-					&& !RootActor->bPreviewEnablePostProcess
-					&& RootActor->DoObserversNeedPostProcessRenderTarget()
-#endif
+					&& PreviewRootActor
+					&& !PreviewRootActor->bPreviewEnablePostProcess
 				)
 				{
 					if (ViewFamily.EngineShowFlags.TemporalAA)
diff --git a/Engine/Plugins/Runtime/nDisplay/Source/DisplayCluster/Private/Render/Viewport/DisplayClusterViewportProxy.cpp b/Engine/Plugins/Runtime/nDisplay/Source/DisplayCluster/Private/Render/Viewport/DisplayClusterViewportProxy.cpp
--- a/Engine/Plugins/Runtime/nDisplay/Source/DisplayCluster/Private/Render/Viewport/DisplayClusterViewportProxy.cpp
+++ b/Engine/Plugins/Runtime/nDisplay/Source/DisplayCluster/Private/Render/Viewport/DisplayClusterViewportProxy.cpp
@@ -1,14 +1,13 @@
 // Copyright Epic Games, Inc. All Rights Reserved.
 
 #include "Render/Viewport/DisplayClusterViewportProxy.h"
-
 #include "Render/Viewport/DisplayClusterViewportManagerProxy.h"
 #include "Render/Viewport/DisplayClusterViewport.h"
 #include "Render/Viewport/DisplayClusterViewportManager.h"
 #include "Render/Viewport/DisplayClusterViewportManagerViewExtension.h"
 
 #include "Render/Viewport/Containers/DisplayClusterViewport_PostRenderSettings.h"
-#include "Render/Viewport/Configuration/DisplayClusterViewportConfiguration.h"
+#include "Render/Viewport/Containers/DisplayClusterViewportProxyData.h"
 
 #include "Render/Viewport/RenderTarget/DisplayClusterRenderTargetResource.h"
 #include "Render/Viewport/RenderFrame/DisplayClusterRenderFrameSettings.h"
@@ -25,11 +24,6 @@
 #include "IDisplayClusterShaders.h"
 #include "TextureResource.h"
 
-#if WITH_EDITOR
-#include "DisplayClusterRootActor.h"
-#include "Render/Viewport/Containers/DisplayClusterViewportReadPixels.h"
-#endif
-
 #include "RHIStaticStates.h"
 
 #include "RenderResource.h"
@@ -136,12 +130,11 @@ namespace UE::DisplayCluster::ViewportProxy
 using namespace UE::DisplayCluster::ViewportProxy;
 
 ///////////////////////////////////////////////////////////////////////////////////////
-FDisplayClusterViewportProxy::FDisplayClusterViewportProxy(const FDisplayClusterViewport& RenderViewport)
-	: ViewportId(RenderViewport.ViewportId)
-	, ClusterNodeId(RenderViewport.ClusterNodeId)
-	, RenderSettings(RenderViewport.RenderSettings)
-	, ProjectionPolicy(RenderViewport.UninitializedProjectionPolicy)
-	, ViewportManagerProxyWeakRef(RenderViewport.GetViewportManagerProxyRefImpl())
+FDisplayClusterViewportProxy::FDisplayClusterViewportProxy(const TSharedRef<FDisplayClusterViewportConfiguration, ESPMode::ThreadSafe>& InConfiguration, const FString& InViewportId, const TSharedPtr<IDisplayClusterProjectionPolicy, ESPMode::ThreadSafe>& InProjectionPolicy)
+	: ConfigurationProxy(InConfiguration->Proxy)
+	, ViewportId(InViewportId)
+	, ClusterNodeId(InConfiguration->GetClusterNodeId())
+	, ProjectionPolicy(InProjectionPolicy)
 {
 	check(ProjectionPolicy.IsValid());
 }
@@ -150,6 +143,39 @@ FDisplayClusterViewportProxy::~FDisplayClusterViewportProxy()
 {
 }
 
+void FDisplayClusterViewportProxy::UpdateViewportProxyData_RenderThread(const FDisplayClusterViewportProxyData& InViewportProxyData)
+{
+	OpenColorIO = InViewportProxyData.OpenColorIO;
+
+	OverscanRuntimeSettings = InViewportProxyData.OverscanRuntimeSettings;
+
+	RemapMesh = InViewportProxyData.RemapMesh;
+
+	RenderSettings = InViewportProxyData.RenderSettings;
+
+	RenderSettingsICVFX.SetParameters(InViewportProxyData.RenderSettingsICVFX);
+	PostRenderSettings.SetParameters(InViewportProxyData.PostRenderSettings);
+
+	ProjectionPolicy = InViewportProxyData.ProjectionPolicy;
+
+	// The RenderThreadData for DstViewportProxy has been updated in DisplayClusterViewportManagerViewExtension on the rendering thread.
+	// Therefore, the RenderThreadData values from the game thread must be overridden by current data from the render thread.
+	{
+		const TArray<FDisplayClusterViewport_Context> CurrentContexts = Contexts;
+		Contexts = InViewportProxyData.Contexts;
+
+		int32 ContextAmmount = FMath::Min(CurrentContexts.Num(), Contexts.Num());
+		for (int32 ContextIndex = 0; ContextIndex < ContextAmmount; ContextIndex++)
+		{
+			Contexts[ContextIndex].RenderThreadData = CurrentContexts[ContextIndex].RenderThreadData;
+		}
+	}
+
+	// Update viewport proxy resources from container
+	Resources = InViewportProxyData.Resources;
+	ViewStates = InViewportProxyData.ViewStates;
+}
+
 EDisplayClusterViewportResourceType FDisplayClusterViewportProxy::GetResourceType_RenderThread(const EDisplayClusterViewportResourceType& InResourceType) const
 {
 	check(IsInRenderingThread());
@@ -167,7 +193,7 @@ EDisplayClusterViewportResourceType FDisplayClusterViewportProxy::GetResourceTyp
 		/**
 		* Output textures for preview rendering.
 		*/
-		if (GetRenderMode() == EDisplayClusterRenderFrameMode::PreviewInScene)
+		if (ConfigurationProxy->IsPreviewRendering_RenderThread())
 		{
 			// [warp] -> OutputPreviewTargetableResource
 			return EDisplayClusterViewportResourceType::OutputPreviewTargetableResource;
@@ -213,8 +239,7 @@ bool FDisplayClusterViewportProxy::ShouldApplyWarpBlend_RenderThread() const
 		return false;
 	}
 
-	const FDisplayClusterRenderFrameSettings* RenderFrameSettings = GetRenderFrameSettings_RenderThread();
-	if (!RenderFrameSettings || !RenderFrameSettings->bAllowWarpBlend)
+	if (!ConfigurationProxy->GetRenderFrameSettings().bAllowWarpBlend)
 	{
 		return false;
 	}
@@ -223,41 +248,6 @@ bool FDisplayClusterViewportProxy::ShouldApplyWarpBlend_RenderThread() const
 	return ProjectionPolicy.IsValid() && ProjectionPolicy->IsWarpBlendSupported();
 }
 
-const FDisplayClusterRenderFrameSettings* FDisplayClusterViewportProxy::GetRenderFrameSettings_RenderThread() const
-{
-	if (FDisplayClusterViewportManagerProxy* ViewportManagerProxy = GetViewportManagerProxyImpl_RenderThread())
-	{
-		return &ViewportManagerProxy->GetRenderFrameSettings_RenderThread();
-	}
-
-	return nullptr;
-}
-
-const IDisplayClusterViewportManagerProxy* FDisplayClusterViewportProxy::GetViewportManagerProxy_RenderThread() const
-{
-	return GetViewportManagerProxyRefImpl_RenderThread().Get();
-}
-
-FDisplayClusterViewportManagerProxy* FDisplayClusterViewportProxy::GetViewportManagerProxyImpl_RenderThread() const
-{
-	return GetViewportManagerProxyRefImpl_RenderThread().Get();
-}
-
-TSharedPtr<FDisplayClusterViewportManagerProxy, ESPMode::ThreadSafe> FDisplayClusterViewportProxy::GetViewportManagerProxyRefImpl_RenderThread() const
-{
-	return ViewportManagerProxyWeakRef.IsValid() ? ViewportManagerProxyWeakRef.Pin() : nullptr;
-}
-
-EDisplayClusterRenderFrameMode FDisplayClusterViewportProxy::GetRenderMode() const
-{
-	if (FDisplayClusterViewportManagerProxy* ViewportManagerProxy = GetViewportManagerProxyImpl_RenderThread())
-	{
-		return ViewportManagerProxy->GetRenderMode();
-	}
-
-	return EDisplayClusterRenderFrameMode::Unknown;
-}
-
 bool FDisplayClusterViewportProxy::ShouldOverrideViewportResource(const EDisplayClusterViewportResourceType InExtResourceType) const
 {
 	// Override resources from other viewport
@@ -314,7 +304,7 @@ const FDisplayClusterViewportProxy& FDisplayClusterViewportProxy::GetRenderingVi
 	{
 	case EDisplayClusterViewportOverrideMode::All:
 	case EDisplayClusterViewportOverrideMode::InernalRTT:
-		if (FDisplayClusterViewportManagerProxy* ViewportManagerProxy = GetViewportManagerProxyImpl_RenderThread())
+		if (FDisplayClusterViewportManagerProxy* ViewportManagerProxy = ConfigurationProxy->GetViewportManagerProxyImpl())
 		{
 			if (FDisplayClusterViewportProxy const* OverrideViewportProxy = ViewportManagerProxy->ImplFindViewportProxy_RenderThread(RenderSettings.GetViewportOverrideId()))
 			{
@@ -394,7 +384,7 @@ bool FDisplayClusterViewportProxy::ImplGetResources_RenderThread(const EDisplayC
 				bResult = true;
 				
 				// Get resources from external UV LightCard manager
-				if (FDisplayClusterViewportManagerProxy* ViewportManagerProxy = GetViewportManagerProxyImpl_RenderThread())
+				if (FDisplayClusterViewportManagerProxy* ViewportManagerProxy = ConfigurationProxy->GetViewportManagerProxyImpl())
 				{
 					TSharedPtr<FDisplayClusterViewportLightCardManagerProxy, ESPMode::ThreadSafe> LightCardManager = ViewportManagerProxy->GetLightCardManagerProxy_RenderThread();
 					if (LightCardManager.IsValid())
@@ -473,15 +463,12 @@ void FDisplayClusterViewportProxy::PostResolveViewport_RenderThread(FRHICommandL
 
 	// Implement ViewportRemap feature
 	ImplViewportRemap_RenderThread(RHICmdList);
-
-	// Implement read pixels for preview DCRA
-	ImplPreviewReadPixels_RenderThread(RHICmdList);
 }
 
 void FDisplayClusterViewportProxy::ImplViewportRemap_RenderThread(FRHICommandListImmediate& RHICmdList) const
 {
 	// Preview in editor not support this feature
-	if (GetRenderMode() == EDisplayClusterRenderFrameMode::PreviewInScene)
+	if (ConfigurationProxy->IsPreviewRendering_RenderThread())
 	{
 		return;
 	}
@@ -514,70 +501,6 @@ void FDisplayClusterViewportProxy::ImplViewportRemap_RenderThread(FRHICommandLis
 	}
 }
 
-void FDisplayClusterViewportProxy::ImplPreviewReadPixels_RenderThread(FRHICommandListImmediate& RHICmdList) const
-{
-#if WITH_EDITOR
-	if (RenderSettings.bPreviewReadPixels)
-	{
-		bPreviewReadPixels = true;
-	}
-
-	// Now try to read until success:
-	if (bPreviewReadPixels && Contexts.Num() > 0 && !Resources[EDisplayClusterViewportResource::OutputPreviewTargetableResources].IsEmpty() && Resources[EDisplayClusterViewportResource::OutputPreviewTargetableResources][0].IsValid())
-	{
-		const FDisplayClusterRenderFrameSettings* RenderFrameSettings = GetRenderFrameSettings_RenderThread();
-		check(RenderFrameSettings && RenderFrameSettings->RenderMode == EDisplayClusterRenderFrameMode::PreviewInScene);
-
-		// We should synchronize thread for preview read
-		FScopeLock Lock(&PreviewPixelsCSGuard);
-
-		if (PreviewPixels.IsValid() == false)
-		{
-			// Read pixels from this texture
-			if (FRHITexture2D* Texture2D = Resources[EDisplayClusterViewportResource::OutputPreviewTargetableResources][0]->GetViewportResourceRHI())
-			{
-				
-				// Clear deferred read flag
-				bPreviewReadPixels = false;
-
-				TSharedPtr<FDisplayClusterViewportReadPixelsData, ESPMode::ThreadSafe> ReadData = MakeShared<FDisplayClusterViewportReadPixelsData, ESPMode::ThreadSafe>();
-
-				RHICmdList.ReadSurfaceData(
-					Texture2D,
-					FIntRect(FIntPoint(0, 0), Texture2D->GetSizeXY()),
-					ReadData->Pixels,
-					FReadSurfaceDataFlags(RCM_UNorm, CubeFace_MAX)
-				);
-
-				ReadData->Size = Texture2D->GetSizeXY();
-
-				// Expose result:
-				PreviewPixels = ReadData;
-			}
-		}
-	}
-#endif
-}
-
-#if WITH_EDITOR
-bool FDisplayClusterViewportProxy::GetPreviewPixels_GameThread(TSharedPtr<FDisplayClusterViewportReadPixelsData, ESPMode::ThreadSafe>& OutPixelsData) const
-{
-	check(IsInGameThread());
-
-	// We should synchronize thread for preview read
-	FScopeLock Lock(&PreviewPixelsCSGuard);
-
-	if (PreviewPixels.IsValid())
-	{
-		OutPixelsData = PreviewPixels;
-		PreviewPixels.Reset();
-		return true;
-	}
-
-	return false;
-}
-#endif
-
 bool FDisplayClusterViewportProxy::GetResourcesWithRects_RenderThread(const EDisplayClusterViewportResourceType InExtResourceType, TArray<FRHITexture2D*>& OutResources, TArray<FIntRect>& OutResourceRects) const
 {
 	return ImplGetResourcesWithRects_RenderThread(InExtResourceType, OutResources, OutResourceRects, false);
@@ -1069,10 +992,7 @@ bool FDisplayClusterViewportProxy::ShouldUsePostProcessPassAfterSSRInput() const
 {	
 	if (ShouldUseAlphaChannel_RenderThread())
 	{
-		if (const FDisplayClusterRenderFrameSettings* RenderFrameSettings = GetRenderFrameSettings_RenderThread())
-		{
-			return RenderFrameSettings->AlphaChannelCaptureMode == EDisplayClusterRenderFrameAlphaChannelCaptureMode::ThroughTonemapper;
-		}
+		return ConfigurationProxy->GetRenderFrameSettings().AlphaChannelCaptureMode == EDisplayClusterRenderFrameAlphaChannelCaptureMode::ThroughTonemapper;
 	}
 
 	return false;
@@ -1082,10 +1002,7 @@ bool FDisplayClusterViewportProxy::ShouldUsePostProcessPassAfterFXAA() const
 {
 	if (ShouldUseAlphaChannel_RenderThread())
 	{
-		if (const FDisplayClusterRenderFrameSettings* RenderFrameSettings = GetRenderFrameSettings_RenderThread())
-		{
-			return RenderFrameSettings->AlphaChannelCaptureMode == EDisplayClusterRenderFrameAlphaChannelCaptureMode::ThroughTonemapper;
-		}
+		return ConfigurationProxy->GetRenderFrameSettings().AlphaChannelCaptureMode == EDisplayClusterRenderFrameAlphaChannelCaptureMode::ThroughTonemapper;
 	}
 
 	return false;
@@ -1106,23 +1023,20 @@ void FDisplayClusterViewportProxy::OnResolvedSceneColor_RenderThread(FRDGBuilder
 	const uint32 InContextNum = InProxyContext.ContextNum;
 	if (ShouldUseAlphaChannel_RenderThread())
 	{
-		if (const FDisplayClusterRenderFrameSettings* RenderFrameSettings = GetRenderFrameSettings_RenderThread())
+		switch (ConfigurationProxy->GetRenderFrameSettings().AlphaChannelCaptureMode)
 		{
-			switch (RenderFrameSettings->AlphaChannelCaptureMode)
-			{
-			case EDisplayClusterRenderFrameAlphaChannelCaptureMode::FXAA:
-			case EDisplayClusterRenderFrameAlphaChannelCaptureMode::Copy:
-			case EDisplayClusterRenderFrameAlphaChannelCaptureMode::CopyAA:
-			{
-				const FIntRect SrcRect = GetFinalContextRect(EDisplayClusterViewportResourceType::InternalRenderTargetResource, Contexts[InContextNum].RenderTargetRect);
-				// Copy alpha channel from 'SceneTextures.Color.Resolve' to 'InputShaderResource'
-				CopyResource_RenderThread(GraphBuilder, EDisplayClusterTextureCopyMode::Alpha, InContextNum, SceneTextures.Color.Resolve, SrcRect, EDisplayClusterViewportResourceType::InputShaderResource);
-			}
-			break;
+		case EDisplayClusterRenderFrameAlphaChannelCaptureMode::FXAA:
+		case EDisplayClusterRenderFrameAlphaChannelCaptureMode::Copy:
+		case EDisplayClusterRenderFrameAlphaChannelCaptureMode::CopyAA:
+		{
+			const FIntRect SrcRect = GetFinalContextRect(EDisplayClusterViewportResourceType::InternalRenderTargetResource, Contexts[InContextNum].RenderTargetRect);
+			// Copy alpha channel from 'SceneTextures.Color.Resolve' to 'InputShaderResource'
+			CopyResource_RenderThread(GraphBuilder, EDisplayClusterTextureCopyMode::Alpha, InContextNum, SceneTextures.Color.Resolve, SrcRect, EDisplayClusterViewportResourceType::InputShaderResource);
+		}
+		break;
 
-			default:
-				break;
-			}
+		default:
+			break;
 		}
 	}
 }
@@ -1197,10 +1111,9 @@ void FDisplayClusterViewportProxy::OnPostRenderViewFamily_RenderThread(FRDGBuild
 		}
 	}
 
-	const FDisplayClusterRenderFrameSettings* RenderFrameSettings = GetRenderFrameSettings_RenderThread();
-	if (ShouldUseAlphaChannel_RenderThread() && RenderFrameSettings)
+	if (ShouldUseAlphaChannel_RenderThread())
 	{
-		switch (RenderFrameSettings->AlphaChannelCaptureMode)
+		switch (ConfigurationProxy->GetRenderFrameSettings().AlphaChannelCaptureMode)
 		{
 		case EDisplayClusterRenderFrameAlphaChannelCaptureMode::FXAA:
 		{
diff --git a/Engine/Plugins/Runtime/nDisplay/Source/DisplayCluster/Private/Render/Viewport/DisplayClusterViewportProxy.h b/Engine/Plugins/Runtime/nDisplay/Source/DisplayCluster/Private/Render/Viewport/DisplayClusterViewportProxy.h
--- a/Engine/Plugins/Runtime/nDisplay/Source/DisplayCluster/Private/Render/Viewport/DisplayClusterViewportProxy.h
+++ b/Engine/Plugins/Runtime/nDisplay/Source/DisplayCluster/Private/Render/Viewport/DisplayClusterViewportProxy.h
@@ -3,24 +3,19 @@
 #pragma once
 
 #include "CoreMinimal.h"
+#include "Render/Viewport/IDisplayClusterViewportProxy.h"
+#include "Render/Viewport/DisplayClusterViewportResources.h"
+#include "Render/Viewport/Containers/DisplayClusterViewport_OverscanRuntimeSettings.h"
+#include "Render/Viewport/Configuration/DisplayClusterViewportConfigurationProxy.h"
 
 #include "EngineUtils.h"
 #include "ScreenRendering.h"
 #include "SceneView.h"
 #include "Templates/SharedPointer.h"
 
-#include "Render/Viewport/IDisplayClusterViewportProxy.h"
-#include "Render/Viewport/DisplayClusterViewportResources.h"
-#include "Render/Viewport/Containers/DisplayClusterViewport_OverscanRuntimeSettings.h"
-
-class FDisplayClusterViewportProxyData;
-class IDisplayClusterViewportManagerProxy;
-class FDisplayClusterViewportManagerProxy;
 class FDisplayClusterViewport;
-class IDisplayClusterShaders;
 class IDisplayClusterProjectionPolicy;
 class IDisplayClusterRender_MeshComponent;
-class FDisplayClusterViewportReadPixelsData;
 class FDisplayClusterViewport_OpenColorIO;
 
 class FRDGBuilder;
@@ -61,13 +56,11 @@ class FDisplayClusterViewportProxy
 	, public TSharedFromThis<FDisplayClusterViewportProxy, ESPMode::ThreadSafe>
 {
 public:
-	FDisplayClusterViewportProxy(const FDisplayClusterViewport& RenderViewport);
+	FDisplayClusterViewportProxy(const TSharedRef<FDisplayClusterViewportConfiguration, ESPMode::ThreadSafe>& InConfiguration, const FString& InViewportId, const TSharedPtr<IDisplayClusterProjectionPolicy, ESPMode::ThreadSafe>& InProjectionPolicy);
 	virtual ~FDisplayClusterViewportProxy();
 
 public:
-	///////////////////////////////
-	// IDisplayClusterViewportProxy
-	///////////////////////////////
+	//~~ BEGIN IDisplayClusterViewportProxy
 	virtual TSharedPtr<IDisplayClusterViewportProxy, ESPMode::ThreadSafe> ToSharedPtr() override
 	{
 		return AsShared();
@@ -78,7 +71,10 @@ public:
 		return AsShared();
 	}
 
-	virtual EDisplayClusterRenderFrameMode GetRenderMode() const override;
+	virtual const IDisplayClusterViewportConfigurationProxy& GetConfigurationProxy() const override
+	{
+		return *ConfigurationProxy;
+	}
 
 	virtual FString GetId() const override
 	{
@@ -143,14 +139,9 @@ public:
 	virtual bool ResolveResources_RenderThread(FRHICommandListImmediate& RHICmdList, const EDisplayClusterViewportResourceType InputResourceType, const EDisplayClusterViewportResourceType OutputResourceType, const int32 InContextNum = INDEX_NONE) const override;
 	virtual bool ResolveResources_RenderThread(FRHICommandListImmediate& RHICmdList, IDisplayClusterViewportProxy* InputResourceViewportProxy, const EDisplayClusterViewportResourceType InputResourceType, const EDisplayClusterViewportResourceType OutputResourceType, const int32 InContextNum = INDEX_NONE) const override;
 
-	virtual const class IDisplayClusterViewportManagerProxy* GetViewportManagerProxy_RenderThread() const override;
-
-	virtual const FDisplayClusterRenderFrameSettings* GetRenderFrameSettings_RenderThread() const override;
-
-	///////////////////////////////
-	// ~IDisplayClusterViewportProxy
-	///////////////////////////////
+	//~~ END IDisplayClusterViewportProxy
 
+public:
 	/** Get valid resource type
 	 * 
 	 * @param InResourceType - the requested resource type from the entire namespace
@@ -250,10 +241,6 @@ public:
 	 */
 	void PostResolveViewport_RenderThread(FRHICommandListImmediate& RHICmdList) const;
 
-#if WITH_EDITOR
-	bool GetPreviewPixels_GameThread(TSharedPtr<FDisplayClusterViewportReadPixelsData, ESPMode::ThreadSafe>& OutPixelsData) const;
-#endif
-
 	inline bool FindContext_RenderThread(const int32 ViewIndex, uint32* OutContextNum)
 	{
 		check(IsInRenderingThread());
@@ -278,12 +265,18 @@ public:
 	/* Returns true if the warp can be applied to this viewport. */
 	bool ShouldApplyWarpBlend_RenderThread() const;
 
+	const FDisplayClusterViewportResources& GetResources_RenderThread() const
+	{
+		return Resources;
+	}
+
+	void UpdateViewportProxyData_RenderThread(const class FDisplayClusterViewportProxyData& InViewportProxyData);
+
 private:
 	bool ImplGetResourcesWithRects_RenderThread(const EDisplayClusterViewportResourceType InResourceType, TArray<FRHITexture2D*>& OutResources, TArray<FIntRect>& OutResourceRects, const int32 InRecursionDepth) const;
 	bool ImplGetResources_RenderThread(const EDisplayClusterViewportResourceType InResourceType, TArray<FRHITexture2D*>& OutResources, const int32 InRecursionDepth) const;
 
 	void ImplViewportRemap_RenderThread(FRHICommandListImmediate& RHICmdList) const;
-	void ImplPreviewReadPixels_RenderThread(FRHICommandListImmediate& RHICmdList) const;
 
 	bool ImplResolveResources_RenderThread(FRHICommandListImmediate& RHICmdList, FDisplayClusterViewportProxy const* SourceProxy, const EDisplayClusterViewportResourceType InputResourceType, const EDisplayClusterViewportResourceType OutputResourceType, const int32 InContextNum = INDEX_NONE) const;
 
@@ -342,12 +335,9 @@ private:
 	/** Returns the OCIO rendering type for the given viewport. */
 	EDisplayClusterViewportOpenColorIOMode GetOpenColorIOMode() const;
 
-	FDisplayClusterViewportManagerProxy* GetViewportManagerProxyImpl_RenderThread() const;
-	TSharedPtr<FDisplayClusterViewportManagerProxy, ESPMode::ThreadSafe> GetViewportManagerProxyRefImpl_RenderThread() const;
-
-protected:
-	friend FDisplayClusterViewportProxyData;
-	friend FDisplayClusterViewportManagerProxy;
+public:
+	// Configuration for proxy
+	const TSharedRef<FDisplayClusterViewportConfigurationProxy, ESPMode::ThreadSafe> ConfigurationProxy;
 
 	/** Unique viewport name. */
 	const FString ViewportId;
@@ -355,11 +345,13 @@ protected:
 	/** Cluster node name. */
 	const FString ClusterNodeId;
 
+protected:
 	/** OpenColorIO nDisplay interface ref. */
 	TSharedPtr<FDisplayClusterViewport_OpenColorIO, ESPMode::ThreadSafe> OpenColorIO;
 
 	// Viewport render params
-	mutable FDisplayClusterViewport_RenderSettings       RenderSettings;
+	mutable FDisplayClusterViewport_RenderSettings RenderSettings;
+
 	FDisplayClusterViewport_RenderSettingsICVFX  RenderSettingsICVFX;
 	FDisplayClusterViewport_PostRenderSettings   PostRenderSettings;
 
@@ -379,12 +371,4 @@ protected:
 
 	// Used ViewStates
 	TArray<TSharedPtr<FSceneViewStateReference, ESPMode::ThreadSafe>> ViewStates;
-	
-#if WITH_EDITOR
-	mutable bool bPreviewReadPixels = false;
-	mutable FCriticalSection PreviewPixelsCSGuard;
-	mutable TSharedPtr<FDisplayClusterViewportReadPixelsData, ESPMode::ThreadSafe> PreviewPixels;
-#endif
-
-	TWeakPtr<FDisplayClusterViewportManagerProxy, ESPMode::ThreadSafe> ViewportManagerProxyWeakRef;
 };
diff --git a/Engine/Plugins/Runtime/nDisplay/Source/DisplayCluster/Private/Render/Viewport/DisplayClusterViewportStereoscopicPass.h b/Engine/Plugins/Runtime/nDisplay/Source/DisplayCluster/Private/Render/Viewport/DisplayClusterViewportStereoscopicPass.h
--- a/Engine/Plugins/Runtime/nDisplay/Source/DisplayCluster/Private/Render/Viewport/DisplayClusterViewportStereoscopicPass.h
+++ b/Engine/Plugins/Runtime/nDisplay/Source/DisplayCluster/Private/Render/Viewport/DisplayClusterViewportStereoscopicPass.h
@@ -14,12 +14,11 @@ struct FDisplayClusterViewportStereoscopicPass
 
 	inline static EStereoscopicPass EncodeStereoscopicPass(const uint32 ContextNum, const uint32 ContextAmount, const FDisplayClusterRenderFrameSettings& InFrameSettings)
 	{
-#if WITH_EDITOR
-		if (InFrameSettings.bIsRenderingInEditor)
+		// Use special rendering mode (DCRA preview)
+		if (!InFrameSettings.bUseDisplayClusterRenderDevice)
 		{
 			return EStereoscopicPass::eSSP_FULL;
 		}
-#endif
 
 		// 'bEnableStereoscopicRenderingOptimization'
 		// When stereoscopic rendering optimization is enabled, one image family and one RTT for both eyes is used. 
diff --git a/Engine/Plugins/Runtime/nDisplay/Source/DisplayCluster/Private/Render/Viewport/DisplayClusterViewport_EditorPreview.cpp b/Engine/Plugins/Runtime/nDisplay/Source/DisplayCluster/Private/Render/Viewport/DisplayClusterViewport_EditorPreview.cpp
--- a/Engine/Plugins/Runtime/nDisplay/Source/DisplayCluster/Private/Render/Viewport/DisplayClusterViewport_EditorPreview.cpp
+++ b/Engine/Plugins/Runtime/nDisplay/Source/DisplayCluster/Private/Render/Viewport/DisplayClusterViewport_EditorPreview.cpp
@@ -63,7 +63,7 @@ void FDisplayClusterViewport::CleanupViewState()
 
 FSceneViewStateInterface* FDisplayClusterViewport::GetViewState(uint32 ViewIndex)
 {
-	if (GDisplayClusterPreviewEnableViewState == 0 || (GDisplayClusterPreviewEnableConfiguratorViewState == 0 && IsCurrentWorldHasAnyType(EWorldType::EditorPreview)))
+	if (GDisplayClusterPreviewEnableViewState == 0 || (GDisplayClusterPreviewEnableConfiguratorViewState == 0 && Configuration->IsCurrentWorldHasAnyType(EWorldType::EditorPreview)))
 	{
 		// Disable ViewState
 		ViewStates.Empty();
@@ -84,7 +84,7 @@ FSceneViewStateInterface* FDisplayClusterViewport::GetViewState(uint32 ViewIndex
 
 	if (ViewStates[ViewIndex]->GetReference() == NULL)
 	{
-		const UWorld* CurrentWorld = GetCurrentWorld();
+		const UWorld* CurrentWorld = Configuration->GetCurrentWorld();
 		const ERHIFeatureLevel::Type FeatureLevel = CurrentWorld ? CurrentWorld->GetFeatureLevel() : GMaxRHIFeatureLevel;
 
 		ViewStates[ViewIndex]->Allocate(FeatureLevel);
@@ -145,12 +145,9 @@ FSceneView* FDisplayClusterViewport::ImplCalcScenePreview(FSceneViewFamilyContex
 
 		ViewInitOptions.BackgroundColor = FLinearColor::Black;
 
-		if (const FDisplayClusterRenderFrameSettings* RenderFrameSettings = GetRenderFrameSettings())
+		if (Configuration->GetRenderFrameSettings().IsPostProcessDisabled())
 		{
-			if (RenderFrameSettings->bPreviewEnablePostProcess == false)
-			{
-				ViewInitOptions.OverlayColor = FLinearColor::Black;
-			}
+			ViewInitOptions.OverlayColor = FLinearColor::Black;
 		}
 
 		ViewInitOptions.bSceneCaptureUsesRayTracing = false;
@@ -244,24 +241,10 @@ FMatrix FDisplayClusterViewport::ImplPreview_GetStereoProjectionMatrix(const uin
 	check(IsInGameThread());
 
 	FMatrix PrjMatrix = FMatrix::Identity;
-
-	if (IsSceneOpened())
+	if (Configuration->IsSceneOpened() && GetProjectionMatrix(InContextNum, PrjMatrix) == false)
 	{
-		if (GetProjectionMatrix(InContextNum, PrjMatrix) == false)
-		{
-			UE_LOG(LogDisplayClusterViewport, Verbose, TEXT("Got invalid projection matrix: Viewport %s, ViewIdx: %d"), *GetId(), InContextNum);
-		}
+		UE_LOG(LogDisplayClusterViewport, Verbose, TEXT("Got invalid projection matrix: Viewport %s, ViewIdx: %d"), *GetId(), InContextNum);
 	}
 
 	return PrjMatrix;
 }
-
-#if WITH_EDITOR
-
-bool FDisplayClusterViewport::GetPreviewPixels(TSharedPtr<FDisplayClusterViewportReadPixelsData, ESPMode::ThreadSafe>& OutPixelsData) const
-{
-	check(IsInGameThread());
-
-	return (ViewportProxy != nullptr) && ViewportProxy->GetPreviewPixels_GameThread(OutPixelsData);
-}
-#endif
diff --git a/Engine/Plugins/Runtime/nDisplay/Source/DisplayCluster/Private/Render/Viewport/DisplayClusterViewport_Math.cpp b/Engine/Plugins/Runtime/nDisplay/Source/DisplayCluster/Private/Render/Viewport/DisplayClusterViewport_Math.cpp
--- a/Engine/Plugins/Runtime/nDisplay/Source/DisplayCluster/Private/Render/Viewport/DisplayClusterViewport_Math.cpp
+++ b/Engine/Plugins/Runtime/nDisplay/Source/DisplayCluster/Private/Render/Viewport/DisplayClusterViewport_Math.cpp
@@ -2,27 +2,25 @@
 
 #include "Render/Viewport/DisplayClusterViewport.h"
 #include "Render/Viewport/DisplayClusterViewportManager.h"
+#include "Render/Viewport/DisplayClusterViewportStereoscopicPass.h"
 
 #include "Render/Viewport/Configuration/DisplayClusterViewportConfiguration.h"
+#include "Render/Viewport/Containers/DisplayClusterViewport_PostRenderSettings.h"
 
 #include "Render/Projection/IDisplayClusterProjectionPolicy.h"
 
-#include "Render/Viewport/DisplayClusterViewportStereoscopicPass.h"
 #include "Render/Viewport/RenderFrame/DisplayClusterRenderFrame.h"
 #include "Render/Viewport/RenderFrame/DisplayClusterRenderFrameSettings.h"
 #include "Render/Viewport/RenderTarget/DisplayClusterRenderTargetResource.h"
 
-#include "Render/Viewport/Containers/DisplayClusterViewport_PostRenderSettings.h"
+#include "DisplayClusterSceneViewExtensions.h"
 
 #include "Components/DisplayClusterCameraComponent.h"
+#include "Misc/DisplayClusterLog.h"
 
 #include "EngineUtils.h"
 #include "SceneView.h"
 
-#include "DisplayClusterSceneViewExtensions.h"
-
-#include "Misc/DisplayClusterLog.h"
-
 namespace UE::DisplayCluster::Viewport::Math
 {
 	/**
@@ -69,10 +67,10 @@ FVector2D FDisplayClusterViewport::GetClippingPlanes() const
 bool FDisplayClusterViewport::GetViewPointCameraEye(const uint32 InContextNum, FVector& OutViewLocation, FRotator& OutViewRotation, FVector& OutViewOffset)
 {
 	// Here we use the ViewPoint component as the eye position
-	if (UDisplayClusterCameraComponent* ViewPoint = GetViewPointCameraComponent())
+	if (UDisplayClusterCameraComponent* SceneCameraComponent = GetViewPointCameraComponent(EDisplayClusterRootActorType::Scene))
 	{
-		OutViewLocation = ViewPoint->GetComponentLocation();
-		OutViewRotation = ViewPoint->GetComponentRotation();
+		OutViewLocation = SceneCameraComponent->GetComponentLocation();
+		OutViewRotation = SceneCameraComponent->GetComponentRotation();
 
 		// Calculate stereo ViewOffset:
 		OutViewOffset = ImplGetViewOffset(GetStereoEyeOffsetDistance(InContextNum), OutViewLocation, OutViewRotation);
diff --git a/Engine/Plugins/Runtime/nDisplay/Source/DisplayCluster/Private/Render/Viewport/DisplayClusterViewport_OpenColorIO.h b/Engine/Plugins/Runtime/nDisplay/Source/DisplayCluster/Private/Render/Viewport/DisplayClusterViewport_OpenColorIO.h
--- a/Engine/Plugins/Runtime/nDisplay/Source/DisplayCluster/Private/Render/Viewport/DisplayClusterViewport_OpenColorIO.h
+++ b/Engine/Plugins/Runtime/nDisplay/Source/DisplayCluster/Private/Render/Viewport/DisplayClusterViewport_OpenColorIO.h
@@ -4,14 +4,15 @@
 
 #include "CoreMinimal.h"
 
-#include "OpenColorIOColorSpace.h"
-#include "OpenColorIORendering.h"
-
 #include "RHI.h"
 #include "RHICommandList.h"
 #include "RHIResources.h"
+#include "TextureResource.h"
 #include "Templates/SharedPointer.h"
 
+#include "OpenColorIOColorSpace.h"
+#include "OpenColorIORendering.h"
+
 class FSceneView;
 class FSceneViewFamily;
 struct FScreenPassTexture;
diff --git a/Engine/Plugins/Runtime/nDisplay/Source/DisplayCluster/Private/Render/Viewport/DisplayClusterViewport_ViewPoint.cpp b/Engine/Plugins/Runtime/nDisplay/Source/DisplayCluster/Private/Render/Viewport/DisplayClusterViewport_ViewPoint.cpp
--- a/Engine/Plugins/Runtime/nDisplay/Source/DisplayCluster/Private/Render/Viewport/DisplayClusterViewport_ViewPoint.cpp
+++ b/Engine/Plugins/Runtime/nDisplay/Source/DisplayCluster/Private/Render/Viewport/DisplayClusterViewport_ViewPoint.cpp
@@ -108,12 +108,12 @@ bool IDisplayClusterViewport::GetPlayerCameraView(UWorld* InWorld, const bool bU
 ///////////////////////////////////////////////////////////////////////////////////////
 //          FDisplayClusterViewport
 ///////////////////////////////////////////////////////////////////////////////////////
-UDisplayClusterCameraComponent* FDisplayClusterViewport::GetViewPointCameraComponent() const
+UDisplayClusterCameraComponent* FDisplayClusterViewport::GetViewPointCameraComponent(const EDisplayClusterRootActorType InRootActorType) const
 {
-	ADisplayClusterRootActor* RootActor = GetRootActor();
+	ADisplayClusterRootActor* RootActor = Configuration->GetRootActor(InRootActorType);
 	if (!RootActor)
 	{
-		UE_LOG(LogDisplayClusterViewport, Warning, TEXT("Viewport '%s' has no root actor found in game manager"), *GetId());
+		UE_LOG(LogDisplayClusterViewport, Warning, TEXT("Viewport '%s' has no root actor found"), *GetId());
 
 		return nullptr;
 	}
@@ -124,8 +124,7 @@ UDisplayClusterCameraComponent* FDisplayClusterViewport::GetViewPointCameraCompo
 		return nullptr;
 	}
 
-	FDisplayClusterViewportManager* ViewportManager = GetViewportManagerImpl();
-	if (!ViewportManager || !ViewportManager->IsSceneOpened())
+	if (!Configuration->IsSceneOpened())
 	{
 		return nullptr;
 	}
@@ -152,20 +151,19 @@ UDisplayClusterCameraComponent* FDisplayClusterViewport::GetViewPointCameraCompo
 
 bool FDisplayClusterViewport::SetupViewPoint(FMinimalViewInfo& InOutViewInfo)
 {
-	if (UDisplayClusterCameraComponent* ViewCamera = GetViewPointCameraComponent())
+	if (UDisplayClusterCameraComponent* SceneCameraComponent = GetViewPointCameraComponent(EDisplayClusterRootActorType::Scene))
 	{
 		// Get ViewPoint from DCRA component
-		ViewCamera->GetDesiredView(InOutViewInfo, &CustomNearClippingPlane);
+		SceneCameraComponent->GetDesiredView(*Configuration, InOutViewInfo, &CustomNearClippingPlane);
 
 		// The projection policy can override these ViewPoint data.
 		if (ProjectionPolicy.IsValid())
 		{
-			float DeltaTime = 0.0f;
-			if (ADisplayClusterRootActor* RootActor = GetRootActor())
+			if (ADisplayClusterRootActor* SceneRootActor = Configuration->GetRootActor(EDisplayClusterRootActorType::Scene))
 			{
-				DeltaTime = RootActor->GetWorldDeltaSeconds();
+				const float DeltaTime = SceneRootActor->GetWorldDeltaSeconds();
+				ProjectionPolicy->SetupProjectionViewPoint(this, DeltaTime, InOutViewInfo, &CustomNearClippingPlane);
 			}
-			ProjectionPolicy->SetupProjectionViewPoint(this, DeltaTime, InOutViewInfo, &CustomNearClippingPlane);
 		}
 
 		return true;
@@ -178,10 +176,10 @@ float FDisplayClusterViewport::GetStereoEyeOffsetDistance(const uint32 InContext
 {
 	float StereoEyeOffsetDistance = 0.f;
 
-	if (UDisplayClusterCameraComponent* ViewCamera = GetViewPointCameraComponent())
+	if (UDisplayClusterCameraComponent* ConfigurationCameraComponent = GetViewPointCameraComponent(EDisplayClusterRootActorType::Configuration))
 	{
 		// Calculate eye offset considering the world scale
-		const float CfgEyeDist = ViewCamera->GetInterpupillaryDistance();
+		const float CfgEyeDist = ConfigurationCameraComponent->GetInterpupillaryDistance();
 		const float EyeOffset = CfgEyeDist / 2.f;
 		const float EyeOffsetValues[] = { -EyeOffset, 0.f, EyeOffset };
 
@@ -198,7 +196,7 @@ float FDisplayClusterViewport::GetStereoEyeOffsetDistance(const uint32 InContext
 			// * Force left (-1) ==> 0 left eye
 			// * Force right (1) ==> 2 right eye
 			// * Default (0) ==> 1 mono
-			const EDisplayClusterEyeStereoOffset CfgEyeOffset = ViewCamera->GetStereoOffset();
+			const EDisplayClusterEyeStereoOffset CfgEyeOffset = ConfigurationCameraComponent->GetStereoOffset();
 			const int32 EyeOffsetIdx =
 				(CfgEyeOffset == EDisplayClusterEyeStereoOffset::None ? 0 :
 					(CfgEyeOffset == EDisplayClusterEyeStereoOffset::Left ? -1 : 1));
@@ -213,7 +211,7 @@ float FDisplayClusterViewport::GetStereoEyeOffsetDistance(const uint32 InContext
 			PassOffset = EyeOffsetValues[EyeType];
 
 			// Apply eye swap
-			const bool  CfgEyeSwap = ViewCamera->GetSwapEyes();
+			const bool  CfgEyeSwap = ConfigurationCameraComponent->GetSwapEyes();
 			StereoEyeOffsetDistance = (CfgEyeSwap ? -PassOffset : PassOffset);
 		}
 	}
diff --git a/Engine/Plugins/Runtime/nDisplay/Source/DisplayCluster/Private/Render/Viewport/DisplayClusterViewport_VisibilitySettings.h b/Engine/Plugins/Runtime/nDisplay/Source/DisplayCluster/Private/Render/Viewport/DisplayClusterViewport_VisibilitySettings.h
--- a/Engine/Plugins/Runtime/nDisplay/Source/DisplayCluster/Private/Render/Viewport/DisplayClusterViewport_VisibilitySettings.h
+++ b/Engine/Plugins/Runtime/nDisplay/Source/DisplayCluster/Private/Render/Viewport/DisplayClusterViewport_VisibilitySettings.h
@@ -26,7 +26,7 @@ public:
 
 public:
 	// Reset actor layers visibility rules
-	void ResetConfiguration()
+	void BeginUpdateSettings()
 	{
 		LayersMode = EDisplayClusterViewport_VisibilityMode::None;
 
@@ -47,7 +47,7 @@ public:
 		return false;
 	}
 
-	void UpdateConfiguration(EDisplayClusterViewport_VisibilityMode InMode, const TArray<FName>& InActorLayers, const TSet<FPrimitiveComponentId>& InAdditionalComponentsList)
+	void UpdateVisibilitySettings(EDisplayClusterViewport_VisibilityMode InMode, const TArray<FName>& InActorLayers, const TSet<FPrimitiveComponentId>& InAdditionalComponentsList)
 	{
 
 		LayersMode = InMode;
diff --git a/Engine/Plugins/Runtime/nDisplay/Source/DisplayCluster/Private/Render/Viewport/LightCard/DisplayClusterViewportLightCardManager.cpp b/Engine/Plugins/Runtime/nDisplay/Source/DisplayCluster/Private/Render/Viewport/LightCard/DisplayClusterViewportLightCardManager.cpp
--- a/Engine/Plugins/Runtime/nDisplay/Source/DisplayCluster/Private/Render/Viewport/LightCard/DisplayClusterViewportLightCardManager.cpp
+++ b/Engine/Plugins/Runtime/nDisplay/Source/DisplayCluster/Private/Render/Viewport/LightCard/DisplayClusterViewportLightCardManager.cpp
@@ -3,6 +3,9 @@
 #include "DisplayClusterViewportLightCardManager.h"
 #include "DisplayClusterViewportLightCardManagerProxy.h"
 #include "DisplayClusterViewportLightCardResource.h"
+
+#include "Render/Viewport/Configuration/DisplayClusterViewportConfiguration.h"
+
 #include "ShaderParameters/DisplayClusterShaderParameters_UVLightCards.h"
 
 #include "DisplayClusterLightCardActor.h"
@@ -25,17 +28,14 @@ static TAutoConsoleVariable<int32> CVarUVLightCardTextureSize(
 ///////////////////////////////////////////////////////////////////////////////////////////////
 // FDisplayClusterViewportLightCardManager
 ///////////////////////////////////////////////////////////////////////////////////////////////
-FDisplayClusterViewportLightCardManager::FDisplayClusterViewportLightCardManager(FDisplayClusterViewportManager& InViewportManager)
-	: ViewportManagerWeakPtr(InViewportManager.AsShared())
-{
-	LightCardManagerProxy = MakeShared<FDisplayClusterViewportLightCardManagerProxy, ESPMode::ThreadSafe>();
-}
+FDisplayClusterViewportLightCardManager::FDisplayClusterViewportLightCardManager(const TSharedRef<const FDisplayClusterViewportConfiguration, ESPMode::ThreadSafe>& InConfiguration)
+	: Configuration(InConfiguration)
+	, LightCardManagerProxy(MakeShared<FDisplayClusterViewportLightCardManagerProxy, ESPMode::ThreadSafe>())
+{ }
 
 FDisplayClusterViewportLightCardManager::~FDisplayClusterViewportLightCardManager()
 {
 	Release();
-
-	LightCardManagerProxy.Reset();
 }
 
 void FDisplayClusterViewportLightCardManager::Release()
@@ -49,11 +49,6 @@ void FDisplayClusterViewportLightCardManager::Release()
 }
 
 ///////////////////////////////////////////////////////////////////////////////////////////////
-void FDisplayClusterViewportLightCardManager::UpdateConfiguration()
-{
-
-}
-
 void FDisplayClusterViewportLightCardManager::HandleStartScene()
 {
 }
@@ -89,19 +84,13 @@ void FDisplayClusterViewportLightCardManager::UpdateUVLightCardData()
 {
 	ReleaseUVLightCardData();
 
-	FDisplayClusterViewportManager* ViewportManager = GetViewportManager();
-	if (!ViewportManager)
-	{
-		return;
-	}
-
 	/** The list of UV light card actors that are referenced by the root actor */
 	TArray<ADisplayClusterLightCardActor*> UVLightCardActors;
 
-	if (ADisplayClusterRootActor* RootActorPtr = ViewportManager->GetRootActor())
+	if (ADisplayClusterRootActor* SceneRootActorPtr = Configuration->GetRootActor(EDisplayClusterRootActorType::Scene))
 	{
 		TSet<ADisplayClusterLightCardActor*> LightCards;
-		UDisplayClusterBlueprintLib::FindLightCardsForRootActor(RootActorPtr, LightCards);
+		UDisplayClusterBlueprintLib::FindLightCardsForRootActor(SceneRootActorPtr, LightCards);
 
 		for (ADisplayClusterLightCardActor* LightCard : LightCards)
 		{
@@ -171,9 +160,9 @@ void FDisplayClusterViewportLightCardManager::UpdateUVLightCardResource()
 void FDisplayClusterViewportLightCardManager::RenderUVLightCard()
 {
 	// Render UV LightCard:
-	FDisplayClusterViewportManager* ViewportManager = GetViewportManager();
-	UWorld* World = ViewportManager ? ViewportManager->GetCurrentWorld() : nullptr;
-	if (IsUVLightCardEnabled() && World && LightCardManagerProxy.IsValid())
+	FDisplayClusterViewportManager* ViewportManager = Configuration->GetViewportManagerImpl();
+	UWorld* CurrentWorld = Configuration->GetCurrentWorld();
+	if (IsUVLightCardEnabled() && CurrentWorld && ViewportManager)
 	{
 		UpdateUVLightCardResource();
 
@@ -202,7 +191,7 @@ void FDisplayClusterViewportLightCardManager::RenderUVLightCard()
 				}
 			}
 
-			LightCardManagerProxy->RenderUVLightCard(World->Scene, UVLightCardParameters);
+			LightCardManagerProxy->RenderUVLightCard(CurrentWorld->Scene, UVLightCardParameters);
 
 			for (UPrimitiveComponent* LoadedComponent : ComponentsToUnload)
 			{
@@ -215,7 +204,6 @@ void FDisplayClusterViewportLightCardManager::RenderUVLightCard()
 	{
 		ReleaseUVLightCardResource();
 	}
-
 }
 
 ///////////////////////////////////////////////////////////////////////////////////////////////
diff --git a/Engine/Plugins/Runtime/nDisplay/Source/DisplayCluster/Private/Render/Viewport/LightCard/DisplayClusterViewportLightCardManager.h b/Engine/Plugins/Runtime/nDisplay/Source/DisplayCluster/Private/Render/Viewport/LightCard/DisplayClusterViewportLightCardManager.h
--- a/Engine/Plugins/Runtime/nDisplay/Source/DisplayCluster/Private/Render/Viewport/LightCard/DisplayClusterViewportLightCardManager.h
+++ b/Engine/Plugins/Runtime/nDisplay/Source/DisplayCluster/Private/Render/Viewport/LightCard/DisplayClusterViewportLightCardManager.h
@@ -9,8 +9,8 @@
 
 #include "UObject/GCObject.h"
 
+class FDisplayClusterViewportConfiguration;
 class ADisplayClusterLightCardActor;
-class FDisplayClusterViewportManager;
 class UWorld;
 
 /**
@@ -20,7 +20,7 @@ class FDisplayClusterViewportLightCardManager
 	: public FGCObject
 {
 public:
-	FDisplayClusterViewportLightCardManager(FDisplayClusterViewportManager& InViewportManager);
+	FDisplayClusterViewportLightCardManager(const TSharedRef<const FDisplayClusterViewportConfiguration, ESPMode::ThreadSafe>& InConfiguration);
 	virtual ~FDisplayClusterViewportLightCardManager();
 
 	void Release();
@@ -37,14 +37,7 @@ public:
 	/** Get UV LightCard texture size. */
 	FIntPoint GetUVLightCardResourceSize() const;
 
-	/** Get proxy object. */
-	TSharedPtr<FDisplayClusterViewportLightCardManagerProxy, ESPMode::ThreadSafe> GetLightCardManagerProxy() const
-	{ return LightCardManagerProxy; }
-
 public:
-	/** Update internals for the new frame. */
-	void UpdateConfiguration();
-
 	/** Handle StartScene event: created and update internal resources. */
 	void HandleStartScene();
 
@@ -73,23 +66,18 @@ private:
 	/** Release UVLightCard data game thread*/
 	void ReleaseUVLightCardData();
 
-	/** Get a pointer to the DC ViewportManager if it still exists. */
-	FDisplayClusterViewportManager* GetViewportManager() const
-	{
-		return ViewportManagerWeakPtr.IsValid() ? ViewportManagerWeakPtr.Pin().Get() : nullptr;
-	}
+public:
+	// Configuration of the current cluster node
+	const TSharedRef<const FDisplayClusterViewportConfiguration, ESPMode::ThreadSafe> Configuration;
 
-private:
-	/** A reference to the owning viewport manager */
-	const TWeakPtr<FDisplayClusterViewportManager, ESPMode::ThreadSafe> ViewportManagerWeakPtr;
+	/** RenderThread Proxy object*/
+	const TSharedRef<FDisplayClusterViewportLightCardManagerProxy, ESPMode::ThreadSafe> LightCardManagerProxy;
 
+private:
 	/** A list of primitive components that have been added to the preview scene for rendering in the current frame */
 	TArray<UPrimitiveComponent*> UVLightCardPrimitiveComponents;
 
 	/** The render target to which the UV light card map is rendered */
 	TSharedPtr<FDisplayClusterViewportLightCardResource, ESPMode::ThreadSafe> UVLightCardResource;
 
-private:
-	/** RenderThread Proxy object*/
-	TSharedPtr<FDisplayClusterViewportLightCardManagerProxy, ESPMode::ThreadSafe> LightCardManagerProxy;
 };
diff --git a/Engine/Plugins/Runtime/nDisplay/Source/DisplayCluster/Private/Render/Viewport/Postprocess/DisplayClusterViewportPostProcessManager.cpp b/Engine/Plugins/Runtime/nDisplay/Source/DisplayCluster/Private/Render/Viewport/Postprocess/DisplayClusterViewportPostProcessManager.cpp
--- a/Engine/Plugins/Runtime/nDisplay/Source/DisplayCluster/Private/Render/Viewport/Postprocess/DisplayClusterViewportPostProcessManager.cpp
+++ b/Engine/Plugins/Runtime/nDisplay/Source/DisplayCluster/Private/Render/Viewport/Postprocess/DisplayClusterViewportPostProcessManager.cpp
@@ -1,26 +1,20 @@
 // Copyright Epic Games, Inc. All Rights Reserved.
 
 #include "Render/Viewport/Postprocess/DisplayClusterViewportPostProcessManager.h"
-
-#include "Misc/DisplayClusterHelpers.h"
-#include "Misc/DisplayClusterLog.h"
-
-#include "HAL/IConsoleManager.h"
-
-#include "Render/IPDisplayClusterRenderManager.h"
-#include "Render/IDisplayClusterRenderManager.h"
-#include "Render/PostProcess/IDisplayClusterPostProcessFactory.h"
-
+#include "Render/Viewport/Postprocess/DisplayClusterViewportPostProcessOutputRemap.h"
 #include "Render/Viewport/DisplayClusterViewportManager.h"
-
 #include "Render/Viewport/DisplayClusterViewportManagerProxy.h"
 #include "Render/Viewport/DisplayClusterViewportProxy.h"
+#include "Render/Viewport/Configuration/DisplayClusterViewportConfiguration.h"
+#include "Render/PostProcess/IDisplayClusterPostProcessFactory.h"
+#include "Render/IPDisplayClusterRenderManager.h"
+#include "Render/IDisplayClusterRenderManager.h"
+#include "DisplayClusterConfigurationTypes.h"
 
-#include "Render/Viewport/Postprocess/DisplayClusterViewportPostProcessOutputRemap.h"
-
+#include "Misc/DisplayClusterHelpers.h"
+#include "Misc/DisplayClusterLog.h"
 #include "Misc/DisplayClusterGlobals.h"
-
-#include "DisplayClusterConfigurationTypes.h"
+#include "HAL/IConsoleManager.h"
 
 //////////////////////////////////////////////////////////////////////////////////////////////
 // Round 1: VIEW before warp&blend
@@ -61,8 +55,8 @@ static TAutoConsoleVariable<int32> CVarPostprocessFrameAfterWarpBlend(
 //----------------------------------------------------------------------------------------------------------
 // FDisplayClusterViewportPostProcessManager
 //----------------------------------------------------------------------------------------------------------
-FDisplayClusterViewportPostProcessManager::FDisplayClusterViewportPostProcessManager(FDisplayClusterViewportManager& InViewportManager)
-	: ViewportManagerWeakPtr(InViewportManager.AsShared())
+FDisplayClusterViewportPostProcessManager::FDisplayClusterViewportPostProcessManager(const TSharedRef<const FDisplayClusterViewportConfiguration, ESPMode::ThreadSafe>& InConfiguration)
+	: Configuration(InConfiguration)
 {
 	OutputRemap = MakeShared<FDisplayClusterViewportPostProcessOutputRemap, ESPMode::ThreadSafe>();
 }
@@ -72,11 +66,6 @@ FDisplayClusterViewportPostProcessManager::~FDisplayClusterViewportPostProcessMa
 	Release();
 }
 
-FDisplayClusterViewportManager* FDisplayClusterViewportPostProcessManager::GetViewportManager() const
-{
-	return ViewportManagerWeakPtr.IsValid() ? ViewportManagerWeakPtr.Pin().Get() : nullptr;
-}
-
 void FDisplayClusterViewportPostProcessManager::Release()
 {
 	if (IsInGameThread())
@@ -94,28 +83,33 @@ void FDisplayClusterViewportPostProcessManager::Release()
 
 bool FDisplayClusterViewportPostProcessManager::HandleStartScene()
 {
-	check(IsInGameThread());;
+	check(IsInGameThread());
+
+	bool bResult = true;
 
-	FDisplayClusterViewportManager* ViewportManager = GetViewportManager();
-	if (ViewportManager && ViewportManager->IsSceneOpened())
+	IDisplayClusterViewportManager* ViewportManager = Configuration->GetViewportManager();
+	if (ViewportManager && Configuration->IsSceneOpened())
 	{
 		for (const TSharedPtr<IDisplayClusterPostProcess, ESPMode::ThreadSafe>& It : Postprocess)
 		{
 			if (It.IsValid())
 			{
-				It->HandleStartScene(ViewportManager);
+				if (!It->HandleStartScene(ViewportManager))
+				{
+					bResult = false;
+				}
 			}
 		}
 	}
 
-	return false;
+	return bResult;
 }
 
 void FDisplayClusterViewportPostProcessManager::HandleEndScene()
 {
 	check(IsInGameThread());
 
-	if (FDisplayClusterViewportManager* ViewportManager = GetViewportManager())
+	if (IDisplayClusterViewportManager* ViewportManager = Configuration->GetViewportManager())
 	{
 		for (const TSharedPtr<IDisplayClusterPostProcess, ESPMode::ThreadSafe>& It : Postprocess)
 		{
@@ -179,8 +173,8 @@ bool FDisplayClusterViewportPostProcessManager::CreatePostprocess(const FString&
 		{
 			UE_LOG(LogDisplayClusterRender, Log, TEXT("PostProcess '%s', type '%s' : Created"), *InConfigurationPostprocess->Type, *InPostprocessId);
 
-			FDisplayClusterViewportManager* ViewportManager = GetViewportManager();
-			if (ViewportManager && ViewportManager->IsSceneOpened())
+			IDisplayClusterViewportManager* ViewportManager = Configuration->GetViewportManager();
+			if (ViewportManager && Configuration->IsSceneOpened())
 			{
 				PostProcessInstance->HandleStartScene(ViewportManager);
 			}
@@ -357,7 +351,7 @@ void FDisplayClusterViewportPostProcessManager::Tick()
 
 void FDisplayClusterViewportPostProcessManager::HandleSetupNewFrame()
 {
-	FDisplayClusterViewportManager* ViewportManager = GetViewportManager();
+	IDisplayClusterViewportManager* ViewportManager = Configuration->GetViewportManager();
 	if (ViewportManager && CVarCustomPPEnabled.GetValueOnGameThread() != 0)
 	{
 		for (const TSharedPtr<IDisplayClusterPostProcess, ESPMode::ThreadSafe >& It : Postprocess)
@@ -369,7 +363,7 @@ void FDisplayClusterViewportPostProcessManager::HandleSetupNewFrame()
 
 void FDisplayClusterViewportPostProcessManager::HandleBeginNewFrame(FDisplayClusterRenderFrame& InOutRenderFrame)
 {
-	FDisplayClusterViewportManager* ViewportManager = GetViewportManager();
+	IDisplayClusterViewportManager* ViewportManager = Configuration->GetViewportManager();
 	if (ViewportManager && CVarCustomPPEnabled.GetValueOnGameThread() != 0)
 	{
 		for (const TSharedPtr<IDisplayClusterPostProcess, ESPMode::ThreadSafe >& It : Postprocess)
diff --git a/Engine/Plugins/Runtime/nDisplay/Source/DisplayCluster/Private/Render/Viewport/Postprocess/DisplayClusterViewportPostProcessManager.h b/Engine/Plugins/Runtime/nDisplay/Source/DisplayCluster/Private/Render/Viewport/Postprocess/DisplayClusterViewportPostProcessManager.h
--- a/Engine/Plugins/Runtime/nDisplay/Source/DisplayCluster/Private/Render/Viewport/Postprocess/DisplayClusterViewportPostProcessManager.h
+++ b/Engine/Plugins/Runtime/nDisplay/Source/DisplayCluster/Private/Render/Viewport/Postprocess/DisplayClusterViewportPostProcessManager.h
@@ -6,9 +6,9 @@
 #include "Render/IDisplayClusterRenderManager.h"
 #include "Templates/SharedPointer.h"
 
-class FDisplayClusterViewportManager;
 class FDisplayClusterViewportManagerProxy;
 class FDisplayClusterViewportPostProcessOutputRemap;
+class FDisplayClusterViewportConfiguration;
 struct FDisplayClusterConfigurationPostprocess;
 
 /**
@@ -19,7 +19,7 @@ class FDisplayClusterViewportPostProcessManager
 	: public TSharedFromThis<FDisplayClusterViewportPostProcessManager, ESPMode::ThreadSafe>
 {
 public:
-	FDisplayClusterViewportPostProcessManager(FDisplayClusterViewportManager& InViewportManager);
+	FDisplayClusterViewportPostProcessManager(const TSharedRef<const FDisplayClusterViewportConfiguration, ESPMode::ThreadSafe>& InConfiguration);
 	virtual ~FDisplayClusterViewportPostProcessManager();
 
 	void Release();
@@ -86,25 +86,22 @@ public:
 private:
 	TSharedPtr<IDisplayClusterPostProcess, ESPMode::ThreadSafe> ImplFindPostProcess(const FString& InPostprocessId) const;
 
-	/** Get a pointer to the DC ViewportManager if it still exists. */
-	FDisplayClusterViewportManager* GetViewportManager() const;
-
 protected:
 	void ImplPerformPostProcessViewBeforeWarpBlend_RenderThread(FRHICommandListImmediate& RHICmdList, const FDisplayClusterViewportManagerProxy* InViewportManagerProxy) const;
 	void ImplPerformPostProcessViewAfterWarpBlend_RenderThread(FRHICommandListImmediate& RHICmdList,  const FDisplayClusterViewportManagerProxy* InViewportManagerProxy) const;
 	void ImplPerformPostProcessFrameAfterWarpBlend_RenderThread(FRHICommandListImmediate& RHICmdList, const FDisplayClusterViewportManagerProxy* InViewportManagerProxy) const;
 
+public:
+	// Configuration of the current cluster node
+	const TSharedRef<const FDisplayClusterViewportConfiguration, ESPMode::ThreadSafe> Configuration;
 
-protected:
+private:
 	// Game thread instances
 	TArray<TSharedPtr<IDisplayClusterPostProcess, ESPMode::ThreadSafe>> Postprocess;
 
-private:
-	/** A reference to the owning viewport manager */
-	const TWeakPtr<FDisplayClusterViewportManager, ESPMode::ThreadSafe> ViewportManagerWeakPtr;
-
 	// Render thread instances
 	TArray<TSharedPtr<IDisplayClusterPostProcess, ESPMode::ThreadSafe>> PostprocessProxy;
 
+	// Output remap implementation
 	TSharedPtr<FDisplayClusterViewportPostProcessOutputRemap, ESPMode::ThreadSafe> OutputRemap;
 };
diff --git a/Engine/Plugins/Runtime/nDisplay/Source/DisplayCluster/Private/Render/Viewport/Postprocess/DisplayClusterViewportPostProcessOutputRemap.cpp b/Engine/Plugins/Runtime/nDisplay/Source/DisplayCluster/Private/Render/Viewport/Postprocess/DisplayClusterViewportPostProcessOutputRemap.cpp
--- a/Engine/Plugins/Runtime/nDisplay/Source/DisplayCluster/Private/Render/Viewport/Postprocess/DisplayClusterViewportPostProcessOutputRemap.cpp
+++ b/Engine/Plugins/Runtime/nDisplay/Source/DisplayCluster/Private/Render/Viewport/Postprocess/DisplayClusterViewportPostProcessOutputRemap.cpp
@@ -28,16 +28,7 @@
 //////////////////////////////////////////////////////////////////////////////////////////////
 // FDisplayClusterViewportPostProcessOutputRemap
 //////////////////////////////////////////////////////////////////////////////////////////////
-FDisplayClusterViewportPostProcessOutputRemap::FDisplayClusterViewportPostProcessOutputRemap()
-	: ShadersAPI(IDisplayClusterShaders::Get())
-{
-}
-
-FDisplayClusterViewportPostProcessOutputRemap::~FDisplayClusterViewportPostProcessOutputRemap()
-{
-}
-
-bool FDisplayClusterViewportPostProcessOutputRemap::ImplInitializeConfiguration()
+bool FDisplayClusterViewportPostProcessOutputRemap::ImplInitializeOutputRemap()
 {
 	if (!OutputRemapMesh.IsValid())
 	{
@@ -92,7 +83,7 @@ bool FDisplayClusterViewportPostProcessOutputRemap::UpdateConfiguration_External
 		// Store input filename for next update
 		ExternalFile = FullPathFileName;
 	}
-	else if (ImplInitializeConfiguration())
+	else if (ImplInitializeOutputRemap())
 	{
 		// Assign loaded geometry
 		OutputRemapMesh->SetGeometryFunc(EDisplayClusterRender_MeshComponentProxyDataFunc::OutputRemapScreenSpace);
@@ -129,7 +120,7 @@ bool FDisplayClusterViewportPostProcessOutputRemap::UpdateConfiguration_StaticMe
 		}
 	}
 
-	if (ImplInitializeConfiguration())
+	if (ImplInitializeOutputRemap())
 	{
 		OutputRemapMesh->SetGeometryFunc(EDisplayClusterRender_MeshComponentProxyDataFunc::OutputRemapScreenSpace);
 		OutputRemapMesh->AssignStaticMesh(InStaticMesh, FDisplayClusterMeshUVs());
@@ -163,7 +154,7 @@ bool FDisplayClusterViewportPostProcessOutputRemap::UpdateConfiguration_StaticMe
 	}
 
 	// Begin new configuration:
-	if (ImplInitializeConfiguration())
+	if (ImplInitializeOutputRemap())
 	{
 		OutputRemapMesh->SetGeometryFunc(EDisplayClusterRender_MeshComponentProxyDataFunc::OutputRemapScreenSpace);
 		OutputRemapMesh->AssignStaticMeshComponentRefs(InStaticMeshComponent, FDisplayClusterMeshUVs());
@@ -197,7 +188,7 @@ bool FDisplayClusterViewportPostProcessOutputRemap::UpdateConfiguration_Procedur
 	}
 
 	// Begin new configuration:
-	if (ImplInitializeConfiguration())
+	if (ImplInitializeOutputRemap())
 	{
 		OutputRemapMesh->SetGeometryFunc(EDisplayClusterRender_MeshComponentProxyDataFunc::OutputRemapScreenSpace);
 		OutputRemapMesh->AssignProceduralMeshComponentRefs(InProceduralMeshComponent, FDisplayClusterMeshUVs());
@@ -245,6 +236,7 @@ void FDisplayClusterViewportPostProcessOutputRemap::PerformPostProcessFrame_Rend
 				FRHITexture2D* InOutTexture = (*InFrameTargets)[Index];
 				FRHITexture2D* TempTargetableTexture = (*InAdditionalFrameTargets)[Index];
 
+				static IDisplayClusterShaders& ShadersAPI = IDisplayClusterShaders::Get();
 				if (ShadersAPI.RenderPostprocess_OutputRemap(RHICmdList, InOutTexture, TempTargetableTexture, *MeshProxy))
 				{
 					TransitionAndCopyTexture(RHICmdList, TempTargetableTexture, InOutTexture, {});
diff --git a/Engine/Plugins/Runtime/nDisplay/Source/DisplayCluster/Private/Render/Viewport/Postprocess/DisplayClusterViewportPostProcessOutputRemap.h b/Engine/Plugins/Runtime/nDisplay/Source/DisplayCluster/Private/Render/Viewport/Postprocess/DisplayClusterViewportPostProcessOutputRemap.h
--- a/Engine/Plugins/Runtime/nDisplay/Source/DisplayCluster/Private/Render/Viewport/Postprocess/DisplayClusterViewportPostProcessOutputRemap.h
+++ b/Engine/Plugins/Runtime/nDisplay/Source/DisplayCluster/Private/Render/Viewport/Postprocess/DisplayClusterViewportPostProcessOutputRemap.h
@@ -19,9 +19,6 @@ class UProceduralMeshComponent;
 class FDisplayClusterViewportPostProcessOutputRemap
 {
 public:
-	FDisplayClusterViewportPostProcessOutputRemap();
-	virtual ~FDisplayClusterViewportPostProcessOutputRemap();
-
 	// Game thread update calls
 	bool UpdateConfiguration_ExternalFile(const FString& InExternalFile);
 	bool UpdateConfiguration_StaticMesh(UStaticMesh* InStaticMesh);
@@ -40,12 +37,12 @@ public:
 	void PerformPostProcessFrame_RenderThread(FRHICommandListImmediate& RHICmdList, const TArray<FRHITexture2D*>* InFrameTargets = nullptr, const TArray<FRHITexture2D*>* InAdditionalFrameTargets = nullptr) const;
 
 private:
-	bool ImplInitializeConfiguration();
+	bool ImplInitializeOutputRemap();
 
 private:
+	// Name of the external file used for remapping. Empty if not used.
 	FString ExternalFile;
 
-private:
+	// Reference to the mesh component that is used for remapping.
 	TSharedPtr<IDisplayClusterRender_MeshComponent, ESPMode::ThreadSafe> OutputRemapMesh;
-	class IDisplayClusterShaders& ShadersAPI;
 };
