nDisplay: Added new interface IDisplayClusterViewportConfiguration

diff --git a/Engine/Plugins/Runtime/nDisplay/Source/DisplayCluster/Private/Render/Viewport/RenderFrame/DisplayClusterRenderFrameManager.cpp b/Engine/Plugins/Runtime/nDisplay/Source/DisplayCluster/Private/Render/Viewport/RenderFrame/DisplayClusterRenderFrameManager.cpp
--- a/Engine/Plugins/Runtime/nDisplay/Source/DisplayCluster/Private/Render/Viewport/RenderFrame/DisplayClusterRenderFrameManager.cpp
+++ b/Engine/Plugins/Runtime/nDisplay/Source/DisplayCluster/Private/Render/Viewport/RenderFrame/DisplayClusterRenderFrameManager.cpp
@@ -1,33 +1,27 @@
 // Copyright Epic Games, Inc. All Rights Reserved.
 
 #include "Render/Viewport/RenderFrame/DisplayClusterRenderFrameManager.h"
-
-#include "Render/Viewport/DisplayClusterViewport.h"
-
 #include "Render/Viewport/RenderFrame/DisplayClusterRenderFrame.h"
 #include "Render/Viewport/RenderFrame/DisplayClusterRenderFrameSettings.h"
+#include "Render/Viewport/Configuration/DisplayClusterViewportConfiguration.h"
+#include "Render/Viewport/DisplayClusterViewport.h"
 
 #include "Misc/Parse.h"
 
-
 ///////////////////////////////////////////////////////////////
 // FDisplayClusterRenderTargetFrame
 ///////////////////////////////////////////////////////////////
-bool FDisplayClusterRenderFrameManager::BuildRenderFrame(FViewport* InViewport, const FDisplayClusterRenderFrameSettings& InRenderFrameSettings, const TArray<TSharedPtr<FDisplayClusterViewport, ESPMode::ThreadSafe>>& InViewports, FDisplayClusterRenderFrame& OutRenderFrame)
+bool FDisplayClusterRenderFrameManager::BuildRenderFrame(FViewport* InViewport, const TArray<TSharedPtr<FDisplayClusterViewport, ESPMode::ThreadSafe>>& InViewports, FDisplayClusterRenderFrame& OutRenderFrame)
 {
+	const FDisplayClusterRenderFrameSettings& InRenderFrameSettings = Configuration->GetRenderFrameSettings();
 
-	switch (InRenderFrameSettings.RenderMode)
+	// Dont use render frame for preview
+	if (!InRenderFrameSettings.IsPreviewRendering())
 	{
-	case EDisplayClusterRenderFrameMode::PreviewInScene:
-		// Dont use render frame for preview
-		break;
-
-	default:
-		if (!FindFrameTargetRect(InViewport, InViewports, InRenderFrameSettings, OutRenderFrame.FrameRect))
+		if (!FindFrameTargetRect(InViewport, InViewports, OutRenderFrame.FrameRect))
 		{
 			return false;
 		}
-		break;
 	}
 
 	// @todo add atlasing, merge multiple viewports in single viewfamily, etc
@@ -43,7 +37,7 @@ bool FDisplayClusterRenderFrameManager::BuildRenderFrame(FViewport* InViewport,
 	{
 		if (Viewport.IsValid())
 		{
-			if (Viewport->RenderSettings.GetParentViewportId().IsEmpty())
+			if (Viewport->GetRenderSettings().GetParentViewportId().IsEmpty())
 			{
 				SortedViewports.Add(Viewport);
 			}
@@ -80,44 +74,13 @@ bool FDisplayClusterRenderFrameManager::BuildRenderFrame(FViewport* InViewport,
 		}
 	}
 
-	bool bResult = false;
-
-	if (InRenderFrameSettings.bAllowRenderTargetAtlasing)
-	{
-// Not implemented yet
-#if 0
-		switch (InRenderFrameSettings.ViewFamilyMode)
-		{
-		case EDisplayClusterRenderFamilyMode::AllowMergeForGroups:
-			bResult = false;
-			break;
-
-		case EDisplayClusterRenderFamilyMode::AllowMergeForGroupsAndStereo:
-			//NOT_IMPLEMENTED
-			bResult = false;
-			break;
-
-		case EDisplayClusterRenderFamilyMode::MergeAnyPossible:
-			//NOT_IMPLEMENTED
-			bResult = false;
-			break;
-
-		default:
-			bResult = false;
-			break;
-		}
-#endif
-	}
-	else
-	{
-		bResult = BuildSimpleFrame(InViewport, InRenderFrameSettings, SortedViewports, OutRenderFrame);
-	}
-
-	return bResult;
+	return BuildSimpleFrame(InViewport, SortedViewports, OutRenderFrame);
 }
 
-bool FDisplayClusterRenderFrameManager::BuildSimpleFrame(FViewport* InViewport, const FDisplayClusterRenderFrameSettings& InRenderFrameSettings, const TArray<TSharedPtr<FDisplayClusterViewport, ESPMode::ThreadSafe>>& InViewports, FDisplayClusterRenderFrame& OutRenderFrame)
+bool FDisplayClusterRenderFrameManager::BuildSimpleFrame(FViewport* InViewport, const TArray<TSharedPtr<FDisplayClusterViewport, ESPMode::ThreadSafe>>& InViewports, FDisplayClusterRenderFrame& OutRenderFrame)
 {
+	const FDisplayClusterRenderFrameSettings& InRenderFrameSettings = Configuration->GetRenderFrameSettings();
+
 	for (const TSharedPtr<FDisplayClusterViewport, ESPMode::ThreadSafe>& ViewportIt : InViewports)
 	{
 		if (ViewportIt.IsValid())
@@ -128,8 +91,8 @@ bool FDisplayClusterRenderFrameManager::BuildSimpleFrame(FViewport* InViewport,
 				{
 					FrameView.ContextNum = ContextIt.ContextNum;
 					FrameView.Viewport = ViewportIt.ToSharedRef();
-					FrameView.bDisableRender = ContextIt.bDisableRender || ViewportIt->RenderSettings.bSkipSceneRenderingButLeaveResourcesAvailable;
-					FrameView.bFreezeRendering = ViewportIt->RenderSettings.bFreezeRendering;
+					FrameView.bDisableRender = ContextIt.bDisableRender || ViewportIt->GetRenderSettings().bSkipSceneRenderingButLeaveResourcesAvailable;
+					FrameView.bFreezeRendering = ViewportIt->GetRenderSettings().bFreezeRendering;
 				}
 
 				FDisplayClusterRenderFrameTargetViewFamily FrameViewFamily;
@@ -142,12 +105,12 @@ bool FDisplayClusterRenderFrameManager::BuildSimpleFrame(FViewport* InViewport,
 				FDisplayClusterRenderFrameTarget FrameRenderTarget;
 				{
 					// Simple frame use unique RTT  for each viewport, so disable RTT when viewport rendering disabled
-					FrameRenderTarget.bShouldUseRenderTarget = ViewportIt->Resources[EDisplayClusterViewportResource::RenderTargets].Num() > 0;
+					FrameRenderTarget.bShouldUseRenderTarget = ViewportIt->GetViewportResources(EDisplayClusterViewportResource::RenderTargets).Num() > 0;
 
 					FrameRenderTarget.ViewFamilies.Add(FrameViewFamily);
 
 					FrameRenderTarget.RenderTargetSize = ContextIt.RenderTargetRect.Max;
-					FrameRenderTarget.CaptureMode = ViewportIt->RenderSettings.CaptureMode;
+					FrameRenderTarget.CaptureMode = ViewportIt->GetRenderSettings().CaptureMode;
 				}
 
 				OutRenderFrame.RenderTargets.Add(FrameRenderTarget);
@@ -158,8 +121,10 @@ bool FDisplayClusterRenderFrameManager::BuildSimpleFrame(FViewport* InViewport,
 	return true;
 }
 
-bool FDisplayClusterRenderFrameManager::FindFrameTargetRect(FViewport* InViewport, const TArray<TSharedPtr<FDisplayClusterViewport, ESPMode::ThreadSafe>>& InOutViewports, const FDisplayClusterRenderFrameSettings& InRenderFrameSettings, FIntRect& OutFrameTargetRect) const
+bool FDisplayClusterRenderFrameManager::FindFrameTargetRect(FViewport* InViewport, const TArray<TSharedPtr<FDisplayClusterViewport, ESPMode::ThreadSafe>>& InOutViewports, FIntRect& OutFrameTargetRect) const
 {
+	const FDisplayClusterRenderFrameSettings& InRenderFrameSettings = Configuration->GetRenderFrameSettings();
+
 	// Calculate Backbuffer frame
 	bool bIsUsed = false;
 
@@ -173,7 +138,7 @@ bool FDisplayClusterRenderFrameManager::FindFrameTargetRect(FViewport* InViewpor
 	// Optimize frame target RTT
 	for (const TSharedPtr<FDisplayClusterViewport, ESPMode::ThreadSafe>& ViewportIt : InOutViewports)
 	{
-		if (ViewportIt.IsValid() && ViewportIt->RenderSettings.bEnable && ViewportIt->RenderSettings.bVisible)
+		if (ViewportIt.IsValid() && ViewportIt->GetRenderSettings().bEnable && ViewportIt->GetRenderSettings().bVisible)
 		{
 			for (const auto& ContextIt : ViewportIt->GetContexts())
 			{
diff --git a/Engine/Plugins/Runtime/nDisplay/Source/DisplayCluster/Private/Render/Viewport/RenderFrame/DisplayClusterRenderFrameManager.h b/Engine/Plugins/Runtime/nDisplay/Source/DisplayCluster/Private/Render/Viewport/RenderFrame/DisplayClusterRenderFrameManager.h
--- a/Engine/Plugins/Runtime/nDisplay/Source/DisplayCluster/Private/Render/Viewport/RenderFrame/DisplayClusterRenderFrameManager.h
+++ b/Engine/Plugins/Runtime/nDisplay/Source/DisplayCluster/Private/Render/Viewport/RenderFrame/DisplayClusterRenderFrameManager.h
@@ -4,20 +4,26 @@
 
 #include "CoreMinimal.h"
 
-class FViewport;
+class FDisplayClusterViewportConfiguration;
 class FDisplayClusterViewport;
 class FViewport;
 class FDisplayClusterRenderFrame;
-struct FDisplayClusterRenderFrameSettings;
 
 class FDisplayClusterRenderFrameManager
 {
 public:
-	bool BuildRenderFrame(FViewport* InViewport, const FDisplayClusterRenderFrameSettings& InRenderFrameSettings, const TArray<TSharedPtr<FDisplayClusterViewport, ESPMode::ThreadSafe>>& InViewports, FDisplayClusterRenderFrame& OutRenderFrame);
+	FDisplayClusterRenderFrameManager(const TSharedRef<const FDisplayClusterViewportConfiguration, ESPMode::ThreadSafe>& InConfiguration)
+		: Configuration(InConfiguration)
+	{ }
 
-private:
-	bool FindFrameTargetRect(FViewport* InViewport, const TArray<TSharedPtr<FDisplayClusterViewport, ESPMode::ThreadSafe>>& InViewports, const FDisplayClusterRenderFrameSettings& InRenderFrameSettings, FIntRect& OutFrameTargetRect) const;
-	bool BuildSimpleFrame(FViewport* InViewport, const FDisplayClusterRenderFrameSettings& InRenderFrameSettings, const TArray<TSharedPtr<FDisplayClusterViewport, ESPMode::ThreadSafe>>& InViewports, FDisplayClusterRenderFrame& OutRenderFrame);
-};
+public:
+	bool BuildRenderFrame(FViewport* InViewport, const TArray<TSharedPtr<FDisplayClusterViewport, ESPMode::ThreadSafe>>& InViewports, FDisplayClusterRenderFrame& OutRenderFrame);
 
+private:
+	bool FindFrameTargetRect(FViewport* InViewport, const TArray<TSharedPtr<FDisplayClusterViewport, ESPMode::ThreadSafe>>& InViewports, FIntRect& OutFrameTargetRect) const;
+	bool BuildSimpleFrame(FViewport* InViewport, const TArray<TSharedPtr<FDisplayClusterViewport, ESPMode::ThreadSafe>>& InViewports, FDisplayClusterRenderFrame& OutRenderFrame);
 
+private:
+	// Configuration of the current cluster node
+	const TSharedRef<const FDisplayClusterViewportConfiguration, ESPMode::ThreadSafe> Configuration;
+};
diff --git a/Engine/Plugins/Runtime/nDisplay/Source/DisplayCluster/Private/Render/Viewport/RenderFrame/DisplayClusterRenderFrameSettings.h b/Engine/Plugins/Runtime/nDisplay/Source/DisplayCluster/Private/Render/Viewport/RenderFrame/DisplayClusterRenderFrameSettings.h
--- a/Engine/Plugins/Runtime/nDisplay/Source/DisplayCluster/Private/Render/Viewport/RenderFrame/DisplayClusterRenderFrameSettings.h
+++ b/Engine/Plugins/Runtime/nDisplay/Source/DisplayCluster/Private/Render/Viewport/RenderFrame/DisplayClusterRenderFrameSettings.h
@@ -6,6 +6,7 @@
 
 #include "Render/Viewport/RenderFrame/DisplayClusterRenderFrameEnums.h"
 #include "Render/Viewport/RenderFrame/DisplayClusterRenderFrameSettingsEnums.h"
+#include "Render/Viewport/Containers/DisplayClusterViewport_PreviewSettings.h"
 
 // Settings for render frame builder
 struct FDisplayClusterRenderFrameSettings
@@ -40,13 +41,6 @@ struct FDisplayClusterRenderFrameSettings
 	// Postprocess can use full size backbuffer. This disables Frame RTT size optimization
 	bool bShouldUseFullSizeFrameTargetableResource = false;
 
-	// Enable Post Process for preview rendering
-	bool bPreviewEnablePostProcess = false;
-
-	// Preview RTT size multiplier
-	float PreviewRenderTargetRatioMult = 1.f;
-
-
 	// Multiply all viewports RTT size's for whole cluster by this value
 	float ClusterRenderTargetRatioMult = 1.f;
 
@@ -66,37 +60,18 @@ struct FDisplayClusterRenderFrameSettings
 	// Multiply outer viewports buffer ratios for whole cluster by this value
 	float ClusterICVFXOuterViewportBufferRatioMult = 1.f;
 
-
 	// Allow warpblend render
 	bool bAllowWarpBlend = true;
 
-	// Render in Editor mode
-	bool bIsRenderingInEditor = false;
-
-	// Configuration used to render preview
-	bool bIsPreviewRendering = false;
+	// Settings for preview rendering
+	FDisplayClusterViewport_PreviewSettings PreviewSettings;
 
-	// If the preview rendering should be frozen
-	bool bFreezePreviewRender = false;
-	
-	// Allow mGPU in editor mode
-	bool bAllowMultiGPURenderingInEditor = false;
-	int32 PreviewMinGPUIndex = 0;
-	int32 PreviewMaxGPUIndex = 0;
-
-	// The maximum dimension of any texture for preview
-	int32 PreviewMaxTextureDimension = 2048;
+	// Use DC render device for rendering
+	bool bUseDisplayClusterRenderDevice = true;
 
 	// Performance: Stereoscopic rendering uses a single viewfamily(RTT) for both eyes
 	bool bEnableStereoscopicRenderingOptimization = false;
 
-	// Performance: Allow merge multiple viewports on single RTT with atlasing (required for bAllowViewFamilyMergeOptimization)
-	bool bAllowRenderTargetAtlasing = false;
-
-	// Performance: Allow viewfamily merge optimization (render multiple viewports contexts within single family)
-	// [not implemented yet] Experimental
-	EDisplayClusterRenderFamilyMode ViewFamilyMode = EDisplayClusterRenderFamilyMode::None;
-
 	// Performance: Allow to use parent ViewFamily from parent viewport 
 	// (icvfx has child viewports: lightcard and chromakey with prj_view matrices copied from parent viewport. May sense to use same viewfamily?)
 	// [not implemented yet] Experimental
@@ -104,6 +79,40 @@ struct FDisplayClusterRenderFrameSettings
 
 	// Cluster node name for render
 	FString ClusterNodeId;
-};
 
+public:
+	/** Current frame is preview. */
+	bool IsPreviewRendering() const;
 
+	/** returns true if the preview rendering has been updated.If this function returns false, the DCRA preview image should be frozen. */
+	bool IsPreviewFreezeRender() const;
+
+	/** [Experimental] Get the GPU range used for rendering preview.
+	* returns nullptr if this rendering function is disabled on mGPU
+	*/
+	const FIntPoint* GetPreviewMultiGPURendering() const;
+
+	/** PostProcess should be disabled for the current frame. */
+	bool IsPostProcessDisabled() const;
+
+	/** Should use linear gamma. */
+	bool ShouldUseLinearGamma() const;
+
+	/** Getting the desired frame size multipliers. */
+	FVector2D GetDesiredFrameMult() const;
+
+	/** Get the maximum texture size that is allowed to be used in the viewport. */
+	int32 GetViewportTextureMaxSize() const;
+
+	/** Should use Cross-GPU transfers. */
+	inline bool ShouldUseCrossGPUTransfers() const
+	{
+		return bUseDisplayClusterRenderDevice || GetPreviewMultiGPURendering() != nullptr;
+	}
+
+	/** Ability to reuse the viewport across all nodes in the cluster (DCRA Preview). */
+	bool CanReuseViewportWithinClusterNodes() const;
+
+	/** Return number of contexts per viewport. */
+	int32 GetViewPerViewportAmount() const;
+};
diff --git a/Engine/Plugins/Runtime/nDisplay/Source/DisplayCluster/Private/Render/Viewport/RenderTarget/DisplayClusterRenderTargetManager.cpp b/Engine/Plugins/Runtime/nDisplay/Source/DisplayCluster/Private/Render/Viewport/RenderTarget/DisplayClusterRenderTargetManager.cpp
--- a/Engine/Plugins/Runtime/nDisplay/Source/DisplayCluster/Private/Render/Viewport/RenderTarget/DisplayClusterRenderTargetManager.cpp
+++ b/Engine/Plugins/Runtime/nDisplay/Source/DisplayCluster/Private/Render/Viewport/RenderTarget/DisplayClusterRenderTargetManager.cpp
@@ -6,6 +6,7 @@
 
 #include "Render/Viewport/DisplayClusterViewport.h"
 #include "Render/Viewport/Containers/DisplayClusterViewport_Enums.h"
+#include "Render/Viewport/Configuration/DisplayClusterViewportConfiguration.h"
 
 #include "Render/Viewport/RenderFrame/DisplayClusterRenderFrame.h"
 #include "Render/Viewport/RenderFrame/DisplayClusterRenderFrameSettings.h"
@@ -89,9 +90,10 @@ using namespace UE::DisplayCluster::RenderTargetManager;
 ////////////////////////////////////////////////////////////////////////////////////////////////////////////
 /// FDisplayClusterRenderTargetManager
 ////////////////////////////////////////////////////////////////////////////////////////////////////////////
-FDisplayClusterRenderTargetManager::FDisplayClusterRenderTargetManager(FDisplayClusterViewportManagerProxy* InViewportManagerProxy)
+FDisplayClusterRenderTargetManager::FDisplayClusterRenderTargetManager(const TSharedRef<FDisplayClusterViewportConfiguration, ESPMode::ThreadSafe>& InConfiguration)
+	: Configuration(InConfiguration)
 {
-	ResourcesPool = MakeUnique<FDisplayClusterRenderTargetResourcesPool>(InViewportManagerProxy);
+	ResourcesPool = MakeUnique<FDisplayClusterRenderTargetResourcesPool>();
 }
 
 FDisplayClusterRenderTargetManager::~FDisplayClusterRenderTargetManager()
@@ -104,11 +106,11 @@ void FDisplayClusterRenderTargetManager::Release()
 	ResourcesPool->Release();
 }
 
-bool FDisplayClusterRenderTargetManager::AllocateRenderFrameResources(FViewport* InViewport, const FDisplayClusterRenderFrameSettings& InRenderFrameSettings, const TArray<TSharedPtr<FDisplayClusterViewport, ESPMode::ThreadSafe>>& InViewports, FDisplayClusterRenderFrame& InOutRenderFrame)
+bool FDisplayClusterRenderTargetManager::AllocateRenderFrameResources(FViewport* InViewport, const TArray<TSharedPtr<FDisplayClusterViewport, ESPMode::ThreadSafe>>& InViewports, FDisplayClusterRenderFrame& InOutRenderFrame)
 {
 	bool bResult = true;
 
-	if(ResourcesPool->BeginReallocateResources(InRenderFrameSettings, InViewport))
+	if(ResourcesPool->BeginReallocateResources(InViewport, Configuration->GetRenderFrameSettings()))
 	{
 		// ReAllocate Render targets for all viewports
 		for (FDisplayClusterRenderFrameTarget& FrameRenderTargetIt : InOutRenderFrame.RenderTargets)
@@ -131,10 +133,10 @@ bool FDisplayClusterRenderTargetManager::AllocateRenderFrameResources(FViewport*
 							if (FDisplayClusterViewport* ViewportPtr = static_cast<FDisplayClusterViewport*>(ViewIt.Viewport.Get()))
 							{
 								// Array already resized in function FDisplayClusterViewport::UpdateFrameContexts() with RenderTargets.AddZeroed(ViewportContextAmount);
-								check(ViewportPtr->Resources[EDisplayClusterViewportResource::RenderTargets].IsValidIndex(ViewIt.ContextNum));
-								check(!ViewportPtr->Contexts[ViewIt.ContextNum].bDisableRender);
+								check(ViewportPtr->GetViewportResources(EDisplayClusterViewportResource::RenderTargets).IsValidIndex(ViewIt.ContextNum));
+								check(!ViewportPtr->GetContexts()[ViewIt.ContextNum].bDisableRender);
 
-								ViewportPtr->Resources[EDisplayClusterViewportResource::RenderTargets][ViewIt.ContextNum] = NewResource;
+								ViewportPtr->GetViewportResourcesImpl(EDisplayClusterViewportResource::RenderTargets)[ViewIt.ContextNum] = NewResource;
 
 								if (EnumHasAnyFlags(NewResource->GetResourceState(), EDisplayClusterViewportResourceState::Initialized) == false)
 								{
@@ -151,9 +153,9 @@ bool FDisplayClusterRenderTargetManager::AllocateRenderFrameResources(FViewport*
 		// Allocate viewport internal resources
 		for (const TSharedPtr<FDisplayClusterViewport, ESPMode::ThreadSafe>& ViewportIt : InViewports)
 		{
-			if (ViewportIt.IsValid() && ViewportIt->RenderSettings.bFreezeRendering == false)
+			if (ViewportIt.IsValid() && ViewportIt->GetRenderSettings().bFreezeRendering == false)
 			{
-				EPixelFormat CustomFormat = ImplGetCustomFormat(ViewportIt->RenderSettings.CaptureMode);
+				EPixelFormat CustomFormat = ImplGetCustomFormat(ViewportIt->GetRenderSettings().CaptureMode);
 
 				// Allocate all context resources:
 				for (const FDisplayClusterViewport_Context& ContextIt : ViewportIt->GetContexts())
@@ -162,29 +164,29 @@ bool FDisplayClusterRenderTargetManager::AllocateRenderFrameResources(FViewport*
 					const uint32& ContextNum = ContextIt.ContextNum;
 
 					// Allocate per-viewport preview output textures:
-					if (ViewportIt->Resources[EDisplayClusterViewportResource::OutputPreviewTargetableResources].Num() > (int32)ContextNum)
+					if (ViewportIt->GetViewportResources(EDisplayClusterViewportResource::OutputPreviewTargetableResources).Num() > (int32)ContextNum)
 					{
-						ViewportIt->Resources[EDisplayClusterViewportResource::OutputPreviewTargetableResources][ContextNum] = ResourcesPool->AllocateResource(ContextSize, CustomFormat, EDisplayClusterViewportResourceSettingsFlags::PreviewTargetableTexture);
-						ImplViewportTextureResourceLogging(ViewportIt, ContextNum, TEXT("InputShader"), ViewportIt->Resources[EDisplayClusterViewportResource::OutputPreviewTargetableResources][ContextNum]);
+						ViewportIt->GetViewportResourcesImpl(EDisplayClusterViewportResource::OutputPreviewTargetableResources)[ContextNum] = ResourcesPool->AllocateResource(ContextSize, CustomFormat, EDisplayClusterViewportResourceSettingsFlags::PreviewTargetableTexture);
+						ImplViewportTextureResourceLogging(ViewportIt, ContextNum, TEXT("InputShader"), ViewportIt->GetViewportResources(EDisplayClusterViewportResource::OutputPreviewTargetableResources)[ContextNum]);
 					}
 
-					if (ViewportIt->Resources[EDisplayClusterViewportResource::InputShaderResources].Num() > (int32)ContextNum)
+					if (ViewportIt->GetViewportResources(EDisplayClusterViewportResource::InputShaderResources).Num() > (int32)ContextNum)
 					{
-						ViewportIt->Resources[EDisplayClusterViewportResource::InputShaderResources][ContextNum] = ResourcesPool->AllocateResource(ContextSize, CustomFormat, EDisplayClusterViewportResourceSettingsFlags::ResolveTargetableTexture);
-						ImplViewportTextureResourceLogging(ViewportIt, ContextNum, TEXT("InputShader"), ViewportIt->Resources[EDisplayClusterViewportResource::InputShaderResources][ContextNum]);
+						ViewportIt->GetViewportResourcesImpl(EDisplayClusterViewportResource::InputShaderResources)[ContextNum] = ResourcesPool->AllocateResource(ContextSize, CustomFormat, EDisplayClusterViewportResourceSettingsFlags::ResolveTargetableTexture);
+						ImplViewportTextureResourceLogging(ViewportIt, ContextNum, TEXT("InputShader"), ViewportIt->GetViewportResources(EDisplayClusterViewportResource::InputShaderResources)[ContextNum]);
 					}
 
 					// Allocate custom resources:
-					if (ViewportIt->Resources[EDisplayClusterViewportResource::AdditionalTargetableResources].Num() > (int32)ContextNum)
+					if (ViewportIt->GetViewportResources(EDisplayClusterViewportResource::AdditionalTargetableResources).Num() > (int32)ContextNum)
 					{
-						ViewportIt->Resources[EDisplayClusterViewportResource::AdditionalTargetableResources][ContextNum] = ResourcesPool->AllocateResource(ContextSize, CustomFormat, EDisplayClusterViewportResourceSettingsFlags::RenderTargetableTexture);
-						ImplViewportTextureResourceLogging(ViewportIt, ContextNum, TEXT("AdditionalRTT"), ViewportIt->Resources[EDisplayClusterViewportResource::AdditionalTargetableResources][ContextNum]);
+						ViewportIt->GetViewportResourcesImpl(EDisplayClusterViewportResource::AdditionalTargetableResources)[ContextNum] = ResourcesPool->AllocateResource(ContextSize, CustomFormat, EDisplayClusterViewportResourceSettingsFlags::RenderTargetableTexture);
+						ImplViewportTextureResourceLogging(ViewportIt, ContextNum, TEXT("AdditionalRTT"), ViewportIt->GetViewportResources(EDisplayClusterViewportResource::AdditionalTargetableResources)[ContextNum]);
 					}
 
-					if (ViewportIt->Resources[EDisplayClusterViewportResource::MipsShaderResources].Num() > (int32)ContextNum)
+					if (ViewportIt->GetViewportResources(EDisplayClusterViewportResource::MipsShaderResources).Num() > (int32)ContextNum)
 					{
-						ViewportIt->Resources[EDisplayClusterViewportResource::MipsShaderResources][ContextNum] = ResourcesPool->AllocateResource(ContextSize, CustomFormat, EDisplayClusterViewportResourceSettingsFlags::ResolveTargetableTexture, ContextIt.NumMips);
-						ImplViewportTextureResourceLogging(ViewportIt, ContextNum, TEXT("Mips"), ViewportIt->Resources[EDisplayClusterViewportResource::MipsShaderResources][ContextNum]);
+						ViewportIt->GetViewportResourcesImpl(EDisplayClusterViewportResource::MipsShaderResources)[ContextNum] = ResourcesPool->AllocateResource(ContextSize, CustomFormat, EDisplayClusterViewportResourceSettingsFlags::ResolveTargetableTexture, ContextIt.NumMips);
+						ImplViewportTextureResourceLogging(ViewportIt, ContextNum, TEXT("Mips"), ViewportIt->GetViewportResources(EDisplayClusterViewportResource::MipsShaderResources)[ContextNum]);
 					}
 				}
 			}
@@ -193,7 +195,7 @@ bool FDisplayClusterRenderTargetManager::AllocateRenderFrameResources(FViewport*
 		// Allocate frame targets for all visible  on backbuffer viewports
 		FIntPoint ViewportSize = InViewport ? InViewport->GetSizeXY() : FIntPoint(0, 0);
 
-		if (!AllocateFrameTargets(InRenderFrameSettings, ViewportSize, InOutRenderFrame))
+		if (!AllocateFrameTargets(ViewportSize, InOutRenderFrame))
 		{
 			UE_LOG(LogDisplayClusterViewport, Error, TEXT("DisplayClusterRenderTargetManager: Can't allocate frame targets."));
 			bResult = false;
@@ -205,8 +207,14 @@ bool FDisplayClusterRenderTargetManager::AllocateRenderFrameResources(FViewport*
 	return bResult;
 }
 
-bool FDisplayClusterRenderTargetManager::AllocateFrameTargets(const FDisplayClusterRenderFrameSettings& InRenderFrameSettings, const FIntPoint& InViewportSize, FDisplayClusterRenderFrame& InOutRenderFrame)
+bool FDisplayClusterRenderTargetManager::AllocateFrameTargets(const FIntPoint& InViewportSize, FDisplayClusterRenderFrame& InOutRenderFrame)
 {
+	// Preview render to external RTT resource
+	if (Configuration->IsPreviewRendering())
+	{
+		return true;
+	}
+
 	uint32 FrameTargetsAmount = 0;
 
 	// Support side_by_side and top_bottom eye offset, aligned to Viewport size:
@@ -214,12 +222,8 @@ bool FDisplayClusterRenderTargetManager::AllocateFrameTargets(const FDisplayClus
 	FIntPoint TargetOffset(ForceInitToZero);
 
 	//Re-allocate frame targets
-	switch (InRenderFrameSettings.RenderMode)
+	switch (Configuration->GetRenderFrameSettings().RenderMode)
 	{
-	case EDisplayClusterRenderFrameMode::PreviewInScene:
-		// Preview model: render to external RTT resource
-		return true;
-
 	case EDisplayClusterRenderFrameMode::Mono:
 		FrameTargetsAmount = 1;
 		break;
@@ -264,7 +268,7 @@ bool FDisplayClusterRenderTargetManager::AllocateFrameTargets(const FDisplayClus
 			NewFrameTargetResources.Add(NewResource);
 		}
 
-		if (InRenderFrameSettings.bShouldUseAdditionalFrameTargetableResource)
+		if (Configuration->GetRenderFrameSettings().bShouldUseAdditionalFrameTargetableResource)
 		{
 			TSharedPtr<FDisplayClusterViewportResource, ESPMode::ThreadSafe> NewAdditionalResource = ResourcesPool->AllocateResource(InOutRenderFrame.FrameRect.Size(), PF_Unknown, EDisplayClusterViewportResourceSettingsFlags::RenderTargetableTexture);
 			if (NewAdditionalResource.IsValid())
@@ -288,17 +292,22 @@ bool FDisplayClusterRenderTargetManager::AllocateFrameTargets(const FDisplayClus
 			for (FDisplayClusterRenderFrameTargetView& ViewIt : ViewFamilieIt.Views)
 			{
 				FDisplayClusterViewport* ViewportPtr = static_cast<FDisplayClusterViewport*>(ViewIt.Viewport.Get());
-				if (ViewportPtr && ViewportPtr->RenderSettings.bVisible)
+				if (ViewportPtr && ViewportPtr->GetRenderSettings().bVisible)
 				{
-					ViewportPtr->Resources[EDisplayClusterViewportResource::OutputFrameTargetableResources] = NewFrameTargetResources;
-					ViewportPtr->Resources[EDisplayClusterViewportResource::AdditionalFrameTargetableResources] = NewAdditionalFrameTargetableResources;
+					ViewportPtr->GetViewportResourcesImpl(EDisplayClusterViewportResource::OutputFrameTargetableResources) = NewFrameTargetResources;
+					ViewportPtr->GetViewportResourcesImpl(EDisplayClusterViewportResource::AdditionalFrameTargetableResources) = NewAdditionalFrameTargetableResources;
 
 					// Adjust viewports frame rects. This offset saved in 'BackbufferFrameOffset'
-					if (ViewportPtr->Contexts.IsValidIndex(ViewIt.ContextNum))
+					if (ViewportPtr->GetContexts().IsValidIndex(ViewIt.ContextNum))
 					{
-						FDisplayClusterViewport_Context& Context = ViewportPtr->Contexts[ViewIt.ContextNum];
+						TArray<FDisplayClusterViewport_Context> ViewportContexts = ViewportPtr->GetContexts();
+
+						FDisplayClusterViewport_Context& Context = ViewportContexts[ViewIt.ContextNum];
 						Context.FrameTargetRect.Min -= InOutRenderFrame.FrameRect.Min;
 						Context.FrameTargetRect.Max -= InOutRenderFrame.FrameRect.Min;
+
+						// Override data in viewport context:
+						ViewportPtr->SetContexts(ViewportContexts);
 					}
 				}
 			}
diff --git a/Engine/Plugins/Runtime/nDisplay/Source/DisplayCluster/Private/Render/Viewport/RenderTarget/DisplayClusterRenderTargetManager.h b/Engine/Plugins/Runtime/nDisplay/Source/DisplayCluster/Private/Render/Viewport/RenderTarget/DisplayClusterRenderTargetManager.h
--- a/Engine/Plugins/Runtime/nDisplay/Source/DisplayCluster/Private/Render/Viewport/RenderTarget/DisplayClusterRenderTargetManager.h
+++ b/Engine/Plugins/Runtime/nDisplay/Source/DisplayCluster/Private/Render/Viewport/RenderTarget/DisplayClusterRenderTargetManager.h
@@ -7,24 +7,26 @@
 class FViewport;
 class FDisplayClusterViewport;
 class FDisplayClusterRenderTargetResourcesPool;
-class FDisplayClusterViewportManagerProxy;
+class FDisplayClusterViewportConfiguration;
 class FDisplayClusterRenderFrame;
-struct FDisplayClusterRenderFrameSettings;
 
 class FDisplayClusterRenderTargetManager
 {
 public:
-	FDisplayClusterRenderTargetManager(FDisplayClusterViewportManagerProxy* InViewportManagerProxy);
+	FDisplayClusterRenderTargetManager(const TSharedRef<FDisplayClusterViewportConfiguration, ESPMode::ThreadSafe>& InConfiguration);
 	~FDisplayClusterRenderTargetManager();
 
 	void Release();
 
 public:
-	bool AllocateRenderFrameResources(FViewport* InViewport, const FDisplayClusterRenderFrameSettings& InRenderFrameSettings, const TArray<TSharedPtr<FDisplayClusterViewport, ESPMode::ThreadSafe>>& InViewports, FDisplayClusterRenderFrame& InOutRenderFrame);
+	bool AllocateRenderFrameResources(FViewport* InViewport, const TArray<TSharedPtr<FDisplayClusterViewport, ESPMode::ThreadSafe>>& InViewports, FDisplayClusterRenderFrame& InOutRenderFrame);
 
 protected:
-	bool AllocateFrameTargets(const FDisplayClusterRenderFrameSettings& InRenderFrameSettings, const FIntPoint& InViewportSize, FDisplayClusterRenderFrame& InOutRenderFrame);
+	bool AllocateFrameTargets(const FIntPoint& InViewportSize, FDisplayClusterRenderFrame& InOutRenderFrame);
 
 private:
 	TUniquePtr<FDisplayClusterRenderTargetResourcesPool> ResourcesPool;
+
+	// Configuration of the current cluster node
+	const TSharedRef<FDisplayClusterViewportConfiguration, ESPMode::ThreadSafe> Configuration;
 };
diff --git a/Engine/Plugins/Runtime/nDisplay/Source/DisplayCluster/Private/Render/Viewport/RenderTarget/DisplayClusterRenderTargetResourcesPool.cpp b/Engine/Plugins/Runtime/nDisplay/Source/DisplayCluster/Private/Render/Viewport/RenderTarget/DisplayClusterRenderTargetResourcesPool.cpp
--- a/Engine/Plugins/Runtime/nDisplay/Source/DisplayCluster/Private/Render/Viewport/RenderTarget/DisplayClusterRenderTargetResourcesPool.cpp
+++ b/Engine/Plugins/Runtime/nDisplay/Source/DisplayCluster/Private/Render/Viewport/RenderTarget/DisplayClusterRenderTargetResourcesPool.cpp
@@ -3,8 +3,8 @@
 #include "DisplayClusterRenderTargetResourcesPool.h"
 #include "DisplayClusterRenderTargetResource.h"
 
+#include "Render/Viewport/Configuration/DisplayClusterViewportConfiguration.h"
 #include "Render/Viewport/Resource/DisplayClusterViewportPreviewResource.h"
-
 #include "Render/Viewport/RenderFrame/DisplayClusterRenderFrameSettings.h"
 #include "Render/Viewport/DisplayClusterViewportHelpers.h"
 #include "Render/Viewport/DisplayClusterViewportManagerProxy.h"
@@ -16,10 +16,6 @@
 ////////////////////////////////////////////////////////////////////////////////////////////////////////////
 /// FDisplayClusterRenderTargetResourcesPool
 ////////////////////////////////////////////////////////////////////////////////////////////////////////////
-FDisplayClusterRenderTargetResourcesPool::FDisplayClusterRenderTargetResourcesPool(FDisplayClusterViewportManagerProxy* InViewportManagerProxy)
-	: ViewportManagerProxyWeakPtr(InViewportManagerProxy ? InViewportManagerProxy->AsWeak() : nullptr)
-{ }
-
 FDisplayClusterRenderTargetResourcesPool::~FDisplayClusterRenderTargetResourcesPool()
 {
 	Release();
@@ -83,7 +79,7 @@ void FDisplayClusterRenderTargetResourcesPool::ImplUpdateResources(TArray<TShare
 	InOutViewportResources.Empty();
 }
 
-bool FDisplayClusterRenderTargetResourcesPool::BeginReallocateResources(const FDisplayClusterRenderFrameSettings& InRenderFrameSettings, FViewport* InViewport)
+bool FDisplayClusterRenderTargetResourcesPool::BeginReallocateResources(FViewport* InViewport, const FDisplayClusterRenderFrameSettings& InRenderFrameSettings)
 {
 	check(ResourceSettings == nullptr);
 
diff --git a/Engine/Plugins/Runtime/nDisplay/Source/DisplayCluster/Private/Render/Viewport/RenderTarget/DisplayClusterRenderTargetResourcesPool.h b/Engine/Plugins/Runtime/nDisplay/Source/DisplayCluster/Private/Render/Viewport/RenderTarget/DisplayClusterRenderTargetResourcesPool.h
--- a/Engine/Plugins/Runtime/nDisplay/Source/DisplayCluster/Private/Render/Viewport/RenderTarget/DisplayClusterRenderTargetResourcesPool.h
+++ b/Engine/Plugins/Runtime/nDisplay/Source/DisplayCluster/Private/Render/Viewport/RenderTarget/DisplayClusterRenderTargetResourcesPool.h
@@ -5,16 +5,12 @@
 #include "CoreMinimal.h"
 #include "Render/Viewport/Resource/DisplayClusterViewportResource.h"
 
-class FDisplayClusterViewportManagerProxy;
-struct FDisplayClusterRenderFrameSettings;
-
 /**
  * DC viewport resources pool
  */
 class FDisplayClusterRenderTargetResourcesPool
 {
 public:
-	FDisplayClusterRenderTargetResourcesPool(FDisplayClusterViewportManagerProxy* InViewportManagerProxy);
 	~FDisplayClusterRenderTargetResourcesPool();
 
 	void Release();
@@ -27,7 +23,7 @@ public:
 	* 
 	* @return true, if success
 	*/
-	bool BeginReallocateResources(const FDisplayClusterRenderFrameSettings& InRenderFrameSettings, class FViewport* InViewport);
+	bool BeginReallocateResources(class FViewport* InViewport, const struct FDisplayClusterRenderFrameSettings& InRenderFrameSettings);
 	
 	/** Allocate a new resource or reuse exists
 	* 
@@ -60,19 +56,10 @@ private:
 	*/
 	void ImplUpdateResources(TArray<TSharedPtr<FDisplayClusterViewportResource, ESPMode::ThreadSafe>>& InOutViewportResources, const EResourceUpdateMode InUpdateMode);
 
-	/** Retrieving a ViewportManager instance. */
-	inline FDisplayClusterViewportManagerProxy* GetViewportManagerProxy() const
-	{
-		return ViewportManagerProxyWeakPtr.IsValid() ? ViewportManagerProxyWeakPtr.Pin().Get() : nullptr;
-	}
-
 private:
 	// Current render resource settings
 	FDisplayClusterViewportResourceSettings* ResourceSettings = nullptr;
 
 	// Viewport render resources
 	TArray<TSharedPtr<FDisplayClusterViewportResource, ESPMode::ThreadSafe>> ViewportResources;
-
-	// Weak ref to the viewport manager
-	TWeakPtr<FDisplayClusterViewportManagerProxy, ESPMode::ThreadSafe> ViewportManagerProxyWeakPtr;
 };
diff --git a/Engine/Plugins/Runtime/nDisplay/Source/DisplayCluster/Private/Render/Viewport/Resource/DisplayClusterViewportResourceSettings.cpp b/Engine/Plugins/Runtime/nDisplay/Source/DisplayCluster/Private/Render/Viewport/Resource/DisplayClusterViewportResourceSettings.cpp
--- a/Engine/Plugins/Runtime/nDisplay/Source/DisplayCluster/Private/Render/Viewport/Resource/DisplayClusterViewportResourceSettings.cpp
+++ b/Engine/Plugins/Runtime/nDisplay/Source/DisplayCluster/Private/Render/Viewport/Resource/DisplayClusterViewportResourceSettings.cpp
@@ -35,21 +35,15 @@ FDisplayClusterViewportResourceSettings::FDisplayClusterViewportResourceSettings
 		// Always use srgb for preview rendering
 		EnumAddFlags(ResourceFlags, EDisplayClusterViewportResourceSettingsFlags::ShouldUseSRGB);
 
-		switch (InRenderFrameSettings.RenderMode)
+		if (InRenderFrameSettings.ShouldUseLinearGamma())
 		{
-		case EDisplayClusterRenderFrameMode::PreviewInScene:
-			Format = FDisplayClusterViewportHelpers::GetPreviewDefaultPixelFormat();
-
-			// Preview display gamma
-			if (InRenderFrameSettings.bPreviewEnablePostProcess == false)
-			{
-				// Disable postprocess for preview. Use Gamma 1.f
 				DisplayGamma = 1.f;
 			}
-			break;
 
-		default:
-			break;
+		if(InRenderFrameSettings.IsPreviewRendering())
+		{
+			// for preview rendering use custom pixel format
+			Format = FDisplayClusterViewportHelpers::GetPreviewDefaultPixelFormat();
 		}
 	}
 }
diff --git a/Engine/Plugins/Runtime/nDisplay/Source/DisplayCluster/Public/Components/DisplayClusterCameraComponent.h b/Engine/Plugins/Runtime/nDisplay/Source/DisplayCluster/Public/Components/DisplayClusterCameraComponent.h
--- a/Engine/Plugins/Runtime/nDisplay/Source/DisplayCluster/Public/Components/DisplayClusterCameraComponent.h
+++ b/Engine/Plugins/Runtime/nDisplay/Source/DisplayCluster/Public/Components/DisplayClusterCameraComponent.h
@@ -12,6 +12,7 @@ class UBillboardComponent;
 class UTexture2D;
 class IDisplayClusterViewportManager;
 class IDisplayClusterWarpPolicy;
+class IDisplayClusterViewportConfiguration;
 struct FMinimalViewInfo;
 
 UENUM()
@@ -38,11 +39,21 @@ public:
 
 	/** Return ViewPoint for this component
 	 * If the component logic supports postprocess, it will also be in the ViewInfo structure.
-	 * 
+	 *
 	 * @param InOutViewInfo - ViewInfo data
 	 * @param OutCustomNearClippingPlane - Custom NCP, or a value less than zero if not defined.
 	 */
-	virtual void GetDesiredView(FMinimalViewInfo& InOutViewInfo, float* OutCustomNearClippingPlane = nullptr);
+	UE_DEPRECATED(5.4, "This function has been deprecated. Please use 'GetDesiredView()'.")
+	virtual void GetDesiredView(FMinimalViewInfo& InOutViewInfo, float* OutCustomNearClippingPlane = nullptr)
+	{ }
+
+	/** Return ViewPoint for this component
+	 * If the component logic supports postprocess, it will also be in the ViewInfo structure.
+	 *
+	 * @param InOutViewInfo - ViewInfo data
+	 * @param OutCustomNearClippingPlane - Custom NCP, or a value less than zero if not defined.
+	 */
+	virtual void GetDesiredView(IDisplayClusterViewportConfiguration& InViewportConfiguration, FMinimalViewInfo& InOutViewInfo, float* OutCustomNearClippingPlane = nullptr);
 
 	/**
 	 * All cluster viewports that reference this component will be created in the background on the current cluster node if the function returns true.
diff --git a/Engine/Plugins/Runtime/nDisplay/Source/DisplayCluster/Public/Components/DisplayClusterICVFXCameraComponent.h b/Engine/Plugins/Runtime/nDisplay/Source/DisplayCluster/Public/Components/DisplayClusterICVFXCameraComponent.h
--- a/Engine/Plugins/Runtime/nDisplay/Source/DisplayCluster/Public/Components/DisplayClusterICVFXCameraComponent.h
+++ b/Engine/Plugins/Runtime/nDisplay/Source/DisplayCluster/Public/Components/DisplayClusterICVFXCameraComponent.h
@@ -54,7 +54,11 @@ public:
 	UCineCameraComponent* GetActualCineCameraComponent();
 
 	// Returns true if this camera is active
-	bool IsICVFXEnabled() const;
+	UE_DEPRECATED(5.4, "This function has been deprecated. Please use 'GetCameraSettingsICVFX().IsICVFXEnabled()'.")
+	bool IsICVFXEnabled() const
+	{
+		return false;
+	}
 
 	// Return unique camera name
 	FString GetCameraUniqueId() const;
diff --git a/Engine/Plugins/Runtime/nDisplay/Source/DisplayCluster/Public/Components/DisplayClusterPreviewComponent.h b/Engine/Plugins/Runtime/nDisplay/Source/DisplayCluster/Public/Components/DisplayClusterPreviewComponent.h
--- a/Engine/Plugins/Runtime/nDisplay/Source/DisplayCluster/Public/Components/DisplayClusterPreviewComponent.h
+++ b/Engine/Plugins/Runtime/nDisplay/Source/DisplayCluster/Public/Components/DisplayClusterPreviewComponent.h
@@ -124,6 +124,10 @@ private:
 protected:
 	bool UpdatePreviewMesh();
 	void ReleasePreviewMesh();
+
+	/** Releases the MovablePreviewMesh property if it exists. */
+	void ReleaseMovablePreviewMesh();
+
 	void UpdatePreviewMeshReference();
 
 	void ReleasePreviewMaterial();
diff --git a/Engine/Plugins/Runtime/nDisplay/Source/DisplayCluster/Public/DisplayClusterRootActor.h b/Engine/Plugins/Runtime/nDisplay/Source/DisplayCluster/Public/DisplayClusterRootActor.h
--- a/Engine/Plugins/Runtime/nDisplay/Source/DisplayCluster/Public/DisplayClusterRootActor.h
+++ b/Engine/Plugins/Runtime/nDisplay/Source/DisplayCluster/Public/DisplayClusterRootActor.h
@@ -104,10 +104,18 @@ public:
 	/** Returns the current rendering mode of this DCRA (not a value from configuration).
 	 * This value can be overridden from DCRenderDevice or other rendering subsystems (e.g. Preview).
 	 */
-	EDisplayClusterRenderFrameMode GetRenderMode() const;
+	UE_DEPRECATED(5.4, "This function has been deprecated.")
+	EDisplayClusterRenderFrameMode GetRenderMode() const
+	{
+		return EDisplayClusterRenderFrameMode::Unknown;
+	}
 
 	/** Returns the preview rendering mode of this DCRA. */
-	EDisplayClusterRenderFrameMode GetPreviewRenderMode() const;
+	UE_DEPRECATED(5.4, "This function has been deprecated.")
+	EDisplayClusterRenderFrameMode GetPreviewRenderMode() const
+	{
+		return EDisplayClusterRenderFrameMode::Unknown;
+	}
 
 protected:
 	//////////////////////////////////////////////////////////////////////////////////////////////
@@ -199,6 +207,9 @@ public:
 public:
 	/** Get ViewportManager API. */
 	IDisplayClusterViewportManager* GetViewportManager() const;
+
+	/** Get ViewportConfiguration API.*/
+	IDisplayClusterViewportConfiguration* GetViewportConfiguration() const;
 	
 	static FName GetCurrentConfigDataMemberName()
 	{
@@ -372,6 +383,9 @@ public:
 	UPROPERTY()
 	bool bMoviePipelineRenderPass = false;
 
+#endif
+	// Some preview settings have been opened be available for Standalone/Package builds:
+
 	/** Render the scene and display it as a preview on the nDisplay root actor in the editor.  This will impact editor performance. */
 	UPROPERTY(EditAnywhere, BlueprintReadWrite, Category = "Editor Preview", meta = (DisplayName = "Enable Editor Preview"))
 	bool bPreviewEnable = true;
@@ -384,6 +398,7 @@ public:
 	UPROPERTY(EditAnywhere, BlueprintReadWrite, Category = "Editor Preview", meta = (DisplayName = "Enable Post Process"), BlueprintSetter = SetPreviewEnablePostProcess)
 	bool bPreviewEnablePostProcess = false;
 
+#if WITH_EDITORONLY_DATA
 	/** Configure the root actor for tech viz rendering with preview components. */
 	UPROPERTY(EditAnywhere, BlueprintReadWrite, Category = "Editor Preview")
 	bool bEnablePreviewTechViz = true;
@@ -407,6 +422,9 @@ public:
 	/** Render Mode */
 	UPROPERTY(EditAnywhere, BlueprintReadWrite, Category = "Editor Preview", meta = (DisplayName = "Render Mode"))
 	EDisplayClusterConfigurationRenderMode RenderMode = EDisplayClusterConfigurationRenderMode::Mono;
+#endif
+
+	// Some preview settings have been opened be available for Standalone/Package builds:
 
 	/** Tick Per Frame */
 	UPROPERTY(EditAnywhere, BlueprintReadWrite, Category = "Editor Preview", AdvancedDisplay, meta = (DisplayName = "Tick Per Frame", ClampMin = "1", UIMin = "1", ClampMax = "200", UIMax = "200"))
@@ -420,6 +438,10 @@ public:
 	UPROPERTY(EditAnywhere, BlueprintReadWrite, Category = "Editor Preview", AdvancedDisplay, meta = (DisplayName = "Preview Texture Max Size", ClampMin = "64", UIMin = "64", ClampMax = "4096", UIMax = "4096"))
 	int PreviewMaxTextureDimension = 2048;
 
+//////////////////////////////////////////////////////////////////////////////////////////////
+// EDITOR RELATED SETTINGS
+//////////////////////////////////////////////////////////////////////////////////////////////
+#if WITH_EDITORONLY_DATA
 private:
 	UPROPERTY(Transient)
 	TMap<FString, TObjectPtr<UDisplayClusterPreviewComponent>> PreviewComponents;
@@ -569,6 +591,9 @@ protected:
 
 	bool ImplUpdatePreviewConfiguration_Editor(const FString& InClusterNodeId);
 
+	/** Gets a preview rendering mode for this root actor. */
+	EDisplayClusterRenderFrameMode GetPreviewRenderMode_Editor() const;
+
 	void ImplRenderPreview_Editor();
 	bool ImplRenderPassPreviewClusterNode_Editor(const FString& InClusterNodeId);
 
diff --git a/Engine/Plugins/Runtime/nDisplay/Source/DisplayCluster/Public/DisplayClusterSceneViewExtensions.h b/Engine/Plugins/Runtime/nDisplay/Source/DisplayCluster/Public/DisplayClusterSceneViewExtensions.h
--- a/Engine/Plugins/Runtime/nDisplay/Source/DisplayCluster/Public/DisplayClusterSceneViewExtensions.h
+++ b/Engine/Plugins/Runtime/nDisplay/Source/DisplayCluster/Public/DisplayClusterSceneViewExtensions.h
@@ -4,8 +4,7 @@
 
 #include "CoreMinimal.h"
 #include "SceneViewExtension.h"
-
-class IDisplayClusterViewportManager;
+#include "Render/Viewport/IDisplayClusterViewportConfiguration.h"
 
 /**
  * Contains information about the context in which this scene view extension will be used.
@@ -27,34 +26,22 @@ public:
 		: FSceneViewExtensionContext()
 	{ }
 
-	FDisplayClusterSceneViewExtensionContext(FViewport* InViewport, const TSharedPtr<IDisplayClusterViewportManager, ESPMode::ThreadSafe>& InViewportManager, const FString& InViewportId)
+	FDisplayClusterSceneViewExtensionContext(FViewport* InViewport, const TSharedRef<IDisplayClusterViewportConfiguration, ESPMode::ThreadSafe>& InConfiguration, const FString& InViewportId)
 		: FSceneViewExtensionContext(InViewport)
 		, ViewportId(InViewportId)
-		, ViewportManagerWeakPtr(InViewportManager)
+		, Configuration(InConfiguration)
 	{ }
 
-	FDisplayClusterSceneViewExtensionContext(FSceneInterface* InScene, const TSharedPtr<IDisplayClusterViewportManager, ESPMode::ThreadSafe>& InViewportManager, const FString& InViewportId)
+	FDisplayClusterSceneViewExtensionContext(FSceneInterface* InScene, const TSharedRef<IDisplayClusterViewportConfiguration, ESPMode::ThreadSafe>& InConfiguration, const FString& InViewportId)
 		: FSceneViewExtensionContext(InScene)
 		, ViewportId(InViewportId)
-		, ViewportManagerWeakPtr(InViewportManager)
+		, Configuration(InConfiguration)
 	{ }
 
-	/** Returns a pointer to the DC ViewportManager used for this VE context. */
-	IDisplayClusterViewportManager* GetViewportManager() const
-	{
-		return ViewportManagerWeakPtr.IsValid() ? ViewportManagerWeakPtr.Pin().Get() : nullptr;
-	}
-
-	/** Returns the name of the DC viewport used for this VE context. */
-	const FString& GetViewportId() const
-	{
-		return ViewportId;
-	}
-
-private:
+public:
 	// The id of the nDisplay viewport being rendered.
 	const FString ViewportId;
 
 	// Reference to viewport manager
-	const TWeakPtr<IDisplayClusterViewportManager, ESPMode::ThreadSafe> ViewportManagerWeakPtr;
+	const TSharedPtr<IDisplayClusterViewportConfiguration, ESPMode::ThreadSafe> Configuration;
 };
diff --git a/Engine/Plugins/Runtime/nDisplay/Source/DisplayCluster/Public/Game/IDisplayClusterGameManager.h b/Engine/Plugins/Runtime/nDisplay/Source/DisplayCluster/Public/Game/IDisplayClusterGameManager.h
--- a/Engine/Plugins/Runtime/nDisplay/Source/DisplayCluster/Public/Game/IDisplayClusterGameManager.h
+++ b/Engine/Plugins/Runtime/nDisplay/Source/DisplayCluster/Public/Game/IDisplayClusterGameManager.h
@@ -29,4 +29,9 @@ public:
 	* @return - Current world
 	*/
 	virtual UWorld* GetWorld() const = 0;
+
+	/**
+	* @return - All root actors in world
+	*/
+	virtual TArray<ADisplayClusterRootActor*> GetAllRootActorsFromWorld(UWorld* InWorld) = 0;
 };
diff --git a/Engine/Plugins/Runtime/nDisplay/Source/DisplayCluster/Public/Misc/DisplayClusterObjectRef.h b/Engine/Plugins/Runtime/nDisplay/Source/DisplayCluster/Public/Misc/DisplayClusterObjectRef.h
--- a/Engine/Plugins/Runtime/nDisplay/Source/DisplayCluster/Public/Misc/DisplayClusterObjectRef.h
+++ b/Engine/Plugins/Runtime/nDisplay/Source/DisplayCluster/Public/Misc/DisplayClusterObjectRef.h
@@ -206,7 +206,7 @@ public:
 	// For killed object ptr, reset and find component new object ptr by name and save to [mutable] ComponentPtr
 	USceneComponent* GetOrFindSceneComponent() const;
 
-	bool SetSceneComponent(USceneComponent* InComponent)
+	bool SetSceneComponent(USceneComponent* const InComponent)
 	{
 		FScopeLock lock(&DataGuard);
 
diff --git a/Engine/Plugins/Runtime/nDisplay/Source/DisplayCluster/Public/Render/Projection/IDisplayClusterProjectionPolicy.h b/Engine/Plugins/Runtime/nDisplay/Source/DisplayCluster/Public/Render/Projection/IDisplayClusterProjectionPolicy.h
--- a/Engine/Plugins/Runtime/nDisplay/Source/DisplayCluster/Public/Render/Projection/IDisplayClusterProjectionPolicy.h
+++ b/Engine/Plugins/Runtime/nDisplay/Source/DisplayCluster/Public/Render/Projection/IDisplayClusterProjectionPolicy.h
@@ -85,12 +85,16 @@ public:
 	/**
 	 * Called before FDisplayClusterViewport::UpdateFrameContexts()
 	 * From this function, the policy can override any viewport settings (custom overscan, etc).
+	 * 
+	 * @param InViewport - a owner viewport
 	 */
 	virtual void BeginUpdateFrameContexts(IDisplayClusterViewport* InViewport) const
 	{ }
 
 	/**
 	 * Called after FDisplayClusterViewport::UpdateFrameContexts()
+	 * 
+	 * @param InViewport - a owner viewport
 	 */
 	virtual void EndUpdateFrameContexts(IDisplayClusterViewport* InViewport) const
 	{ }
@@ -136,13 +140,6 @@ public:
 		return false;
 	}
 
-	// This policy can support ICVFX rendering
-	UE_DEPRECATED(5.3, "This function has been deprecated. Please use 'ShouldSupportICVFX(IDisplayClusterViewport*)'.")
-	virtual bool ShouldSupportICVFX() const
-	{
-		return false;
-	}
-
 	/** Returns true if this policy supports ICVFX rendering
 	 * 
 	 * @param InViewport - a owner viewport
@@ -170,7 +167,7 @@ public:
 
 	/** Override view from this projection policy
 	 *
-	 * @param InViewport                 - the viewport of this projection policy
+	 * @param InViewport                 - a owner viewport
 	 * @param InDeltaTime                - delta time in current frame
 	 * @param InOutViewInfo              - ViewInfo data
 	 * @param OutCustomNearClippingPlane - Custom NCP, or a value less than zero if not defined.
@@ -178,11 +175,16 @@ public:
 	virtual void SetupProjectionViewPoint(IDisplayClusterViewport* InViewport, const float InDeltaTime, FMinimalViewInfo& InOutViewInfo, float* OutCustomNearClippingPlane = nullptr)
 	{ }
 
-	/** Projection policy can override PP */
+	/** Projection policy can override PP
+	* 
+	* @param InViewport - a owner viewport
+	*/
 	virtual void UpdatePostProcessSettings(IDisplayClusterViewport* InViewport)
 	{ }
 
-	/**
+	/** Calculate view projection data
+	* 
+	* @param InViewport        - a owner viewport
 	* @param ViewIdx           - Index of view that is being processed for this viewport
 	* @param InOutViewLocation - (in/out) View location with ViewOffset (i.e. left eye pre-computed location)
 	* @param InOutViewRotation - (in/out) View rotation
@@ -195,7 +197,9 @@ public:
 	*/
 	virtual bool CalculateView(IDisplayClusterViewport* InViewport, const uint32 InContextNum, FVector& InOutViewLocation, FRotator& InOutViewRotation, const FVector& ViewOffset, const float WorldToMeters, const float NCP, const float FCP) = 0;
 
-	/**
+	/** Gets projection matrix
+	* 
+	* @param InViewport   - a owner viewport
 	* @param ViewIdx      - Index of view that is being processed for this viewport
 	* @param OutPrjMatrix - (out) projection matrix
 	*
@@ -287,9 +291,10 @@ public:
 	/**
 	* Ask projection policy instance if it has any mesh based preview
 	*
+	* @param InViewport - a owner viewport
 	* @return - True if mesh based preview is available
 	*/
-	virtual bool HasPreviewMesh()
+	virtual bool HasPreviewMesh(IDisplayClusterViewport* InViewport)
 	{
 		return false;
 	}
@@ -308,12 +313,23 @@ public:
 		return nullptr;
 	}
 
+	/**
+	 * Return Origin point component used by preview mesh
+	 * 
+	 * @param InViewport - a owner viewport
+	 */
+	virtual USceneComponent* const GetPreviewMeshOriginComponent(IDisplayClusterViewport* InViewport) const
+	{
+		return nullptr;
+	}
+
 	/**
 	* Ask projection policy instance if it has any movable mesh based preview
-	*
+	* 
+	* @param InViewport - a owner viewport
 	* @return - True if mesh based preview is available
 	*/
-	virtual bool HasPreviewMovableMesh()
+	virtual bool HasPreviewMovableMesh(IDisplayClusterViewport* InViewport)
 	{
 		return false;
 	}
@@ -322,10 +338,42 @@ public:
 	* Build preview movable mesh
 	* This MeshComponent is a copy of the preview mesh and can be moved freely with the UI visualization.
 	*
-	* @param InViewport - Projection specific parameters.
+	* @param InViewport - a owner viewport
 	*/
 	virtual UMeshComponent* GetOrCreatePreviewMovableMeshComponent(IDisplayClusterViewport* InViewport)
 	{
 		return nullptr;
 	}
+
+	/**
+	 * Return Origin point component used by preview movable mesh
+	 * 
+	 * @param InViewport - a owner viewport
+	 */
+	virtual USceneComponent* const GetPreviewMovableMeshOriginComponent(IDisplayClusterViewport* InViewport) const
+	{
+		return nullptr;
+	}
+
+	//////////// UE_DEPRECATED 5.3 ////////////
+
+	// This policy can support ICVFX rendering
+	UE_DEPRECATED(5.3, "This function has been deprecated. Please use 'ShouldSupportICVFX(IDisplayClusterViewport*)'.")
+		virtual bool ShouldSupportICVFX() const
+	{
+		return false;
+	}
+
+	//////////// UE_DEPRECATED 5.4 ////////////
+
+	/**
+	* Ask projection policy instance if it has any mesh based preview
+	*
+	* @return - True if mesh based preview is available
+	*/
+	UE_DEPRECATED(5.4, "This function has been deprecated. Please use 'HasPreviewMesh(IDisplayClusterViewport*)'.")
+	virtual bool HasPreviewMesh()
+	{
+		return false;
+	}
 };
diff --git a/Engine/Plugins/Runtime/nDisplay/Source/DisplayCluster/Public/Render/Viewport/Containers/DisplayClusterViewportProxy_Context.h b/Engine/Plugins/Runtime/nDisplay/Source/DisplayCluster/Public/Render/Viewport/Containers/DisplayClusterViewportProxy_Context.h
--- a/Engine/Plugins/Runtime/nDisplay/Source/DisplayCluster/Public/Render/Viewport/Containers/DisplayClusterViewportProxy_Context.h
+++ b/Engine/Plugins/Runtime/nDisplay/Source/DisplayCluster/Public/Render/Viewport/Containers/DisplayClusterViewportProxy_Context.h
@@ -4,9 +4,12 @@
 
 #include "CoreMinimal.h"
 
-class FDisplayClusterViewportProxy_Context
+/**
+* The viewport proxy context (rendering thread data)
+*/
+struct FDisplayClusterViewportProxy_Context
 {
-public:
+	// The number of context
 	uint32  ContextNum = 0;
 
 	// Profile description from ViewFamily
diff --git a/Engine/Plugins/Runtime/nDisplay/Source/DisplayCluster/Public/Render/Viewport/Containers/DisplayClusterViewport_CameraMotionBlur.h b/Engine/Plugins/Runtime/nDisplay/Source/DisplayCluster/Public/Render/Viewport/Containers/DisplayClusterViewport_CameraMotionBlur.h
--- a/Engine/Plugins/Runtime/nDisplay/Source/DisplayCluster/Public/Render/Viewport/Containers/DisplayClusterViewport_CameraMotionBlur.h
+++ b/Engine/Plugins/Runtime/nDisplay/Source/DisplayCluster/Public/Render/Viewport/Containers/DisplayClusterViewport_CameraMotionBlur.h
@@ -4,6 +4,9 @@
 
 #include "CoreMinimal.h"
 
+/**
+* Camera motion blur mode.
+*/
 enum class EDisplayClusterViewport_CameraMotionBlur : uint8
 {
 	Undefined = 0,
@@ -12,6 +15,9 @@ enum class EDisplayClusterViewport_CameraMotionBlur : uint8
 	Override
 };
 
+/**
+* Settings for the CameraMotion blur.
+*/
 struct FDisplayClusterViewport_CameraMotionBlur
 {
 	EDisplayClusterViewport_CameraMotionBlur Mode = EDisplayClusterViewport_CameraMotionBlur::Undefined;
@@ -21,4 +27,3 @@ struct FDisplayClusterViewport_CameraMotionBlur
 
 	float    TranslationScale = 1.f;
 };
-
diff --git a/Engine/Plugins/Runtime/nDisplay/Source/DisplayCluster/Public/Render/Viewport/Containers/DisplayClusterViewport_Context.h b/Engine/Plugins/Runtime/nDisplay/Source/DisplayCluster/Public/Render/Viewport/Containers/DisplayClusterViewport_Context.h
--- a/Engine/Plugins/Runtime/nDisplay/Source/DisplayCluster/Public/Render/Viewport/Containers/DisplayClusterViewport_Context.h
+++ b/Engine/Plugins/Runtime/nDisplay/Source/DisplayCluster/Public/Render/Viewport/Containers/DisplayClusterViewport_Context.h
@@ -6,32 +6,7 @@
 #include "StereoRendering.h"
 #include "ShowFlags.h"
 
-/**
- * Special flags for nDisplay Viewport context
- */
-enum class EDisplayClusterViewportContextState : uint8
-{
-	None = 0,
-
-	// The FDisplayClusterViewport::CalculateView() function can be called several times per frame.
-	// Each time it must return the same values. For optimization purposes, after the first call this function
-	// stores the result in the context variables 'ViewLocation' and 'ViewRotation'.
-	// Finally, raises this flag for subsequent calls in the current frame.
-	HasCalculatedViewPoint = 1 << 0,
-
-	// Viewpoint is not valid for this viewport (cannot be calculated)
-	InvalidViewPoint = 1 << 1,
-
-	// The FDisplayClusterViewport::GetProjectionMatrix() function can also be called several times per frame.
-	// stores the result in the context variables 'ProjectionMatrix' and 'OverscanProjectionMatrix'.
-	// Finally, raises this flag for subsequent calls in the current frame.
-	HasCalculatedProjectionMatrix = 1 << 2,
-	HasCalculatedOverscanProjectionMatrix = 1 << 3,
-
-	// The projection matrix is not valid (cannot be calculated)
-	InvalidProjectionMatrix = 1 << 4,
-};
-ENUM_CLASS_FLAGS(EDisplayClusterViewportContextState);
+#include "Render/Viewport/Containers/DisplayClusterViewport_Enums.h"
 
 /**
  * Viewport context with cahched data and states
@@ -43,11 +18,13 @@ public:
 		: ContextNum(InContextNum)
 		, StereoscopicPass(InStereoscopicPass)
 		, StereoViewIndex(InStereoViewIndex)
-	{}
+	{ }
 
 public:
+	// Index of context (eye #)
 	const uint32            ContextNum;
 
+	// References to IStereoRendering
 	EStereoscopicPass StereoscopicPass;
 	int32             StereoViewIndex;
 
@@ -67,14 +44,11 @@ public:
 	// World scale
 	float WorldToMeters = 100.f;
 
-	//////////////////
-	// Rendering data, for internal usage
-
 	// GPU index for this context render target
 	int32 GPUIndex = INDEX_NONE;
 
-	/* Enables nDisplay's native implementation of cross-GPU transfer.
-	 * This disables cross-GPU transfer by default for nDisplay viewports in FSceneViewFamily structure. **/
+	// Enables nDisplay's native implementation of cross-GPU transfer.
+	// This disables cross-GPU transfer by default for nDisplay viewports in FSceneViewFamily structure.
 	bool bOverrideCrossGPUTransfer = false;
 
 	// Location and size on a render target texture
@@ -95,6 +69,9 @@ public:
 	// Disable render for this viewport (Overlay)
 	bool bDisableRender = false;
 
+	/**
+	* Viewport context data for rendering thread
+	*/
 	struct FRenderThreadData
 	{
 		FRenderThreadData()
diff --git a/Engine/Plugins/Runtime/nDisplay/Source/DisplayCluster/Public/Render/Viewport/Containers/DisplayClusterViewport_Enums.h b/Engine/Plugins/Runtime/nDisplay/Source/DisplayCluster/Public/Render/Viewport/Containers/DisplayClusterViewport_Enums.h
--- a/Engine/Plugins/Runtime/nDisplay/Source/DisplayCluster/Public/Render/Viewport/Containers/DisplayClusterViewport_Enums.h
+++ b/Engine/Plugins/Runtime/nDisplay/Source/DisplayCluster/Public/Render/Viewport/Containers/DisplayClusterViewport_Enums.h
@@ -131,3 +131,30 @@ enum class EDisplayClusterViewport_FrustumUnit: uint8
 	// 1 unit = 1 per cent
 	Percent
 };
+
+/**
+ * Special flags for nDisplay Viewport context
+ */
+enum class EDisplayClusterViewportContextState : uint8
+{
+	None = 0,
+
+	// The FDisplayClusterViewport::CalculateView() function can be called several times per frame.
+	// Each time it must return the same values. For optimization purposes, after the first call this function
+	// stores the result in the context variables 'ViewLocation' and 'ViewRotation'.
+	// Finally, raises this flag for subsequent calls in the current frame.
+	HasCalculatedViewPoint = 1 << 0,
+
+	// Viewpoint is not valid for this viewport (cannot be calculated)
+	InvalidViewPoint = 1 << 1,
+
+	// The FDisplayClusterViewport::GetProjectionMatrix() function can also be called several times per frame.
+	// stores the result in the context variables 'ProjectionMatrix' and 'OverscanProjectionMatrix'.
+	// Finally, raises this flag for subsequent calls in the current frame.
+	HasCalculatedProjectionMatrix = 1 << 2,
+	HasCalculatedOverscanProjectionMatrix = 1 << 3,
+
+	// The projection matrix is not valid (cannot be calculated)
+	InvalidProjectionMatrix = 1 << 4,
+};
+ENUM_CLASS_FLAGS(EDisplayClusterViewportContextState);
diff --git a/Engine/Plugins/Runtime/nDisplay/Source/DisplayCluster/Public/Render/Viewport/Containers/DisplayClusterViewport_PostRenderSettings.h b/Engine/Plugins/Runtime/nDisplay/Source/DisplayCluster/Public/Render/Viewport/Containers/DisplayClusterViewport_PostRenderSettings.h
--- a/Engine/Plugins/Runtime/nDisplay/Source/DisplayCluster/Public/Render/Viewport/Containers/DisplayClusterViewport_PostRenderSettings.h
+++ b/Engine/Plugins/Runtime/nDisplay/Source/DisplayCluster/Public/Render/Viewport/Containers/DisplayClusterViewport_PostRenderSettings.h
@@ -8,6 +8,9 @@
 #include "ShaderParameters/DisplayClusterShaderParameters_PostprocessBlur.h"
 #include "ShaderParameters/DisplayClusterShaderParameters_GenerateMips.h"
 
+/**
+* Additional DC viewport rendering settings
+*/
 class FDisplayClusterViewport_PostRenderSettings
 {
 public:
@@ -17,14 +20,13 @@ public:
 	~FDisplayClusterViewport_PostRenderSettings()
 	{}
 
-	void SetParameters(const FDisplayClusterViewport_PostRenderSettings& InPostRenderSettings)
+	inline void SetParameters(const FDisplayClusterViewport_PostRenderSettings& InPostRenderSettings)
 	{
 		Replace.SetParameters(InPostRenderSettings.Replace);
 		PostprocessBlur = InPostRenderSettings.PostprocessBlur;
 		GenerateMips = InPostRenderSettings.GenerateMips;
 	}
 
-public:
 	inline void BeginUpdateSettings()
 	{
 		Replace.Reset();
@@ -37,5 +39,3 @@ public:
 	FDisplayClusterShaderParameters_PostprocessBlur PostprocessBlur;
 	FDisplayClusterShaderParameters_GenerateMips    GenerateMips;
 };
-
-
diff --git a/Engine/Plugins/Runtime/nDisplay/Source/DisplayCluster/Public/Render/Viewport/Containers/DisplayClusterViewport_PreviewSettings.h b/Engine/Plugins/Runtime/nDisplay/Source/DisplayCluster/Public/Render/Viewport/Containers/DisplayClusterViewport_PreviewSettings.h
new file mode 100644
--- /dev/null
+++ b/Engine/Plugins/Runtime/nDisplay/Source/DisplayCluster/Public/Render/Viewport/Containers/DisplayClusterViewport_PreviewSettings.h
@@ -0,0 +1,30 @@
+// Copyright Epic Games, Inc. All Rights Reserved.
+
+#pragma once
+
+#include "CoreMinimal.h"
+
+/**
+* Viewport preview-in-scene rendering settings.
+*/
+struct FDisplayClusterViewport_PreviewSettings
+{
+	// Preview RTT size multiplier
+	float RenderTargetRatioMult = 1.f;
+
+	// Enable/Disable preview rendering. When disabled preview image freeze
+	bool bFreezeRender = false;
+
+	// Hack preview gamma.
+	// In a scene, PostProcess always renders on top of the preview textures.
+	// But in it, PostProcess is also rendered with the flag turned off.
+	bool bEnablePostProcess = false;
+
+	// The maximum dimension of any texture for preview
+	// Limit preview textures max size
+	int32 MaxTextureDimension = 2048;
+
+	// [Experimental] Render preview in multi-GPU
+	// Specifies the mGPU index range for rendering the DCRA preview.
+	TOptional<FIntPoint> MultiGPURendering;
+};
diff --git a/Engine/Plugins/Runtime/nDisplay/Source/DisplayCluster/Public/Render/Viewport/Containers/DisplayClusterViewport_RenderSettings.h b/Engine/Plugins/Runtime/nDisplay/Source/DisplayCluster/Public/Render/Viewport/Containers/DisplayClusterViewport_RenderSettings.h
--- a/Engine/Plugins/Runtime/nDisplay/Source/DisplayCluster/Public/Render/Viewport/Containers/DisplayClusterViewport_RenderSettings.h
+++ b/Engine/Plugins/Runtime/nDisplay/Source/DisplayCluster/Public/Render/Viewport/Containers/DisplayClusterViewport_RenderSettings.h
@@ -53,21 +53,17 @@ public:
 	// Viewport frustum overscan settings
 	FDisplayClusterViewport_OverscanSettings OverscanSettings;
 
-	// Read viewport pixels for preview (this flag is cleared at the end of the frame)
-	bool bPreviewReadPixels = false;
-
 	// Useful to render some viewports in mono, then copied to stereo backbuffers identical image
 	bool bForceMono = false;
 
 	// Is this viewport being captured by a media capture device?
 	bool bIsBeingCaptured = false;
 
-	/** Enable cross-GPU transfer for this viewport.
-	  * It may be disabled in some configurations. For example, when using offscreen rendering with TextureShare,
-	  * cross-gpu transfer can be disabled for this viewport to improve performance, because when transfer is called,
-	  * it freezes the GPUs until synchronization is reached.
-	  * (TextureShare uses its own implementation of the crossGPU transfer for the shared textures.)
-	  */
+	// Enable cross-GPU transfer for this viewport.
+	// It may be disabled in some configurations. For example, when using offscreen rendering with TextureShare,
+	// cross-gpu transfer can be disabled for this viewport to improve performance, because when transfer is called,
+	// it freezes the GPUs until synchronization is reached.
+	// (TextureShare uses its own implementation of the crossGPU transfer for the shared textures.)
 	bool bEnableCrossGPUTransfer = true;
 
 	// Performance, Multi-GPU: Asign GPU for viewport rendering. The Value '-1' used to default gpu mapping
@@ -88,10 +84,6 @@ public:
 	// Viewport can overlap each other on backbuffer. This value uses to sorting order
 	int32 OverlapOrder = 0;
 
-	// Performance: Support special frame builder mode - merge viewports to single viewfamily by group num
-	// [not implemented yet] Experimental
-	int32 RenderFamilyGroup = -1;
-
 	// Special capture modes (chromakey, lightcard) change RTT format and render flags
 	EDisplayClusterViewportCaptureMode CaptureMode = EDisplayClusterViewportCaptureMode::Default;
 
@@ -116,9 +108,7 @@ public:
 
 	/** Finishes setting the viewport in the game thread. Called once per frame at the end. */
 	inline void FinishUpdateSettings()
-	{
-		bPreviewReadPixels = false;
-	}
+	{ }
 
 	/** Returns true if the viewport is assigned to a parent viewport. */
 	inline bool IsViewportHasParent() const
@@ -152,8 +142,6 @@ public:
 
 		GPUIndex = (GPUIndex < 0) ? InParentSettings.GPUIndex : GPUIndex;
 		StereoGPUIndex = (StereoGPUIndex < 0) ? InParentSettings.StereoGPUIndex : StereoGPUIndex;
-
-		RenderFamilyGroup = (RenderFamilyGroup < 0) ? InParentSettings.RenderFamilyGroup : RenderFamilyGroup;
 	}
 
 	/** The viewport can be overridden from another viewport. This function returns true if it is. */
diff --git a/Engine/Plugins/Runtime/nDisplay/Source/DisplayCluster/Public/Render/Viewport/IDisplayClusterViewport.h b/Engine/Plugins/Runtime/nDisplay/Source/DisplayCluster/Public/Render/Viewport/IDisplayClusterViewport.h
--- a/Engine/Plugins/Runtime/nDisplay/Source/DisplayCluster/Public/Render/Viewport/IDisplayClusterViewport.h
+++ b/Engine/Plugins/Runtime/nDisplay/Source/DisplayCluster/Public/Render/Viewport/IDisplayClusterViewport.h
@@ -3,6 +3,7 @@
 #pragma once
 
 #include "Render/Viewport/IDisplayClusterViewport_CustomPostProcessSettings.h"
+#include "Render/Viewport/IDisplayClusterViewportConfiguration.h"
 
 #include "Render/Viewport/Containers/DisplayClusterViewport_Context.h"
 #include "Render/Viewport/Containers/DisplayClusterViewport_RenderSettings.h"
@@ -34,8 +35,19 @@
 	virtual FString GetId() const = 0;
 	virtual FString GetClusterNodeId() const = 0;
 
+	/** Get viewport manager configuration interface. */
+	virtual IDisplayClusterViewportConfiguration& GetConfiguration() = 0;
+
+	/** [const] Get viewport manager configuration interface. */
+	virtual const IDisplayClusterViewportConfiguration& GetConfiguration() const = 0;
+
+	/** Get main viewport render settings. */
 	virtual const FDisplayClusterViewport_RenderSettings&      GetRenderSettings() const = 0;
+
+	/** Get ICVFX render settings. */
 	virtual const FDisplayClusterViewport_RenderSettingsICVFX& GetRenderSettingsICVFX() const = 0;
+
+	/** Get post-render settings. */
 	virtual const FDisplayClusterViewport_PostRenderSettings&  GetPostRenderSettings() const = 0;
 
 	/** Calculate projection matrices for viewport context
@@ -52,24 +64,6 @@
 	 */
 	virtual void CalculateProjectionMatrix(const uint32 InContextNum, float Left, float Right, float Top, float Bottom, float ZNear, float ZFar, bool bIsAnglesInput) = 0;
 
-	/** Calculate viewport projection (This function has been deprecated.)
-	 *
-	 * @param InContextNum      - viewport eye context index
-	 * @param InOutViewLocation - (in, out) View location
-	 * @param InOutViewRotation - (in,out) View rotator
-	 * @param ViewOffset         - eye offset from view location
-	 * @param WorldToMeters     - UE world scale
-	 * @param NCP - Near clipping plane
-	 * @param FCP - Far clipping plane
-	 *
-	 * @return - true, if the calculation of the viewport projection is successful.
-	 */
-	UE_DEPRECATED(5.3, "This function has been deprecated. Please use 'CalculateView()'.")
-	virtual bool CalculateView(const uint32 InContextNum, FVector& InOutViewLocation, FRotator& InOutViewRotation, const FVector& ViewOffset, const float WorldToMeters, const float NCP, const float FCP)
-	{
-		return false;
-	}
-
 	/** Calculate viewport projection
 	 * 
 	 * @param InContextNum      - viewport eye context index
@@ -99,8 +93,10 @@
 	 */
 	virtual bool SetupViewPoint(struct FMinimalViewInfo& InOutViewInfo) = 0;
 
-	/** Return view point camera component for this viewport. */
-	virtual class UDisplayClusterCameraComponent* GetViewPointCameraComponent() const = 0;
+	/** Return view point camera component for this viewport.
+	* @param InRootActorType - the root type
+	*/
+	virtual class UDisplayClusterCameraComponent* GetViewPointCameraComponent(const EDisplayClusterRootActorType InRootActorType) const = 0;
 
 	/** Retrieves the view position from the ViewPoint component used by this viewport.
 	 * The CalculateView() function receives arguments that may not match this position because they can be overridden.
@@ -114,8 +110,10 @@
 	 */
 	virtual float GetStereoEyeOffsetDistance(const uint32 InContextNum) = 0;
 
+	/** Get projection policy ref. */
 	virtual const TSharedPtr<class IDisplayClusterProjectionPolicy, ESPMode::ThreadSafe>& GetProjectionPolicy() const = 0;
 
+	/** Get viewport contexts */
 	virtual const TArray<FDisplayClusterViewport_Context>& GetContexts() const = 0;
 
 	/** Get clipping planes used by this viewport
@@ -124,42 +122,34 @@
 	 */
 	virtual FVector2D GetClippingPlanes() const = 0;
 
-	// Override postprocess settings for this viewport
+	/**
+	* [const] Get custom postprocess settings for this viewport.
+	*/
 	virtual const IDisplayClusterViewport_CustomPostProcessSettings& GetViewport_CustomPostProcessSettings() const = 0;
+
+	/**
+	* Get custom postprocess settings for this viewport.
+	*/
 	virtual IDisplayClusterViewport_CustomPostProcessSettings& GetViewport_CustomPostProcessSettings() = 0;
 
-	// Setup scene view for rendering specified Context
+	/**
+	* Setup scene view for rendering specified Context.
+	*/
 	virtual void SetupSceneView(uint32 ContextNum, class UWorld* World, class FSceneViewFamily& InViewFamily, FSceneView& InView) const = 0;
 
-	UE_DEPRECATED(5.3, "This function has beend deprecate. Please use 'GetViewportManager'.")
-	virtual class IDisplayClusterViewportManager& GetOwner() const
-	{
-		check(false);
-
-		return *GetViewportManager();
-	}
-
-	/** Return the viewport manager that owns this viewport */
-	virtual class IDisplayClusterViewportManager* GetViewportManager() const = 0;
-
-	/** Return the DCRA that owns this viewport */
-	virtual class ADisplayClusterRootActor* GetRootActor() const = 0;
-
-	/** Return current render mode. */
-	virtual EDisplayClusterRenderFrameMode GetRenderMode() const = 0;
-
-	/** Return current world. */
-	virtual class UWorld* GetCurrentWorld() const = 0;
-
-	/** Returns true if the scene is open now (The current world is assigned and DCRA has already initialized for it). */
-	virtual bool IsSceneOpened() const = 0;
-
-	/** Returns true if the current world type is equal to one of the input types. */
-	virtual bool IsCurrentWorldHasAnyType(const EWorldType::Type InWorldType1, const EWorldType::Type InWorldType2 = EWorldType::None, const EWorldType::Type InWorldType3 = EWorldType::None) const = 0;
-
+	/**
+	* Change render setting for viewport.
+	* This function must be called in a valid moment in time.
+	*/
 	virtual void SetRenderSettings(const FDisplayClusterViewport_RenderSettings& InRenderSettings) = 0;
+
+	/**
+	* Change viewport contexts data.
+	* This function must be called in a valid moment in time.
+	*/
 	virtual void SetContexts(TArray<FDisplayClusterViewport_Context>& InContexts) = 0;
 
+	/** Make asymmetrical projection matrix for nDisplay. */
 	static FMatrix MakeProjectionMatrix(float InLeft, float InRight, float InTop, float InBottom, float ZNear, float ZFar);
 
 	/** Get View from CameraComponent
@@ -183,4 +173,82 @@
 	* @return false if the viewpoint data cannot be retrieved.
 	*/
 	static bool GetPlayerCameraView(UWorld* InWorld, const bool bUseCameraPostprocess, FMinimalViewInfo& InOutViewInfo);
+
+	///////////////// UE_DEPRECATED 5.3 ///////////////////
+
+	UE_DEPRECATED(5.3, "This function has beend deprecate. Please use 'GetViewportManager()'.")
+	virtual class IDisplayClusterViewportManager& GetOwner() const
+	{
+		check(false);
+	
+		return *GetConfiguration().GetViewportManager();
+	}
+
+	/** Calculate viewport projection (This function has been deprecated.)
+	 *
+	 * @param InContextNum      - viewport eye context index
+	 * @param InOutViewLocation - (in, out) View location
+	 * @param InOutViewRotation - (in,out) View rotator
+	 * @param ViewOffset         - eye offset from view location
+	 * @param WorldToMeters     - UE world scale
+	 * @param NCP - Near clipping plane
+	 * @param FCP - Far clipping plane
+	 *
+	 * @return - true, if the calculation of the viewport projection is successful.
+	 */
+	UE_DEPRECATED(5.3, "This function has been deprecated. Please use 'CalculateView()'.")
+	virtual bool CalculateView(const uint32 InContextNum, FVector& InOutViewLocation, FRotator& InOutViewRotation, const FVector& ViewOffset, const float WorldToMeters, const float NCP, const float FCP)
+	{
+		return false;
+	}
+
+	///////////////// UE_DEPRECATED 5.4 ///////////////////
+	/** Return view point camera component for this viewport. */
+	UE_DEPRECATED(5.4, "This function has been deprecated. Please use 'GetViewPointCameraComponent(EDisplayClusterRootActorType)'.")
+	virtual class UDisplayClusterCameraComponent* GetViewPointCameraComponent() const
+	{
+		return nullptr;
+	}
+
+	/** Return the viewport manager that owns this viewport */
+	UE_DEPRECATED(5.4, "This function has been deprecated. Please use 'GetConfiguration()'.")
+	virtual class IDisplayClusterViewportManager* GetViewportManager() const
+	{
+		return nullptr;
+	}
+
+	/** Gets a reference to the DCRA with actual configuration inside. */
+	UE_DEPRECATED(5.4, "This function has been deprecated. Please use 'GetConfiguration()'.")
+	virtual class ADisplayClusterRootActor* GetRootActor() const
+	{
+		return nullptr;
+	}
+
+	/** Return current render mode. */
+	UE_DEPRECATED(5.4, "This function has been deprecated. Please use 'GetConfiguration()'.")
+	virtual EDisplayClusterRenderFrameMode GetRenderMode() const
+	{
+		return EDisplayClusterRenderFrameMode::Unknown;
+	}
+
+	/** Return current world. */
+	UE_DEPRECATED(5.4, "This function has been deprecated. Please use 'GetConfiguration()'.")
+	virtual class UWorld* GetCurrentWorld() const
+	{
+		return nullptr;
+	}
+
+	/** Returns true if the scene is open now (The current world is assigned and DCRA has already initialized for it). */
+	UE_DEPRECATED(5.4, "This function has been deprecated. Please use 'GetConfiguration()'.")
+	virtual bool IsSceneOpened() const
+	{
+		return false;
+	}
+
+	/** Returns true if the current world type is equal to one of the input types. */
+	UE_DEPRECATED(5.4, "This function has been deprecated. Please use 'GetConfiguration()'.")
+	virtual bool IsCurrentWorldHasAnyType(const EWorldType::Type InWorldType1, const EWorldType::Type InWorldType2 = EWorldType::None, const EWorldType::Type InWorldType3 = EWorldType::None) const
+	{
+		return false;
+	}
 };
diff --git a/Engine/Plugins/Runtime/nDisplay/Source/DisplayCluster/Public/Render/Viewport/IDisplayClusterViewportConfiguration.h b/Engine/Plugins/Runtime/nDisplay/Source/DisplayCluster/Public/Render/Viewport/IDisplayClusterViewportConfiguration.h
new file mode 100644
--- /dev/null
+++ b/Engine/Plugins/Runtime/nDisplay/Source/DisplayCluster/Public/Render/Viewport/IDisplayClusterViewportConfiguration.h
@@ -0,0 +1,144 @@
+// Copyright Epic Games, Inc. All Rights Reserved.
+
+#pragma once
+
+#include "CoreMinimal.h"
+#include "Engine/EngineTypes.h"
+#include "Render/Viewport/RenderFrame/DisplayClusterRenderFrameEnums.h"
+
+class UWorld;
+class UDisplayClusterConfigurationData;
+class ADisplayClusterRootActor;
+class IDisplayClusterViewportManager;
+struct FDisplayClusterConfigurationICVFX_StageSettings;
+struct FDisplayClusterConfigurationRenderFrame;
+
+/**
+* Type of DCRA used by configuration
+*/
+enum class EDisplayClusterRootActorType : uint8
+{
+	// This DCRA will be used to render previews. The meshes and preview materials are created at runtime.
+	Preview = 0,
+
+	// A reference to DCRA in the scene, used as a source for math calculations and references.
+	// Locations in the scene and math data are taken from this DCRA.
+	Scene,
+
+	// Reference to DCRA, used as a source of configuration data from DCRA and its components.
+	Configuration,
+
+	// This value can only be used in very specific cases:
+	// For function GetRootActor() : Return any of the DCRAs that are not nullptr, in ascending order of type: Preview, Scene, Configuration.
+	// For function SetRootActor() : Sets all references to DRCA to the specified value.
+	Any
+};
+
+/**
+ * Viewport manager configuration.
+ */
+class DISPLAYCLUSTER_API IDisplayClusterViewportConfiguration
+{
+public:
+	virtual ~IDisplayClusterViewportConfiguration() = default;
+
+public:
+	/**
+	 *  Sets a reference to the current world to be rendered in DCRA
+	 * @param InWorld - ptr to the world to be rendered
+	 */
+	virtual void SetCurrentWorld(UWorld* InWorld) = 0;
+
+	/**
+	 *  Sets a reference to the DCRA's
+	 * @param InRootActor     - a ref to DCRA
+	 * @param InRootActorType - Type of the DCRA
+	 */
+	virtual void SetRootActor(ADisplayClusterRootActor* InRootActor, const EDisplayClusterRootActorType InRootActorType) = 0;
+
+	/**
+	* Update\Create\Delete local node viewports
+	* Updating the configuration to render a ClusterNode in the specified mode
+	*
+	* @param InRenderMode     - Render mode
+	* @param InClusterNodeId  - cluster node for rendering
+	*/
+	virtual bool UpdateConfigurationForClusterNode(EDisplayClusterRenderFrameMode InRenderMode, const FString& InClusterNodeId) = 0;
+
+	/**
+	* Update\Create\Delete local node viewports
+	* Updating the configuration to render a list of viewports in a given mode
+	*
+	* @param InRenderMode    - Render mode
+	* @param InViewportNames - Viewports names for next frame
+	*/
+	virtual bool UpdateConfigurationForViewportsList(EDisplayClusterRenderFrameMode InRenderMode, const TArray<FString>& InViewportNames) = 0;
+
+public:
+	/** Return the configuration proxy object. */
+	virtual const class IDisplayClusterViewportConfigurationProxy& GetProxy() const = 0;
+
+	/** Return the viewport manager that used by this configuration. */
+	virtual IDisplayClusterViewportManager* GetViewportManager() const = 0;
+
+	/**
+	 *  Gets a reference to the current world being rendered in DCRA
+	 */
+	virtual UWorld* GetCurrentWorld() const = 0;
+
+	/**
+	 * Gets a reference to the DCRA by type.
+	 * If a DCRA with the specified type is not assigned, the default DCRA is used.
+	 * 
+	 * @param InRootActorType - Type of the DCRA
+	 */
+	virtual ADisplayClusterRootActor* GetRootActor(const EDisplayClusterRootActorType InRootActorType) const = 0;
+
+	/**
+	 * Gets a configuration data from the Configuration RootActor
+	 */
+	virtual const UDisplayClusterConfigurationData* GetConfigurationData() const = 0;
+
+	/**
+	 * Gets a configuration stage settings from the Configuration RootActor
+	 */
+	virtual const FDisplayClusterConfigurationICVFX_StageSettings* GetStageSettings() const = 0;
+
+	/**
+	 * Gets a configuration render settings from the Configuration RootActor
+	 */
+	virtual const FDisplayClusterConfigurationRenderFrame* GetConfigurationRenderFrameSettings() const = 0;
+
+	/**
+	* Returns true if the current world type is equal to one of the input types.
+	*/
+	virtual bool IsCurrentWorldHasAnyType(
+		const EWorldType::Type InWorldType1,
+		const EWorldType::Type InWorldType2 = EWorldType::None,
+		const EWorldType::Type InWorldType3 = EWorldType::None
+	) const = 0;
+
+	/**
+	* Returns true if the given DCRA has a world type equal to one of the input types.
+	*/
+	virtual bool IsRootActorWorldHasAnyType(
+		const EDisplayClusterRootActorType InRootActorType,
+		const EWorldType::Type InWorldType1,
+		const EWorldType::Type InWorldType2 = EWorldType::None,
+		const EWorldType::Type InWorldType3 = EWorldType::None
+	) const = 0;
+
+	/** 
+	* Returns true if the scene is open now (The current world is assigned and DCRA has already initialized for it).
+	*/
+	virtual bool IsSceneOpened() const = 0;
+
+	/** Returns true if preview rendering mode is used. */
+	virtual bool IsPreviewRendering() const = 0;
+
+	/** Returns the rendering mode for PIE. */
+	virtual EDisplayClusterRenderFrameMode GetRenderModeForPIE() const = 0;
+
+	/** Return current cluster node id. */
+	virtual const FString& GetClusterNodeId() const = 0;
+};
diff --git a/Engine/Plugins/Runtime/nDisplay/Source/DisplayCluster/Public/Render/Viewport/IDisplayClusterViewportManager.h b/Engine/Plugins/Runtime/nDisplay/Source/DisplayCluster/Public/Render/Viewport/IDisplayClusterViewportManager.h
--- a/Engine/Plugins/Runtime/nDisplay/Source/DisplayCluster/Public/Render/Viewport/IDisplayClusterViewportManager.h
+++ b/Engine/Plugins/Runtime/nDisplay/Source/DisplayCluster/Public/Render/Viewport/IDisplayClusterViewportManager.h
@@ -1,11 +1,7 @@
 // Copyright Epic Games, Inc. All Rights Reserved.
 
 #pragma once
-
-#include "CoreMinimal.h"
-
 #include "Render/Viewport/IDisplayClusterViewport.h"
-#include "Render/Viewport/Containers/DisplayClusterPreviewSettings.h"
 #include "Render/Viewport/RenderFrame/DisplayClusterRenderFrameEnums.h"
 
 #include "SceneView.h"
@@ -35,45 +31,13 @@
 	virtual TSharedPtr<IDisplayClusterViewportManager, ESPMode::ThreadSafe> ToSharedPtr() = 0;
 	virtual TSharedPtr<const IDisplayClusterViewportManager, ESPMode::ThreadSafe> ToSharedPtr() const = 0;
 
+	/** Get viewport manager proxy interface. */
 	virtual const IDisplayClusterViewportManagerProxy* GetProxy() const = 0;
 	virtual       IDisplayClusterViewportManagerProxy* GetProxy() = 0;
 
-	virtual UWorld* GetCurrentWorld() const = 0;
-
-	virtual ADisplayClusterRootActor* GetRootActor() const = 0;
+	/** Get viewport manager configuration interface. */
+	virtual IDisplayClusterViewportConfiguration& GetConfiguration() = 0;
+	virtual const IDisplayClusterViewportConfiguration& GetConfiguration() const = 0;
 
 	/**
-	* Returns true if the scene is open now (The current world is assigned and DCRA has already initialized for it).
-	* [Game thread func]
-	*/
-	virtual bool IsSceneOpened() const = 0;
-
-	/**
-	* Update\Create\Delete local node viewports
-	* Update ICVFX configuration from root actor components
-	* [Game thread func]
-	*
-	* @param InRenderMode     - Render mode
-	* @param InClusterNodeId  - cluster node for rendering
-	* @param InRootActorPtr   - reference to RootActor with actual configuration inside
-	* @param InPreviewSettings - support preview rendering
-	*
-	* @return - true, if success
-	*/
-	virtual bool UpdateConfiguration(EDisplayClusterRenderFrameMode InRenderMode, const FString& InClusterNodeId, class ADisplayClusterRootActor* InRootActorPtr, const FDisplayClusterPreviewSettings* InPreviewSettings = nullptr) = 0;
-
-	/**
-	* Update\Create\Delete viewports for frame. For rendering outside of cluster nodes
-	* Update ICVFX configuration from root actor components
-	* [Game thread func]
-	*
-	* @param InRenderMode    - Render mode
-	* @param InViewportNames - Viewports names for next frame
-	* @param InRootActorPtr  - reference to RootActor with actual configuration inside
-	*
-	* @return - true, if success
-	*/
-	virtual bool UpdateCustomConfiguration(EDisplayClusterRenderFrameMode InRenderMode, const TArray<FString>& InViewportNames, class ADisplayClusterRootActor* InRootActorPtr) = 0;
-
-	/**
 	* Initialize new frame for all viewports on game thread, and update context, render resources with viewport new settings
@@ -85,7 +49,7 @@
 	*
 	* @return - true, if success
 	*/
-	virtual bool BeginNewFrame(FViewport* InViewport, UWorld* InWorld, FDisplayClusterRenderFrame& OutRenderFrame) = 0;
+	virtual bool BeginNewFrame(FViewport* InViewport, FDisplayClusterRenderFrame& OutRenderFrame) = 0;
 
 	/**
 	* Initialize frame for render on game thread
@@ -132,9 +96,6 @@
 	/** Send to render thread. */
 	virtual void RenderFrame(FViewport* InViewport) = 0;
 
-	/** Return current render mode. */
-	virtual EDisplayClusterRenderFrameMode GetRenderMode() const = 0;
-
 	/** Add internal DCVM objects to the reference collector. */
 	virtual void AddReferencedObjects(FReferenceCollector& Collector) = 0;
 
@@ -175,12 +136,6 @@
 	*/
 	virtual IDisplayClusterViewport* FindViewport(const int32 StereoViewIndex, uint32* OutContextNum = nullptr) const = 0;
 	
-	UE_DEPRECATED(5.3, "This function has been deprecated. Please use 'GetCurrentRenderFrameViewports'.")
-	virtual const TArrayView<IDisplayClusterViewport*> GetViewports() const
-	{
-		return TArrayView<IDisplayClusterViewport*>();
-	}
-
 	/**
 	* Return all exist viewports objects for current cluster node
 	* [Game thread func]
@@ -198,13 +153,13 @@
 	virtual const TArrayView<TSharedPtr<IDisplayClusterViewport, ESPMode::ThreadSafe>> GetEntireClusterViewports() const = 0;
 
 	/**
-	 * Return all viewports associated with the specified warp policy
+	 * Returns the viewports from the entire cluster associated with the specified warp policy
 	 * [Game thread func]
 	 *
 	 * @param InWarpPolicy - The warp policy to get viewports for
 	 * @return - A list of viewports associated with the warp policy
 	 */
-	virtual TArray<TSharedPtr<IDisplayClusterViewport, ESPMode::ThreadSafe>> GetViewportsForWarpPolicy(const TSharedPtr<IDisplayClusterWarpPolicy>& InWarpPolicy) const = 0;
+	virtual TArray<TSharedPtr<IDisplayClusterViewport, ESPMode::ThreadSafe>> GetEntireClusterViewportsForWarpPolicy(const TSharedPtr<IDisplayClusterWarpPolicy>& InWarpPolicy) const = 0;
 
 	/**
 	* Mark the geometry of the referenced component(s) as dirty (ProceduralMesh, etc)
@@ -226,6 +181,102 @@
 	* @param ViewDrawer - optional drawing in the view
 	* @param StereoViewIndex - index of the view when using stereoscopy
 	*/
-	virtual class FSceneView* CalcSceneView(class ULocalPlayer* LocalPlayer, class FSceneViewFamily* ViewFamily, FVector& OutViewLocation, FRotator& OutViewRotation, class FViewport* Viewport, class FViewElementDrawer* ViewDrawer, int32 StereoViewIndex) = 0;
+	virtual class FSceneView* CalcSceneView(
+		class ULocalPlayer* LocalPlayer,
+		class FSceneViewFamily* ViewFamily,
+		FVector& OutViewLocation,
+		FRotator& OutViewRotation,
+		class FViewport* Viewport,
+		class FViewElementDrawer* ViewDrawer,
+		int32 StereoViewIndex) = 0;
+
+	///////////////// UE_DEPRECATED 5.3 ///////////////////
+
+	UE_DEPRECATED(5.3, "This function has been deprecated. Please use 'GetCurrentRenderFrameViewports'.")
+		virtual const TArrayView<IDisplayClusterViewport*> GetViewports() const
+	{
+		return TArrayView<IDisplayClusterViewport*>();
+	}
+
+	///////////////// UE_DEPRECATED 5.4 ///////////////////
+	
+	/** Return current render mode. */
+	UE_DEPRECATED(5.4, "This function has been deprecated. Please use 'GetConfiguration()'.")
+	virtual EDisplayClusterRenderFrameMode GetRenderMode() const
+	{
+		return EDisplayClusterRenderFrameMode::Unknown;
+	}
+
+	UE_DEPRECATED(5.4, "This function has been deprecated. Please use 'GetConfiguration()'.")
+	virtual UWorld* GetCurrentWorld() const
+	{
+		return nullptr;
+	}
+
+	UE_DEPRECATED(5.4, "This function has been deprecated. Please use 'GetConfiguration()'.")
+	virtual ADisplayClusterRootActor* GetRootActor() const
+	{
+		return nullptr;
+	}
+
+	/**
+	* Returns true if the scene is open now (The current world is assigned and DCRA has already initialized for it).
+	* [Game thread func]
+	*/
+	UE_DEPRECATED(5.4, "This function has been deprecated. Please use 'GetConfiguration()'.")
+	virtual bool IsSceneOpened() const
+	{
+		return false;
+	}
+
+	/**
+	* Update\Create\Delete local node viewports
+	* Update ICVFX configuration from root actor components
+	* [Game thread func]
+	*
+	* @param InRenderMode     - Render mode
+	* @param InClusterNodeId  - cluster node for rendering
+	* @param InRootActorPtr   - reference to RootActor with actual configuration inside
+	* @param InPreviewSettings - support preview rendering
+	*
+	* @return - true, if success
+	*/
+	UE_DEPRECATED(5.4, "This function has been deprecated. Please use 'GetConfiguration()'.")
+	virtual bool UpdateConfiguration(EDisplayClusterRenderFrameMode InRenderMode, const FString& InClusterNodeId, class ADisplayClusterRootActor* InRootActorPtr, const struct FDisplayClusterPreviewSettings* InPreviewSettings = nullptr)
+	{
+		return false;
+	}
+
+	/**
+	* Update\Create\Delete viewports for frame. For rendering outside of cluster nodes
+	* Update ICVFX configuration from root actor components
+	* [Game thread func]
+	*
+	* @param InRenderMode    - Render mode
+	* @param InViewportNames - Viewports names for next frame
+	* @param InRootActorPtr  - reference to RootActor with actual configuration inside
+	*
+	* @return - true, if success
+	*/
+	UE_DEPRECATED(5.4, "This function has been deprecated. Please use 'GetConfiguration()'.")
+	virtual bool UpdateCustomConfiguration(EDisplayClusterRenderFrameMode InRenderMode, const TArray<FString>& InViewportNames, class ADisplayClusterRootActor* InRootActorPtr)
+	{
+		return false;
+	}
+
+	/**
+	* Initialize new frame for all viewports on game thread, and update context, render resources with viewport new settings
+	* And finally build render frame structure and send to render thread proxy viewport objects
+	* [Game thread func]
+	*
+	* @param InViewport          - target viewport
+	* @param OutRenderFrame      - output render frame container
+	*
+	* @return - true, if success
+	*/
+	UE_DEPRECATED(5.4, "This function has been deprecated. Please use 'BeginNewFrame()'.")
+	virtual bool BeginNewFrame(FViewport* InViewport, UWorld* InWorld, FDisplayClusterRenderFrame& OutRenderFrame)
+	{
+		return false;
+	}
 };
-
diff --git a/Engine/Plugins/Runtime/nDisplay/Source/DisplayCluster/Public/Render/Viewport/IDisplayClusterViewportManagerProxy.h b/Engine/Plugins/Runtime/nDisplay/Source/DisplayCluster/Public/Render/Viewport/IDisplayClusterViewportManagerProxy.h
--- a/Engine/Plugins/Runtime/nDisplay/Source/DisplayCluster/Public/Render/Viewport/IDisplayClusterViewportManagerProxy.h
+++ b/Engine/Plugins/Runtime/nDisplay/Source/DisplayCluster/Public/Render/Viewport/IDisplayClusterViewportManagerProxy.h
@@ -3,7 +3,6 @@
 #pragma once
 
 #include "CoreMinimal.h"
-
 #include "Render/Viewport/IDisplayClusterViewportProxy.h"
 
 /**
@@ -19,8 +18,8 @@ public:
 	virtual TSharedPtr<IDisplayClusterViewportManagerProxy, ESPMode::ThreadSafe> ToSharedPtr() = 0;
 	virtual TSharedPtr<const IDisplayClusterViewportManagerProxy, ESPMode::ThreadSafe> ToSharedPtr() const = 0;
 
-	/** Return current render mode. */
-	virtual EDisplayClusterRenderFrameMode GetRenderMode() const = 0;
+	/** Get viewport manager proxy configuration interface. */
+	virtual const IDisplayClusterViewportConfigurationProxy& GetConfigurationProxy() const = 0;
 
 	/**
 	* Find viewport render thread proxy object by name
@@ -74,5 +73,14 @@ public:
 	* @return - true if success
 	*/
 	virtual bool ResolveFrameTargetToBackBuffer_RenderThread(FRHICommandListImmediate& RHICmdList, const uint32 InContextNum, const int32 DestArrayIndex, FRHITexture2D* DstBackBuffer, FVector2D WindowSize) const = 0;
+
+	///////////////// UE_DEPRECATED 5.4 ///////////////////
+
+	/** Return current render mode. */
+	UE_DEPRECATED(5.4, "This function has been deprecated. Please use 'GetConfigurationProxy()'.")
+	virtual EDisplayClusterRenderFrameMode GetRenderMode() const
+	{
+		return EDisplayClusterRenderFrameMode::Unknown;
+	}
 };
 
diff --git a/Engine/Plugins/Runtime/nDisplay/Source/DisplayCluster/Public/Render/Viewport/IDisplayClusterViewportProxy.h b/Engine/Plugins/Runtime/nDisplay/Source/DisplayCluster/Public/Render/Viewport/IDisplayClusterViewportProxy.h
--- a/Engine/Plugins/Runtime/nDisplay/Source/DisplayCluster/Public/Render/Viewport/IDisplayClusterViewportProxy.h
+++ b/Engine/Plugins/Runtime/nDisplay/Source/DisplayCluster/Public/Render/Viewport/IDisplayClusterViewportProxy.h
@@ -4,18 +4,18 @@
 
 #include "CoreMinimal.h"
 
-#include "RHI.h"
-#include "RHICommandList.h"
-#include "RHIResources.h"
-
 #include "Render/Viewport/Containers/DisplayClusterViewport_Enums.h"
 #include "Render/Viewport/Containers/DisplayClusterViewport_Context.h"
 #include "Render/Viewport/Containers/DisplayClusterViewport_RenderSettings.h"
 #include "Render/Viewport/Containers/DisplayClusterViewport_RenderSettingsICVFX.h"
 #include "Render/Viewport/Containers/DisplayClusterViewport_PostRenderSettings.h"
 #include "Render/Viewport/Containers/DisplayClusterViewportProxy_Context.h"
-
 #include "Render/Viewport/RenderFrame/DisplayClusterRenderFrameEnums.h"
+#include "Render/Viewport/IDisplayClusterViewportConfigurationProxy.h"
+
+#include "RHI.h"
+#include "RHICommandList.h"
+#include "RHIResources.h"
 
 /**
  * nDisplay: ViewportProxy (interface for RenderThread)
@@ -30,16 +30,25 @@ public:
 	virtual TSharedPtr<IDisplayClusterViewportProxy, ESPMode::ThreadSafe> ToSharedPtr() = 0;
 	virtual TSharedPtr<const IDisplayClusterViewportProxy, ESPMode::ThreadSafe> ToSharedPtr() const = 0;
 
-	/** Return current render mode. */
-	virtual EDisplayClusterRenderFrameMode GetRenderMode() const = 0;
-
+	/** Get the unique viewport name. */
 	virtual FString GetId() const = 0;
+
+	/** Get the name of the cluster node to which this viewport belongs. */
 	virtual FString GetClusterNodeId() const = 0;
 
+	/** Get viewport manager configuration interface. */
+	virtual const IDisplayClusterViewportConfigurationProxy& GetConfigurationProxy() const = 0;
+
+	/** Get projection policy ref. */
 	virtual const TSharedPtr<class IDisplayClusterProjectionPolicy, ESPMode::ThreadSafe>& GetProjectionPolicy_RenderThread() const = 0;
 
+	/** Get main viewport render settings. */
 	virtual const FDisplayClusterViewport_RenderSettings& GetRenderSettings_RenderThread() const = 0;
+
+	/** Get ICVFX render settings. */
 	virtual const FDisplayClusterViewport_RenderSettingsICVFX& GetRenderSettingsICVFX_RenderThread() const = 0;
+
+	/** Get post-render settings. */
 	virtual const FDisplayClusterViewport_PostRenderSettings& GetPostRenderSettings_RenderThread() const = 0;
 
 	/** Return viewport proxy contexts data. */
@@ -87,6 +96,20 @@ public:
 	 */
 	virtual bool ResolveResources_RenderThread(FRHICommandListImmediate& RHICmdList, IDisplayClusterViewportProxy* InputResourceViewportProxy, const EDisplayClusterViewportResourceType InputResourceType, const EDisplayClusterViewportResourceType OutputResourceType, const int32 InContextNum = INDEX_NONE) const = 0;
 
+	/**
+	* Change render setting for viewport.
+	* This function must be called in a valid moment in time.
+	*/
+	virtual void SetRenderSettings_RenderThread(const FDisplayClusterViewport_RenderSettings& InRenderSettings) const = 0;
+
+	/**
+	* Change viewport contexts data.
+	* This function must be called in a valid moment in time.
+	*/
+	virtual void SetContexts_RenderThread(const TArray<FDisplayClusterViewport_Context>& InContexts) const = 0;
+
+	///////////////// UE_DEPRECATED 5.3 ///////////////////
+
 	/** Return output resource type (support preview, remap, etc). */
 	UE_DEPRECATED(5.3, "This function has beend deprecate. Please use the new enumeration values.")
 	virtual EDisplayClusterViewportResourceType GetOutputResourceType_RenderThread() const
@@ -94,12 +117,26 @@ public:
 		return EDisplayClusterViewportResourceType::OutputTargetableResource;
 	}
 
+	///////////////// UE_DEPRECATED 5.4 ///////////////////
+
+	/** Return current render mode. */
+	UE_DEPRECATED(5.4, "This function has been deprecated. Please use 'GetConfigurationProxy().GetRenderMode_RenderThread()'.")
+	virtual EDisplayClusterRenderFrameMode GetRenderMode() const
+	{
+		return EDisplayClusterRenderFrameMode::Unknown;
+	}
+
 	/** Returns ptr to ViewportManagerProxy (the owner of this viewport proxy) if it still exists. */
-	virtual const class IDisplayClusterViewportManagerProxy* GetViewportManagerProxy_RenderThread() const = 0;
+	UE_DEPRECATED(5.4, "This function has been deprecated. Please use 'GetConfigurationProxy().GetViewportManagerProxy_RenderThread'.")
+	virtual const class IDisplayClusterViewportManagerProxy* GetViewportManagerProxy_RenderThread() const
+	{
+		return nullptr;
+	}
 
 	/** Get render frame settings from viewport proxy owner, if it still exists. */
-	virtual const struct FDisplayClusterRenderFrameSettings* GetRenderFrameSettings_RenderThread() const = 0;
-
-	virtual void SetRenderSettings_RenderThread(const FDisplayClusterViewport_RenderSettings& InRenderSettings) const = 0;
-	virtual void SetContexts_RenderThread(const TArray<FDisplayClusterViewport_Context>& InContexts) const = 0;
+	UE_DEPRECATED(5.4, "This function has been deprecated.")
+	virtual const struct FDisplayClusterRenderFrameSettings* GetRenderFrameSettings_RenderThread() const
+	{
+		return nullptr;
+	}
 };
diff --git a/Engine/Plugins/Runtime/nDisplay/Source/DisplayCluster/Public/Render/Viewport/IDisplayClusterViewport_CustomPostProcessSettings.h b/Engine/Plugins/Runtime/nDisplay/Source/DisplayCluster/Public/Render/Viewport/IDisplayClusterViewport_CustomPostProcessSettings.h
--- a/Engine/Plugins/Runtime/nDisplay/Source/DisplayCluster/Public/Render/Viewport/IDisplayClusterViewport_CustomPostProcessSettings.h
+++ b/Engine/Plugins/Runtime/nDisplay/Source/DisplayCluster/Public/Render/Viewport/IDisplayClusterViewport_CustomPostProcessSettings.h
@@ -4,7 +4,9 @@
 
 #include "CoreMinimal.h"
 
-
+/**
+* DC Viewport Postprocess interface
+*/
 class DISPLAYCLUSTER_API IDisplayClusterViewport_CustomPostProcessSettings
 {
 public:
diff --git a/Engine/Plugins/Runtime/nDisplay/Source/DisplayCluster/Public/Render/Viewport/RenderFrame/DisplayClusterRenderFrame.h b/Engine/Plugins/Runtime/nDisplay/Source/DisplayCluster/Public/Render/Viewport/RenderFrame/DisplayClusterRenderFrame.h
--- a/Engine/Plugins/Runtime/nDisplay/Source/DisplayCluster/Public/Render/Viewport/RenderFrame/DisplayClusterRenderFrame.h
+++ b/Engine/Plugins/Runtime/nDisplay/Source/DisplayCluster/Public/Render/Viewport/RenderFrame/DisplayClusterRenderFrame.h
@@ -66,6 +66,8 @@ struct FDisplayClusterRenderFrameTarget
 	// required Render target size (resource can be bigger)
 	FIntPoint RenderTargetSize;
 
+	// Viewport capture mode
+	// This mode affects many viewport rendering settings.
 	EDisplayClusterViewportCaptureMode CaptureMode = EDisplayClusterViewportCaptureMode::Default;
 
 	// Render target resource ref
