nDisplay: Added new interface IDisplayClusterViewportConfiguration

diff --git a/Engine/Plugins/Runtime/nDisplay/Source/DisplayCluster/Public/Render/Viewport/RenderFrame/DisplayClusterRenderFrameEnums.h b/Engine/Plugins/Runtime/nDisplay/Source/DisplayCluster/Public/Render/Viewport/RenderFrame/DisplayClusterRenderFrameEnums.h
--- a/Engine/Plugins/Runtime/nDisplay/Source/DisplayCluster/Public/Render/Viewport/RenderFrame/DisplayClusterRenderFrameEnums.h
+++ b/Engine/Plugins/Runtime/nDisplay/Source/DisplayCluster/Public/Render/Viewport/RenderFrame/DisplayClusterRenderFrameEnums.h
@@ -4,6 +4,9 @@
 
 #include "CoreMinimal.h"
 
+/**
+* Rendering Mode.
+*/
 enum class EDisplayClusterRenderFrameMode: uint8
 {
 	Unknown = 0,
@@ -22,20 +25,8 @@ enum class EDisplayClusterRenderFrameMode: uint8
 
 	// Special render for preview in scene
 	PreviewInScene,
-};
-
-// (experimental, not implemented, reserved)
-enum class EDisplayClusterRenderFamilyMode : uint8
-{
-	// Render all viewports to unique RenderTargets
-	None = 0,
-
-	// Merge views by ViewFamilyGroupNum
-	AllowMergeForGroups,
-
-	// Merge views by ViewFamilyGroupNum and stereo
-	AllowMergeForGroupsAndStereo,
 
-	// Use rules to merge views to minimal num of families (separate by: buffer_ratio, viewExtension, max RTT size)
-	MergeAnyPossible,
+	// Special render of ProxyHit for preview in scene
+	// *** Not implemented.
+	PreviewProxyHitInScene,
 };
diff --git a/Engine/Plugins/Runtime/nDisplay/Source/DisplayCluster/Public/Render/Warp/IDisplayClusterWarpPolicy.h b/Engine/Plugins/Runtime/nDisplay/Source/DisplayCluster/Public/Render/Warp/IDisplayClusterWarpPolicy.h
--- a/Engine/Plugins/Runtime/nDisplay/Source/DisplayCluster/Public/Render/Warp/IDisplayClusterWarpPolicy.h
+++ b/Engine/Plugins/Runtime/nDisplay/Source/DisplayCluster/Public/Render/Warp/IDisplayClusterWarpPolicy.h
@@ -104,4 +104,14 @@ public:
 	 */
 	virtual void EndCalcFrustum(IDisplayClusterViewport* InViewport, const uint32 ContextNum)
 	{ }
+
+	/**
+	* Ask warp  policy instance if it has any movable mesh based preview
+	* @param InViewport - a owner viewport
+	* @return - True if mesh based preview is available
+	*/
+	virtual bool HasPreviewMovableMesh(IDisplayClusterViewport* InViewport)
+	{
+		return false;
+	}
 };
diff --git a/Engine/Plugins/Runtime/nDisplay/Source/DisplayClusterConfiguration/Private/DisplayClusterConfigurationTypes_ICVFX.cpp b/Engine/Plugins/Runtime/nDisplay/Source/DisplayClusterConfiguration/Private/DisplayClusterConfigurationTypes_ICVFX.cpp
--- a/Engine/Plugins/Runtime/nDisplay/Source/DisplayClusterConfiguration/Private/DisplayClusterConfigurationTypes_ICVFX.cpp
+++ b/Engine/Plugins/Runtime/nDisplay/Source/DisplayClusterConfiguration/Private/DisplayClusterConfigurationTypes_ICVFX.cpp
@@ -1,6 +1,8 @@
 // Copyright Epic Games, Inc. All Rights Reserved.
 
 #include "DisplayClusterConfigurationTypes_ICVFX.h"
+#include "DisplayClusterConfigurationTypes.h"
+#include "IDisplayCluster.h"
 
 namespace UE::DisplayClusterConfiguration::ICVFX
 {
@@ -95,13 +97,12 @@ const FOpenColorIOColorConversionSettings* FDisplayClusterConfigurationICVFX_Cam
 	return nullptr;
 }
 
-#if WITH_EDITOR
-bool FDisplayClusterConfigurationICVFX_CameraOCIO::IsChromakeyViewportSettingsEqual_Editor(const FString& InClusterNodeId1, const FString& InClusterNodeId2) const
+bool FDisplayClusterConfigurationICVFX_CameraOCIO::IsChromakeyViewportSettingsEqual(const FString& InClusterNodeId1, const FString& InClusterNodeId2) const
 {
-	return IsInnerFrustumViewportSettingsEqual_Editor(InClusterNodeId1, InClusterNodeId2);
+	return IsInnerFrustumViewportSettingsEqual(InClusterNodeId1, InClusterNodeId2);
 }
 
-bool FDisplayClusterConfigurationICVFX_CameraOCIO::IsInnerFrustumViewportSettingsEqual_Editor(const FString& InClusterNodeId1, const FString& InClusterNodeId2) const
+bool FDisplayClusterConfigurationICVFX_CameraOCIO::IsInnerFrustumViewportSettingsEqual(const FString& InClusterNodeId1, const FString& InClusterNodeId2) const
 {
 	if (AllNodesOCIOConfiguration.bIsEnabled)
 	{
@@ -136,7 +137,6 @@ bool FDisplayClusterConfigurationICVFX_CameraOCIO::IsInnerFrustumViewportSetting
 
 	return true;
 }
-#endif
 
 ///////////////////////////////////////////////////////////////////////////////////////
 // FDisplayClusterConfigurationICVFX_LightcardOCIO
@@ -169,6 +169,39 @@ FDisplayClusterConfigurationICVFX_CameraSettings::FDisplayClusterConfigurationIC
 	AllNodesColorGrading.bEnableEntireClusterColorGrading = true;
 }
 
+bool FDisplayClusterConfigurationICVFX_CameraSettings::IsICVFXEnabled(const UDisplayClusterConfigurationData& InConfigurationData, const FString& InClusterNodeId) const
+{
+	if (bEnable)
+	{
+		// When rendering offscreen, we have an extended logic for camera rendering activation
+		static const bool bIsRunningClusterModeOffscreen =
+			(IDisplayCluster::Get().GetOperationMode() == EDisplayClusterOperationMode::Cluster) &&
+			FParse::Param(FCommandLine::Get(), TEXT("RenderOffscreen"));
+
+		// If cluster mode + rendering offscreen, discover media output settings
+		if (bIsRunningClusterModeOffscreen)
+		{
+			// First condition to render offscreen: it has media output assigned
+			const bool bUsesMediaOutput = (RenderSettings.Media.bEnable && RenderSettings.Media.IsMediaOutputAssigned(InClusterNodeId));
+
+			// Get backbuffer media settings
+			const UDisplayClusterConfigurationClusterNode* const NodeCfg = InConfigurationData.Cluster->GetNode(InClusterNodeId);
+			const FDisplayClusterConfigurationMedia* BackbufferMediaSettings = NodeCfg ? &NodeCfg->Media : nullptr;
+
+			// Second condition to render offscreen: the backbuffer has media output assigned.
+			// This means the whole frame including ICVFX cameras need to be rendered.
+			const bool bIsBackbufferBeingCaptured = BackbufferMediaSettings ? BackbufferMediaSettings->bEnable && BackbufferMediaSettings->IsMediaOutputAssigned() : false;
+
+			// Finally make a decision if the camera should be rendered
+			return bUsesMediaOutput && bIsBackbufferBeingCaptured;
+		}
+
+		return true;
+	}
+
+	return false;
+}
+
 const FOpenColorIOColorConversionSettings* FDisplayClusterConfigurationICVFX_CameraSettings::FindInnerFrustumOCIOConfiguration(const FString& InClusterNodeId) const
 {
 	return CameraOCIO.FindOCIOConfiguration(InClusterNodeId);
@@ -180,17 +213,15 @@ const FOpenColorIOColorConversionSettings* FDisplayClusterConfigurationICVFX_Cam
 	return CameraOCIO.FindOCIOConfiguration(InClusterNodeId);
 }
 
-#if WITH_EDITOR
-bool FDisplayClusterConfigurationICVFX_CameraSettings::IsInnerFrustumViewportSettingsEqual_Editor(const FString& InClusterNodeId1, const FString& InClusterNodeId2) const
+bool FDisplayClusterConfigurationICVFX_CameraSettings::IsInnerFrustumViewportSettingsEqual(const FString& InClusterNodeId1, const FString& InClusterNodeId2) const
 {
-	return CameraOCIO.IsInnerFrustumViewportSettingsEqual_Editor(InClusterNodeId1, InClusterNodeId2);
+	return CameraOCIO.IsInnerFrustumViewportSettingsEqual(InClusterNodeId1, InClusterNodeId2);
 }
 
-bool FDisplayClusterConfigurationICVFX_CameraSettings::IsChromakeyViewportSettingsEqual_Editor(const FString& InClusterNodeId1, const FString& InClusterNodeId2) const
+bool FDisplayClusterConfigurationICVFX_CameraSettings::IsChromakeyViewportSettingsEqual(const FString& InClusterNodeId1, const FString& InClusterNodeId2) const
 {
-	return CameraOCIO.IsChromakeyViewportSettingsEqual_Editor(InClusterNodeId1, InClusterNodeId2);
+	return CameraOCIO.IsChromakeyViewportSettingsEqual(InClusterNodeId1, InClusterNodeId2);
 }
-#endif
 
 float FDisplayClusterConfigurationICVFX_CameraSettings::GetCameraFieldOfViewMultiplier(const FDisplayClusterConfigurationICVFX_StageSettings& InStageSettings) const
 {
diff --git a/Engine/Plugins/Runtime/nDisplay/Source/DisplayClusterConfiguration/Public/DisplayClusterConfigurationTypes_ICVFX.h b/Engine/Plugins/Runtime/nDisplay/Source/DisplayClusterConfiguration/Public/DisplayClusterConfigurationTypes_ICVFX.h
--- a/Engine/Plugins/Runtime/nDisplay/Source/DisplayClusterConfiguration/Public/DisplayClusterConfigurationTypes_ICVFX.h
+++ b/Engine/Plugins/Runtime/nDisplay/Source/DisplayClusterConfiguration/Public/DisplayClusterConfigurationTypes_ICVFX.h
@@ -68,13 +68,11 @@ public:
 	/** Return InCamera OCIO configuration for the specified cluster node. Return nullptr if no OCIO. */
 	const FOpenColorIOColorConversionSettings* FindOCIOConfiguration(const FString& InClusterNodeId) const;
 
-#if WITH_EDITOR
 	/** Returns true if the InCamera OCIO configuration is the same for the input nodes. */
-	bool IsInnerFrustumViewportSettingsEqual_Editor(const FString& InClusterNodeId1, const FString& InClusterNodeId2) const;
+	bool IsInnerFrustumViewportSettingsEqual(const FString& InClusterNodeId1, const FString& InClusterNodeId2) const;
 
 	/** Returns true if the Chromakey OCIO configuration is the same for the input nodes. */
-	bool IsChromakeyViewportSettingsEqual_Editor(const FString& InClusterNodeId1, const FString& InClusterNodeId2) const;
-#endif
+	bool IsChromakeyViewportSettingsEqual(const FString& InClusterNodeId1, const FString& InClusterNodeId2) const;
 
 public:
 	/** OCIO Display look configuration for all nodes */
@@ -647,6 +645,16 @@ struct DISPLAYCLUSTERCONFIGURATION_API FDisplayClusterConfigurationICVFX_CameraS
 {
 	GENERATED_BODY()
 
+public:
+	/**
+	* Returns true if this camera is active
+	* @param InConfigurationData - cluster configuration data
+	* @param InClusterNodeId     - current cluster node name
+	* 
+	* @return - true if this camera can be used
+	**/
+	bool IsICVFXEnabled(const class UDisplayClusterConfigurationData& InConfigurationData, const FString& InClusterNodeId) const;
+
 public:
 	FDisplayClusterConfigurationICVFX_CameraSettings();
 
@@ -656,10 +664,8 @@ public:
 	/** Return Chromakey OCIO configuration for the specified cluster node. Return nullptr if no OCIO. */
 	const FOpenColorIOColorConversionSettings* FindChromakeyOCIOConfiguration(const FString& InClusterNodeId) const;
 
-#if WITH_EDITOR
-	bool IsInnerFrustumViewportSettingsEqual_Editor(const FString& InClusterNodeId1, const FString& InClusterNodeId2) const;
-	bool IsChromakeyViewportSettingsEqual_Editor(const FString& InClusterNodeId1, const FString& InClusterNodeId2) const;
-#endif
+	bool IsInnerFrustumViewportSettingsEqual(const FString& InClusterNodeId1, const FString& InClusterNodeId2) const;
+	bool IsChromakeyViewportSettingsEqual(const FString& InClusterNodeId1, const FString& InClusterNodeId2) const;
 
 	/** Return calculated soft edges values. */
 	FVector4 GetCameraSoftEdge(const FDisplayClusterConfigurationICVFX_StageSettings& InStageSettings) const;
diff --git a/Engine/Plugins/Runtime/nDisplay/Source/DisplayClusterMedia/Private/Capture/DisplayClusterMediaCaptureCamera.cpp b/Engine/Plugins/Runtime/nDisplay/Source/DisplayClusterMedia/Private/Capture/DisplayClusterMediaCaptureCamera.cpp
--- a/Engine/Plugins/Runtime/nDisplay/Source/DisplayClusterMedia/Private/Capture/DisplayClusterMediaCaptureCamera.cpp
+++ b/Engine/Plugins/Runtime/nDisplay/Source/DisplayClusterMedia/Private/Capture/DisplayClusterMediaCaptureCamera.cpp
@@ -19,10 +19,14 @@ FDisplayClusterMediaCaptureCamera::FDisplayClusterMediaCaptureCamera(const FStri
 	: FDisplayClusterMediaCaptureViewport(InMediaId, InClusterNodeId, InViewportId, InMediaOutput, SyncPolicy)
 	, CameraId(InCameraId)
 {
-	if (const ADisplayClusterRootActor* const RootActor = IDisplayCluster::Get().GetGameMgr()->GetRootActor())
+	if (const ADisplayClusterRootActor* const ActiveRootActor = IDisplayCluster::Get().GetGameMgr()->GetRootActor())
 	{
 		TArray<UActorComponent*> ICVFXCameraComponents;
-		RootActor->GetComponents(UDisplayClusterICVFXCameraComponent::StaticClass(), ICVFXCameraComponents);
+		IDisplayClusterViewportConfiguration* ViewportConfiguration = ActiveRootActor->GetViewportConfiguration();
+		if (const ADisplayClusterRootActor* const ConfigurationRootActor = ViewportConfiguration ? ViewportConfiguration->GetRootActor(EDisplayClusterRootActorType::Configuration) : nullptr)
+		{
+			ConfigurationRootActor->GetComponents(UDisplayClusterICVFXCameraComponent::StaticClass(), ICVFXCameraComponents);
+		}
 
 		for (const UActorComponent* const Component : ICVFXCameraComponents)
 		{
diff --git a/Engine/Plugins/Runtime/nDisplay/Source/DisplayClusterMoviePipeline/Private/DisplayClusterMoviePipelineViewportPass.cpp b/Engine/Plugins/Runtime/nDisplay/Source/DisplayClusterMoviePipeline/Private/DisplayClusterMoviePipelineViewportPass.cpp
--- a/Engine/Plugins/Runtime/nDisplay/Source/DisplayClusterMoviePipeline/Private/DisplayClusterMoviePipelineViewportPass.cpp
+++ b/Engine/Plugins/Runtime/nDisplay/Source/DisplayClusterMoviePipeline/Private/DisplayClusterMoviePipelineViewportPass.cpp
@@ -478,10 +478,9 @@ bool UDisplayClusterMoviePipelineViewportPassBase::InitializeDisplayCluster()
 			IDisplayClusterViewportManager* ViewportManager = DCRootActor->GetViewportManager();
 			check(ViewportManager);
 
-			const EDisplayClusterRenderFrameMode RenderFrameMode = EDisplayClusterRenderFrameMode::Mono;
-
 			// Update local node viewports (update\create\delete) and build new render frame
-			if (ViewportManager->UpdateCustomConfiguration(RenderFrameMode, DisplayClusterViewports, DCRootActor))
+			const EDisplayClusterRenderFrameMode RenderFrameMode = EDisplayClusterRenderFrameMode::Mono;
+			if (ViewportManager->GetConfiguration().UpdateConfigurationForViewportsList(RenderFrameMode, DisplayClusterViewports))
 			{
 				return true;
 			}
@@ -547,8 +546,10 @@ IDisplayClusterViewport* UDisplayClusterMoviePipelineViewportPassBase::GetAndCal
 			DCViewport->SetRenderSettings(RenderSettings);
 		}
 
+		ViewportManager->GetConfiguration().SetCurrentWorld(CurrentWorld);
+
 		FDisplayClusterRenderFrame RenderFrame;
-		if (ViewportManager->BeginNewFrame(GameViewportClient->Viewport, CurrentWorld, RenderFrame))
+		if (ViewportManager->BeginNewFrame(GameViewportClient->Viewport, RenderFrame))
 		{
 			// Obtaining the internal viewpoint for a given viewport with stereo eye offset distance.
 			FMinimalViewInfo ViewInfo;
@@ -570,7 +571,7 @@ IDisplayClusterViewport* UDisplayClusterMoviePipelineViewportPassBase::GetAndCal
 			{
 				OutView.ProjectionMatrix = FMatrix::Identity;
 
-				bResult = ViewportManager->IsSceneOpened() && DCViewport->GetProjectionMatrix(InContextNum, OutView.ProjectionMatrix);
+				bResult = ViewportManager->GetConfiguration().IsSceneOpened() && DCViewport->GetProjectionMatrix(InContextNum, OutView.ProjectionMatrix);
 			}
 
 			if (bFrameWarpBlend)
diff --git a/Engine/Plugins/Runtime/nDisplay/Source/DisplayClusterProjection/Private/Policy/Camera/DisplayClusterProjectionCameraPolicy.cpp b/Engine/Plugins/Runtime/nDisplay/Source/DisplayClusterProjection/Private/Policy/Camera/DisplayClusterProjectionCameraPolicy.cpp
--- a/Engine/Plugins/Runtime/nDisplay/Source/DisplayClusterProjection/Private/Policy/Camera/DisplayClusterProjectionCameraPolicy.cpp
+++ b/Engine/Plugins/Runtime/nDisplay/Source/DisplayClusterProjection/Private/Policy/Camera/DisplayClusterProjectionCameraPolicy.cpp
@@ -40,9 +40,9 @@ void FDisplayClusterProjectionCameraPolicy::UpdatePostProcessSettings(IDisplayCl
 	if (InViewport && !EnumHasAnyFlags(InViewport->GetRenderSettingsICVFX().RuntimeFlags, EDisplayClusterViewportRuntimeICVFXFlags::InCamera))
 	{
 		float DeltaTime = 0.0f;
-		if (ADisplayClusterRootActor* RootActor = InViewport->GetRootActor())
+		if (ADisplayClusterRootActor* SceneRootActor = InViewport->GetConfiguration().GetRootActor(EDisplayClusterRootActorType::Scene))
 		{
-			DeltaTime = RootActor->GetWorldDeltaSeconds();
+			DeltaTime = SceneRootActor->GetWorldDeltaSeconds();
 		}
 
 		FMinimalViewInfo ViewInfo;
@@ -150,7 +150,7 @@ bool FDisplayClusterProjectionCameraPolicy::ImplSetupProjectionViewPoint(IDispla
 		// Store CustomNearClippingPlane locally, and then use that value for the projection matrix
 		bResult = IDisplayClusterViewport::GetCameraComponentView(CameraComponent, InDeltaTime, CameraSettings.bUseCameraPostprocess, InOutViewInfo, OutCustomNearClippingPlane);
 	}
-	else if (UWorld* CurrentWorld = InViewport ? InViewport->GetCurrentWorld() : nullptr)
+	else if (UWorld* CurrentWorld = InViewport ? InViewport->GetConfiguration().GetCurrentWorld() : nullptr)
 	{
 		// Get active player camera
 		bResult = IDisplayClusterViewport::GetPlayerCameraView(CurrentWorld, CameraSettings.bUseCameraPostprocess, InOutViewInfo);
diff --git a/Engine/Plugins/Runtime/nDisplay/Source/DisplayClusterProjection/Private/Policy/DisplayClusterProjectionPolicyBase.cpp b/Engine/Plugins/Runtime/nDisplay/Source/DisplayClusterProjection/Private/Policy/DisplayClusterProjectionPolicyBase.cpp
--- a/Engine/Plugins/Runtime/nDisplay/Source/DisplayClusterProjection/Private/Policy/DisplayClusterProjectionPolicyBase.cpp
+++ b/Engine/Plugins/Runtime/nDisplay/Source/DisplayClusterProjection/Private/Policy/DisplayClusterProjectionPolicyBase.cpp
@@ -12,6 +12,35 @@
 #include "Render/Viewport/IDisplayClusterViewportManager.h"
 #include "Render/Viewport/IDisplayClusterViewport.h"
 
+namespace UE::DisplayCluster::Projection::PolicyBaseHelpers
+{
+	static inline USceneComponent* FindOriginComponent(IDisplayClusterViewport* InViewport, const EDisplayClusterRootActorType InRootActorType, const FString& OriginCompId, const FString& PolicyId, bool bAllowExtraLog)
+	{
+		if (ADisplayClusterRootActor* InRootActor = InViewport->GetConfiguration().GetRootActor(InRootActorType))
+		{
+			// Try to get a node specified in the config file
+			if (!OriginCompId.IsEmpty())
+			{
+				UE_LOG(LogDisplayClusterProjection, Verbose, TEXT("Looking for an origin component '%s'..."), *OriginCompId);
+				if (USceneComponent* PolicyOriginComp = InRootActor->GetComponentByName<USceneComponent>(OriginCompId))
+				{
+					return PolicyOriginComp;
+				}
+
+				if (bAllowExtraLog)
+				{
+					UE_LOG(LogDisplayClusterProjection, Error, TEXT("No custom origin set or component '%s' not found for policy '%s'. VR root will be used."), *OriginCompId, *PolicyId);
+				}
+			}
+
+			// default use root actor as Origin
+			return InRootActor->GetRootComponent();
+		}
+
+		return nullptr;
+	}
+};
+using namespace UE::DisplayCluster::Projection;
 
 //////////////////////////////////////////////////////////////////////////////////////////////
 // FDisplayClusterProjectionPolicyBase
@@ -46,7 +75,7 @@ bool FDisplayClusterProjectionPolicyBase::IsEditorOperationMode(IDisplayClusterV
 	}
 
 	// Get state from viewport world (UE-114493)
-	return InViewport && InViewport->IsCurrentWorldHasAnyType(EWorldType::Editor, EWorldType::EditorPreview);
+	return InViewport && InViewport->GetConfiguration().IsCurrentWorldHasAnyType(EWorldType::Editor, EWorldType::EditorPreview);
 }
 
 bool FDisplayClusterProjectionPolicyBase::IsConfigurationChanged(const FDisplayClusterConfigurationProjection* InConfigurationProjectionPolicy) const
@@ -70,49 +99,22 @@ bool FDisplayClusterProjectionPolicyBase::IsConfigurationChanged(const FDisplayC
 
 void FDisplayClusterProjectionPolicyBase::InitializeOriginComponent(IDisplayClusterViewport* InViewport, const FString& OriginCompId)
 {
-	// Reset previous one
-	PolicyOriginComponentRef.ResetSceneComponent();
+	const bool bAllowExtraLog = !IsEditorOperationMode(InViewport);
+
+	ReleaseOriginComponent();
 
 	if (InViewport)
 	{
-		USceneComponent* PolicyOriginComp = nullptr;
-		if (ADisplayClusterRootActor* RootActor = InViewport->GetRootActor())
-		{
-			// default use root actor as Origin
-			PolicyOriginComp = RootActor->GetRootComponent();
-
-			// Try to get a node specified in the config file
-			if (!OriginCompId.IsEmpty())
-			{
-				UE_LOG(LogDisplayClusterProjection, Verbose, TEXT("Looking for an origin component '%s'..."), *OriginCompId);
-				PolicyOriginComp = RootActor->GetComponentByName<USceneComponent>(OriginCompId);
-
-				if (PolicyOriginComp == nullptr)
-				{
-					if (!IsEditorOperationMode(InViewport))
-					{
-						UE_LOG(LogDisplayClusterProjection, Error, TEXT("No custom origin set or component '%s' not found for policy '%s'. VR root will be used."), *OriginCompId, *GetId());
-					}
-
-					PolicyOriginComp = RootActor->GetRootComponent();
-				}
-			}
-		}
-
-		if (!PolicyOriginComp)
-		{
-			if (!IsEditorOperationMode(InViewport))
-			{
-				UE_LOG(LogDisplayClusterProjection, Error, TEXT("Couldn't set origin component"));
-			}
-			return;
-		}
+		// Get origin from scene DCRA (uses for math)
+		PolicySceneOriginComponentRef.SetSceneComponent(PolicyBaseHelpers::FindOriginComponent(InViewport, EDisplayClusterRootActorType::Scene, OriginCompId, GetId(), bAllowExtraLog));
 
-		PolicyOriginComponentRef.SetSceneComponent(PolicyOriginComp);
+		// This origin used to get origin from preview DCRA (uses for preview mesh)
+		PolicyPreviewMeshOriginComponentRef.SetSceneComponent(PolicyBaseHelpers::FindOriginComponent(InViewport, EDisplayClusterRootActorType::Preview, OriginCompId, GetId(), bAllowExtraLog));
 	}
 }
 
 void FDisplayClusterProjectionPolicyBase::ReleaseOriginComponent()
 {
-	PolicyOriginComponentRef.ResetSceneComponent();
+	PolicySceneOriginComponentRef.ResetSceneComponent();
+	PolicyPreviewMeshOriginComponentRef.ResetSceneComponent();
 }
diff --git a/Engine/Plugins/Runtime/nDisplay/Source/DisplayClusterProjection/Private/Policy/DisplayClusterProjectionPolicyBase.h b/Engine/Plugins/Runtime/nDisplay/Source/DisplayClusterProjection/Private/Policy/DisplayClusterProjectionPolicyBase.h
--- a/Engine/Plugins/Runtime/nDisplay/Source/DisplayClusterProjection/Private/Policy/DisplayClusterProjectionPolicyBase.h
+++ b/Engine/Plugins/Runtime/nDisplay/Source/DisplayClusterProjection/Private/Policy/DisplayClusterProjectionPolicyBase.h
@@ -20,19 +20,20 @@ public:
 	FDisplayClusterProjectionPolicyBase(const FString& ProjectionPolicyId, const FDisplayClusterConfigurationProjection* InConfigurationProjectionPolicy);
 
 public:
+	//~BEGIN IDisplayClusterProjectionPolicy
 	virtual const FString& GetId() const override
 	{
 		return PolicyInstanceId;
 	}
 
-	void SetOriginComp(USceneComponent* OriginComp)
+	virtual USceneComponent* const GetOriginComponent() const override
 	{
-		PolicyOriginComponentRef.SetSceneComponent(OriginComp);
+		return PolicySceneOriginComponentRef.GetOrFindSceneComponent();
 	}
 
-	virtual USceneComponent* const GetOriginComponent() const override
+	virtual USceneComponent* const GetPreviewMeshOriginComponent(IDisplayClusterViewport* InViewport) const override
 	{
-		return PolicyOriginComponentRef.GetOrFindSceneComponent();
+		return PolicyPreviewMeshOriginComponentRef.GetOrFindSceneComponent();
 	}
 
 	virtual const TMap<FString, FString>& GetParameters() const override
@@ -41,6 +42,8 @@ public:
 	}
 
 	virtual bool IsConfigurationChanged(const FDisplayClusterConfigurationProjection* InConfigurationProjectionPolicy) const override;
+	//~~END IDisplayClusterProjectionPolicy
+
 
 	static bool IsEditorOperationMode(IDisplayClusterViewport* InViewport);
 	static bool IsEditorOperationMode_RenderThread(const IDisplayClusterViewportProxy* InViewportProxy);
@@ -50,9 +53,18 @@ protected:
 	void ReleaseOriginComponent();
 
 private:
+	// The unique name of this projection policy.
 	const FString PolicyInstanceId;
 
+	// Origin component name.
 	FString PolicyOriginCompId;
+
+	// Projection policy parameters. Used for initialization.
 	TMap<FString, FString> Parameters;
-	FDisplayClusterSceneComponentRef PolicyOriginComponentRef;
+
+	// Origin component in DCRA from the scene
+	FDisplayClusterSceneComponentRef PolicySceneOriginComponentRef;
+
+	// Origin component in the parent DCRA that is used to render the preview.
+	FDisplayClusterSceneComponentRef PolicyPreviewMeshOriginComponentRef;
 };
diff --git a/Engine/Plugins/Runtime/nDisplay/Source/DisplayClusterProjection/Private/Policy/Link/DisplayClusterProjectionLinkPolicy.cpp b/Engine/Plugins/Runtime/nDisplay/Source/DisplayClusterProjection/Private/Policy/Link/DisplayClusterProjectionLinkPolicy.cpp
--- a/Engine/Plugins/Runtime/nDisplay/Source/DisplayClusterProjection/Private/Policy/Link/DisplayClusterProjectionLinkPolicy.cpp
+++ b/Engine/Plugins/Runtime/nDisplay/Source/DisplayClusterProjection/Private/Policy/Link/DisplayClusterProjectionLinkPolicy.cpp
@@ -30,7 +30,7 @@ const FString& FDisplayClusterProjectionLinkPolicy::GetType() const
 void FDisplayClusterProjectionLinkPolicy::SetupProjectionViewPoint(IDisplayClusterViewport* InViewport, const float InDeltaTime, FMinimalViewInfo& InOutViewInfo, float* OutCustomNearClippingPlane)
 {
 	// Getting the right data from a parent
-	if (IDisplayClusterViewportManager* ViewportManager = InViewport->GetViewportManager())
+	if (IDisplayClusterViewportManager* ViewportManager = InViewport->GetConfiguration().GetViewportManager())
 	{
 		if (IDisplayClusterViewport* ParentViewport = ViewportManager->FindViewport(InViewport->GetRenderSettings().GetParentViewportId()))
 		{
@@ -47,7 +47,7 @@ bool FDisplayClusterProjectionLinkPolicy::CalculateView(IDisplayClusterViewport*
 	check(IsInGameThread());
 	check(InViewport);
 
-	if (IDisplayClusterViewportManager* ViewportManager = InViewport->GetViewportManager())
+	if (IDisplayClusterViewportManager* ViewportManager = InViewport->GetConfiguration().GetViewportManager())
 	{
 		if (IDisplayClusterViewport* ParentViewport = ViewportManager->FindViewport(InViewport->GetRenderSettings().GetParentViewportId()))
 		{
@@ -78,7 +78,7 @@ bool FDisplayClusterProjectionLinkPolicy::GetProjectionMatrix(IDisplayClusterVie
 	check(IsInGameThread());
 	check(InViewport);
 
-	if (IDisplayClusterViewportManager* ViewportManager = InViewport->GetViewportManager())
+	if (IDisplayClusterViewportManager* ViewportManager = InViewport->GetConfiguration().GetViewportManager())
 	{
 		if (IDisplayClusterViewport* ParentViewport = ViewportManager->FindViewport(InViewport->GetRenderSettings().GetParentViewportId()))
 		{
diff --git a/Engine/Plugins/Runtime/nDisplay/Source/DisplayClusterProjection/Private/Policy/MPCDI/DisplayClusterProjectionMPCDIPolicy.cpp b/Engine/Plugins/Runtime/nDisplay/Source/DisplayClusterProjection/Private/Policy/MPCDI/DisplayClusterProjectionMPCDIPolicy.cpp
--- a/Engine/Plugins/Runtime/nDisplay/Source/DisplayClusterProjection/Private/Policy/MPCDI/DisplayClusterProjectionMPCDIPolicy.cpp
+++ b/Engine/Plugins/Runtime/nDisplay/Source/DisplayClusterProjection/Private/Policy/MPCDI/DisplayClusterProjectionMPCDIPolicy.cpp
@@ -30,7 +30,7 @@
 #include "Blueprints/DisplayClusterWarpGeometry.h"
 
 #include "DisplayClusterRootActor.h"
-
+#include "Components/DisplayClusterScreenComponent.h"
 
 namespace UE::DisplayClusterProjection::MPCDIPolicy
 {
@@ -71,11 +71,11 @@ void FDisplayClusterProjectionMPCDIPolicy::SetWarpPolicy(IDisplayClusterWarpPoli
 }
 
 IDisplayClusterWarpPolicy* FDisplayClusterProjectionMPCDIPolicy::GetWarpPolicy() const
-			{
+{
 	check(IsInGameThread());
 
 	return WarpPolicyInterface.Get();
-		}
+}
 
 IDisplayClusterWarpPolicy* FDisplayClusterProjectionMPCDIPolicy::GetWarpPolicy_RenderThread() const
 {
@@ -110,9 +110,9 @@ void FDisplayClusterProjectionMPCDIPolicy::UpdateProxyData(IDisplayClusterViewpo
 
 	ENQUEUE_RENDER_COMMAND(DisplayClusterProjectionMPCDIPolicy_UpdateProxyData)(
 		[ProjectionPolicy = InViewport->GetProjectionPolicy(), WarpBlendInterfacePtr = WarpBlendInterface, WarPolicyInterfacePtr = WarpPolicyInterface, Contexts = WarpBlendContexts](FRHICommandListImmediate& RHICmdList)
-{
+	{
 		if (ProjectionPolicy.IsValid())
-{
+		{
 			FDisplayClusterProjectionMPCDIPolicy* MPCDIPolicy = static_cast<FDisplayClusterProjectionMPCDIPolicy*>(ProjectionPolicy.Get());
 			if (MPCDIPolicy)
 			{
@@ -122,9 +122,6 @@ void FDisplayClusterProjectionMPCDIPolicy::UpdateProxyData(IDisplayClusterViewpo
 			}
 		}
 	});
-
-	// reset warp policy after game thread is done
-	WarpPolicyInterface.Reset();
 }
 
 //////////////////////////////////////////////////////////////////////////////////////////////
@@ -304,7 +301,7 @@ void FDisplayClusterProjectionMPCDIPolicy::ApplyWarpBlend_RenderThread(FRHIComma
 		return;
 	}
 
-	const IDisplayClusterViewportManagerProxy* ViewportManagerProxyPtr = InViewportProxy->GetViewportManagerProxy_RenderThread();
+	const IDisplayClusterViewportManagerProxy* ViewportManagerProxyPtr = InViewportProxy->GetConfigurationProxy().GetViewportManagerProxy_RenderThread();
 	if (!ViewportManagerProxyPtr)
 	{
 		return;
@@ -444,43 +441,90 @@ void FDisplayClusterProjectionMPCDIPolicy::ApplyWarpBlend_RenderThread(FRHIComma
 	}
 }
 
-UMeshComponent* FDisplayClusterProjectionMPCDIPolicy::GetOrCreatePreviewMeshComponent(IDisplayClusterViewport* InViewport, bool& bOutIsRootActorComponent)
+bool FDisplayClusterProjectionMPCDIPolicy::HasPreviewMesh(IDisplayClusterViewport* InViewport)
 {
-	return (bIsPreviewMeshEnabled && WarpBlendInterface.IsValid()) ? WarpBlendInterface->GetOrCreateMeshComponent(InViewport, bOutIsRootActorComponent) : nullptr;
+	if (bIsPreviewMeshEnabled)
+	{
+		return true;
+	}
+
+	PreviewMeshComponentRef.ResetSceneComponent();
+
+	return false;
 }
 
-UMeshComponent* FDisplayClusterProjectionMPCDIPolicy::GetOrCreatePreviewMovableMeshComponent(IDisplayClusterViewport* InViewport)
+UMeshComponent* FDisplayClusterProjectionMPCDIPolicy::GetOrCreatePreviewMeshComponent(IDisplayClusterViewport* InViewport, bool& bOutIsRootActorComponent)
 {
-#if WITH_EDITOR
+	if (!HasPreviewMesh(InViewport))
+	{
+		return nullptr;
+	}
+
 	// If we have already created a preview mesh component before, return that component
-	if (UMeshComponent* MovableMeshComp = Cast<UMeshComponent>(MovablePreviewMeshComponentRef.GetOrFindSceneComponent()))
+	if (UMeshComponent* ExistsPreviewMeshComp = Cast<UMeshComponent>(PreviewMeshComponentRef.GetOrFindSceneComponent()))
 	{
-		return MovableMeshComp;
+		return ExistsPreviewMeshComp;
 	}
-	// Create a new mesh component. Best option is to simply create a duplicate of the static preview component by passing a template into NewObject
-	// Attach the movable mesh component to the DCRA's root component, since we don't care about hierarchy here
-	if (ADisplayClusterRootActor* RootActor = InViewport->GetRootActor())
+
+	// Get a new one
+	if (UMeshComponent* PreviewMeshComp = WarpBlendInterface.IsValid() ? WarpBlendInterface->GetOrCreatePreviewMeshComponent(InViewport, bOutIsRootActorComponent) : nullptr)
 	{
-		USceneComponent* RootComponent = RootActor->GetRootComponent();
-		bool bIsRootActorComponent = false;
-		if (UMeshComponent* PreviewMeshComponent = GetOrCreatePreviewMeshComponent(InViewport, bIsRootActorComponent))
+		// Store reference to mesh component
+		PreviewMeshComponentRef.SetSceneComponent(PreviewMeshComp);
+
+		return PreviewMeshComp;
+	}
+
+	return nullptr;
+}
+
+USceneComponent* const FDisplayClusterProjectionMPCDIPolicy::GetPreviewMovableMeshOriginComponent(IDisplayClusterViewport* InViewport) const
+{
+	// Note: currently for the movable mesh component we expect it to be used only in the scene,
+	// so we use the root component from the scene all the time.
+	// But if other use cases are found, we need to refine this logic.
+
+	return GetOriginComponent();
+}
+
+bool FDisplayClusterProjectionMPCDIPolicy::HasPreviewMovableMesh(IDisplayClusterViewport* InViewport)
+{
+	if (InViewport)
+	{
+		// The movable preview grid is a feature for the warp policy, so we must request permission to use it.
+		const bool bWarpPolicyUseMovableMesh = WarpPolicyInterface.IsValid() && WarpPolicyInterface->HasPreviewMovableMesh(InViewport);
+		if (bIsPreviewMeshEnabled && bWarpPolicyUseMovableMesh)
 		{
-			const FString CompName = FString::Printf(TEXT("MPCDI_%s_movable_impl"), *GetId());
-			const EObjectFlags ObjectFlags = EObjectFlags::RF_DuplicateTransient | RF_Transient | RF_TextExportTransient;
-			UMeshComponent* NewMovableMeshComp = NewObject<UMeshComponent>(RootComponent, PreviewMeshComponent->GetClass(), *CompName, ObjectFlags, PreviewMeshComponent);
-			if (NewMovableMeshComp)
-			{
-				NewMovableMeshComp->RegisterComponent();
-				NewMovableMeshComp->AttachToComponent(RootComponent, FAttachmentTransformRules(EAttachmentRule::KeepRelative, false));
-				NewMovableMeshComp->SetIsVisualizationComponent(true);
-				NewMovableMeshComp->SetHiddenInGame(true);
-				// Store reference to mesh component
-				MovablePreviewMeshComponentRef.SetSceneComponent(NewMovableMeshComp);
-				return NewMovableMeshComp;
-			}
+			return true;
 		}
 	}
-#endif
+
+	PreviewMovableMeshComponentRef.ResetSceneComponent();
+
+	return false;
+}
+
+UMeshComponent* FDisplayClusterProjectionMPCDIPolicy::GetOrCreatePreviewMovableMeshComponent(IDisplayClusterViewport* InViewport)
+{
+	if (!HasPreviewMovableMesh(InViewport))
+	{
+		return nullptr;
+	}
+
+	// If we have already created a preview mesh component before, return that component
+	if (UMeshComponent* ExistsPreviewMovableMeshComp = Cast<UMeshComponent>(PreviewMovableMeshComponentRef.GetOrFindSceneComponent()))
+	{
+		return ExistsPreviewMovableMeshComp;
+	}
+
+	// Get a new one
+	if (UMeshComponent* PreviewMovableMeshComp = WarpBlendInterface.IsValid() ? WarpBlendInterface->GetOrCreatePreviewMovableMeshComponent(InViewport) : nullptr)
+	{
+		// Store reference to mesh component
+		PreviewMovableMeshComponentRef.SetSceneComponent(PreviewMovableMeshComp);
+
+		return PreviewMovableMeshComp;
+	}
 
 	return nullptr;
 }
@@ -519,22 +563,27 @@ bool FDisplayClusterProjectionMPCDIPolicy::CreateWarpBlendFromConfig(IDisplayClu
 				return false;
 			}
 
-			if (CfgData.ScreenComponent == nullptr)
+			// Using the screen component as the "region" for the MPCDI 2d profile.
+			const bool bUseScreenComponentForMPCDIProfile2D = CfgData.ScreenComponent && CfgData.MPCDIAttributes.ProfileType == EDisplayClusterWarpProfileType::warp_2D;
+			if (!bUseScreenComponentForMPCDIProfile2D)
 			{
 				FDisplayClusterWarpInitializer_MPCDIFile CreateParameters;
-			CreateParameters.MPCDIFileName = CfgData.MPCDIFileName;
-			CreateParameters.BufferId = CfgData.BufferId;
-			CreateParameters.RegionId = CfgData.RegionId;
+				CreateParameters.MPCDIFileName = CfgData.MPCDIFileName;
+				CreateParameters.BufferId = CfgData.BufferId;
+				CreateParameters.RegionId = CfgData.RegionId;
 
 				WarpBlendInterface = GetWarpAPI().Create(CreateParameters);
 			}
 			else
 			{
 				FDisplayClusterWarpInitializer_MPCDIFile_Profile2DScreen CreateParameters;
-			CreateParameters.MPCDIFileName = CfgData.MPCDIFileName;
-			CreateParameters.BufferId = CfgData.BufferId;
-			CreateParameters.RegionId = CfgData.RegionId;
-				CreateParameters.StaticMeshComponent = CfgData.ScreenComponent;
+				CreateParameters.MPCDIFileName = CfgData.MPCDIFileName;
+				CreateParameters.BufferId = CfgData.BufferId;
+				CreateParameters.RegionId = CfgData.RegionId;
+
+				// Using the screen component as the "region" for the MPCDI 2d profile.
+				CreateParameters.WarpMeshComponent = CfgData.ScreenComponent;
+				CreateParameters.PreviewMeshComponent = CfgData.PreviewScreenComponent;
 
 				WarpBlendInterface = GetWarpAPI().Create(CreateParameters);
 			}
@@ -576,4 +625,3 @@ bool FDisplayClusterProjectionMPCDIPolicy::CreateWarpBlendFromConfig(IDisplayClu
 
 	return false;
 }
-
diff --git a/Engine/Plugins/Runtime/nDisplay/Source/DisplayClusterProjection/Private/Policy/MPCDI/DisplayClusterProjectionMPCDIPolicy.h b/Engine/Plugins/Runtime/nDisplay/Source/DisplayClusterProjection/Private/Policy/MPCDI/DisplayClusterProjectionMPCDIPolicy.h
--- a/Engine/Plugins/Runtime/nDisplay/Source/DisplayClusterProjection/Private/Policy/MPCDI/DisplayClusterProjectionMPCDIPolicy.h
+++ b/Engine/Plugins/Runtime/nDisplay/Source/DisplayClusterProjection/Private/Policy/MPCDI/DisplayClusterProjectionMPCDIPolicy.h
@@ -6,6 +6,8 @@
 #include "Containers/DisplayClusterWarpContext.h"
 #include "IDisplayClusterWarpBlend.h"
 
+class UMeshComponent;
+
 /**
  * MPCDI projection policy
  * Supported load from 'MPCDI' and 'PFM' files
@@ -55,6 +57,13 @@ public:
 	virtual IDisplayClusterWarpPolicy* GetWarpPolicy() const override;
 	virtual IDisplayClusterWarpPolicy* GetWarpPolicy_RenderThread() const override;
 
+	virtual bool HasPreviewMesh(IDisplayClusterViewport* InViewport) override;
+	virtual UMeshComponent* GetOrCreatePreviewMeshComponent(IDisplayClusterViewport* InViewport, bool& bOutIsRootActorComponent) override;
+
+	virtual bool HasPreviewMovableMesh(IDisplayClusterViewport* InViewport) override;
+	virtual UMeshComponent* GetOrCreatePreviewMovableMeshComponent(IDisplayClusterViewport* InViewport) override;
+	virtual USceneComponent* const GetPreviewMovableMeshOriginComponent(IDisplayClusterViewport* InViewport) const override;
+
 	//~~End IDisplayClusterProjectionPolicy
 
 protected:
@@ -74,21 +83,10 @@ protected:
 	TArray<FDisplayClusterWarpContext> WarpBlendContexts;
 	TArray<FDisplayClusterWarpContext> WarpBlendContexts_Proxy;
 
-protected:
 	bool bInvalidConfiguration = false;
 	bool bIsPreviewMeshEnabled = false;
 
-protected:
-	virtual bool HasPreviewMesh() override
-	{
-		return true;
-	}
-
-	virtual bool HasPreviewMovableMesh() override { return true; }
-
-	virtual class UMeshComponent* GetOrCreatePreviewMeshComponent(IDisplayClusterViewport* InViewport, bool& bOutIsRootActorComponent) override;
-	virtual class UMeshComponent* GetOrCreatePreviewMovableMeshComponent(IDisplayClusterViewport* InViewport) override;
-
 private:
-	FDisplayClusterSceneComponentRef MovablePreviewMeshComponentRef;
+	FDisplayClusterSceneComponentRef PreviewMeshComponentRef;
+	FDisplayClusterSceneComponentRef PreviewMovableMeshComponentRef;
 };
diff --git a/Engine/Plugins/Runtime/nDisplay/Source/DisplayClusterProjection/Private/Policy/MPCDI/DisplayClusterProjectionMPCDIPolicy_Config.cpp b/Engine/Plugins/Runtime/nDisplay/Source/DisplayClusterProjection/Private/Policy/MPCDI/DisplayClusterProjectionMPCDIPolicy_Config.cpp
--- a/Engine/Plugins/Runtime/nDisplay/Source/DisplayClusterProjection/Private/Policy/MPCDI/DisplayClusterProjectionMPCDIPolicy_Config.cpp
+++ b/Engine/Plugins/Runtime/nDisplay/Source/DisplayClusterProjection/Private/Policy/MPCDI/DisplayClusterProjectionMPCDIPolicy_Config.cpp
@@ -109,13 +109,21 @@ bool FDisplayClusterProjectionMPCDIPolicy_ConfigParser::ImplGetMPCDIConfig()
 	{
 		if (!ScreenComponentName.IsEmpty())
 		{
-			if (ADisplayClusterRootActor* RootActorPtr = Viewport->GetRootActor())
+			if (ADisplayClusterRootActor* SceneRootActor = Viewport->GetConfiguration().GetRootActor(EDisplayClusterRootActorType::Scene))
 			{
-				if (UDisplayClusterScreenComponent* ScreenComp = RootActorPtr->GetComponentByName<UDisplayClusterScreenComponent>(ScreenComponentName))
+				if (UDisplayClusterScreenComponent* ScreenComp = SceneRootActor->GetComponentByName<UDisplayClusterScreenComponent>(ScreenComponentName))
 				{
 					ScreenComponent = ScreenComp;
 				}
 			}
+
+			if (ADisplayClusterRootActor* PreviewRootActor = Viewport->GetConfiguration().GetRootActor(EDisplayClusterRootActorType::Preview))
+			{
+				if (UDisplayClusterScreenComponent* PreviewScreenComp = PreviewRootActor->GetComponentByName<UDisplayClusterScreenComponent>(ScreenComponentName))
+				{
+					PreviewScreenComponent = PreviewScreenComp;
+				}
+			}
 		}
 	}
 
diff --git a/Engine/Plugins/Runtime/nDisplay/Source/DisplayClusterProjection/Private/Policy/MPCDI/DisplayClusterProjectionMPCDIPolicy_Config.h b/Engine/Plugins/Runtime/nDisplay/Source/DisplayClusterProjection/Private/Policy/MPCDI/DisplayClusterProjectionMPCDIPolicy_Config.h
--- a/Engine/Plugins/Runtime/nDisplay/Source/DisplayClusterProjection/Private/Policy/MPCDI/DisplayClusterProjectionMPCDIPolicy_Config.h
+++ b/Engine/Plugins/Runtime/nDisplay/Source/DisplayClusterProjection/Private/Policy/MPCDI/DisplayClusterProjectionMPCDIPolicy_Config.h
@@ -8,6 +8,7 @@
 
 class IDisplayClusterViewport;
 class UStaticMeshComponent;
+class UDisplayClusterScreenComponent;
 
 /**
  * MPCDI projection policy configuration parser
@@ -30,7 +31,10 @@ struct FDisplayClusterProjectionMPCDIPolicy_ConfigParser
 	FDisplayClusterWarpMPCDIAttributes MPCDIAttributes;
 
 	// Screen component with source geometry
-	UStaticMeshComponent* ScreenComponent = nullptr;
+	UDisplayClusterScreenComponent* ScreenComponent = nullptr;
+
+	// Screen component with source geometry used as a preview mesh
+	UDisplayClusterScreenComponent* PreviewScreenComponent = nullptr;
 
 	// Name of the PFM file that is used as the data source
 	FString PFMFile;
diff --git a/Engine/Plugins/Runtime/nDisplay/Source/DisplayClusterProjection/Private/Policy/Mesh/DisplayClusterProjectionMeshPolicy.cpp b/Engine/Plugins/Runtime/nDisplay/Source/DisplayClusterProjection/Private/Policy/Mesh/DisplayClusterProjectionMeshPolicy.cpp
--- a/Engine/Plugins/Runtime/nDisplay/Source/DisplayClusterProjection/Private/Policy/Mesh/DisplayClusterProjectionMeshPolicy.cpp
+++ b/Engine/Plugins/Runtime/nDisplay/Source/DisplayClusterProjection/Private/Policy/Mesh/DisplayClusterProjectionMeshPolicy.cpp
@@ -23,12 +23,14 @@ struct FDisplayClusterProjectionMeshPolicyConfiguration
 {
 	// StaticMesh component with source geometry
 	UStaticMeshComponent* StaticMeshComponent = nullptr;
+	UStaticMeshComponent* PreviewStaticMeshComponent = nullptr;
 
 	// StaticMesh geometry LOD
 	int32 StaticMeshComponentLODIndex = 0;
 
 	// ProceduralMesh component with source geometry
 	UProceduralMeshComponent* ProceduralMeshComponent = nullptr;
+	UProceduralMeshComponent* PreviewProceduralMeshComponent = nullptr;
 
 	// ProceduralMesh section index
 	int32 ProceduralMeshComponentSectionIndex = 0;
@@ -73,7 +75,9 @@ bool FDisplayClusterProjectionMeshPolicy::CreateWarpMeshInterface(IDisplayCluste
 
 			CreateParameters.OriginComponent = GetOriginComponent();
 
-			CreateParameters.StaticMeshComponent = WarpCfg.StaticMeshComponent;
+			CreateParameters.PreviewMeshComponent = WarpCfg.PreviewStaticMeshComponent;
+
+			CreateParameters.WarpMeshComponent = WarpCfg.StaticMeshComponent;
 			CreateParameters.StaticMeshComponentLODIndex = WarpCfg.StaticMeshComponentLODIndex;
 
 			CreateParameters.BaseUVIndex      = WarpCfg.BaseUVIndex;
@@ -87,7 +91,9 @@ bool FDisplayClusterProjectionMeshPolicy::CreateWarpMeshInterface(IDisplayCluste
 
 			CreateParameters.OriginComponent = GetOriginComponent();
 
-			CreateParameters.ProceduralMeshComponent = WarpCfg.ProceduralMeshComponent;
+			CreateParameters.PreviewMeshComponent = WarpCfg.PreviewProceduralMeshComponent;
+
+			CreateParameters.WarpMeshComponent = WarpCfg.ProceduralMeshComponent;
 			CreateParameters.ProceduralMeshComponentSectionIndex = WarpCfg.ProceduralMeshComponentSectionIndex;
 
 			CreateParameters.BaseUVIndex      = WarpCfg.BaseUVIndex;
@@ -146,8 +152,9 @@ bool FDisplayClusterProjectionMeshPolicy::GetWarpMeshConfiguration(IDisplayClust
 	check(InViewport);
 
 	// Get our VR root
-	ADisplayClusterRootActor* RootActorPtr = InViewport->GetRootActor();
-	if (!RootActorPtr)
+	ADisplayClusterRootActor* PreviewRootActor = InViewport->GetConfiguration().GetRootActor(EDisplayClusterRootActorType::Preview);
+	ADisplayClusterRootActor* SceneRootActor = InViewport->GetConfiguration().GetRootActor(EDisplayClusterRootActorType::Scene);
+	if (!SceneRootActor)
 	{
 		UE_LOG(LogDisplayClusterProjectionMesh, Error, TEXT("Couldn't get a VR root object"));
 		return false;
@@ -204,7 +211,9 @@ bool FDisplayClusterProjectionMeshPolicy::GetWarpMeshConfiguration(IDisplayClust
 	OutWarpCfg.ChromakeyUVIndex = CfgChromakeyUVIndex;
 
 	// Get the StaticMeshComponent
-	OutWarpCfg.StaticMeshComponent = RootActorPtr->GetComponentByName<UStaticMeshComponent>(ComponentId);
+	OutWarpCfg.StaticMeshComponent        = SceneRootActor->GetComponentByName<UStaticMeshComponent>(ComponentId);
+	OutWarpCfg.PreviewStaticMeshComponent = PreviewRootActor ? PreviewRootActor->GetComponentByName<UStaticMeshComponent>(ComponentId) : nullptr;
+
 	if (OutWarpCfg.StaticMeshComponent != nullptr)
 	{
 		int CfgLODIndex;
@@ -225,7 +234,9 @@ bool FDisplayClusterProjectionMeshPolicy::GetWarpMeshConfiguration(IDisplayClust
 	else
 	{
 		// Get the ProceduralMeshComponent
-		OutWarpCfg.ProceduralMeshComponent = RootActorPtr->GetComponentByName<UProceduralMeshComponent>(ComponentId);
+		OutWarpCfg.ProceduralMeshComponent        = SceneRootActor->GetComponentByName<UProceduralMeshComponent>(ComponentId);
+		OutWarpCfg.PreviewProceduralMeshComponent = PreviewRootActor ? PreviewRootActor->GetComponentByName<UProceduralMeshComponent>(ComponentId) : nullptr;
+
 		if (OutWarpCfg.ProceduralMeshComponent == nullptr)
 		{
 			if (!IsEditorOperationMode(InViewport))
diff --git a/Engine/Plugins/Runtime/nDisplay/Source/DisplayClusterProjection/Private/Policy/Simple/DisplayClusterProjectionSimplePolicy.cpp b/Engine/Plugins/Runtime/nDisplay/Source/DisplayClusterProjection/Private/Policy/Simple/DisplayClusterProjectionSimplePolicy.cpp
--- a/Engine/Plugins/Runtime/nDisplay/Source/DisplayClusterProjection/Private/Policy/Simple/DisplayClusterProjectionSimplePolicy.cpp
+++ b/Engine/Plugins/Runtime/nDisplay/Source/DisplayClusterProjection/Private/Policy/Simple/DisplayClusterProjectionSimplePolicy.cpp
@@ -59,6 +59,7 @@ bool FDisplayClusterProjectionSimplePolicy::HandleStartScene(IDisplayClusterView
 		return true;
 	}
 
+	ReleaseMeshData(InViewport);
 	return false;
 }
 
@@ -66,8 +67,7 @@ void FDisplayClusterProjectionSimplePolicy::HandleEndScene(IDisplayClusterViewpo
 {
 	check(IsInGameThread());
 
-	// Forget about screen. It will be released by the Engine.
-	ScreenCompRef.ResetSceneComponent();
+	ReleaseMeshData(InViewport);
 }
 
 bool FDisplayClusterProjectionSimplePolicy::CalculateView(IDisplayClusterViewport* InViewport, const uint32 InContextNum, FVector& InOutViewLocation, FRotator& InOutViewRotation, const FVector& ViewOffset, const float WorldToMeters, const float NCP, const float FCP)
@@ -183,11 +183,13 @@ bool FDisplayClusterProjectionSimplePolicy::InitializeMeshData(IDisplayClusterVi
 {
 	UE_LOG(LogDisplayClusterProjectionSimple, Verbose, TEXT("Initializing screen geometry for viewport policy  %s"), *GetId());
 
+	const bool bExtraLog = !IsEditorOperationMode(InViewport);
+
 	// Get our VR root
-	ADisplayClusterRootActor* const RootActorPtr = InViewport->GetRootActor();
-	if (!RootActorPtr)
+	ADisplayClusterRootActor* const SceneRootActor = InViewport->GetConfiguration().GetRootActor(EDisplayClusterRootActorType::Scene);
+	if (!SceneRootActor)
 	{
-		if (!IsEditorOperationMode(InViewport))
+		if (bExtraLog)
 		{
 			UE_LOG(LogDisplayClusterProjectionSimple, Error, TEXT("Couldn't get a VR root object"));
 		}
@@ -195,32 +197,44 @@ bool FDisplayClusterProjectionSimplePolicy::InitializeMeshData(IDisplayClusterVi
 		return false;
 	}
 
-	// Get screen component
-	UDisplayClusterScreenComponent* ScreenComp = RootActorPtr->GetComponentByName<UDisplayClusterScreenComponent>(ScreenId);
+	// Get preview screen components (optional)
+	if(ADisplayClusterRootActor* const PreviewRootActor = InViewport->GetConfiguration().GetRootActor(EDisplayClusterRootActorType::Preview))
+	{
+		UDisplayClusterScreenComponent* PreviewScreenComp = PreviewRootActor->GetComponentByName<UDisplayClusterScreenComponent>(ScreenId);
+		PreviewScreenCompRef.SetSceneComponent(PreviewScreenComp);
+	}
+
+	// Get screen components
+	UDisplayClusterScreenComponent* ScreenComp = SceneRootActor->GetComponentByName<UDisplayClusterScreenComponent>(ScreenId);
+	ScreenCompRef.SetSceneComponent(ScreenComp);
+
 	if (!ScreenComp)
 	{
-		if (!IsEditorOperationMode(InViewport))
+		if (bExtraLog)
 		{
-			UE_LOG(LogDisplayClusterProjectionSimple, Warning, TEXT("Couldn't initialize screen component"));
+			UE_LOG(LogDisplayClusterProjectionSimple, Warning, TEXT("Couldn't initialize screen component '%s'"), *ScreenId);
 		}
 
 		return false;
 	}
 
-	return ScreenCompRef.SetSceneComponent(ScreenComp);
+	return true;
 }
 
 void FDisplayClusterProjectionSimplePolicy::ReleaseMeshData(IDisplayClusterViewport* InViewport)
 {
+	// Forget about screen. It will be released by the Engine.
+	ScreenCompRef.ResetSceneComponent();
+	PreviewScreenCompRef.ResetSceneComponent();
 }
 
 UMeshComponent* FDisplayClusterProjectionSimplePolicy::GetOrCreatePreviewMeshComponent(IDisplayClusterViewport* InViewport, bool& bOutIsRootActorComponent)
 {
 	bOutIsRootActorComponent = true;
-	if (USceneComponent* SceneComponent = ScreenCompRef.GetOrFindSceneComponent())
+	if (USceneComponent* PreviewSceneComponent = PreviewScreenCompRef.GetOrFindSceneComponent())
 	{
-		UDisplayClusterScreenComponent* ScreenComp = StaticCast<UDisplayClusterScreenComponent*>(SceneComponent);
-		return (ScreenComp == nullptr) ? nullptr : ScreenComp;
+		UDisplayClusterScreenComponent* PreviewScreenComp = StaticCast<UDisplayClusterScreenComponent*>(PreviewSceneComponent);
+		return (PreviewScreenComp == nullptr) ? nullptr : PreviewScreenComp;
 	}
 
 	return nullptr;
diff --git a/Engine/Plugins/Runtime/nDisplay/Source/DisplayClusterProjection/Private/Policy/Simple/DisplayClusterProjectionSimplePolicy.h b/Engine/Plugins/Runtime/nDisplay/Source/DisplayClusterProjection/Private/Policy/Simple/DisplayClusterProjectionSimplePolicy.h
--- a/Engine/Plugins/Runtime/nDisplay/Source/DisplayClusterProjection/Private/Policy/Simple/DisplayClusterProjectionSimplePolicy.h
+++ b/Engine/Plugins/Runtime/nDisplay/Source/DisplayClusterProjection/Private/Policy/Simple/DisplayClusterProjectionSimplePolicy.h
@@ -29,7 +29,7 @@ public:
 	virtual bool CalculateView(IDisplayClusterViewport* InViewport, const uint32 InContextNum, FVector& InOutViewLocation, FRotator& InOutViewRotation, const FVector& ViewOffset, const float WorldToMeters, const float NCP, const float FCP) override;
 	virtual bool GetProjectionMatrix(IDisplayClusterViewport* InViewport, const uint32 InContextNum, FMatrix& OutPrjMatrix) override;
 
-	virtual bool HasPreviewMesh() override
+	virtual bool HasPreviewMesh(IDisplayClusterViewport* InViewport) override
 	{
 		return true;
 	}
@@ -67,6 +67,9 @@ private:
 	// Weak ptr screen component
 	FDisplayClusterSceneComponentRef ScreenCompRef;
 
+	// Weak ptr to preview screen component
+	FDisplayClusterSceneComponentRef PreviewScreenCompRef;
+
 	struct FViewData
 	{
 		FVector  ViewLoc;
diff --git a/Engine/Plugins/Runtime/nDisplay/Source/DisplayClusterProjection/Private/Policy/VIOSO/Windows/DisplayClusterProjectionVIOSOPolicy.cpp b/Engine/Plugins/Runtime/nDisplay/Source/DisplayClusterProjection/Private/Policy/VIOSO/Windows/DisplayClusterProjectionVIOSOPolicy.cpp
--- a/Engine/Plugins/Runtime/nDisplay/Source/DisplayClusterProjection/Private/Policy/VIOSO/Windows/DisplayClusterProjectionVIOSOPolicy.cpp
+++ b/Engine/Plugins/Runtime/nDisplay/Source/DisplayClusterProjection/Private/Policy/VIOSO/Windows/DisplayClusterProjectionVIOSOPolicy.cpp
@@ -222,8 +222,7 @@ UMeshComponent* FDisplayClusterProjectionVIOSOPolicy::GetOrCreatePreviewMeshComp
 	// Create a new DCRA mesh component
 	bOutIsRootActorComponent = false;
 
-	// Get geometry data
-	USceneComponent* OriginComp = GetOriginComponent();
+	USceneComponent* OriginComp = GetPreviewMeshOriginComponent(InViewport);
 	TSharedPtr<FDisplayClusterProjectionVIOSOGeometryExportData, ESPMode::ThreadSafe> GeometryExportData = FDisplayClusterProjectionVIOSOGeometryExportData::Create(VIOSOLibrary, ViosoConfigData);
 	if (OriginComp && GeometryExportData.IsValid())
 	{
diff --git a/Engine/Plugins/Runtime/nDisplay/Source/DisplayClusterProjection/Private/Policy/VIOSO/Windows/DisplayClusterProjectionVIOSOPolicy.h b/Engine/Plugins/Runtime/nDisplay/Source/DisplayClusterProjection/Private/Policy/VIOSO/Windows/DisplayClusterProjectionVIOSOPolicy.h
--- a/Engine/Plugins/Runtime/nDisplay/Source/DisplayClusterProjection/Private/Policy/VIOSO/Windows/DisplayClusterProjectionVIOSOPolicy.h
+++ b/Engine/Plugins/Runtime/nDisplay/Source/DisplayClusterProjection/Private/Policy/VIOSO/Windows/DisplayClusterProjectionVIOSOPolicy.h
@@ -59,7 +59,7 @@ protected:
 	//////////////////////////////////////////////////////////////////////////////////////////////
 	// IDisplayClusterProjectionPolicyPreview
 	//////////////////////////////////////////////////////////////////////////////////////////////
-	virtual bool HasPreviewMesh() override
+	virtual bool HasPreviewMesh(IDisplayClusterViewport* InViewport) override
 	{
 		return true;
 	}
diff --git a/Engine/Plugins/Runtime/nDisplay/Source/DisplayClusterWarp/Private/Components/DisplayClusterInFrustumFitCameraComponent.cpp b/Engine/Plugins/Runtime/nDisplay/Source/DisplayClusterWarp/Private/Components/DisplayClusterInFrustumFitCameraComponent.cpp
--- a/Engine/Plugins/Runtime/nDisplay/Source/DisplayClusterWarp/Private/Components/DisplayClusterInFrustumFitCameraComponent.cpp
+++ b/Engine/Plugins/Runtime/nDisplay/Source/DisplayClusterWarp/Private/Components/DisplayClusterInFrustumFitCameraComponent.cpp
@@ -55,36 +55,39 @@ UDisplayClusterInFrustumFitCameraComponent::UDisplayClusterInFrustumFitCameraCom
 	bAutoActivate = true;
 }
 
-void UDisplayClusterInFrustumFitCameraComponent::TickComponent(float DeltaTime, ELevelTick TickType, FActorComponentTickFunction* ThisTickFunction)
+const UDisplayClusterInFrustumFitCameraComponent& UDisplayClusterInFrustumFitCameraComponent::GetConfigurationInFrustumFitCameraComponent(IDisplayClusterViewportConfiguration& InViewportConfiguration) const
 {
-	// todo: render preview
-
-#if WITH_EDITOR
-	if (WarpPolicy.IsValid())
+	if (ADisplayClusterRootActor* ConfigurationRootActor = InViewportConfiguration.GetRootActor(EDisplayClusterRootActorType::Configuration))
 	{
-		if (bRefreshPreviewState)
+		if (ConfigurationRootActor != GetOwner())
 		{
-			// Set the preview mesh visibility state to match the bShowPreviewFrustumFit flag
-			if (ADisplayClusterRootActor* ParentRootActor = Cast<ADisplayClusterRootActor>(GetOwner()))
+			if (UDisplayClusterInFrustumFitCameraComponent* ConfigurationCameraComponent = ConfigurationRootActor->GetComponentByName<UDisplayClusterInFrustumFitCameraComponent>(GetName()))
 			{
-				if (IDisplayClusterViewportManager* ViewportManager = ParentRootActor->GetViewportManager())
-				{
-					TArray<TSharedPtr<IDisplayClusterViewport, ESPMode::ThreadSafe>> Viewports = ViewportManager->GetViewportsForWarpPolicy(WarpPolicy);
-					for (const TSharedPtr<IDisplayClusterViewport, ESPMode::ThreadSafe>& Viewport : Viewports)
-					{
-						if (UMeshComponent* MovableMeshComponent = Viewport->GetProjectionPolicy()->GetOrCreatePreviewMovableMeshComponent(Viewport.Get()))
-						{
-							MovableMeshComponent->SetVisibility(bShowPreviewFrustumFit);
-						}
-					}
-				}
+				return *ConfigurationCameraComponent;
 			}
-
-			bRefreshPreviewState = false;
 		}
 	}
-#endif
 
+	return *this;
+}
+
+bool UDisplayClusterInFrustumFitCameraComponent::IsEnabled() const
+{
+	return bEnableCameraProjection;
+}
+
+UCameraComponent* UDisplayClusterInFrustumFitCameraComponent::GetExternalCameraComponent() const
+{
+	if (ACineCameraActor* CineCamera = ExternalCameraActor.Get())
+	{
+		return CineCamera->GetCameraComponent();
+	}
+
+	return nullptr;
+}
+
+void UDisplayClusterInFrustumFitCameraComponent::TickComponent(float DeltaTime, ELevelTick TickType, FActorComponentTickFunction* ThisTickFunction)
+{
 	// Tick warp policy instance
 	if (WarpPolicy.IsValid())
 	{
@@ -98,37 +101,37 @@ void UDisplayClusterInFrustumFitCameraComponent::TickComponent(float DeltaTime,
 	}
 }
 
-void UDisplayClusterInFrustumFitCameraComponent::GetDesiredView(FMinimalViewInfo& InOutViewInfo, float* OutCustomNearClippingPlane)
+void UDisplayClusterInFrustumFitCameraComponent::GetDesiredView(IDisplayClusterViewportConfiguration& InViewportConfiguration, FMinimalViewInfo& InOutViewInfo, float* OutCustomNearClippingPlane)
 {
-	if (IsEnabled())
+	const UDisplayClusterInFrustumFitCameraComponent& ConfigurationCameraComponent = GetConfigurationInFrustumFitCameraComponent(InViewportConfiguration);
+	if (ConfigurationCameraComponent.IsEnabled())
 	{
-		if (ADisplayClusterRootActor* ParentRootActor = Cast<ADisplayClusterRootActor>(GetOwner()))
+		if (ADisplayClusterRootActor* SceneRootActor = InViewportConfiguration.GetRootActor(EDisplayClusterRootActorType::Scene))
 		{
-			if (IDisplayClusterViewport::GetCameraComponentView(GetExternalCameraComponent(), ParentRootActor->GetWorldDeltaSeconds(), bUseCameraPostprocess, InOutViewInfo, OutCustomNearClippingPlane))
+			if (IDisplayClusterViewport::GetCameraComponentView(ConfigurationCameraComponent.GetExternalCameraComponent(), SceneRootActor->GetWorldDeltaSeconds(), ConfigurationCameraComponent.bUseCameraPostprocess, InOutViewInfo, OutCustomNearClippingPlane))
 			{
 				// 1. Use external camera for rendering
 				return;
 			}
+		}
 
-			if(IDisplayClusterViewportManager* ViewportManager = ParentRootActor->GetViewportManager())
-			{
-				if (IDisplayClusterViewport::GetPlayerCameraView(ViewportManager->GetCurrentWorld(), bUseCameraPostprocess, InOutViewInfo))
-				{
-					// 2. Use active game camera
-					return;
-				}
-			}
+		if (IDisplayClusterViewport::GetPlayerCameraView(InViewportConfiguration.GetCurrentWorld(), ConfigurationCameraComponent.bUseCameraPostprocess, InOutViewInfo))
+		{
+			// 2. Use active game camera
+			return;
 		}
 	}
 
 	// use default logic
-	return UDisplayClusterCameraComponent::GetDesiredView(InOutViewInfo, OutCustomNearClippingPlane);
+	return UDisplayClusterCameraComponent::GetDesiredView(InViewportConfiguration, InOutViewInfo, OutCustomNearClippingPlane);
 }
 
 bool UDisplayClusterInFrustumFitCameraComponent::ShouldUseEntireClusterViewports(IDisplayClusterViewportManager* InViewportManager) const
 {
+	 const UDisplayClusterInFrustumFitCameraComponent& ConfigurationCameraComponent = InViewportManager ? GetConfigurationInFrustumFitCameraComponent(InViewportManager->GetConfiguration()) : *this;
+
 	// Only when this component is enabled should viewports be created for the entire cluster that accesses this component.
-	return IsEnabled();
+	return ConfigurationCameraComponent.IsEnabled();
 }
 
 IDisplayClusterWarpPolicy* UDisplayClusterInFrustumFitCameraComponent::GetWarpPolicy(IDisplayClusterViewportManager* InViewportManager)
@@ -170,32 +173,7 @@ void UDisplayClusterInFrustumFitCameraComponent::OnRegister()
 
 }
 
-bool UDisplayClusterInFrustumFitCameraComponent::IsEnabled() const
-{
-	return bEnableCameraProjection;
-}
-
-UCameraComponent* UDisplayClusterInFrustumFitCameraComponent::GetExternalCameraComponent() const
-{
-	if (ACineCameraActor* CineCamera = ExternalCameraActor.Get())
-	{
-		return CineCamera->GetCameraComponent();
-	}
-
-	return nullptr;
-}
-
 #if WITH_EDITOR
-void UDisplayClusterInFrustumFitCameraComponent::PostEditChangeProperty(FPropertyChangedEvent& PropertyChangedEvent)
-{
-	Super::PostEditChangeProperty(PropertyChangedEvent);
-
-	if (PropertyChangedEvent.Property && PropertyChangedEvent.Property->GetFName() == GET_MEMBER_NAME_CHECKED(UDisplayClusterInFrustumFitCameraComponent, bShowPreviewFrustumFit))
-	{
-		InvalidatePreviewState();
-	}
-}
-
 bool UDisplayClusterInFrustumFitCameraComponent::GetEditorPreviewInfo(float DeltaTime, FMinimalViewInfo& ViewOut)
 {
 	if (UCameraComponent* CameraComponent = GetExternalCameraComponent())
@@ -215,5 +193,4 @@ TSharedPtr<SWidget> UDisplayClusterInFrustumFitCameraComponent::GetCustomEditorP
 
 	return nullptr;
 }
-
 #endif
diff --git a/Engine/Plugins/Runtime/nDisplay/Source/DisplayClusterWarp/Private/Policy/InFrustumFit/DisplayClusterWarpInFrustumFitPolicy.cpp b/Engine/Plugins/Runtime/nDisplay/Source/DisplayClusterWarp/Private/Policy/InFrustumFit/DisplayClusterWarpInFrustumFitPolicy.cpp
--- a/Engine/Plugins/Runtime/nDisplay/Source/DisplayClusterWarp/Private/Policy/InFrustumFit/DisplayClusterWarpInFrustumFitPolicy.cpp
+++ b/Engine/Plugins/Runtime/nDisplay/Source/DisplayClusterWarp/Private/Policy/InFrustumFit/DisplayClusterWarpInFrustumFitPolicy.cpp
@@ -94,7 +94,6 @@ void FDisplayClusterWarpInFrustumFitPolicy::HandleNewFrame(const TArray<TSharedP
 
 	// Recalculate warp projection angles
 	GroupGeometryWarpProjection.ResetProjectionAngles();
-	FittedViewports.Empty();
 	SymmetricForwardCorrection.Reset();
 
 	bool bCanCalcFrustumContext = true;
@@ -117,9 +116,9 @@ void FDisplayClusterWarpInFrustumFitPolicy::HandleNewFrame(const TArray<TSharedP
 				continue;
 			}
 
-			if (UDisplayClusterInFrustumFitCameraComponent* WarpViewPointComponent = Cast<UDisplayClusterInFrustumFitCameraComponent>(Viewport->GetViewPointCameraComponent()))
+			if (UDisplayClusterInFrustumFitCameraComponent* ConfigurationCameraComponent = Cast<UDisplayClusterInFrustumFitCameraComponent>(Viewport->GetViewPointCameraComponent(EDisplayClusterRootActorType::Configuration)))
 			{
-				bHasFixedViewDirection = WarpViewPointComponent->CameraViewTarget == EDisplayClusterWarpCameraViewTarget::MatchViewOrigin;
+				bHasFixedViewDirection = ConfigurationCameraComponent->CameraViewTarget == EDisplayClusterWarpCameraViewTarget::MatchViewOrigin;
 			}
 
 			TSharedPtr<IDisplayClusterWarpBlend, ESPMode::ThreadSafe> WarpBlend;
@@ -162,68 +161,109 @@ void FDisplayClusterWarpInFrustumFitPolicy::HandleNewFrame(const TArray<TSharedP
 void FDisplayClusterWarpInFrustumFitPolicy::Tick(IDisplayClusterViewportManager* InViewportManager, float DeltaSeconds)
 {
 #if WITH_EDITOR
-	bool bHasDrawnDebugFrustum = false;
-
-	TArray<TSharedPtr<IDisplayClusterViewport, ESPMode::ThreadSafe>> Viewports = InViewportManager->GetViewportsForWarpPolicy(SharedThis(this));
-	for (const TSharedPtr<IDisplayClusterViewport, ESPMode::ThreadSafe>& Viewport : Viewports)
+	UDisplayClusterInFrustumFitCameraComponent* SceneCameraComponent = nullptr;
+	for (const TSharedPtr<IDisplayClusterViewport, ESPMode::ThreadSafe>& Viewport : InViewportManager->GetEntireClusterViewportsForWarpPolicy(SharedThis(this)))
 	{
-		// If the viewport doesn't have a valid warp projection yet, skip it
-		if (!FittedViewports.Contains(Viewport->GetId()))
-		{
-			continue;
-		}
-
+		// Process only viewports with a projection policy based on the warpblend interface.
 		TSharedPtr<IDisplayClusterWarpBlend, ESPMode::ThreadSafe> WarpBlend;
-		if (Viewport->GetProjectionPolicy()->GetWarpBlendInterface(WarpBlend))
+		if (HasPreviewMovableMesh(Viewport.Get()) && Viewport->GetProjectionPolicy()->GetWarpBlendInterface(WarpBlend))
 		{
 			FDisplayClusterWarpData& WarpData = WarpBlend->GetWarpData(0);
+			UMeshComponent* MovableMeshComponent = Viewport->GetProjectionPolicy()->GetOrCreatePreviewMovableMeshComponent(Viewport.Get());
+
+			// Not all projection policies support a movable mesh.
+			if (MovableMeshComponent && WarpData.bValid && WarpData.bHasWarpPolicyChanges)
+			{
+				const FTransform CameraTransform(WarpData.WarpProjection.CameraRotation.Quaternion(), WarpData.WarpProjection.CameraLocation);
 
-			const FTransform CameraTransform(WarpData.WarpProjection.CameraRotation.Quaternion(), WarpData.WarpProjection.CameraLocation);
+				const float HScale = (WarpData.WarpProjection.Left - WarpData.WarpProjection.Right) / (WarpData.GeometryWarpProjection.Left - WarpData.GeometryWarpProjection.Right);
+				const float VScale = (WarpData.WarpProjection.Top - WarpData.WarpProjection.Bottom) / (WarpData.GeometryWarpProjection.Top - WarpData.GeometryWarpProjection.Bottom);
 
-			const float HScale = (WarpData.WarpProjection.Left - WarpData.WarpProjection.Right) / (WarpData.GeometryWarpProjection.Left - WarpData.GeometryWarpProjection.Right);
-			const float VScale = (WarpData.WarpProjection.Top - WarpData.WarpProjection.Bottom) / (WarpData.GeometryWarpProjection.Top - WarpData.GeometryWarpProjection.Bottom);
+				checkf(FMath::IsNearlyEqual(HScale, VScale), TEXT("Streching the stage geometry to fit a different aspect ratio is not supported!"));
 
-			checkf(FMath::IsNearlyEqual(HScale, VScale), TEXT("Streching the stage geometry to fit a different aspect ratio is not supported!"));
+				const FVector Scale = FVector(1, HScale, HScale);
 
-			const FVector Scale = FVector(1, HScale, HScale);
+				// Compute the relative transform from the view origin to the geometry
+				FTransform RelativeTransform(WarpData.WarpContext.MeshToStageMatrix * WarpData.Local2World.Inverse());
+				RelativeTransform.ScaleTranslation(Scale);
 
-			// Compute the relative transform from the view origin to the geometry
-			FTransform RelativeTransform(WarpData.WarpContext.MeshToStageMatrix * WarpData.Local2World.Inverse());
-			RelativeTransform.ScaleTranslation(Scale);
+				// Final transform is computed from the relative transform of the geometry to the view point, the frustum fit transform
+				// which will scale and position the geometry based on the fitted frustum, and the camera transform
+				const FTransform FinalTransform = RelativeTransform * CameraTransform;
 
-			// Final transform is computed from the relative transform of the geometry to the view point, the frustum fit transform
-			// which will scale and position the geometry based on the fitted frustum, and the camera transform
-			const FTransform FinalTransform = RelativeTransform * CameraTransform;
+				MovableMeshComponent->SetRelativeTransform(FinalTransform);
 
-			UMeshComponent* MovableMeshComponent = Viewport->GetProjectionPolicy()->GetOrCreatePreviewMovableMeshComponent(Viewport.Get());
-			MovableMeshComponent->SetRelativeTransform(FinalTransform);
+				// Since the mesh needs to be skewed to scale appropriately, and since Unreal Engine does not support a skew transform
+				// through FTransform, the mesh needs to be skewed through the vertex shader using WorldPositionOffset,
+				// so pass in the "global" scale to the preview mesh's material instance
+				UMaterialInterface* MovableMeshMaterial = MovableMeshComponent->GetMaterial(0);
+				if (UMaterialInstanceDynamic* MaterialInstance = (MovableMeshMaterial && MovableMeshMaterial->IsA<UMaterialInstanceDynamic>()) ? Cast<UMaterialInstanceDynamic>(MovableMeshMaterial) : nullptr)
+				{
+					FMatrix CameraBasis = FRotationMatrix::Make(WarpData.WarpProjection.CameraRotation).Inverse();
 
-			// Since the mesh needs to be skewed to scale appropriately, and since Unreal Engine does not support a skew transform
-			// through FTransform, the mesh needs to be skewed through the vertex shader using WorldPositionOffset,
-			// so pass in the "global" scale to the preview mesh's material instance
-			if (UMaterialInstanceDynamic* MaterialInstance = Cast<UMaterialInstanceDynamic>(MovableMeshComponent->GetMaterial(0)))
-			{
-				FMatrix CameraBasis = FRotationMatrix::Make(WarpData.WarpProjection.CameraRotation).Inverse();
+					MaterialInstance->SetVectorParameterValue(TEXT("GlobalScale"), Scale);
+					MaterialInstance->SetVectorParameterValue(TEXT("GlobalForward"), CameraBasis.GetUnitAxis(EAxis::X));
+					MaterialInstance->SetVectorParameterValue(TEXT("GlobalRight"), CameraBasis.GetUnitAxis(EAxis::Y));
+					MaterialInstance->SetVectorParameterValue(TEXT("GlobalUp"), CameraBasis.GetUnitAxis(EAxis::Z));
+				}
 
-				MaterialInstance->SetVectorParameterValue(TEXT("GlobalScale"), Scale);
-				MaterialInstance->SetVectorParameterValue(TEXT("GlobalForward"), CameraBasis.GetUnitAxis(EAxis::X));
-				MaterialInstance->SetVectorParameterValue(TEXT("GlobalRight"), CameraBasis.GetUnitAxis(EAxis::Y));
-				MaterialInstance->SetVectorParameterValue(TEXT("GlobalUp"), CameraBasis.GetUnitAxis(EAxis::Z));
+				if (!SceneCameraComponent)
+				{
+					SceneCameraComponent = Cast<UDisplayClusterInFrustumFitCameraComponent>(Viewport->GetViewPointCameraComponent(EDisplayClusterRootActorType::Scene));
+				}
 			}
+		}
+	}
 
-			if (GDisplayClusterWarpInFrustumFitPolicyDrawFrustum && !bHasDrawnDebugFrustum)
+	if (GDisplayClusterWarpInFrustumFitPolicyDrawFrustum)
+	{
+		if (ADisplayClusterRootActor* SceneRootActor = InViewportManager->GetConfiguration().GetRootActor(EDisplayClusterRootActorType::Scene))
+		{
+			if (SceneCameraComponent)
 			{
-				DrawDebugGroupFrustum(InViewportManager->GetRootActor(), Cast<UDisplayClusterInFrustumFitCameraComponent>(Viewport->GetViewPointCameraComponent()), FColor::Blue);
-				bHasDrawnDebugFrustum = true;
+				DrawDebugGroupFrustum(SceneRootActor, SceneCameraComponent, FColor::Blue);
 			}
+
+			DrawDebugGroupBoundingBox(SceneRootActor, FColor::Red);
 		}
 	}
+#endif
+}
 
-	if (GDisplayClusterWarpInFrustumFitPolicyDrawFrustum)
+bool FDisplayClusterWarpInFrustumFitPolicy::HasPreviewMovableMesh(IDisplayClusterViewport* InViewport)
+{
+#if WITH_EDITOR
+	// This warp policy is based on IDisplayClusterWarpBlend only.
+	// Process only viewports with a projection policy based on the warpblend interface.
+	TSharedPtr<IDisplayClusterWarpBlend, ESPMode::ThreadSafe> WarpBlend;
+	if (!InViewport || !InViewport->GetProjectionPolicy().IsValid() || !InViewport->GetProjectionPolicy()->GetWarpBlendInterface(WarpBlend))
+	{
+		return false;
+	}
+
+	// If the preview is not used in this configuration
+	if (!InViewport->GetConfiguration().IsPreviewRendering() || InViewport->GetConfiguration().GetRootActor(EDisplayClusterRootActorType::Preview) == nullptr)
 	{
-		DrawDebugGroupBoundingBox(InViewportManager->GetRootActor(), FColor::Red);
+		return false;
+	}
+
+	// If owner DCRA world is EditorPreview dont show movable mesh (Configurator, ICVFX Panel, etc)
+	if (InViewport->GetConfiguration().IsRootActorWorldHasAnyType(EDisplayClusterRootActorType::Preview, EWorldType::EditorPreview))
+	{
+		return false;
+	}
+
+	// The movable mesh is an option for the UDisplayClusterInFrustumFitCameraComponent.
+	if (UDisplayClusterInFrustumFitCameraComponent* ConfigurationCameraComponent = Cast<UDisplayClusterInFrustumFitCameraComponent>(InViewport->GetViewPointCameraComponent(EDisplayClusterRootActorType::Configuration)))
+	{
+		if (ConfigurationCameraComponent->bShowPreviewFrustumFit)
+		{
+			return true;
+		}
 	}
 #endif
+
+	return false;
 }
 
 void FDisplayClusterWarpInFrustumFitPolicy::BeginCalcFrustum(IDisplayClusterViewport* InViewport, const uint32 ContextNum)
@@ -233,7 +273,7 @@ void FDisplayClusterWarpInFrustumFitPolicy::BeginCalcFrustum(IDisplayClusterView
 		TSharedPtr<IDisplayClusterWarpBlend, ESPMode::ThreadSafe> WarpBlend;
 		if (InViewport->GetProjectionPolicy()->GetWarpBlendInterface(WarpBlend))
 		{
-			if (UDisplayClusterInFrustumFitCameraComponent* WarpViewPointComponent = Cast<UDisplayClusterInFrustumFitCameraComponent>(InViewport->GetViewPointCameraComponent()))
+			if (UDisplayClusterInFrustumFitCameraComponent* ConfigurationCameraComponent = Cast<UDisplayClusterInFrustumFitCameraComponent>(InViewport->GetViewPointCameraComponent(EDisplayClusterRootActorType::Configuration)))
 			{
 				FDisplayClusterWarpData& WarpData = WarpBlend->GetWarpData(ContextNum);
 				if (WarpData.WarpEye.IsValid())
@@ -241,7 +281,7 @@ void FDisplayClusterWarpInFrustumFitPolicy::BeginCalcFrustum(IDisplayClusterView
 					// geometry context already updated.
 					WarpData.WarpEye->bUpdateGeometryContext = !bGeometryContextsUpdated;
 
-					if (WarpViewPointComponent->CameraViewTarget == EDisplayClusterWarpCameraViewTarget::MatchViewOrigin)
+					if (ConfigurationCameraComponent->CameraViewTarget == EDisplayClusterWarpCameraViewTarget::MatchViewOrigin)
 					{
 						WarpData.WarpEye->OverrideViewDirection = WarpData.WarpEye->ViewPoint.Rotation.RotateVector(FVector::XAxisVector);
 					}
@@ -287,23 +327,18 @@ void FDisplayClusterWarpInFrustumFitPolicy::EndCalcFrustum(IDisplayClusterViewpo
 		TSharedPtr<IDisplayClusterWarpBlend, ESPMode::ThreadSafe> WarpBlend;
 		if (InViewport->GetProjectionPolicy()->GetWarpBlendInterface(WarpBlend))
 		{
-			if (UDisplayClusterCameraComponent* ViewPointComponent = InViewport->GetViewPointCameraComponent())
-			{
-				if (ViewPointComponent->IsA<UDisplayClusterInFrustumFitCameraComponent>())
-				{
-					if (UDisplayClusterInFrustumFitCameraComponent* WarpViewPointComponent = CastChecked<UDisplayClusterInFrustumFitCameraComponent>(ViewPointComponent))
-					{
-						// Change warp settings:
-						FDisplayClusterWarpData& WarpData = WarpBlend->GetWarpData(ContextNum);
-
-						// Apply camera frustum fitting:
-						FDisplayClusterWarpProjection NewWarpProjection = ApplyInFrustumFit(WarpViewPointComponent, WarpData.WarpEye->World2LocalTransform, WarpData.WarpProjection);
+			// Change warp settings:
+			FDisplayClusterWarpData& WarpData = WarpBlend->GetWarpData(ContextNum);
 
-						WarpData.WarpProjection = NewWarpProjection;
+			// Apply camera frustum fitting:
+			FDisplayClusterWarpProjection NewWarpProjection = ApplyInFrustumFit(InViewport, WarpData.WarpEye->World2LocalTransform, WarpData.WarpProjection);
+			if (NewWarpProjection.IsValidProjection())
+			{
+				WarpData.WarpProjection = NewWarpProjection;
 
-						FittedViewports.Add(InViewport->GetId());
-					}
-				}
+				// The warp policy Tick() function uses warp data, and it must be sure that this data is updated in the previous frame.
+				//.This value must be set to true from the EndCalcFrustum() warp policy function when changes are made to this structure.
+				WarpData.bHasWarpPolicyChanges = true;
 			}
 		}
 	}
@@ -327,15 +362,20 @@ void FDisplayClusterWarpInFrustumFitPolicy::DrawDebugGroupFrustum(ADisplayCluste
 {
 	if (RootActor && CameraComponent)
 	{
-		if (UWorld* World = RootActor->GetWorld())
+		UWorld* World = RootActor->GetWorld();
+		IDisplayClusterViewportConfiguration* ViewportConfiguration = RootActor->GetViewportConfiguration();
+		if (ViewportConfiguration && World)
 		{
+			// Get the configuration in use
+			const UDisplayClusterInFrustumFitCameraComponent& ConfigurationCameraComponent = CameraComponent->GetConfigurationInFrustumFitCameraComponent(*ViewportConfiguration);
+
 			const float NearPlane = 10;
 			const float FarPlane = 1000;
 
 			const FVector CameraLoc = CameraComponent->GetComponentLocation();
 			FVector ViewDirection;
 
-			if (CameraComponent->CameraViewTarget == EDisplayClusterWarpCameraViewTarget::MatchViewOrigin)
+			if (ConfigurationCameraComponent.CameraViewTarget == EDisplayClusterWarpCameraViewTarget::MatchViewOrigin)
 			{
 				ViewDirection = CameraComponent->GetComponentRotation().RotateVector(FVector::XAxisVector);
 			}
diff --git a/Engine/Plugins/Runtime/nDisplay/Source/DisplayClusterWarp/Private/Policy/InFrustumFit/DisplayClusterWarpInFrustumFitPolicy.h b/Engine/Plugins/Runtime/nDisplay/Source/DisplayClusterWarp/Private/Policy/InFrustumFit/DisplayClusterWarpInFrustumFitPolicy.h
--- a/Engine/Plugins/Runtime/nDisplay/Source/DisplayClusterWarp/Private/Policy/InFrustumFit/DisplayClusterWarpInFrustumFitPolicy.h
+++ b/Engine/Plugins/Runtime/nDisplay/Source/DisplayClusterWarp/Private/Policy/InFrustumFit/DisplayClusterWarpInFrustumFitPolicy.h
@@ -3,6 +3,7 @@
 #pragma once
 
 #include "Policy/DisplayClusterWarpPolicyBase.h"
+#include "Render/Viewport/IDisplayClusterViewportConfiguration.h"
 #include "Blueprints/DisplayClusterWarpBlueprint_Enums.h"
 
 #include "Containers/DisplayClusterWarpContext.h"
@@ -12,7 +13,7 @@ class UDisplayClusterInFrustumFitCameraComponent;
 class FDisplayClusterWarpEye;
 
 /**
- * FurstumFit warp policy
+ * InFrustumFit warp policy
  */
 class FDisplayClusterWarpInFrustumFitPolicy
 	: public FDisplayClusterWarpPolicyBase
@@ -28,14 +29,19 @@ public:
 
 	virtual void BeginCalcFrustum(IDisplayClusterViewport* InViewport, const uint32 ContextNum) override;
 	virtual void EndCalcFrustum(IDisplayClusterViewport* InViewport, const uint32 ContextNum) override;
+
+	virtual bool HasPreviewMovableMesh(IDisplayClusterViewport* InViewport) override;
 	//~~ End IDisplayClusterWarpPolicy
 
 private:
+	/** Return camera component that used for configuration.. */
+	const UDisplayClusterInFrustumFitCameraComponent* GetConfigurationInFrustumFitCameraComponent(IDisplayClusterViewport* InViewport) const;
+
 	/** Converts an asymmetric group frustum to a symmetric group frustum and computes the correction rotator to the view forward vector needed for a symmetric frustum  */
 	void MakeGroupFrustumSymmetrical(bool bFixedViewTarget);
 
 	/** Apply frustum fit to the specified warp projection */
-	FDisplayClusterWarpProjection ApplyInFrustumFit(UDisplayClusterInFrustumFitCameraComponent* InComponent, const FTransform& World2OriginTransform, const FDisplayClusterWarpProjection& InWarpProjection);
+	FDisplayClusterWarpProjection ApplyInFrustumFit(IDisplayClusterViewport* InViewport, const FTransform& World2OriginTransform, const FDisplayClusterWarpProjection& InWarpProjection);
 
 	/** Find final projection scale.*/
 	FVector2D FindFrustumFit(const EDisplayClusterWarpCameraProjectionMode InProjectionMode, const FVector2D& InCameraFOV, const FVector2D& InGeometryFOV);
@@ -52,12 +58,6 @@ private:
 	// Warp projection data
 	FDisplayClusterWarpProjection GroupGeometryWarpProjection;
 
-	/**
-	 * Due to the piecemeal nature of the viewport rendering, not all viewports will have a valid warp projection every tick. This list
-	 * stores the viewports that have been fitted since the last frame was started.
-	 */
-	TSet<FString> FittedViewports;
-
 	/** The AABB of the group computed in DCRA space */
 	FDisplayClusterWarpAABB GroupAABBox;
 
diff --git a/Engine/Plugins/Runtime/nDisplay/Source/DisplayClusterWarp/Private/Policy/InFrustumFit/DisplayClusterWarpInFrustumFitPolicy_Math.cpp b/Engine/Plugins/Runtime/nDisplay/Source/DisplayClusterWarp/Private/Policy/InFrustumFit/DisplayClusterWarpInFrustumFitPolicy_Math.cpp
--- a/Engine/Plugins/Runtime/nDisplay/Source/DisplayClusterWarp/Private/Policy/InFrustumFit/DisplayClusterWarpInFrustumFitPolicy_Math.cpp
+++ b/Engine/Plugins/Runtime/nDisplay/Source/DisplayClusterWarp/Private/Policy/InFrustumFit/DisplayClusterWarpInFrustumFitPolicy_Math.cpp
@@ -2,6 +2,7 @@
 
 #include "Policy/InFrustumFit/DisplayClusterWarpInFrustumFitPolicy.h"
 
+#include "Render/Viewport/IDisplayClusterViewport.h"
 #include "Components/DisplayClusterInFrustumFitCameraComponent.h"
 #include "Camera/CameraTypes.h"
 #include "HAL/IConsoleManager.h"
@@ -71,14 +72,23 @@ void FDisplayClusterWarpInFrustumFitPolicy::MakeGroupFrustumSymmetrical(bool bFi
 	}
 }
 
-FDisplayClusterWarpProjection FDisplayClusterWarpInFrustumFitPolicy::ApplyInFrustumFit(UDisplayClusterInFrustumFitCameraComponent* InComponent, const FTransform& World2OriginTransform, const FDisplayClusterWarpProjection& InWarpProjection)
+FDisplayClusterWarpProjection FDisplayClusterWarpInFrustumFitPolicy::ApplyInFrustumFit(IDisplayClusterViewport* InViewport, const FTransform& World2OriginTransform, const FDisplayClusterWarpProjection& InWarpProjection)
 {
-	check(InComponent);
+	UDisplayClusterInFrustumFitCameraComponent* SceneCameraComponent = Cast<UDisplayClusterInFrustumFitCameraComponent>(InViewport->GetViewPointCameraComponent(EDisplayClusterRootActorType::Scene));
+	if (!SceneCameraComponent)
+	{
+		// By default, this structure has invalid values.
+		return FDisplayClusterWarpProjection();
+	}
+
+	// Get the configuration in use
+	const UDisplayClusterInFrustumFitCameraComponent& ConfigurationCameraComponent = SceneCameraComponent->GetConfigurationInFrustumFitCameraComponent(InViewport->GetConfiguration());
 
 	FDisplayClusterWarpProjection OutWarpProjection(InWarpProjection);
 
+	// Get location from scene component
 	FMinimalViewInfo CameraViewInfo;
-	InComponent->GetDesiredView(CameraViewInfo);
+	SceneCameraComponent->GetDesiredView(InViewport->GetConfiguration(), CameraViewInfo);
 
 	// Use camera position to render:
 	if (GDisplayClusterWarpInFrustumFitPolicyEnableCameraPositionFit)
@@ -104,7 +114,8 @@ FDisplayClusterWarpProjection FDisplayClusterWarpInFrustumFitPolicy::ApplyInFrus
 	const FVector2D CameraHalfFOV(CameraHalfFOVProjection, CameraHalfFOVProjection / CameraViewInfo.AspectRatio);
 	const FVector2D GeometryHalfFOV = GeometryFOV * 0.5;
 
-	const FVector2D FinalHalfFOV = FindFrustumFit(InComponent->CameraProjectionMode, CameraHalfFOV, GeometryHalfFOV);
+	// Receive configuration from InCfgComponent
+	const FVector2D FinalHalfFOV = FindFrustumFit(ConfigurationCameraComponent.CameraProjectionMode, CameraHalfFOV, GeometryHalfFOV);
 
 	// Sample code: convert back to projection angles
 	if (GDisplayClusterWarpInFrustumFitPolicyEnableProjectionAnglesFit)
diff --git a/Engine/Plugins/Runtime/nDisplay/Source/DisplayClusterWarp/Private/WarpBlend/DisplayClusterWarpBlend.cpp b/Engine/Plugins/Runtime/nDisplay/Source/DisplayClusterWarp/Private/WarpBlend/DisplayClusterWarpBlend.cpp
--- a/Engine/Plugins/Runtime/nDisplay/Source/DisplayClusterWarp/Private/WarpBlend/DisplayClusterWarpBlend.cpp
+++ b/Engine/Plugins/Runtime/nDisplay/Source/DisplayClusterWarp/Private/WarpBlend/DisplayClusterWarpBlend.cpp
@@ -75,16 +75,21 @@ const FDisplayClusterWarpGeometryContext& FDisplayClusterWarpBlend::GetGeometryC
 
 bool FDisplayClusterWarpBlend::CalcFrustumContext(const TSharedPtr<FDisplayClusterWarpEye, ESPMode::ThreadSafe>& InWarpEye)
 {
-	if(!InWarpEye.IsValid())
-	{
-		return false;
-	}
+	check(WarpData.IsValidIndex(InWarpEye->ContextNum));
 
 	// Update current warp data
-	check(WarpData.IsValidIndex(InWarpEye->ContextNum));
 	FDisplayClusterWarpData& CurrentWarpData = WarpData[InWarpEye->ContextNum];
+
 	CurrentWarpData.WarpEye = InWarpEye;
 
+	CurrentWarpData.bValid = false;
+	CurrentWarpData.bHasWarpPolicyChanges = false;
+
+	if(!InWarpEye.IsValid())
+	{
+		return false;
+	}
+
 	FDisplayClusterWarpBlendMath_Frustum Frustum(CurrentWarpData, GeometryContext);
 
 	// Override settings from warp policy
@@ -104,6 +109,9 @@ bool FDisplayClusterWarpBlend::CalcFrustumContext(const TSharedPtr<FDisplayClust
 		return false;
 	}
 
+	// Now this data becomes valid.
+	CurrentWarpData.bValid = true;
+
 	return true;
 }
 
@@ -136,7 +144,18 @@ const FDisplayClusterWarpData& FDisplayClusterWarpBlend::GetWarpData(const uint3
 
 }
 
-UMeshComponent* FDisplayClusterWarpBlend::GetOrCreateMeshComponent(IDisplayClusterViewport* InViewport, bool& bExistingComponent) const
+UMeshComponent* FDisplayClusterWarpBlend::GetOrCreatePreviewMeshComponent(IDisplayClusterViewport* InViewport, bool& bExistingComponent) const
+{
+	return GetOrCreatePreviewMeshComponentImpl(InViewport, false, bExistingComponent);
+}
+
+UMeshComponent* FDisplayClusterWarpBlend::GetOrCreatePreviewMovableMeshComponent(IDisplayClusterViewport* InViewport) const
+{
+	bool bExistingComponentDummy;
+	return GetOrCreatePreviewMeshComponentImpl(InViewport, true, bExistingComponentDummy);
+}
+
+UMeshComponent* FDisplayClusterWarpBlend::GetOrCreatePreviewMeshComponentImpl(IDisplayClusterViewport* InViewport, bool bMovableMesh, bool& bExistingComponent) const
 {
 	const TSharedPtr<IDisplayClusterProjectionPolicy, ESPMode::ThreadSafe> ProjectionPolicy = InViewport ? InViewport->GetProjectionPolicy() : nullptr;
 
@@ -145,9 +164,46 @@ UMeshComponent* FDisplayClusterWarpBlend::GetOrCreateMeshComponent(IDisplayClust
 	case EDisplayClusterWarpGeometryType::WarpMesh:
 	case EDisplayClusterWarpGeometryType::WarpProceduralMesh:
 		// use the existing DCRA component
-		bExistingComponent = true;
+		if (USceneComponent* PreviewMeshComponent = GeometryContext.GeometryProxy.PreviewMeshComponentRef.GetOrFindSceneComponent())
+		{
+			if (PreviewMeshComponent->IsA<UMeshComponent>())
+			{
+				UMeshComponent* ExistMeshComponent = static_cast<UMeshComponent*>(PreviewMeshComponent);
+
+				if (bMovableMesh)
+				{
+					// create a mesh component copy
+					bExistingComponent = false;
+
+					// Get movable mesh root
+					USceneComponent* SceneOriginComp = ProjectionPolicy->GetPreviewMovableMeshOriginComponent(InViewport);
+
+					const FString CompName = FString::Printf(TEXT("DCWarpBlend_MovableMesh_%s"), *ProjectionPolicy->GetId());
+					const EObjectFlags ObjectFlags = EObjectFlags::RF_DuplicateTransient | RF_Transient | RF_TextExportTransient;
+					if (UMeshComponent* PreviewMeshComponentCopy = NewObject<UMeshComponent>(SceneOriginComp, ExistMeshComponent->GetClass(), *CompName, ObjectFlags, ExistMeshComponent))
+					{
+						PreviewMeshComponentCopy->RegisterComponent();
+						PreviewMeshComponentCopy->AttachToComponent(SceneOriginComp, FAttachmentTransformRules(EAttachmentRule::KeepRelative, false));
+						PreviewMeshComponentCopy->SetHiddenInGame(true);
+
+#if WITH_EDITOR
+						PreviewMeshComponentCopy->SetIsVisualizationComponent(true);
+#endif
+
+						return PreviewMeshComponentCopy;
+					}
+				}
+				else
+				{
+					bExistingComponent = true;
+
+					return ExistMeshComponent;
+				}
+				
+			}
+		}
 
-		return GeometryContext.GeometryProxy.MeshComponent.IsValid() ? GeometryContext.GeometryProxy.MeshComponent->GetMeshComponent() : nullptr;
+		return nullptr;
 
 	case EDisplayClusterWarpGeometryType::WarpMap:
 		if (ProjectionPolicy.IsValid())
@@ -158,20 +214,20 @@ UMeshComponent* FDisplayClusterWarpBlend::GetOrCreateMeshComponent(IDisplayClust
 			// Downscale preview mesh dimension to max limit
 			const uint32 PreviewGeometryDimLimit = 128;
 
-			USceneComponent* OriginComp = ProjectionPolicy->GetOriginComponent();
+			USceneComponent* PreviewOriginComp = bMovableMesh ? ProjectionPolicy->GetPreviewMovableMeshOriginComponent(InViewport) : ProjectionPolicy->GetPreviewMeshOriginComponent(InViewport);
 
 			// Create new WarpMesh component
 			FDisplayClusterWarpGeometryOBJ MeshData;
-			if (OriginComp && ExportWarpMapGeometry(MeshData, PreviewGeometryDimLimit))
+			if (PreviewOriginComp && ExportWarpMapGeometry(MeshData, PreviewGeometryDimLimit))
 			{
-				const FString CompName = FString::Printf(TEXT("ExportPFM_%s"), *ProjectionPolicy->GetId());
+				const FString CompName = FString::Printf(TEXT("DCWarpBlend_PFM_%s_mesh"), *ProjectionPolicy->GetId());
 
 				// Creta new object
-				UProceduralMeshComponent* MeshComp = NewObject<UProceduralMeshComponent>(OriginComp, FName(*CompName), EObjectFlags::RF_DuplicateTransient | RF_Transient | RF_TextExportTransient);
+				UProceduralMeshComponent* MeshComp = NewObject<UProceduralMeshComponent>(PreviewOriginComp, FName(*CompName), EObjectFlags::RF_DuplicateTransient | RF_Transient | RF_TextExportTransient);
 				if (MeshComp)
 				{
 					MeshComp->RegisterComponent();
-					MeshComp->AttachToComponent(OriginComp, FAttachmentTransformRules(EAttachmentRule::KeepRelative, false));
+					MeshComp->AttachToComponent(PreviewOriginComp, FAttachmentTransformRules(EAttachmentRule::KeepRelative, false));
 					MeshComp->CreateMeshSection(0, MeshData.Vertices, MeshData.Triangles, MeshData.Normal, MeshData.UV, TArray<FColor>(), TArray<FProcMeshTangent>(), false);
 #if WITH_EDITOR
 					MeshComp->SetIsVisualizationComponent(true);
diff --git a/Engine/Plugins/Runtime/nDisplay/Source/DisplayClusterWarp/Private/WarpBlend/DisplayClusterWarpBlend.h b/Engine/Plugins/Runtime/nDisplay/Source/DisplayClusterWarp/Private/WarpBlend/DisplayClusterWarpBlend.h
--- a/Engine/Plugins/Runtime/nDisplay/Source/DisplayClusterWarp/Private/WarpBlend/DisplayClusterWarpBlend.h
+++ b/Engine/Plugins/Runtime/nDisplay/Source/DisplayClusterWarp/Private/WarpBlend/DisplayClusterWarpBlend.h
@@ -82,11 +82,22 @@ public:
 	virtual FDisplayClusterWarpData& GetWarpData(const uint32 ContextNum) override;
 	virtual const FDisplayClusterWarpData& GetWarpData(const uint32 ContextNum) const override;
 
-	virtual UMeshComponent* GetOrCreateMeshComponent(IDisplayClusterViewport* InViewport, bool& bExistingComponent) const override;
+	virtual UMeshComponent* GetOrCreatePreviewMeshComponent(IDisplayClusterViewport* InViewport, bool& bExistingComponent) const override;
+	virtual UMeshComponent* GetOrCreatePreviewMovableMeshComponent(IDisplayClusterViewport* InViewport) const override;
 
 	//~!IDisplayClusterWarpBlend
 
 private:
+	/**
+	* Get or create a mesh component for previewing
+	* 
+	* @param InViewport - The viewport that do this request 
+	* @param bCreateCopy - true to always create a copy of the mesh
+	* @param bExistingComponent - (out) true if the mesh component is not a copy
+	*/
+	UMeshComponent* GetOrCreatePreviewMeshComponentImpl(IDisplayClusterViewport* InViewport, bool bMovableMesh, bool& bExistingComponent) const;
+
+	/** Begin frustum calc. */
 	void BeginCalcFrustum(const TSharedPtr<FDisplayClusterWarpEye, ESPMode::ThreadSafe>& InWarpEye);
 
 public:
diff --git a/Engine/Plugins/Runtime/nDisplay/Source/DisplayClusterWarp/Private/WarpBlend/DisplayClusterWarpBlend_GeometryProxy.cpp b/Engine/Plugins/Runtime/nDisplay/Source/DisplayClusterWarp/Private/WarpBlend/DisplayClusterWarpBlend_GeometryProxy.cpp
--- a/Engine/Plugins/Runtime/nDisplay/Source/DisplayClusterWarp/Private/WarpBlend/DisplayClusterWarpBlend_GeometryProxy.cpp
+++ b/Engine/Plugins/Runtime/nDisplay/Source/DisplayClusterWarp/Private/WarpBlend/DisplayClusterWarpBlend_GeometryProxy.cpp
@@ -45,7 +45,7 @@ const IDisplayClusterRender_MeshComponentProxy* FDisplayClusterWarpBlend_Geometr
 	{
 	case EDisplayClusterWarpGeometryType::WarpMesh:
 	case EDisplayClusterWarpGeometryType::WarpProceduralMesh:
-		return MeshComponent.IsValid() ? MeshComponent->GetMeshComponentProxy_RenderThread() : nullptr;
+		return WarpMeshComponent.IsValid() ? WarpMeshComponent->GetMeshComponentProxy_RenderThread() : nullptr;
 
 	default:
 		break;
@@ -60,11 +60,11 @@ bool FDisplayClusterWarpBlend_GeometryProxy::MarkWarpFrustumGeometryComponentDir
 	{
 	case EDisplayClusterWarpFrustumGeometryType::WarpMesh:
 	case EDisplayClusterWarpFrustumGeometryType::WarpProceduralMesh:
-		if (MeshComponent.IsValid())
+		if (WarpMeshComponent.IsValid())
 		{
-			if (InComponentName == NAME_None || MeshComponent->EqualsMeshComponentName(InComponentName))
+			if (InComponentName == NAME_None || WarpMeshComponent->EqualsMeshComponentName(InComponentName))
 			{
-				MeshComponent->MarkMeshComponentRefGeometryDirty();
+				WarpMeshComponent->MarkMeshComponentRefGeometryDirty();
 				return true;
 			}
 		}
@@ -113,27 +113,27 @@ bool FDisplayClusterWarpBlend_GeometryProxy::ImplUpdateFrustumGeometry_WarpMesh(
 {
 	bIsGeometryValid = false;
 
-	if (!MeshComponent.IsValid())
+	if (!WarpMeshComponent.IsValid())
 	{
 		return false;
 	}
 
-	UStaticMeshComponent* StaticMeshComponent = MeshComponent->GetStaticMeshComponent();
-	USceneComponent*      OriginComponent     = MeshComponent->GetOriginComponent();
+	UStaticMeshComponent* StaticMeshComponent = WarpMeshComponent->GetStaticMeshComponent();
+	USceneComponent*      OriginComponent     = WarpMeshComponent->GetOriginComponent();
 
-	const FStaticMeshLODResources* StaticMeshLODResources = (StaticMeshComponent!=nullptr) ? MeshComponent->GetStaticMeshComponentLODResources(StaticMeshComponentLODIndex) : nullptr;
+	const FStaticMeshLODResources* StaticMeshLODResources = (StaticMeshComponent!=nullptr) ? WarpMeshComponent->GetStaticMeshComponentLODResources(StaticMeshComponentLODIndex) : nullptr;
 	if (StaticMeshLODResources == nullptr)
 	{
 		// mesh deleted?
-		MeshComponent->ReleaseProxyGeometry();
+		WarpMeshComponent->ReleaseProxyGeometry();
 		bIsMeshComponentLost = true;
 		return false;
 	};
 
 	// If StaticMesh geometry changed, update mpcdi math and RHI resources
-	if (MeshComponent->IsMeshComponentRefGeometryDirty() || bIsMeshComponentLost)
+	if (WarpMeshComponent->IsMeshComponentRefGeometryDirty() || bIsMeshComponentLost)
 	{
-		MeshComponent->AssignStaticMeshComponentRefs(StaticMeshComponent, WarpMeshUVs, OriginComponent, StaticMeshComponentLODIndex);
+		WarpMeshComponent->AssignStaticMeshComponentRefs(StaticMeshComponent, WarpMeshUVs, OriginComponent, StaticMeshComponentLODIndex);
 		bIsMeshComponentLost = false;
 	}
 	
@@ -162,27 +162,27 @@ bool FDisplayClusterWarpBlend_GeometryProxy::ImplUpdateFrustumGeometry_WarpProce
 {
 	bIsGeometryValid = false;
 
-	if (!MeshComponent.IsValid())
+	if (!WarpMeshComponent.IsValid())
 	{
 		return false;
 	}
 
-	UProceduralMeshComponent* ProceduralMeshComponent = MeshComponent->GetProceduralMeshComponent();
-	USceneComponent*          OriginComponent         = MeshComponent->GetOriginComponent();
+	UProceduralMeshComponent* ProceduralMeshComponent = WarpMeshComponent->GetProceduralMeshComponent();
+	USceneComponent*          OriginComponent         = WarpMeshComponent->GetOriginComponent();
 
-	const FProcMeshSection* ProcMeshSection = (ProceduralMeshComponent != nullptr) ? MeshComponent->GetProceduralMeshComponentSection(ProceduralMeshComponentSectionIndex) : nullptr;
+	const FProcMeshSection* ProcMeshSection = (ProceduralMeshComponent != nullptr) ? WarpMeshComponent->GetProceduralMeshComponentSection(ProceduralMeshComponentSectionIndex) : nullptr;
 	if (ProcMeshSection == nullptr)
 	{
 		// mesh deleted, lost or section not defined
-		MeshComponent->ReleaseProxyGeometry();
+		WarpMeshComponent->ReleaseProxyGeometry();
 		bIsMeshComponentLost = true;
 		return false;
 	};
 
 	// If ProceduralMesh geometry changed, update mpcdi math and RHI resources
-	if (MeshComponent->IsMeshComponentRefGeometryDirty() || bIsMeshComponentLost)
+	if (WarpMeshComponent->IsMeshComponentRefGeometryDirty() || bIsMeshComponentLost)
 	{
-		MeshComponent->AssignProceduralMeshComponentRefs(ProceduralMeshComponent, WarpMeshUVs, OriginComponent, ProceduralMeshComponentSectionIndex);
+		WarpMeshComponent->AssignProceduralMeshComponentRefs(ProceduralMeshComponent, WarpMeshUVs, OriginComponent, ProceduralMeshComponentSectionIndex);
 		bIsMeshComponentLost = false;
 	}
 
@@ -226,9 +226,9 @@ bool FDisplayClusterWarpBlend_GeometryProxy::ImplUpdateFrustumGeometryCache_Warp
 {
 	bIsGeometryValid = false;
 
-	if (MeshComponent.IsValid())
+	if (WarpMeshComponent.IsValid())
 	{
-		const FStaticMeshLODResources* StaticMeshLODResources = MeshComponent->GetStaticMeshComponentLODResources(StaticMeshComponentLODIndex);
+		const FStaticMeshLODResources* StaticMeshLODResources = WarpMeshComponent->GetStaticMeshComponentLODResources(StaticMeshComponentLODIndex);
 		if (StaticMeshLODResources != nullptr)
 		{
 			FDisplayClusterWarpBlendMath_WarpMesh MeshHelper(*StaticMeshLODResources);
@@ -248,9 +248,9 @@ bool FDisplayClusterWarpBlend_GeometryProxy::ImplUpdateFrustumGeometryCache_Warp
 {
 	bIsGeometryValid = false;
 
-	if (MeshComponent.IsValid())
+	if (WarpMeshComponent.IsValid())
 	{
-		const FProcMeshSection* ProcMeshSection = MeshComponent->GetProceduralMeshComponentSection(ProceduralMeshComponentSectionIndex);
+		const FProcMeshSection* ProcMeshSection = WarpMeshComponent->GetProceduralMeshComponentSection(ProceduralMeshComponentSectionIndex);
 		if (ProcMeshSection != nullptr)
 		{
 			FDisplayClusterWarpBlendMath_WarpProceduralMesh ProceduralMeshHelper(*ProcMeshSection);
@@ -351,12 +351,12 @@ bool FDisplayClusterWarpBlend_GeometryProxy::UpdateFrustumGeometryLOD(const FInt
 
 const FStaticMeshLODResources* FDisplayClusterWarpBlend_GeometryProxy::GetStaticMeshComponentLODResources() const
 {
-	return MeshComponent.IsValid() ? MeshComponent->GetStaticMeshComponentLODResources(StaticMeshComponentLODIndex) : nullptr;
+	return WarpMeshComponent.IsValid() ? WarpMeshComponent->GetStaticMeshComponentLODResources(StaticMeshComponentLODIndex) : nullptr;
 }
 
 const FProcMeshSection* FDisplayClusterWarpBlend_GeometryProxy::GetProceduralMeshComponentSection() const
 {
-	return MeshComponent.IsValid() ? MeshComponent->GetProceduralMeshComponentSection(ProceduralMeshComponentSectionIndex) : nullptr;
+	return WarpMeshComponent.IsValid() ? WarpMeshComponent->GetProceduralMeshComponentSection(ProceduralMeshComponentSectionIndex) : nullptr;
 }
 
 void FDisplayClusterWarpBlend_GeometryProxy::ReleaseResources()
@@ -364,5 +364,6 @@ void FDisplayClusterWarpBlend_GeometryProxy::ReleaseResources()
 	WarpMapTexture.Reset();
 	AlphaMapTexture.Reset();
 	BetaMapTexture.Reset();
-	MeshComponent.Reset();
+	WarpMeshComponent.Reset();
+	PreviewMeshComponentRef.ResetSceneComponent();
 }
diff --git a/Engine/Plugins/Runtime/nDisplay/Source/DisplayClusterWarp/Private/WarpBlend/DisplayClusterWarpBlend_GeometryProxy.h b/Engine/Plugins/Runtime/nDisplay/Source/DisplayClusterWarp/Private/WarpBlend/DisplayClusterWarpBlend_GeometryProxy.h
--- a/Engine/Plugins/Runtime/nDisplay/Source/DisplayClusterWarp/Private/WarpBlend/DisplayClusterWarpBlend_GeometryProxy.h
+++ b/Engine/Plugins/Runtime/nDisplay/Source/DisplayClusterWarp/Private/WarpBlend/DisplayClusterWarpBlend_GeometryProxy.h
@@ -4,6 +4,7 @@
 #include "Containers/DisplayClusterWarpContainers.h"
 #include "Render/Containers/IDisplayClusterRender_MeshComponent.h"
 #include "Render/Containers/IDisplayClusterRender_Texture.h"
+#include "Misc/DisplayClusterObjectRef.h"
 
 class IDisplayClusterRender_MeshComponentProxy;
 struct FStaticMeshLODResources;
@@ -81,7 +82,8 @@ public:
 	TSharedPtr<IDisplayClusterRender_Texture, ESPMode::ThreadSafe> AlphaMapTexture;
 	TSharedPtr<IDisplayClusterRender_Texture, ESPMode::ThreadSafe> BetaMapTexture;
 
-	TSharedPtr<IDisplayClusterRender_MeshComponent, ESPMode::ThreadSafe> MeshComponent;
+	TSharedPtr<IDisplayClusterRender_MeshComponent, ESPMode::ThreadSafe> WarpMeshComponent;
+	FDisplayClusterSceneComponentRef PreviewMeshComponentRef;
 
 	FDisplayClusterMeshUVs WarpMeshUVs;
 
diff --git a/Engine/Plugins/Runtime/nDisplay/Source/DisplayClusterWarp/Private/WarpBlend/Loader/DisplayClusterWarpBlendLoader_MPCDI.cpp b/Engine/Plugins/Runtime/nDisplay/Source/DisplayClusterWarp/Private/WarpBlend/Loader/DisplayClusterWarpBlendLoader_MPCDI.cpp
--- a/Engine/Plugins/Runtime/nDisplay/Source/DisplayClusterWarp/Private/WarpBlend/Loader/DisplayClusterWarpBlendLoader_MPCDI.cpp
+++ b/Engine/Plugins/Runtime/nDisplay/Source/DisplayClusterWarp/Private/WarpBlend/Loader/DisplayClusterWarpBlendLoader_MPCDI.cpp
@@ -82,8 +82,10 @@ TSharedPtr<FDisplayClusterWarpBlend, ESPMode::ThreadSafe> FDisplayClusterWarpBle
 
 		FDisplayClusterMeshUVs MeshUVs;
 
-		Proxy.MeshComponent = FDisplayClusterWarpBlendLoader_MeshComponent::CreateMeshComponent();
-		Proxy.MeshComponent->AssignStaticMeshComponentRefs(InConstructParameters.StaticMeshComponent, MeshUVs, InConstructParameters.OriginComponent, 0);
+		Proxy.PreviewMeshComponentRef.SetSceneComponent(InConstructParameters.PreviewMeshComponent);
+
+		Proxy.WarpMeshComponent = FDisplayClusterWarpBlendLoader_MeshComponent::CreateMeshComponent();
+		Proxy.WarpMeshComponent->AssignStaticMeshComponentRefs(InConstructParameters.WarpMeshComponent, MeshUVs, InConstructParameters.OriginComponent, 0);
 
 		Proxy.WarpMeshUVs = MeshUVs;
 
