nDisplay: Added new interface IDisplayClusterViewportConfiguration

diff --git a/Engine/Plugins/Runtime/nDisplay/Source/DisplayClusterWarp/Private/WarpBlend/Loader/DisplayClusterWarpBlendLoader_MeshComponent.cpp b/Engine/Plugins/Runtime/nDisplay/Source/DisplayClusterWarp/Private/WarpBlend/Loader/DisplayClusterWarpBlendLoader_MeshComponent.cpp
--- a/Engine/Plugins/Runtime/nDisplay/Source/DisplayClusterWarp/Private/WarpBlend/Loader/DisplayClusterWarpBlendLoader_MeshComponent.cpp
+++ b/Engine/Plugins/Runtime/nDisplay/Source/DisplayClusterWarp/Private/WarpBlend/Loader/DisplayClusterWarpBlendLoader_MeshComponent.cpp
@@ -19,6 +19,8 @@
 #include "Stats/Stats.h"
 #include "Engine/Engine.h"
 
+#include "ProceduralMeshComponent.h"
+
 //---------------------------------------------------------------------
 // FDisplayClusterWarpBlendLoader_MeshComponent
 //---------------------------------------------------------------------
@@ -31,7 +33,7 @@ TSharedPtr<IDisplayClusterRender_MeshComponent, ESPMode::ThreadSafe> FDisplayClu
 
 TSharedPtr<FDisplayClusterWarpBlend, ESPMode::ThreadSafe> FDisplayClusterWarpBlendLoader_MeshComponent::Create(const FDisplayClusterWarpInitializer_StaticMesh& InConstructParameters)
 {
-	if (InConstructParameters.StaticMeshComponent != nullptr)
+	if (InConstructParameters.WarpMeshComponent != nullptr)
 	{
 		//ok, Create and initialize warpblend interface:
 		TSharedPtr<FDisplayClusterWarpBlend, ESPMode::ThreadSafe> WarpBlend = MakeShared<FDisplayClusterWarpBlend, ESPMode::ThreadSafe>();
@@ -42,8 +44,10 @@ TSharedPtr<FDisplayClusterWarpBlend, ESPMode::ThreadSafe> FDisplayClusterWarpBle
 
 		FDisplayClusterMeshUVs MeshUVs(InConstructParameters.BaseUVIndex, InConstructParameters.ChromakeyUVIndex);
 
-		Proxy.MeshComponent = CreateMeshComponent();
-		Proxy.MeshComponent->AssignStaticMeshComponentRefs(InConstructParameters.StaticMeshComponent, MeshUVs, InConstructParameters.OriginComponent, InConstructParameters.StaticMeshComponentLODIndex);
+		Proxy.PreviewMeshComponentRef.SetSceneComponent(InConstructParameters.PreviewMeshComponent);
+
+		Proxy.WarpMeshComponent = CreateMeshComponent();
+		Proxy.WarpMeshComponent->AssignStaticMeshComponentRefs(InConstructParameters.WarpMeshComponent, MeshUVs, InConstructParameters.OriginComponent, InConstructParameters.StaticMeshComponentLODIndex);
 
 		Proxy.StaticMeshComponentLODIndex = InConstructParameters.StaticMeshComponentLODIndex;
 		Proxy.WarpMeshUVs = MeshUVs;
@@ -59,7 +63,7 @@ TSharedPtr<FDisplayClusterWarpBlend, ESPMode::ThreadSafe> FDisplayClusterWarpBle
 
 TSharedPtr<FDisplayClusterWarpBlend, ESPMode::ThreadSafe> FDisplayClusterWarpBlendLoader_MeshComponent::Create(const FDisplayClusterWarpInitializer_ProceduralMesh& InConstructParameters)
 {
-	if (InConstructParameters.ProceduralMeshComponent != nullptr)
+	if (InConstructParameters.WarpMeshComponent != nullptr)
 	{
 		//ok, Create and initialize warpblend interface:
 		TSharedPtr<FDisplayClusterWarpBlend, ESPMode::ThreadSafe> WarpBlend = MakeShared<FDisplayClusterWarpBlend, ESPMode::ThreadSafe>();
@@ -70,9 +74,11 @@ TSharedPtr<FDisplayClusterWarpBlend, ESPMode::ThreadSafe> FDisplayClusterWarpBle
 
 		FDisplayClusterMeshUVs MeshUVs(InConstructParameters.BaseUVIndex, InConstructParameters.ChromakeyUVIndex);
 
+		Proxy.PreviewMeshComponentRef.SetSceneComponent(InConstructParameters.PreviewMeshComponent);
+
 		// Assign procedural mesh
-		Proxy.MeshComponent = CreateMeshComponent();
-		Proxy.MeshComponent->AssignProceduralMeshComponentRefs(InConstructParameters.ProceduralMeshComponent, MeshUVs, InConstructParameters.OriginComponent, InConstructParameters.ProceduralMeshComponentSectionIndex);
+		Proxy.WarpMeshComponent = CreateMeshComponent();
+		Proxy.WarpMeshComponent->AssignProceduralMeshComponentRefs(InConstructParameters.WarpMeshComponent, MeshUVs, InConstructParameters.OriginComponent, InConstructParameters.ProceduralMeshComponentSectionIndex);
 
 		Proxy.ProceduralMeshComponentSectionIndex = InConstructParameters.ProceduralMeshComponentSectionIndex;
 		Proxy.WarpMeshUVs = MeshUVs;
diff --git a/Engine/Plugins/Runtime/nDisplay/Source/DisplayClusterWarp/Private/WarpBlend/Math/DisplayClusterWarpBlendMath_FrustumProjection.cpp b/Engine/Plugins/Runtime/nDisplay/Source/DisplayClusterWarp/Private/WarpBlend/Math/DisplayClusterWarpBlendMath_FrustumProjection.cpp
--- a/Engine/Plugins/Runtime/nDisplay/Source/DisplayClusterWarp/Private/WarpBlend/Math/DisplayClusterWarpBlendMath_FrustumProjection.cpp
+++ b/Engine/Plugins/Runtime/nDisplay/Source/DisplayClusterWarp/Private/WarpBlend/Math/DisplayClusterWarpBlendMath_FrustumProjection.cpp
@@ -233,7 +233,7 @@ bool FDisplayClusterWarpBlendMath_Frustum::ImplCalcFrustumProjection_WarpMap(con
 
 bool FDisplayClusterWarpBlendMath_Frustum::ImplCalcFrustumProjection_WarpMesh(const FMatrix& InGeometryToFrustumMatrix)
 {
-	if (!GeometryContext.GeometryProxy.MeshComponent.IsValid())
+	if (!GeometryContext.GeometryProxy.WarpMeshComponent.IsValid())
 	{
 		return false;
 	}
@@ -266,7 +266,7 @@ bool FDisplayClusterWarpBlendMath_Frustum::ImplCalcFrustumProjection_WarpMesh(co
 
 bool FDisplayClusterWarpBlendMath_Frustum::ImplCalcFrustumProjection_WarpProceduralMesh(const FMatrix& InGeometryToFrustumMatrix)
 {
-	if (!GeometryContext.GeometryProxy.MeshComponent.IsValid())
+	if (!GeometryContext.GeometryProxy.WarpMeshComponent.IsValid())
 	{
 		return false;
 	}
@@ -279,7 +279,7 @@ bool FDisplayClusterWarpBlendMath_Frustum::ImplCalcFrustumProjection_WarpProcedu
 
 	bool bAllPointsInFrustum = true;
 
-	const FProcMeshSection* ProcMeshSection = GeometryContext.GeometryProxy.MeshComponent->GetProceduralMeshComponentSection();
+	const FProcMeshSection* ProcMeshSection = GeometryContext.GeometryProxy.WarpMeshComponent->GetProceduralMeshComponentSection();
 	if (ProcMeshSection == nullptr)
 	{
 		return false;
diff --git a/Engine/Plugins/Runtime/nDisplay/Source/DisplayClusterWarp/Public/Components/DisplayClusterInFrustumFitCameraComponent.h b/Engine/Plugins/Runtime/nDisplay/Source/DisplayClusterWarp/Public/Components/DisplayClusterInFrustumFitCameraComponent.h
--- a/Engine/Plugins/Runtime/nDisplay/Source/DisplayClusterWarp/Public/Components/DisplayClusterInFrustumFitCameraComponent.h
+++ b/Engine/Plugins/Runtime/nDisplay/Source/DisplayClusterWarp/Public/Components/DisplayClusterInFrustumFitCameraComponent.h
@@ -30,26 +30,26 @@ public:
 	//~~~End UActorComponent
 
 	//~Begin UDisplayClusterCameraComponent
-	virtual void GetDesiredView(FMinimalViewInfo& InOutViewInfo, float* OutCustomNearClippingPlane = nullptr) override;
+	virtual void GetDesiredView(IDisplayClusterViewportConfiguration& InViewportConfiguration, FMinimalViewInfo& InOutViewInfo, float* OutCustomNearClippingPlane = nullptr) override;
 	virtual bool ShouldUseEntireClusterViewports(IDisplayClusterViewportManager* InViewportManager) const override;
 	virtual IDisplayClusterWarpPolicy* GetWarpPolicy(IDisplayClusterViewportManager* InViewportManager) override;
 	//~~End UDisplayClusterCameraComponent
 
 #if WITH_EDITOR
-	virtual void PostEditChangeProperty(FPropertyChangedEvent& PropertyChangedEvent) override;
 	virtual bool GetEditorPreviewInfo(float DeltaTime, FMinimalViewInfo& ViewOut) override;
 	virtual TSharedPtr<SWidget> GetCustomEditorPreviewWidget() override;
-
-	void InvalidatePreviewState() { bRefreshPreviewState = true; }
 #endif
 
-public:
-	/** true, if camera projection is used. */
-	bool IsEnabled() const;
+	/** Return component that used for configuration. */
+	const UDisplayClusterInFrustumFitCameraComponent& GetConfigurationInFrustumFitCameraComponent(IDisplayClusterViewportConfiguration& InViewportConfiguration) const;
 
+private:
 	/** Return external camera component. */
 	UCameraComponent* GetExternalCameraComponent() const;
 
+	/** true, if camera projection is used. */
+	bool IsEnabled() const;
+
 public:
 	/** Camera projection mode is used. */
 	UPROPERTY(EditAnywhere, Category = "Projection")
@@ -80,9 +80,4 @@ private:
 	// a unique type of warp policy for this component
 	// this policy class knows the properties of the component and implements the corresponding logic
 	TSharedPtr<IDisplayClusterWarpPolicy, ESPMode::ThreadSafe> WarpPolicy;
-
-#if WITH_EDITOR
-	/** Indicates that the state of the preview meshes for the frustum fit are invalid and need to be refreshed on the next tick */
-	bool bRefreshPreviewState = true;
-#endif
 };
diff --git a/Engine/Plugins/Runtime/nDisplay/Source/DisplayClusterWarp/Public/Containers/DisplayClusterWarpContainers.h b/Engine/Plugins/Runtime/nDisplay/Source/DisplayClusterWarp/Public/Containers/DisplayClusterWarpContainers.h
--- a/Engine/Plugins/Runtime/nDisplay/Source/DisplayClusterWarp/Public/Containers/DisplayClusterWarpContainers.h
+++ b/Engine/Plugins/Runtime/nDisplay/Source/DisplayClusterWarp/Public/Containers/DisplayClusterWarpContainers.h
@@ -17,6 +17,11 @@ struct FDisplayClusterWarpProjection
 		Bottom = DBL_MAX;
 	}
 
+	inline bool IsValidProjection() const
+	{
+		return Left < Right && Bottom < Top;
+	}
+
 	inline void ExpandProjectionAngles(const FDisplayClusterWarpProjection& In)
 	{
 		Left   = FMath::Min(Left,   In.Left);
diff --git a/Engine/Plugins/Runtime/nDisplay/Source/DisplayClusterWarp/Public/Containers/DisplayClusterWarpContext.h b/Engine/Plugins/Runtime/nDisplay/Source/DisplayClusterWarp/Public/Containers/DisplayClusterWarpContext.h
--- a/Engine/Plugins/Runtime/nDisplay/Source/DisplayClusterWarp/Public/Containers/DisplayClusterWarpContext.h
+++ b/Engine/Plugins/Runtime/nDisplay/Source/DisplayClusterWarp/Public/Containers/DisplayClusterWarpContext.h
@@ -68,6 +68,13 @@ struct FDisplayClusterWarpData
 	// Control projection type
 	bool bFindBestProjectionType = false;
 	EDisplayClusterWarpBlendProjectionType ProjectionType = EDisplayClusterWarpBlendProjectionType::DynamicAABBCenter;
+
+	// After successful warpdata calculation, this value is set to true.
+	bool bValid = false;
+
+	// The warp policy Tick() function uses warp data, and it must be sure that this data is updated in the previous frame.
+	// This value must be set to true from the EndCalcFrustum() warp policy function when changes are made to this structure.
+	bool bHasWarpPolicyChanges = false;
 };
 
 struct FDisplayClusterWarpGeometryContext
diff --git a/Engine/Plugins/Runtime/nDisplay/Source/DisplayClusterWarp/Public/Containers/DisplayClusterWarpInitializer.h b/Engine/Plugins/Runtime/nDisplay/Source/DisplayClusterWarp/Public/Containers/DisplayClusterWarpInitializer.h
--- a/Engine/Plugins/Runtime/nDisplay/Source/DisplayClusterWarp/Public/Containers/DisplayClusterWarpInitializer.h
+++ b/Engine/Plugins/Runtime/nDisplay/Source/DisplayClusterWarp/Public/Containers/DisplayClusterWarpInitializer.h
@@ -33,7 +33,10 @@ struct FDisplayClusterWarpInitializer_MPCDIFile_Profile2DScreen
 	USceneComponent* OriginComponent = nullptr;
 
 	// StaticMesh component ptr (warp geometry data source)
-	UStaticMeshComponent* StaticMeshComponent = nullptr;
+	UStaticMeshComponent* WarpMeshComponent = nullptr;
+
+	// Preview StaticMesh component ptr (warp geometry data source)
+	UStaticMeshComponent* PreviewMeshComponent = nullptr;
 };
 
 /**
@@ -83,7 +86,10 @@ struct FDisplayClusterWarpInitializer_StaticMesh
 	: public FDisplayClusterWarpInitializer_BaseMesh
 {
 	// StaticMesh component ptr (warp geometry data source)
-	UStaticMeshComponent* StaticMeshComponent = nullptr;
+	UStaticMeshComponent* WarpMeshComponent = nullptr;
+
+	// Preview StaticMesh component ptr (warp geometry data source)
+	UStaticMeshComponent* PreviewMeshComponent = nullptr;
 
 	// (optional) get geometry data from specified LOD index
 	int32 StaticMeshComponentLODIndex = 0;
@@ -96,7 +102,10 @@ struct FDisplayClusterWarpInitializer_ProceduralMesh
 	: public FDisplayClusterWarpInitializer_BaseMesh
 {
 	// ProceduralMesh component ptr (warp geometry data source)
-	UProceduralMeshComponent* ProceduralMeshComponent = nullptr;
+	UProceduralMeshComponent* WarpMeshComponent = nullptr;
+
+	// Preview ProceduralMesh component ptr (warp geometry data source)
+	UProceduralMeshComponent* PreviewMeshComponent = nullptr;
 
 	// (optional) get geometry data from specified SectionIndex
 	int32 ProceduralMeshComponentSectionIndex = 0;
diff --git a/Engine/Plugins/Runtime/nDisplay/Source/DisplayClusterWarp/Public/IDisplayClusterWarpBlend.h b/Engine/Plugins/Runtime/nDisplay/Source/DisplayClusterWarp/Public/IDisplayClusterWarpBlend.h
--- a/Engine/Plugins/Runtime/nDisplay/Source/DisplayClusterWarp/Public/IDisplayClusterWarpBlend.h
+++ b/Engine/Plugins/Runtime/nDisplay/Source/DisplayClusterWarp/Public/IDisplayClusterWarpBlend.h
@@ -102,5 +102,13 @@ public:
 	* 
 	* @return - ptr to the mesh component with the geometry, or nullptr in case of failure
 	*/
-	virtual UMeshComponent* GetOrCreateMeshComponent(IDisplayClusterViewport* InViewport, bool& bExistingComponent) const = 0;
+	virtual UMeshComponent* GetOrCreatePreviewMeshComponent(IDisplayClusterViewport* InViewport, bool& bExistingComponent) const = 0;
+
+	/**
+	* Build preview movable mesh
+	* This MeshComponent is a copy of the preview mesh and can be moved freely with the UI visualization.
+	*
+	* @param InViewport - Projection specific parameters.
+	*/
+	virtual UMeshComponent* GetOrCreatePreviewMovableMeshComponent(IDisplayClusterViewport* InViewport) const = 0;
 };
diff --git a/Engine/Plugins/VirtualProduction/ICVFXTesting/Source/ICVFXTesting/Public/ICVFXTestControllerAutoTest.h b/Engine/Plugins/VirtualProduction/ICVFXTesting/Source/ICVFXTesting/Public/ICVFXTestControllerAutoTest.h
--- a/Engine/Plugins/VirtualProduction/ICVFXTesting/Source/ICVFXTesting/Public/ICVFXTestControllerAutoTest.h
+++ b/Engine/Plugins/VirtualProduction/ICVFXTesting/Source/ICVFXTesting/Public/ICVFXTestControllerAutoTest.h
@@ -17,6 +17,7 @@
 #include "LiveLinkPreset.h"
 #include "LiveLinkPresetTypes.h"
 #include "Roles/LiveLinkTransformRole.h"
+#include "Render/Viewport/IDisplayClusterViewportManager.h"
 
 #include "ICVFXTestControllerAutoTest.generated.h"
 
@@ -164,8 +165,13 @@ public:
 			CameraComponent->CameraSettings.RenderSettings.AdvancedRenderSettings.GPUIndex = GetInnerGPUIndex();
 		}
 
-		const FString NodeId = IDisplayCluster::Get().GetClusterMgr()->GetNodeId();
-		RootActor->GetViewportManager()->UpdateConfiguration(EDisplayClusterRenderFrameMode::Mono, NodeId, RootActor, nullptr);
+		if (IDisplayClusterViewportManager* ViewportManager = RootActor->GetViewportManager())
+		{
+			const FString NodeId = IDisplayCluster::Get().GetClusterMgr()->GetNodeId();
+			const EDisplayClusterRenderFrameMode RenderMode = EDisplayClusterRenderFrameMode::Mono;
+
+			ViewportManager->GetConfiguration().UpdateConfigurationForClusterNode(RenderMode, NodeId);
+		}
 	}
 
 	void UpdateTestLocations()
diff --git a/Engine/Plugins/VirtualProduction/TextureShare/Source/TextureShareDisplayCluster/Private/PostProcess/TextureSharePostprocess.cpp b/Engine/Plugins/VirtualProduction/TextureShare/Source/TextureShareDisplayCluster/Private/PostProcess/TextureSharePostprocess.cpp
--- a/Engine/Plugins/VirtualProduction/TextureShare/Source/TextureShareDisplayCluster/Private/PostProcess/TextureSharePostprocess.cpp
+++ b/Engine/Plugins/VirtualProduction/TextureShare/Source/TextureShareDisplayCluster/Private/PostProcess/TextureSharePostprocess.cpp
@@ -40,7 +40,7 @@ namespace UE::TextureShare::PostProcess
 	{
 		if (InViewportManager)
 		{
-			if (UWorld* CurrentWorld = InViewportManager->GetCurrentWorld())
+			if (UWorld* CurrentWorld = InViewportManager->GetConfiguration().GetCurrentWorld())
 			{
 				if (UGameViewportClient* GameViewportClientPtr = CurrentWorld->GetGameViewport())
 				{
diff --git a/Engine/Plugins/VirtualProduction/TextureShare/Source/TextureShareDisplayCluster/Private/Projection/TextureShareProjectionPolicyBase.cpp b/Engine/Plugins/VirtualProduction/TextureShare/Source/TextureShareDisplayCluster/Private/Projection/TextureShareProjectionPolicyBase.cpp
--- a/Engine/Plugins/VirtualProduction/TextureShare/Source/TextureShareDisplayCluster/Private/Projection/TextureShareProjectionPolicyBase.cpp
+++ b/Engine/Plugins/VirtualProduction/TextureShare/Source/TextureShareDisplayCluster/Private/Projection/TextureShareProjectionPolicyBase.cpp
@@ -47,7 +47,7 @@ bool FTextureShareProjectionPolicyBase::IsEditorOperationMode(class IDisplayClus
 	}
 
 	// Get state from viewport world (UE-114493)
-	return InViewport && InViewport->IsCurrentWorldHasAnyType(EWorldType::Editor, EWorldType::EditorPreview);
+	return InViewport && InViewport->GetConfiguration().IsCurrentWorldHasAnyType(EWorldType::Editor, EWorldType::EditorPreview);
 }
 
 bool FTextureShareProjectionPolicyBase::IsConfigurationChanged(const struct FDisplayClusterConfigurationProjection* InConfigurationProjectionPolicy) const
