nDisplay: Fixed a crash when the root actor was not created in the game thread.
The logic in DCRA related to the viewport manager has been slightly reworked: the viewport manager is only created when it needs to be used for rendering.

diff --git a/Engine/Plugins/Runtime/nDisplay/Source/DisplayCluster/Private/Game/EngineClasses/Basics/DisplayClusterViewportClient.cpp b/Engine/Plugins/Runtime/nDisplay/Source/DisplayCluster/Private/Game/EngineClasses/Basics/DisplayClusterViewportClient.cpp
--- a/Engine/Plugins/Runtime/nDisplay/Source/DisplayCluster/Private/Game/EngineClasses/Basics/DisplayClusterViewportClient.cpp
+++ b/Engine/Plugins/Runtime/nDisplay/Source/DisplayCluster/Private/Game/EngineClasses/Basics/DisplayClusterViewportClient.cpp
@@ -999,7 +999,7 @@ bool UDisplayClusterViewportClient::Draw_PIE(FViewport* InViewport, FCanvas* Sce
 {
 	IDisplayClusterGameManager* GameMgr = GDisplayCluster->GetGameMgr();
 
-	if (GameMgr == nullptr)
+	if (GameMgr == nullptr || !IsInGameThread())
 	{
 		return false;
 	}
@@ -1023,7 +1023,8 @@ bool UDisplayClusterViewportClient::Draw_PIE(FViewport* InViewport, FCanvas* Sce
 	//Get world for render
 	UWorld* const MyWorld = GetWorld();
 
-	IDisplayClusterViewportManager* ViewportManager = RootActor->GetViewportManager();
+	// When the PIE is used by this DCRA, we must create a new ViewportManager
+	IDisplayClusterViewportManager* ViewportManager = RootActor->GetOrCreateViewportManager();
 	if (ViewportManager == nullptr)
 	{
 		return false;
diff --git a/Engine/Plugins/Runtime/nDisplay/Source/DisplayCluster/Private/Game/EngineClasses/Scene/DisplayClusterRootActor.cpp b/Engine/Plugins/Runtime/nDisplay/Source/DisplayCluster/Private/Game/EngineClasses/Scene/DisplayClusterRootActor.cpp
--- a/Engine/Plugins/Runtime/nDisplay/Source/DisplayCluster/Private/Game/EngineClasses/Scene/DisplayClusterRootActor.cpp
+++ b/Engine/Plugins/Runtime/nDisplay/Source/DisplayCluster/Private/Game/EngineClasses/Scene/DisplayClusterRootActor.cpp
@@ -92,8 +92,6 @@ ADisplayClusterRootActor::ADisplayClusterRootActor(const FObjectInitializer& Obj
 	DefaultViewPoint->AttachToComponent(GetRootComponent(), FAttachmentTransformRules::KeepRelativeTransform);
 	DefaultViewPoint->SetRelativeLocation(FVector(0.f, 0.f, 50.f));
 
-	CreateViewportManagerImpl();
-
 	PrimaryActorTick.bCanEverTick = true;
 	PrimaryActorTick.TickGroup = ETickingGroup::TG_PostUpdateWork;
 
@@ -124,38 +122,53 @@ ADisplayClusterRootActor::~ADisplayClusterRootActor()
 
 IDisplayClusterViewportManager* ADisplayClusterRootActor::GetViewportManager() const
 {
-	return ViewportManager.Get();
+	return GetViewportManagerImpl();
 }
 
-IDisplayClusterViewportConfiguration* ADisplayClusterRootActor::GetViewportConfiguration() const
+IDisplayClusterViewportManager* ADisplayClusterRootActor::GetOrCreateViewportManager()
 {
-	return ViewportManager.IsValid() ? &ViewportManager->GetConfiguration() : nullptr;
-}
+	check(IsInGameThread());
 
-void ADisplayClusterRootActor::CreateViewportManagerImpl()
-{
-	if (!ViewportManager.IsValid())
+	if (!ViewportManagerPtr.IsValid())
 	{
-		ViewportManager = MakeShared<FDisplayClusterViewportManager, ESPMode::ThreadSafe>();
+		ViewportManagerPtr = MakeShared<FDisplayClusterViewportManager, ESPMode::ThreadSafe>();
 
 		// After the constructor, we should always call this function to initialize internal references.
-		ViewportManager->Initialize();
+		ViewportManagerPtr->Initialize();
 
 		// Set the owner's DCRA to the newly created viewport manager.
-		ViewportManager->GetConfiguration().SetRootActor(this, EDisplayClusterRootActorType::Any);
+		ViewportManagerPtr->GetConfiguration().SetRootActor(this, EDisplayClusterRootActorType::Any);
 
 		// Preview rendering depends on the DC VM
 		FDisplayClusterRootActorPreviewRenderingManager::HandleEvent(EDisplayClusterRootActorPreviewEvent::Create, this);
 	}
+
+	return GetViewportManagerImpl();
 }
 
-void ADisplayClusterRootActor::RemoveViewportManagerImpl()
+FDisplayClusterViewportManager* ADisplayClusterRootActor::GetViewportManagerImpl() const
 {
-	// Preview rendering depends on the  DC VM
-	FDisplayClusterRootActorPreviewRenderingManager::HandleEvent(EDisplayClusterRootActorPreviewEvent::Remove, this);
+	return ViewportManagerPtr.IsValid() ? ViewportManagerPtr.Get() : nullptr;
+}
 
-	// Immediately release the viewport manager with resources
-	ViewportManager.Reset();
+IDisplayClusterViewportConfiguration* ADisplayClusterRootActor::GetViewportConfiguration() const
+{
+	return ViewportManagerPtr.IsValid() ? &ViewportManagerPtr->GetConfiguration() : nullptr;
+}
+
+void ADisplayClusterRootActor::RemoveViewportManager()
+{
+	if (ViewportManagerPtr.IsValid())
+	{
+		// Preview rendering depends on the  DC VM
+		FDisplayClusterRootActorPreviewRenderingManager::HandleEvent(EDisplayClusterRootActorPreviewEvent::Remove, this);
+
+		// Reset all DCRA references
+		ViewportManagerPtr->GetConfiguration().SetRootActor(nullptr, EDisplayClusterRootActorType::Any);
+
+		// Immediately release the viewport manager with resources
+		ViewportManagerPtr.Reset();
+	}
 }
 
 bool ADisplayClusterRootActor::IsRunningGameOrPIE() const
@@ -711,8 +724,6 @@ void ADisplayClusterRootActor::InitializeRootActor()
 		UpdateConfigDataInstance(GetDefaultConfigDataFromAsset());
 	}
 
-	CreateViewportManagerImpl();
-
 	StageGeometryComponent->Invalidate();
 
 	// Packaged, PIE and -game runtime
@@ -746,7 +757,7 @@ void ADisplayClusterRootActor::UpdateProceduralMeshComponentData(const UProcedur
 {
 	check(IsInGameThread());
 
-	if (ViewportManager.IsValid())
+	if (FDisplayClusterViewportManager* ViewportManager = GetViewportManagerImpl())
 	{
 		FName ProceduralComponentName = (InProceduralMeshComponent==nullptr) ? NAME_None : InProceduralMeshComponent->GetFName();
 
@@ -1076,7 +1087,7 @@ void ADisplayClusterRootActor::Destroyed()
 #endif
 
 	// Release viewport manager with resources immediatelly
-	RemoveViewportManagerImpl();
+	RemoveViewportManager();
 
 	Super::Destroyed();
 }
@@ -1087,7 +1098,8 @@ void ADisplayClusterRootActor::BeginDestroy()
 	BeginDestroy_Editor();
 #endif
 
-	RemoveViewportManagerImpl();
+	// Release viewport manager with resources immediatelly
+	RemoveViewportManager();
 
 	Super::BeginDestroy();
 }
@@ -1464,10 +1476,9 @@ bool ADisplayClusterRootActor::SetFreezeOuterViewports(bool bEnable)
 void ADisplayClusterRootActor::AddReferencedObjects(UObject* InThis, FReferenceCollector& Collector)
 {
 	ADisplayClusterRootActor* This = CastChecked<ADisplayClusterRootActor>(InThis);
-
-	if (This && This->ViewportManager.IsValid())
+	if (FDisplayClusterViewportManager* ViewportManager = This ? This->GetViewportManagerImpl() : nullptr)
 	{
-		This->ViewportManager->AddReferencedObjects(Collector);
+		ViewportManager->AddReferencedObjects(Collector);
 	}
 
 	Super::AddReferencedObjects(InThis, Collector);
diff --git a/Engine/Plugins/Runtime/nDisplay/Source/DisplayCluster/Private/Game/EngineClasses/Scene/DisplayClusterRootActorEditor.cpp b/Engine/Plugins/Runtime/nDisplay/Source/DisplayCluster/Private/Game/EngineClasses/Scene/DisplayClusterRootActorEditor.cpp
--- a/Engine/Plugins/Runtime/nDisplay/Source/DisplayCluster/Private/Game/EngineClasses/Scene/DisplayClusterRootActorEditor.cpp
+++ b/Engine/Plugins/Runtime/nDisplay/Source/DisplayCluster/Private/Game/EngineClasses/Scene/DisplayClusterRootActorEditor.cpp
@@ -126,43 +126,42 @@ void ADisplayClusterRootActor::Destructor_Editor()
 
 void ADisplayClusterRootActor::RenderPreview_Editor()
 {
-	if (IsPreviewEnabled())
+	if (IsPreviewEnabled() && IsInGameThread())
 	{
-		// Restore ViewportManager
-		CreateViewportManagerImpl();
-
-		if (bDeferPreviewGeneration)
+		// When the preview is used by this DCRA, we must create a new ViewportManager
+		if (GetOrCreateViewportManager())
 		{
-			// Hack to generate preview components on instances during map load.
-			// TODO: See if we can move InitializeRootActor out of PostLoad.
-			bDeferPreviewGeneration = false;
-			UpdatePreviewComponents();
-		}
-
-		// Update preview RTTs correspond to 'TickPerFrame' value
-		if (++TickPerFrameCounter >= TickPerFrame)
-		{
-			TickPerFrameCounter = 0;
-
 			// Render viewport for preview material RTTs
-			if (ViewportManager.IsValid())
+			if (bDeferPreviewGeneration)
 			{
+				// Hack to generate preview components on instances during map load.
+				// TODO: See if we can move InitializeRootActor out of PostLoad.
+				bDeferPreviewGeneration = false;
+				UpdatePreviewComponents();
+			}
+
+			// Update preview RTTs correspond to 'TickPerFrame' value
+			if (++TickPerFrameCounter >= TickPerFrame)
+			{
+				TickPerFrameCounter = 0;
+
 				ImplRenderPreview_Editor();
 			}
-		}
 
-		// preview frustums on each tick
-		ImplRenderPreviewFrustums_Editor();
+			// preview frustums on each tick
+			ImplRenderPreviewFrustums_Editor();
+		}
 	}
 	else
 	{
 		ResetPreviewInternals_Editor();
-		if (ViewportManager.IsValid())
+		if (FDisplayClusterViewportManager* ViewportManager = GetViewportManagerImpl())
 		{
 			if(ViewportManager->GetConfiguration().IsPreviewRendering())
 			{
-				// Release viewport manager with resources immediatelly for preview in scene
-				RemoveViewportManagerImpl();
+				// Preview is no longer in use.
+				// Release viewport manager with resources immediatelly
+				RemoveViewportManager();
 			}
 		}
 	}
@@ -353,7 +352,7 @@ void ADisplayClusterRootActor::SetIsSelectedInEditor(bool bValue)
 
 IDisplayClusterViewport* ADisplayClusterRootActor::FindPreviewViewport(const FString& InViewportId) const
 {
-	if (ViewportManager.IsValid())
+	if (FDisplayClusterViewportManager* ViewportManager = GetViewportManagerImpl())
 	{
 		return ViewportManager->FindViewport(InViewportId);
 	}
@@ -376,17 +375,20 @@ bool ADisplayClusterRootActor::ImplUpdatePreviewConfiguration_Editor(const FStri
 	// Reset current preview rendering.
 	ResetClusterNodePreviewRendering_Editor();
 
-	if (IsPreviewEnabled() && ViewportManager.IsValid())
+	if (IsPreviewEnabled())
 	{
-		// Now we render this node
-		PreviewRenderFrameClusterNodeId = InClusterNodeId;
+		if (FDisplayClusterViewportManager* ViewportManager = GetViewportManagerImpl())
+		{
+			// Now we render this node
+			PreviewRenderFrameClusterNodeId = InClusterNodeId;
 
-		// Update current world from scene DCRA
-		ADisplayClusterRootActor* SceneRootActor = ViewportManager->GetConfiguration().GetRootActor(EDisplayClusterRootActorType::Scene);
-		ViewportManager->GetConfiguration().SetCurrentWorld(SceneRootActor ? SceneRootActor->GetWorld() : GetWorld());
+			// Update current world from scene DCRA
+			ADisplayClusterRootActor* SceneRootActor = ViewportManager->GetConfiguration().GetRootActor(EDisplayClusterRootActorType::Scene);
+			ViewportManager->GetConfiguration().SetCurrentWorld(SceneRootActor ? SceneRootActor->GetWorld() : GetWorld());
 
-		// Update local node viewports (update\create\delete) and build new render frame
-		return ViewportManager->GetConfiguration().UpdateConfigurationForClusterNode(GetPreviewRenderMode_Editor(), InClusterNodeId);
+			// Update local node viewports (update\create\delete) and build new render frame
+			return ViewportManager->GetConfiguration().UpdateConfigurationForClusterNode(GetPreviewRenderMode_Editor(), InClusterNodeId);
+		}
 	}
 
 	return false;
@@ -436,19 +438,22 @@ bool ADisplayClusterRootActor::ImplUpdatePreviewRenderFrame_Editor(const FString
 
 		PreviewRenderFrame = MakeUnique<FDisplayClusterRenderFrame>();
 
-		// Update preview viewports from settings
-		if (!ViewportManager->BeginNewFrame(nullptr, *PreviewRenderFrame))
+		if (FDisplayClusterViewportManager* ViewportManager = GetViewportManagerImpl())
 		{
-			PreviewRenderFrame.Reset();
+			// Update preview viewports from settings
+			if (!ViewportManager->BeginNewFrame(nullptr, *PreviewRenderFrame))
+			{
+				PreviewRenderFrame.Reset();
 
-			return false;
-		}
+				return false;
+			}
 
-		// Begin Render Preview For Cluster Node
-		PreviewViewportIndex = 0;
+			// Begin Render Preview For Cluster Node
+			PreviewViewportIndex = 0;
 
-		// Initialize frame for render
-		ViewportManager->InitializeNewFrame();
+			// Initialize frame for render
+			ViewportManager->InitializeNewFrame();
+		}
 	}
 
 	return IsActiveClusterNodePreviewRendering_Editor();
@@ -468,26 +473,29 @@ bool ADisplayClusterRootActor::ImplRenderPassPreviewClusterNode_Editor(const FSt
 		return false;
 	}
 
-	bool bFrameRendered = false;
-	int32 RenderedViewportsAmount = 0;
+	if (FDisplayClusterViewportManager* ViewportManager = GetViewportManagerImpl())
+	{
+		bool bFrameRendered = false;
+		int32 RenderedViewportsAmount = 0;
 
-	ViewportManager->RenderInEditor(*PreviewRenderFrame, nullptr, PreviewViewportIndex, ViewportsAmount, RenderedViewportsAmount, bFrameRendered);
+		ViewportManager->RenderInEditor(*PreviewRenderFrame, nullptr, PreviewViewportIndex, ViewportsAmount, RenderedViewportsAmount, bFrameRendered);
 
-	// Increase viewport index
-	PreviewViewportIndex += ViewportsAmount;
+		// Increase viewport index
+		PreviewViewportIndex += ViewportsAmount;
 
-	// Count only rendered viewports
-	PreviewViewportsRenderedInThisFrameCnt += RenderedViewportsAmount;
+		// Count only rendered viewports
+		PreviewViewportsRenderedInThisFrameCnt += RenderedViewportsAmount;
 
-	if (bFrameRendered)
-	{
-		// current cluster node is composed
-		ResetClusterNodePreviewRendering_Editor();
+		if (bFrameRendered)
+		{
+			// current cluster node is composed
+			ResetClusterNodePreviewRendering_Editor();
 
-		// Send event about RTT changed
-		OnPreviewGenerated.ExecuteIfBound();
+			// Send event about RTT changed
+			OnPreviewGenerated.ExecuteIfBound();
 
-		return true;
+			return true;
+		}
 	}
 
 	return false;
@@ -495,7 +503,7 @@ bool ADisplayClusterRootActor::ImplRenderPassPreviewClusterNode_Editor(const FSt
 
 void ADisplayClusterRootActor::ImplRenderPreview_Editor()
 {
-	if (CurrentConfigData == nullptr || !IsPreviewEnabled() || !ViewportManager.IsValid())
+	if (CurrentConfigData == nullptr || !IsPreviewEnabled())
 	{
 		// no preview
 		return;
@@ -569,7 +577,8 @@ void ADisplayClusterRootActor::ImplRenderPreview_Editor()
 
 void ADisplayClusterRootActor::ImplRenderPreviewFrustums_Editor()
 {
-	if (CurrentConfigData == nullptr || !ViewportManager.IsValid())
+	FDisplayClusterViewportManager* ViewportManager = GetViewportManagerImpl();
+	if (CurrentConfigData == nullptr || !ViewportManager)
 	{
 		return;
 	}
diff --git a/Engine/Plugins/Runtime/nDisplay/Source/DisplayCluster/Private/Render/Device/DisplayClusterDeviceBase.cpp b/Engine/Plugins/Runtime/nDisplay/Source/DisplayCluster/Private/Render/Device/DisplayClusterDeviceBase.cpp
--- a/Engine/Plugins/Runtime/nDisplay/Source/DisplayCluster/Private/Render/Device/DisplayClusterDeviceBase.cpp
+++ b/Engine/Plugins/Runtime/nDisplay/Source/DisplayCluster/Private/Render/Device/DisplayClusterDeviceBase.cpp
@@ -324,7 +324,7 @@ bool FDisplayClusterDeviceBase::BeginNewFrame(FViewport* InViewport, UWorld* InW
 
 	if (ADisplayClusterRootActor* RootActor = DeviceBaseHelpers::GetDisplayClusterAPI().GetGameMgr()->GetRootActor())
 	{
-		if (IDisplayClusterViewportManager* ViewportManagerPtr = RootActor->GetViewportManager())
+		if (IDisplayClusterViewportManager* ViewportManagerPtr = RootActor->GetOrCreateViewportManager())
 		{
 			const FString LocalNodeId = DeviceBaseHelpers::GetDisplayClusterAPI().GetConfigMgr()->GetLocalNodeId();
 
@@ -350,9 +350,11 @@ bool FDisplayClusterDeviceBase::BeginNewFrame(FViewport* InViewport, UWorld* InW
 
 void FDisplayClusterDeviceBase::InitializeNewFrame()
 {
+	check(IsInGameThread());
+
 	if (ADisplayClusterRootActor* RootActor = DeviceBaseHelpers::GetDisplayClusterAPI().GetGameMgr()->GetRootActor())
 	{
-		if (IDisplayClusterViewportManager* ViewportManager = RootActor->GetViewportManager())
+		if (IDisplayClusterViewportManager* ViewportManager = RootActor->GetOrCreateViewportManager())
 		{
 			// Begin use viewport manager for current frame
 			ViewportManagerWeakPtr = ViewportManager->ToSharedPtr();
diff --git a/Engine/Plugins/Runtime/nDisplay/Source/DisplayCluster/Private/Render/Viewport/DisplayClusterViewportManager.h b/Engine/Plugins/Runtime/nDisplay/Source/DisplayCluster/Private/Render/Viewport/DisplayClusterViewportManager.h
--- a/Engine/Plugins/Runtime/nDisplay/Source/DisplayCluster/Private/Render/Viewport/DisplayClusterViewportManager.h
+++ b/Engine/Plugins/Runtime/nDisplay/Source/DisplayCluster/Private/Render/Viewport/DisplayClusterViewportManager.h
@@ -132,12 +132,9 @@ public:
 	/** Get DC viewport manager proxy. */
 	TSharedPtr<FDisplayClusterViewportManagerProxy, ESPMode::ThreadSafe> GetViewportManagerProxy() const
 	{
-		check(IsInGameThread());
-
 		return ViewportManagerProxy;
 	}
 
-
 	/** Return initial StereoViewIndex for the input viewport. */
 	int32 FindFirstViewportStereoViewIndex(const TSharedPtr<FDisplayClusterViewport, ESPMode::ThreadSafe>& InViewport) const;
 
diff --git a/Engine/Plugins/Runtime/nDisplay/Source/DisplayCluster/Public/DisplayClusterRootActor.h b/Engine/Plugins/Runtime/nDisplay/Source/DisplayCluster/Public/DisplayClusterRootActor.h
--- a/Engine/Plugins/Runtime/nDisplay/Source/DisplayCluster/Public/DisplayClusterRootActor.h
+++ b/Engine/Plugins/Runtime/nDisplay/Source/DisplayCluster/Public/DisplayClusterRootActor.h
@@ -32,6 +32,7 @@ class IDisplayClusterConfiguratorBlueprintEditor;
 class FTransactionObjectEvent;
 #endif
 
+class FDisplayClusterViewportManager;
 class IDisplayClusterStageActor;
 class USceneComponent;
 class UDisplayClusterConfigurationData;
@@ -208,6 +209,12 @@ public:
 	/** Get ViewportManager API. */
 	IDisplayClusterViewportManager* GetViewportManager() const;
 
+	/** Get or Create ViewportManager API. */
+	IDisplayClusterViewportManager* GetOrCreateViewportManager();
+
+	/** Release the viewport manager instance, if it exists. */
+	void RemoveViewportManager();
+
 	/** Get ViewportConfiguration API.*/
 	IDisplayClusterViewportConfiguration* GetViewportConfiguration() const;
 	
@@ -217,15 +224,12 @@ public:
 	}
 
 private:
-	/** Create a new instance of the viewport manager if it does not exist. */
-	void CreateViewportManagerImpl();
-
-	/** Release the viewport manager instance, if it exists. */
-	void RemoveViewportManagerImpl();
+	/** Get ViewportManager API. */
+	FDisplayClusterViewportManager* GetViewportManagerImpl() const;
 
 private:
 	// DC ViewportManager instance for this DCRA
-	TSharedPtr<class FDisplayClusterViewportManager, ESPMode::ThreadSafe> ViewportManager;
+	TSharedPtr<FDisplayClusterViewportManager, ESPMode::ThreadSafe> ViewportManagerPtr;
 
 //////////////////////////////////////////////////////////////////////////////////////////////
 // Details Panel Property Referencers
diff --git a/Engine/Plugins/Runtime/nDisplay/Source/DisplayClusterMoviePipeline/Private/DisplayClusterMoviePipelineViewportPass.cpp b/Engine/Plugins/Runtime/nDisplay/Source/DisplayClusterMoviePipeline/Private/DisplayClusterMoviePipelineViewportPass.cpp
--- a/Engine/Plugins/Runtime/nDisplay/Source/DisplayClusterMoviePipeline/Private/DisplayClusterMoviePipelineViewportPass.cpp
+++ b/Engine/Plugins/Runtime/nDisplay/Source/DisplayClusterMoviePipeline/Private/DisplayClusterMoviePipelineViewportPass.cpp
@@ -434,6 +434,12 @@ void UDisplayClusterMoviePipelineViewportPassBase::GetViewportCutOffset(const FM
 
 void UDisplayClusterMoviePipelineViewportPassBase::ReleaseDisplayCluster()
 {
+	if (DCRootActor)
+	{
+		// Delete all used rendering resources
+		DCRootActor->RemoveViewportManager();
+	}
+
 	// Reset all runtime values
 	DCRootActor = nullptr;
 
@@ -446,6 +452,8 @@ void UDisplayClusterMoviePipelineViewportPassBase::ReleaseDisplayCluster()
 
 bool UDisplayClusterMoviePipelineViewportPassBase::InitializeDisplayCluster()
 {
+	check(IsInGameThread());
+
 	ReleaseDisplayCluster();
 
 	UWorld* CurrentWorld = GetWorld();
@@ -475,14 +483,14 @@ bool UDisplayClusterMoviePipelineViewportPassBase::InitializeDisplayCluster()
 				return false;
 			}
 
-			IDisplayClusterViewportManager* ViewportManager = DCRootActor->GetViewportManager();
-			check(ViewportManager);
-
-			// Update local node viewports (update\create\delete) and build new render frame
-			const EDisplayClusterRenderFrameMode RenderFrameMode = EDisplayClusterRenderFrameMode::Mono;
-			if (ViewportManager->GetConfiguration().UpdateConfigurationForViewportsList(RenderFrameMode, DisplayClusterViewports))
+			if (IDisplayClusterViewportManager* ViewportManager = DCRootActor->GetOrCreateViewportManager())
 			{
-				return true;
+				// Update local node viewports (update\create\delete) and build new render frame
+				const EDisplayClusterRenderFrameMode RenderFrameMode = EDisplayClusterRenderFrameMode::Mono;
+				if (ViewportManager->GetConfiguration().UpdateConfigurationForViewportsList(RenderFrameMode, DisplayClusterViewports))
+				{
+					return true;
+				}
 			}
 		}
 	}
