nDisplay: Rename TechViz references to Techvis.

diff --git a/Engine/Plugins/Runtime/nDisplay/Source/DisplayCluster/Private/Game/EngineClasses/Scene/DisplayClusterPreviewComponent.cpp b/Engine/Plugins/Runtime/nDisplay/Source/DisplayCluster/Private/Game/EngineClasses/Scene/DisplayClusterPreviewComponent.cpp
--- a/Engine/Plugins/Runtime/nDisplay/Source/DisplayCluster/Private/Game/EngineClasses/Scene/DisplayClusterPreviewComponent.cpp
+++ b/Engine/Plugins/Runtime/nDisplay/Source/DisplayCluster/Private/Game/EngineClasses/Scene/DisplayClusterPreviewComponent.cpp
@@ -294,12 +294,12 @@ bool UDisplayClusterPreviewComponent::UpdatePreviewMesh()
 				{
 					PreviewMesh->SetCastShadow(false);
 
-					if (RootActor && RootActor->IsTechVizPreviewRenderingEnabled())
+					if (RootActor && RootActor->IsTechvisPreviewRenderingEnabled())
 					{
-						// Force set these properties to on with techviz rendering enabled. It won't work properly
-						// with them off. We don't need to force them off when techviz rendering is disabled, PostEditChangeProperty
+						// Force set these properties to on with Techvis rendering enabled. It won't work properly
+						// with them off. We don't need to force them off when Techvis rendering is disabled, PostEditChangeProperty
 						// on the root actor will take care of that.
-						UE::DisplayClusterDisplayDeviceUtils::ConfigureTechVizForMesh(PreviewMesh, true);
+						UE::DisplayClusterDisplayDeviceUtils::ConfigureTechvisForMesh(PreviewMesh, true);
 					}
 				}
 
@@ -386,7 +386,7 @@ void UDisplayClusterPreviewComponent::UpdatePreviewMaterial()
 		}
 
 		// Allow display device to perform any processing on the material instance.
-		if (RootActor && RootActor->IsTechVizPreviewRenderingEnabled())
+		if (RootActor && RootActor->IsTechvisPreviewRenderingEnabled())
 		{
 			if (UDisplayClusterDisplayDeviceBaseComponent* CachedComponent =
 				Cast<UDisplayClusterDisplayDeviceBaseComponent>(CachedDisplayDevice.GetComponent(GetOwner())))
@@ -465,7 +465,7 @@ void UDisplayClusterPreviewComponent::UpdatePreviewRenderTarget()
 			DestinationRenderTarget = SourceRenderTarget;
 
 			// Perform any render passes on the render target by the display device
-			if (RootActor->IsTechVizPreviewRenderingEnabled())
+			if (RootActor->IsTechvisPreviewRenderingEnabled())
 			{
 				if (UDisplayClusterDisplayDeviceBaseComponent* CachedComponent =
 				Cast<UDisplayClusterDisplayDeviceBaseComponent>(CachedDisplayDevice.GetComponent(GetOwner())))
diff --git a/Engine/Plugins/Runtime/nDisplay/Source/DisplayCluster/Private/Game/EngineClasses/Scene/DisplayClusterRootActorEditor.cpp b/Engine/Plugins/Runtime/nDisplay/Source/DisplayCluster/Private/Game/EngineClasses/Scene/DisplayClusterRootActorEditor.cpp
--- a/Engine/Plugins/Runtime/nDisplay/Source/DisplayCluster/Private/Game/EngineClasses/Scene/DisplayClusterRootActorEditor.cpp
+++ b/Engine/Plugins/Runtime/nDisplay/Source/DisplayCluster/Private/Game/EngineClasses/Scene/DisplayClusterRootActorEditor.cpp
@@ -76,7 +76,7 @@ void ADisplayClusterRootActor::ResetPreviewInternals_Editor()
 	ResetClusterNodePreviewRendering_Editor();
 }
 
-void ADisplayClusterRootActor::ConfigureTechViz_Editor()
+void ADisplayClusterRootActor::ConfigureTechvis_Editor()
 {
 	TArray<UDisplayClusterPreviewComponent*> Components;
 	GetComponents(Components);
@@ -85,7 +85,7 @@ void ADisplayClusterRootActor::ConfigureTechViz_Editor()
 	{
 		if (UMeshComponent* Mesh = PreviewComponent->GetPreviewMesh())
 		{
-			UE::DisplayClusterDisplayDeviceUtils::ConfigureTechVizForMesh(Mesh, bEnablePreviewTechViz);
+			UE::DisplayClusterDisplayDeviceUtils::ConfigureTechvisForMesh(Mesh, bEnablePreviewTechvis);
 		}
 	}
 }
@@ -882,9 +882,9 @@ void ADisplayClusterRootActor::PostEditChangeProperty(FPropertyChangedEvent& Pro
 		PreviewRenderFrame.Reset();
 		bReinitializeActor = false;
 	}
-	else if (PropertyName == GET_MEMBER_NAME_CHECKED(ADisplayClusterRootActor, bEnablePreviewTechViz))
+	else if (PropertyName == GET_MEMBER_NAME_CHECKED(ADisplayClusterRootActor, bEnablePreviewTechvis))
 	{
-		ConfigureTechViz_Editor();
+		ConfigureTechvis_Editor();
 	}
 	
 	if (bReinitializeActor)
diff --git a/Engine/Plugins/Runtime/nDisplay/Source/DisplayCluster/Private/Game/EngineClasses/Scene/DisplayDevice/DisplayClusterDisplayDeviceUtils.cpp b/Engine/Plugins/Runtime/nDisplay/Source/DisplayCluster/Private/Game/EngineClasses/Scene/DisplayDevice/DisplayClusterDisplayDeviceUtils.cpp
--- a/Engine/Plugins/Runtime/nDisplay/Source/DisplayCluster/Private/Game/EngineClasses/Scene/DisplayDevice/DisplayClusterDisplayDeviceUtils.cpp
+++ b/Engine/Plugins/Runtime/nDisplay/Source/DisplayCluster/Private/Game/EngineClasses/Scene/DisplayDevice/DisplayClusterDisplayDeviceUtils.cpp
@@ -96,12 +96,12 @@ UDisplayClusterDisplayDeviceBaseComponent* UE::DisplayClusterDisplayDeviceUtils:
 	return nullptr;
 }
 
-void UE::DisplayClusterDisplayDeviceUtils::ConfigureTechVizForMesh(UMeshComponent* InMeshComponent,
-	bool bIsTechVizEnabled)
+void UE::DisplayClusterDisplayDeviceUtils::ConfigureTechvisForMesh(UMeshComponent* InMeshComponent,
+	bool bIsTechvisEnabled)
 {
 	check(InMeshComponent);
 	
-	if (bIsTechVizEnabled)
+	if (bIsTechvisEnabled)
 	{
 		InMeshComponent->bAffectDynamicIndirectLighting = true;
 		InMeshComponent->bAffectIndirectLightingWhileHidden = true;
@@ -109,7 +109,7 @@ void UE::DisplayClusterDisplayDeviceUtils::ConfigureTechVizForMesh(UMeshComponen
 	else
 	{
 		// When disabling just revert to the archetype values. The original values at the time of enabling
-		// techviz aren't saved, and techviz is enabled by default so they'll always be set to true on new meshes.
+		// Techvis aren't saved, and Techvis is enabled by default so they'll always be set to true on new meshes.
 		if (const UMeshComponent* MeshArchetype = Cast<UMeshComponent>(InMeshComponent->GetArchetype()))
 		{
 			InMeshComponent->bAffectDynamicIndirectLighting = MeshArchetype->bAffectDynamicIndirectLighting;
diff --git a/Engine/Plugins/Runtime/nDisplay/Source/DisplayCluster/Private/Game/EngineClasses/Scene/DisplayDevice/DisplayClusterDisplayDeviceUtils.h b/Engine/Plugins/Runtime/nDisplay/Source/DisplayCluster/Private/Game/EngineClasses/Scene/DisplayDevice/DisplayClusterDisplayDeviceUtils.h
--- a/Engine/Plugins/Runtime/nDisplay/Source/DisplayCluster/Private/Game/EngineClasses/Scene/DisplayDevice/DisplayClusterDisplayDeviceUtils.h
+++ b/Engine/Plugins/Runtime/nDisplay/Source/DisplayCluster/Private/Game/EngineClasses/Scene/DisplayDevice/DisplayClusterDisplayDeviceUtils.h
@@ -21,11 +21,11 @@ namespace UE::DisplayClusterDisplayDeviceUtils
 	UDisplayClusterDisplayDeviceBaseComponent* FindAndSyncDisplayDeviceFromViewport(UDisplayClusterConfigurationViewport* InViewport, bool* bOutIsDefaultDevice = nullptr);
 
 	/**
-	 * Configure techviz parameters on a given mesh component.
+	 * Configure Techvis parameters on a given mesh component.
 	 * 
 	 * @param InMeshComponent The impacted mesh component.
-	 * @param bIsTechVizEnabled Whether techviz should be enabled or disabled on this mesh.
+	 * @param bIsTechvisEnabled Whether Techvis should be enabled or disabled on this mesh.
 	 */
-	void ConfigureTechVizForMesh(UMeshComponent* InMeshComponent, bool bIsTechVizEnabled);
+	void ConfigureTechvisForMesh(UMeshComponent* InMeshComponent, bool bIsTechvisEnabled);
 #endif
 }
diff --git a/Engine/Plugins/Runtime/nDisplay/Source/DisplayCluster/Public/DisplayClusterRootActor.h b/Engine/Plugins/Runtime/nDisplay/Source/DisplayCluster/Public/DisplayClusterRootActor.h
--- a/Engine/Plugins/Runtime/nDisplay/Source/DisplayCluster/Public/DisplayClusterRootActor.h
+++ b/Engine/Plugins/Runtime/nDisplay/Source/DisplayCluster/Public/DisplayClusterRootActor.h
@@ -403,9 +403,9 @@ public:
 	bool bPreviewEnablePostProcess = false;
 
 #if WITH_EDITORONLY_DATA
-	/** Configure the root actor for tech viz rendering with preview components. */
+	/** Configure the root actor for Techvis rendering with preview components. */
 	UPROPERTY(EditAnywhere, BlueprintReadWrite, Category = "Editor Preview")
-	bool bEnablePreviewTechViz = true;
+	bool bEnablePreviewTechvis = true;
 	
 	/** Freeze preview render.  This will impact editor performance. */
 	UPROPERTY(EditAnywhere, BlueprintReadWrite, Category = "Editor Preview", meta = (DisplayName = "Freeze Editor Preview"))
@@ -583,15 +583,15 @@ public:
 	/** Retrieve the default display device, creating it if it doesn't exist */
 	UDisplayClusterDisplayDeviceBaseComponent* GetDefaultDisplayDevice();
 
-	/** If the preview should use techviz rendering. */
-	bool IsTechVizPreviewRenderingEnabled() const { return bEnablePreviewTechViz; }
+	/** If the preview should use Techvis rendering. */
+	bool IsTechvisPreviewRenderingEnabled() const { return bEnablePreviewTechvis; }
 	
 protected:
 	FString GeneratePreviewComponentName_Editor(const FString& NodeId, const FString& ViewportId) const;
 	void ResetPreviewInternals_Editor();
 
-	/** Configure preview components for techviz rendering. */
-	void ConfigureTechViz_Editor();
+	/** Configure preview components for Techvis rendering. */
+	void ConfigureTechvis_Editor();
 
 	bool ImplUpdatePreviewConfiguration_Editor(const FString& InClusterNodeId);
 
