nDisplay: updated vioso dll and fixed rotation calculations

diff --git a/Engine/Plugins/Runtime/nDisplay/Source/DisplayClusterProjection/Private/Policy/VIOSO/Windows/DisplayClusterProjectionVIOSOWarper.cpp b/Engine/Plugins/Runtime/nDisplay/Source/DisplayClusterProjection/Private/Policy/VIOSO/Windows/DisplayClusterProjectionVIOSOWarper.cpp
--- a/Engine/Plugins/Runtime/nDisplay/Source/DisplayClusterProjection/Private/Policy/VIOSO/Windows/DisplayClusterProjectionVIOSOWarper.cpp
+++ b/Engine/Plugins/Runtime/nDisplay/Source/DisplayClusterProjection/Private/Policy/VIOSO/Windows/DisplayClusterProjectionVIOSOWarper.cpp
@@ -58,7 +58,7 @@ static FAutoConsoleVariableRef CVarDisplayClusterProjectionVIOSOPolicyPitchInver
 	ECVF_RenderThreadSafe
 );
 
-int32 GDisplayClusterProjectionVIOSOPolicyYawInverse = 1;
+int32 GDisplayClusterProjectionVIOSOPolicyYawInverse = 0;
 static FAutoConsoleVariableRef CVarDisplayClusterProjectionVIOSOPolicyYawInverse(
 	TEXT("nDisplay.render.projection.VIOSO.Yaw.Inverse"),
 	GDisplayClusterProjectionVIOSOPolicyYawInverse,
@@ -151,17 +151,6 @@ namespace UE::DisplayClusterProjection::VIOSOWarper
 
 		return EAxis::X;
 	}
-
-	static inline FRotator FromViosoEuler(const FVector3f& InEuler)
-	{
-		const FVector InEulerDegree = FMath::RadiansToDegrees(FVector(InEuler));
-
-		float Pitch = InEulerDegree.GetComponentForAxis(GetViosoAxis(GDisplayClusterProjectionVIOSOPolicyPitchAxis)) * (GDisplayClusterProjectionVIOSOPolicyPitchInverse ? -1 : 1);
-		float Yaw   = InEulerDegree.GetComponentForAxis(GetViosoAxis(GDisplayClusterProjectionVIOSOPolicyYawAxis))   * (GDisplayClusterProjectionVIOSOPolicyYawInverse ? -1 : 1);
-		float Roll  = InEulerDegree.GetComponentForAxis(GetViosoAxis(GDisplayClusterProjectionVIOSOPolicyRollAxis))  * (GDisplayClusterProjectionVIOSOPolicyRollInverse ? -1 : 1);
-
-		return FRotator(Pitch, Yaw, Roll);
-	}
 };
 using namespace UE::DisplayClusterProjection;
 
@@ -174,24 +163,26 @@ bool FDisplayClusterProjectionVIOSOWarper::CalculateViewProjection(IDisplayClust
 	FVector3f InOutEye = VIOSOWarper::ToVioso(InOutViewLocation, WorldToMeters);
 	FVector3f InOutRot = FVector3f::ZeroVector;
 
-	FVector3f OutPos = InOutEye;
-	FVector3f OutDir = InOutRot;
-
+	FVector3f OutPos = FVector3f::ZeroVector;
+	FVector3f OutDir = FVector3f::ZeroVector;
+	FRotator CamDir = FRotator::ZeroRotator;
+	
 	const bool bSymmetric = GDisplayClusterProjectionVIOSOPolicySymmetricFrustum != 0;
-
 	const FIntPoint ViewportSize = InViewport->GetContexts()[InContextNum].ContextSize;
 	const float ViewportAspectRatio = GDisplayClusterProjectionVIOSOPolicyFrustumFitToViewport ? float(ViewportSize.X) / float(ViewportSize.Y) : 0;
 
-	// VWB_getPosDirClip(VWB_Warper* pWarper, VWB_float* pEye, VWB_float* pRot, VWB_float* pPos, VWB_float* pDir, VWB_float* pClip, bool symmetric, VWB_float aspect);
 	if (VWB_ERROR_NONE == VIOSOLibrary->VWB_getPosDirClip(pWarper, &InOutEye.X, &InOutRot.X, &OutPos.X, &OutDir.X, &VWB_ViewClip[0], bSymmetric, ViewportAspectRatio))
 	{
-		InOutViewRotation = VIOSOWarper::FromViosoEuler(OutDir);
+		//Direction values in euler degrees - directly from the Warper - Roll inverted
+		CamDir.Pitch	=	pWarper->dir[0];
+		CamDir.Yaw		=	pWarper->dir[1];
+		CamDir.Roll		= -	pWarper->dir[2];
 
+		InOutViewRotation = CamDir;
 		OutProjMatrix = GetProjMatrix(InViewport, InContextNum);
 
 		return true;
 	}
-
 	return false;
 }
 
@@ -199,7 +190,6 @@ FMatrix FDisplayClusterProjectionVIOSOWarper::GetProjMatrix(IDisplayClusterViewp
 {
 	FViewClip PrjClip = ViewClip;
 
-	// Todo: unknown math comes from VIOSO, any value can be less than zero, what its mean??
 	PrjClip.Left   = -ViewClip.Left;
 	PrjClip.Right  = PrjClip.Left + (ViewClip.Left + ViewClip.Right);
 	PrjClip.Bottom = -ViewClip.Bottom;
