nDisplay: Fixed icvfx camera activation logic when it has media i/o assigned.

diff --git a/Engine/Plugins/Runtime/nDisplay/Source/DisplayClusterConfiguration/Private/DisplayClusterConfigurationTypes_ICVFX.cpp b/Engine/Plugins/Runtime/nDisplay/Source/DisplayClusterConfiguration/Private/DisplayClusterConfigurationTypes_ICVFX.cpp
--- a/Engine/Plugins/Runtime/nDisplay/Source/DisplayClusterConfiguration/Private/DisplayClusterConfigurationTypes_ICVFX.cpp
+++ b/Engine/Plugins/Runtime/nDisplay/Source/DisplayClusterConfiguration/Private/DisplayClusterConfigurationTypes_ICVFX.cpp
@@ -171,35 +171,36 @@ FDisplayClusterConfigurationICVFX_CameraSettings::FDisplayClusterConfigurationIC
 
 bool FDisplayClusterConfigurationICVFX_CameraSettings::IsICVFXEnabled(const UDisplayClusterConfigurationData& InConfigurationData, const FString& InClusterNodeId) const
 {
-	if (bEnable)
+	// When rendering offscreen, we have an extended logic for camera rendering activation
+	static const bool bIsRunningClusterModeOffscreen =
+		(IDisplayCluster::Get().GetOperationMode() == EDisplayClusterOperationMode::Cluster) &&
+		FParse::Param(FCommandLine::Get(), TEXT("RenderOffscreen"));
+
+	if (!bIsRunningClusterModeOffscreen)
 	{
-		// When rendering offscreen, we have an extended logic for camera rendering activation
-		static const bool bIsRunningClusterModeOffscreen =
-			(IDisplayCluster::Get().GetOperationMode() == EDisplayClusterOperationMode::Cluster) &&
-			FParse::Param(FCommandLine::Get(), TEXT("RenderOffscreen"));
+		return bEnable;
+	}
 
-		// If cluster mode + rendering offscreen, discover media output settings
-		if (bIsRunningClusterModeOffscreen)
-		{
-			// First condition to render offscreen: it has media output assigned
-			const bool bUsesMediaOutput = (RenderSettings.Media.bEnable && RenderSettings.Media.IsMediaOutputAssigned(InClusterNodeId));
+	if (!bEnable)
+	{
+		return false;
+	}
 
-			// Get backbuffer media settings
-			const UDisplayClusterConfigurationClusterNode* const NodeCfg = InConfigurationData.Cluster->GetNode(InClusterNodeId);
-			const FDisplayClusterConfigurationMedia* BackbufferMediaSettings = NodeCfg ? &NodeCfg->Media : nullptr;
+	// If cluster mode + rendering offscreen, discover media output settings
 
-			// Second condition to render offscreen: the backbuffer has media output assigned.
-			// This means the whole frame including ICVFX cameras need to be rendered.
-			const bool bIsBackbufferBeingCaptured = BackbufferMediaSettings ? BackbufferMediaSettings->bEnable && BackbufferMediaSettings->IsMediaOutputAssigned() : false;
+	// First condition to render offscreen: it has media output assigned
+	const bool bUsesMediaOutput = (RenderSettings.Media.bEnable && RenderSettings.Media.IsMediaOutputAssigned(InClusterNodeId));
 
-			// Finally make a decision if the camera should be rendered
-			return bUsesMediaOutput && bIsBackbufferBeingCaptured;
-		}
+	// Get backbuffer media settings
+	const UDisplayClusterConfigurationClusterNode* const NodeCfg = InConfigurationData.Cluster->GetNode(InClusterNodeId);
+	const FDisplayClusterConfigurationMedia* BackbufferMediaSettings = NodeCfg ? &NodeCfg->Media : nullptr;
 
-		return true;
-	}
+	// Second condition to render offscreen: the backbuffer has media output assigned.
+	// This means the whole frame including ICVFX cameras need to be rendered.
+	const bool bIsBackbufferBeingCaptured = BackbufferMediaSettings ? BackbufferMediaSettings->bEnable && BackbufferMediaSettings->IsMediaOutputAssigned() : false;
 
-	return false;
+	// Finally make a decision if the camera should be rendered
+	return bUsesMediaOutput || bIsBackbufferBeingCaptured;
 }
 
 const FOpenColorIOColorConversionSettings* FDisplayClusterConfigurationICVFX_CameraSettings::FindInnerFrustumOCIOConfiguration(const FString& InClusterNodeId) const
