nDisplay: Fix possible crash when adding an unloaded component to an nDisplay config.

diff --git a/Engine/Plugins/Runtime/nDisplay/Source/DisplayClusterConfigurator/Private/Views/SCSEditor/SDisplayClusterConfiguratorComponentCombo.cpp b/Engine/Plugins/Runtime/nDisplay/Source/DisplayClusterConfigurator/Private/Views/SCSEditor/SDisplayClusterConfiguratorComponentCombo.cpp
--- a/Engine/Plugins/Runtime/nDisplay/Source/DisplayClusterConfigurator/Private/Views/SCSEditor/SDisplayClusterConfiguratorComponentCombo.cpp
+++ b/Engine/Plugins/Runtime/nDisplay/Source/DisplayClusterConfigurator/Private/Views/SCSEditor/SDisplayClusterConfiguratorComponentCombo.cpp
@@ -240,12 +240,23 @@ void SDisplayClusterConfiguratorComponentClassCombo::OnAddComponentSelectionChan
 				if (ComponentClass == nullptr)
 				{
 					// The class is not loaded yet, so load it:
-					const ELoadFlags LoadFlags = LOAD_None;
-					UBlueprint* LoadedObject = LoadObject<UBlueprint>(nullptr, *InItem->GetComponentPath(), nullptr, LoadFlags, nullptr);
-					ComponentClass = GetAuthoritativeBlueprintClass(LoadedObject);
+					if (UObject* LoadedObject = LoadObject<UObject>(nullptr, *InItem->GetComponentPath()))
+					{
+						if (UClass* LoadedClass = Cast<UClass>(LoadedObject))
+						{
+							ComponentClass = LoadedClass;
+						}
+						else if (const UBlueprint* LoadedBP = Cast<UBlueprint>(LoadedObject))
+						{
+							ComponentClass = GetAuthoritativeBlueprintClass(LoadedBP);
+						}
+					}
 				}
 				
-				FSubobjectDataHandle NewActorCompHandle = OnSubobjectClassSelected.Execute(ComponentClass, InItem->GetComponentCreateAction(), InItem->GetAssetOverride());
+				const FSubobjectDataHandle NewActorCompHandle =
+					OnSubobjectClassSelected.IsBound() ?
+					OnSubobjectClassSelected.Execute(ComponentClass, InItem->GetComponentCreateAction(), InItem->GetAssetOverride())
+					: FSubobjectDataHandle::InvalidHandle;
 				if(NewActorCompHandle.IsValid())
 				{
 					InItem->GetOnSubobjectCreated().ExecuteIfBound(NewActorCompHandle);
