Removed PSO cache flush from command list flush method and manually placed it to happen after each scene render.

diff --git a/Engine/Plugins/Runtime/nDisplay/Source/DisplayCluster/Private/Game/EngineClasses/Basics/DisplayClusterViewportClient.cpp b/Engine/Plugins/Runtime/nDisplay/Source/DisplayCluster/Private/Game/EngineClasses/Basics/DisplayClusterViewportClient.cpp
--- a/Engine/Plugins/Runtime/nDisplay/Source/DisplayCluster/Private/Game/EngineClasses/Basics/DisplayClusterViewportClient.cpp
+++ b/Engine/Plugins/Runtime/nDisplay/Source/DisplayCluster/Private/Game/EngineClasses/Basics/DisplayClusterViewportClient.cpp
@@ -894,13 +894,6 @@ void UDisplayClusterViewportClient::Draw(FViewport* InViewport, FCanvas* SceneCa
 	{
 		// Or if none to render, do logic for when rendering is skipped
 		GetRendererModule().PerFrameCleanupIfSkipRenderer();
-
-		// Make sure RHI resources get flushed if we're not using a renderer
-		ENQUEUE_RENDER_COMMAND(UDisplayClusterViewportClient_FlushRHIResources)(
-			[](FRHICommandListImmediate& RHICmdList)
-			{
-				RHICmdList.ImmediateFlush(EImmediateFlushType::FlushRHIThreadFlushResources);
-			});
 	}
 
 	// Handle special viewports game-thread logic at frame end
diff --git a/Engine/Source/Runtime/Engine/Private/GameEngine.cpp b/Engine/Source/Runtime/Engine/Private/GameEngine.cpp
--- a/Engine/Source/Runtime/Engine/Private/GameEngine.cpp
+++ b/Engine/Source/Runtime/Engine/Private/GameEngine.cpp
@@ -1948,11 +1948,6 @@ void UGameEngine::Tick( float DeltaSeconds, bool bIdleMode )
 		if (bRenderingSuspended)
 		{
 			GetRendererModule().PerFrameCleanupIfSkipRenderer();
-			ENQUEUE_RENDER_COMMAND(UGameEngine_Tick_FlushRHIResources)(
-				[](FRHICommandListImmediate& RHICmdList)
-				{
-					RHICmdList.ImmediateFlush(EImmediateFlushType::FlushRHIThreadFlushResources);
-				});
 		}
 	}
 
diff --git a/Engine/Source/Runtime/Engine/Private/GameViewportClient.cpp b/Engine/Source/Runtime/Engine/Private/GameViewportClient.cpp
--- a/Engine/Source/Runtime/Engine/Private/GameViewportClient.cpp
+++ b/Engine/Source/Runtime/Engine/Private/GameViewportClient.cpp
@@ -1795,13 +1795,6 @@ void UGameViewportClient::Draw(FViewport* InViewport, FCanvas* SceneCanvas)
 	else
 	{
 		GetRendererModule().PerFrameCleanupIfSkipRenderer();
-
-		// Make sure RHI resources get flushed if we're not using a renderer
-		ENQUEUE_RENDER_COMMAND(UGameViewportClient_FlushRHIResources)(
-			[](FRHICommandListImmediate& RHICmdList)
-			{
-				RHICmdList.ImmediateFlush(EImmediateFlushType::FlushRHIThreadFlushResources);
-			});
 	}
 
 	// Beyond this point, only UI rendering independent from dynamc resolution.
diff --git a/Engine/Source/Runtime/RHI/Public/RHICommandList.inl b/Engine/Source/Runtime/RHI/Public/RHICommandList.inl
--- a/Engine/Source/Runtime/RHI/Public/RHICommandList.inl
+++ b/Engine/Source/Runtime/RHI/Public/RHICommandList.inl
@@ -95,7 +95,6 @@ FORCEINLINE_DEBUGGABLE void FRHICommandListImmediate::ImmediateFlush(EImmediateF
 		if (FlushType >= EImmediateFlushType::FlushRHIThreadFlushResources)
 		{
 			CSV_SCOPED_TIMING_STAT(RHITFlushes, FlushRHIThreadFlushResourcesTotal);
-			PipelineStateCache::FlushResources();
 
 			// RHIPerFrameRHIFlushComplete was originally called from FlushPendingDeletes, which used to be here, so it was
 			// running as part of the next command list. FlushPendingDeletes now runs as part of the command list being finalized
diff --git a/Engine/Source/Runtime/RenderCore/Private/RenderGraphBuilder.cpp b/Engine/Source/Runtime/RenderCore/Private/RenderGraphBuilder.cpp
--- a/Engine/Source/Runtime/RenderCore/Private/RenderGraphBuilder.cpp
+++ b/Engine/Source/Runtime/RenderCore/Private/RenderGraphBuilder.cpp
@@ -385,6 +385,7 @@ void FRDGBuilder::EndFlushResourcesRHI()
 	CSV_SCOPED_TIMING_STAT_EXCLUSIVE(STAT_RDG_FlushResourcesRHI);
 	SCOPED_NAMED_EVENT(EndFlushResourcesRHI, FColor::Emerald);
 	RHICmdList.ImmediateFlush(EImmediateFlushType::FlushRHIThreadFlushResources);
+	PipelineStateCache::FlushResources();
 }
 
 void FRDGBuilder::TickPoolElements()
diff --git a/Engine/Source/Runtime/RenderCore/Private/RenderingThread.cpp b/Engine/Source/Runtime/RenderCore/Private/RenderingThread.cpp
--- a/Engine/Source/Runtime/RenderCore/Private/RenderingThread.cpp
+++ b/Engine/Source/Runtime/RenderCore/Private/RenderingThread.cpp
@@ -1698,7 +1698,7 @@ inline ERenderCommandPipeMode GetValidatedRenderCommandPipeMode(int32 CVarValue)
 ERenderCommandPipeMode GRenderCommandPipeMode = ERenderCommandPipeMode::None;
 FAutoConsoleVariable CVarRenderCommandPipeMode(
 	TEXT("r.RenderCommandPipeMode"),
-	1,
+	2,
 	TEXT("Controls behavior of the main render thread command pipe.")
 	TEXT(" 0: Render commands are launched individually as tasks;\n")
 	TEXT(" 1: Render commands are enqueued into a render command pipe for the render thread only.;\n")
diff --git a/Engine/Source/Runtime/Renderer/Private/SceneRendering.cpp b/Engine/Source/Runtime/Renderer/Private/SceneRendering.cpp
--- a/Engine/Source/Runtime/Renderer/Private/SceneRendering.cpp
+++ b/Engine/Source/Runtime/Renderer/Private/SceneRendering.cpp
@@ -4930,10 +4930,14 @@ void FRendererModule::PostRenderAllViewports()
 
 void FRendererModule::PerFrameCleanupIfSkipRenderer()
 {
+	UE::RenderCommandPipe::FSyncScope SyncScope;
+
 	// Some systems (e.g. Slate) can still draw (via FRendererModule::DrawTileMesh for example) when scene renderer is not used
 	ENQUEUE_RENDER_COMMAND(CmdPerFrameCleanupIfSkipRenderer)(
 		[](FRHICommandListImmediate& RHICmdList)
 	{
+		RHICmdList.ImmediateFlush(EImmediateFlushType::FlushRHIThreadFlushResources);
+		PipelineStateCache::FlushResources();
 		FSceneRenderer::CleanUp(RHICmdList);
 		GPrimitiveIdVertexBufferPool.DiscardAll();
 	});
