nDisplay: Fix ICVFX Panel's Place Actor / Paste Actor Here not transacting to MU. Fix possible crash when connecting to an MU session and then adding an actor.

diff --git a/Engine/Plugins/Runtime/nDisplay/Source/DisplayClusterLightCardEditor/Private/DisplayClusterLightCardEditor.cpp b/Engine/Plugins/Runtime/nDisplay/Source/DisplayClusterLightCardEditor/Private/DisplayClusterLightCardEditor.cpp
--- a/Engine/Plugins/Runtime/nDisplay/Source/DisplayClusterLightCardEditor/Private/DisplayClusterLightCardEditor.cpp
+++ b/Engine/Plugins/Runtime/nDisplay/Source/DisplayClusterLightCardEditor/Private/DisplayClusterLightCardEditor.cpp
@@ -271,7 +271,7 @@ AActor* FDisplayClusterLightCardEditor::SpawnActor(const UDisplayClusterLightCar
 	return SpawnActor(nullptr, NAME_None, InTemplate, InLevel, bIsPreview);
 }
 
-void FDisplayClusterLightCardEditor::AddNewLightCard()
+ADisplayClusterLightCardActor* FDisplayClusterLightCardEditor::AddNewLightCard()
 {
 	check(ActiveRootActor.IsValid());
 
@@ -294,7 +294,11 @@ void FDisplayClusterLightCardEditor::AddNewLightCard()
 		
 		AddRecentlyPlacedItem(MoveTemp(RecentlyPlacedItem));
 		SelectActors({NewLightCard});
+
+		return NewLightCard;
 	}
+
+	return nullptr;
 }
 
 void FDisplayClusterLightCardEditor::AddExistingLightCard()
@@ -405,7 +409,7 @@ void FDisplayClusterLightCardEditor::AddExistingLightCard()
 	SelectedActorPtr.Reset();
 }
 
-void FDisplayClusterLightCardEditor::AddNewFlag()
+ADisplayClusterLightCardActor* FDisplayClusterLightCardEditor::AddNewFlag()
 {
 	check(ActiveRootActor.IsValid());
 
@@ -432,10 +436,13 @@ void FDisplayClusterLightCardEditor::AddNewFlag()
 		
 		AddRecentlyPlacedItem(MoveTemp(RecentlyPlacedItem));
 		SelectActors({NewLightCard});
+		return NewLightCard;
 	}
+
+	return nullptr;
 }
 
-void FDisplayClusterLightCardEditor::AddNewDynamic(UClass* InClass)
+AActor* FDisplayClusterLightCardEditor::AddNewDynamic(UClass* InClass)
 {
 	check(ActiveRootActor.IsValid());
 
@@ -451,7 +458,10 @@ void FDisplayClusterLightCardEditor::AddNewDynamic(UClass* InClass)
 		
 		AddRecentlyPlacedItem(MoveTemp(RecentlyPlacedItem));
 		SelectActors({NewActor});
+		return NewActor;
 	}
+
+	return nullptr;
 }
 
 void FDisplayClusterLightCardEditor::AddLightCardsToActor(const TArray<ADisplayClusterLightCardActor*>& LightCards)
@@ -580,14 +590,14 @@ bool FDisplayClusterLightCardEditor::CanCopySelectedActors() const
 	return SelectedActors.Num() > 0;
 }
 
-void FDisplayClusterLightCardEditor::PasteActors()
+TArray<AActor*> FDisplayClusterLightCardEditor::PasteActors()
 {
 	if (DoOutlinerFoldersNeedEditorDelegates() && LightCardOutliner->CanPasteSelectedFolder())
 	{
 		// We still must fire the delegates in case folders are selected
 		FEditorDelegates::OnEditPasteActorsBegin.Broadcast();
 		FEditorDelegates::OnEditPasteActorsEnd.Broadcast();
-		return;
+		return {};
 	}
 
 	UWorld* EditorWorld = GEditor->GetEditorWorldContext().World();
@@ -618,6 +628,8 @@ void FDisplayClusterLightCardEditor::PasteActors()
 
 	AddLightCardsToActor(PastedLightCards);
 	SelectActors(PastedActors);
+
+	return PastedActors;
 }
 
 bool FDisplayClusterLightCardEditor::CanPasteActors() const
@@ -1147,8 +1159,11 @@ TSharedRef<SWidget> FDisplayClusterLightCardEditor::GeneratePlaceActorsMenu()
 			FText Label = Class->GetDisplayNameText();
 			FSlateIcon StageActorIcon = FSlateIconFinder::FindIconForClass(Class);
 			MenuBuilder.AddMenuEntry(Label, LOCTEXT("AddStageActorHeader", "Add a stage actor to the scene"), StageActorIcon,
-				FUIAction(FExecuteAction::CreateSP(this, &FDisplayClusterLightCardEditor::AddNewDynamic, Class),
-					FCanExecuteAction::CreateSP(this, &FDisplayClusterLightCardEditor::CanAddNewActor, Class)));
+				FUIAction(FExecuteAction::CreateSPLambda(this, [this, Class]()
+				{
+					AddNewDynamic(Class);
+				}),
+				FCanExecuteAction::CreateSP(this, &FDisplayClusterLightCardEditor::CanAddNewActor, Class)));
 		}
 
 		if (CanAddNewActor())
@@ -1541,12 +1556,18 @@ void FDisplayClusterLightCardEditor::BindCommands()
 
 	CommandList->MapAction(
 		FDisplayClusterLightCardEditorCommands::Get().AddNewLightCard,
-		FExecuteAction::CreateSP(this, &FDisplayClusterLightCardEditor::AddNewLightCard),
+		FExecuteAction::CreateSPLambda(this, [this]()
+		{
+			AddNewLightCard();
+		}),
 		FCanExecuteAction::CreateSP(this, &FDisplayClusterLightCardEditor::CanAddNewActor));
 
 	CommandList->MapAction(
 		FDisplayClusterLightCardEditorCommands::Get().AddNewFlag,
-		FExecuteAction::CreateSP(this, &FDisplayClusterLightCardEditor::AddNewFlag),
+		FExecuteAction::CreateSPLambda(this, [this]()
+		{
+			AddNewFlag();
+		}),
 		FCanExecuteAction::CreateSP(this, &FDisplayClusterLightCardEditor::CanAddNewActor));
 	
 	CommandList->MapAction(
@@ -1571,7 +1592,10 @@ void FDisplayClusterLightCardEditor::BindCommands()
 
 	CommandList->MapAction(
 		FGenericCommands::Get().Paste,
-		FExecuteAction::CreateSP(this, &FDisplayClusterLightCardEditor::PasteActors),
+		FExecuteAction::CreateSPLambda(this, [this]()
+		{
+			PasteActors();
+		}),
 		FCanExecuteAction::CreateSP(this, &FDisplayClusterLightCardEditor::CanPasteActors));
 
 	CommandList->MapAction(
diff --git a/Engine/Plugins/Runtime/nDisplay/Source/DisplayClusterLightCardEditor/Private/DisplayClusterLightCardEditor.h b/Engine/Plugins/Runtime/nDisplay/Source/DisplayClusterLightCardEditor/Private/DisplayClusterLightCardEditor.h
--- a/Engine/Plugins/Runtime/nDisplay/Source/DisplayClusterLightCardEditor/Private/DisplayClusterLightCardEditor.h
+++ b/Engine/Plugins/Runtime/nDisplay/Source/DisplayClusterLightCardEditor/Private/DisplayClusterLightCardEditor.h
@@ -103,16 +103,16 @@ public:
 	}
 	
 	/** Adds a new light card to the root actor and centers it in the viewport */
-	void AddNewLightCard();
+	ADisplayClusterLightCardActor* AddNewLightCard();
 
 	/** Select an existing Light Card from a menu */
 	void AddExistingLightCard();
 
 	/** Adds a new light card configured as a flag */
-	void AddNewFlag();
+	ADisplayClusterLightCardActor* AddNewFlag();
 
 	/** Add a new actor dynamically based on a class */
-	void AddNewDynamic(UClass* InClass);
+	AActor* AddNewDynamic(UClass* InClass);
 
 	/** Adds the given Light Card to the root actor */
 	void AddLightCardsToActor(const TArray<ADisplayClusterLightCardActor*>& LightCards);
@@ -134,7 +134,7 @@ public:
 	bool CanCopySelectedActors() const;
 
 	/** Pastes any actors in the clipboard to the current root actor */
-	void PasteActors();
+	TArray<AActor*> PasteActors();
 
 	/** Determines if there are any actors that can be pasted from the clipboard */
 	bool CanPasteActors() const;
diff --git a/Engine/Plugins/Runtime/nDisplay/Source/DisplayClusterLightCardEditor/Private/Viewport/DisplayClusterLightCardEditorViewportClient.cpp b/Engine/Plugins/Runtime/nDisplay/Source/DisplayClusterLightCardEditor/Private/Viewport/DisplayClusterLightCardEditorViewportClient.cpp
--- a/Engine/Plugins/Runtime/nDisplay/Source/DisplayClusterLightCardEditor/Private/Viewport/DisplayClusterLightCardEditorViewportClient.cpp
+++ b/Engine/Plugins/Runtime/nDisplay/Source/DisplayClusterLightCardEditor/Private/Viewport/DisplayClusterLightCardEditorViewportClient.cpp
@@ -2164,6 +2164,31 @@ void FDisplayClusterLightCardEditorViewportClient::MoveSelectedActorsToPixel(con
 	MoveActorsToPixel(PixelPos, SelectedActors);
 }
 
+void FDisplayClusterLightCardEditorViewportClient::MoveActorsToPixel(const FIntPoint& PixelPos, const TArray<FDisplayClusterWeakStageActorPtr>& InActors)
+{
+	FSceneViewFamilyContext ViewFamily(FSceneViewFamily::ConstructionValues(
+		Viewport,
+		GetScene(),
+		EngineShowFlags)
+		.SetRealtimeUpdate(IsRealtime())
+	);
+
+	FSceneView* View = CalcSceneView(&ViewFamily);
+
+	ProjectionHelper->MoveActorsToPixel(InActors, PixelPos, *View);
+
+	// Update each light card with the delta coordinates; the flush constraint is applied by MoveActorTo, ensuring the light card is always flush to screens
+	for (const FDisplayClusterWeakStageActorPtr& LightCard : InActors)
+	{
+		if (LightCard.IsValid() &&
+			((LightCard->IsUVActor() && ProjectionMode == EDisplayClusterMeshProjectionType::UV) ||
+			(!LightCard->IsUVActor() && ProjectionMode != EDisplayClusterMeshProjectionType::UV)))
+		{
+			PropagateActorTransform(LightCard);
+		}
+	}
+}
+
 void FDisplayClusterLightCardEditorViewportClient::BeginTransaction(const FText& Description)
 {
 	GEditor->BeginTransaction(Description);
@@ -2547,31 +2572,6 @@ void FDisplayClusterLightCardEditorViewportClient::MoveSelectedActors(FViewport*
 	}
 }
 
-void FDisplayClusterLightCardEditorViewportClient::MoveActorsToPixel(const FIntPoint& PixelPos, const TArray<FDisplayClusterWeakStageActorPtr>& InActors)
-{
-	FSceneViewFamilyContext ViewFamily(FSceneViewFamily::ConstructionValues(
-		Viewport,
-		GetScene(),
-		EngineShowFlags)
-		.SetRealtimeUpdate(IsRealtime())
-	);
-
-	FSceneView* View = CalcSceneView(&ViewFamily);
-
-	ProjectionHelper->MoveActorsToPixel(SelectedActors, PixelPos, *View);
-
-	// Update each light card with the delta coordinates; the flush constraint is applied by MoveActorTo, ensuring the light card is always flush to screens
-	for (const FDisplayClusterWeakStageActorPtr& LightCard : InActors)
-	{
-		if (LightCard.IsValid() &&
-			((LightCard->IsUVActor() && ProjectionMode == EDisplayClusterMeshProjectionType::UV) ||
-			(!LightCard->IsUVActor() && ProjectionMode != EDisplayClusterMeshProjectionType::UV)))
-		{
-			PropagateActorTransform(LightCard);
-		}
-	}
-}
-
 void FDisplayClusterLightCardEditorViewportClient::MoveSelectedUVActors(FViewport* InViewport, EAxisList::Type CurrentAxis)
 {
 	FSceneViewFamilyContext ViewFamily(FSceneViewFamily::ConstructionValues(
diff --git a/Engine/Plugins/Runtime/nDisplay/Source/DisplayClusterLightCardEditor/Private/Viewport/DisplayClusterLightcardEditorViewport.cpp b/Engine/Plugins/Runtime/nDisplay/Source/DisplayClusterLightCardEditor/Private/Viewport/DisplayClusterLightcardEditorViewport.cpp
--- a/Engine/Plugins/Runtime/nDisplay/Source/DisplayClusterLightCardEditor/Private/Viewport/DisplayClusterLightcardEditorViewport.cpp
+++ b/Engine/Plugins/Runtime/nDisplay/Source/DisplayClusterLightCardEditor/Private/Viewport/DisplayClusterLightcardEditorViewport.cpp
@@ -14,6 +14,7 @@
 
 #include "DisplayClusterRootActor.h"
 
+#include "Algo/Transform.h"
 #include "EditorViewportCommands.h"
 #include "ScopedTransaction.h"
 #include "SEditorViewportToolBarButton.h"
@@ -914,15 +915,30 @@ void SDisplayClusterLightCardEditorViewport::PasteLightCardsHere()
 {
 	if (LightCardEditorPtr.IsValid())
 	{
-		LightCardEditorPtr.Pin()->PasteActors();
+		const TSharedPtr<FScopedTransaction> Transaction = MakeShared<FScopedTransaction>(LOCTEXT("PasteActorsHereTransactionMessage", "Paste Actors Here"));
+		
+		TArray<AActor*> PastedActors = LightCardEditorPtr.Pin()->PasteActors();
 
-		if (ViewportClient.IsValid())
+		if (ViewportClient.IsValid() && PastedActors.Num() > 0)
 		{
 			// Perform the positioning of the pasted light cards on the next scene refresh, since the light card proxies will have not been 
 			// regenerated until then
-			ViewportClient->GetOnNextSceneRefresh().AddLambda([this]()
+
+			AddActorHereTransactions.Add(Transaction);
+			ViewportClient->GetOnNextSceneRefresh().AddLambda([this, Transaction, PastedActors=MoveTemp(PastedActors)]()
 			{
-				ViewportClient->MoveSelectedActorsToPixel(FIntPoint(PasteHerePos.X, PasteHerePos.Y));
+				TArray<FDisplayClusterWeakStageActorPtr> PastedStageActors;
+				Algo::TransformIf(PastedActors, PastedStageActors, [](const AActor* InItem)
+				{
+					return IsValid(InItem);
+				},
+				[](const AActor* InItem) -> FDisplayClusterWeakStageActorPtr
+				{
+					return InItem;
+				});
+				
+				ViewportClient->MoveActorsToPixel(FIntPoint(PasteHerePos.X, PasteHerePos.Y), PastedStageActors);
+				AddActorHereTransactions.Remove(Transaction);
 			});
 		}
 	}
@@ -942,15 +958,19 @@ void SDisplayClusterLightCardEditorViewport::AddLightCardHere()
 {
 	if (LightCardEditorPtr.IsValid())
 	{
-		LightCardEditorPtr.Pin()->AddNewLightCard();
+		const TSharedPtr<FScopedTransaction> Transaction = MakeShared<FScopedTransaction>(LOCTEXT("AddLightCardHereTransactionMessage", "Add Light Card Here"));
+		
+		ADisplayClusterLightCardActor* NewLightCard = LightCardEditorPtr.Pin()->AddNewLightCard();
 
-		if (ViewportClient.IsValid())
+		if (ViewportClient.IsValid() && NewLightCard)
 		{
 			// Perform the positioning of the pasted light cards on the next scene refresh, since the light card proxies will have not been 
 			// regenerated until then
-			ViewportClient->GetOnNextSceneRefresh().AddLambda([this]()
+			AddActorHereTransactions.Add(Transaction);
+			ViewportClient->GetOnNextSceneRefresh().AddLambda([this, Transaction, NewLightCard]()
 			{
-				ViewportClient->MoveSelectedActorsToPixel(FIntPoint(PasteHerePos.X, PasteHerePos.Y));
+				ViewportClient->MoveActorsToPixel(FIntPoint(PasteHerePos.X, PasteHerePos.Y), { NewLightCard });
+				AddActorHereTransactions.Remove(Transaction);
 			});
 		}
 	}
@@ -960,15 +980,19 @@ void SDisplayClusterLightCardEditorViewport::AddFlagHere()
 {
 	if (LightCardEditorPtr.IsValid())
 	{
-		LightCardEditorPtr.Pin()->AddNewFlag();
+		const TSharedPtr<FScopedTransaction> Transaction = MakeShared<FScopedTransaction>(LOCTEXT("AddFlagHereTransactionMessage", "Add Flag Here"));
+		
+		ADisplayClusterLightCardActor* NewFlag = LightCardEditorPtr.Pin()->AddNewFlag();
 
-		if (ViewportClient.IsValid())
+		if (ViewportClient.IsValid() && NewFlag)
 		{
 			// Perform the positioning of the pasted light cards on the next scene refresh, since the light card proxies will have not been 
 			// regenerated until then
-			ViewportClient->GetOnNextSceneRefresh().AddLambda([this]()
+			AddActorHereTransactions.Add(Transaction);
+			ViewportClient->GetOnNextSceneRefresh().AddLambda([this, Transaction, NewFlag]()
 			{
-				ViewportClient->MoveSelectedActorsToPixel(FIntPoint(PasteHerePos.X, PasteHerePos.Y));
+				ViewportClient->MoveActorsToPixel(FIntPoint(PasteHerePos.X, PasteHerePos.Y), { NewFlag });
+				AddActorHereTransactions.Remove(Transaction);
 			});
 		}
 	}
@@ -978,15 +1002,19 @@ void SDisplayClusterLightCardEditorViewport::AddStageActorHere(UClass* InClass)
 {
 	if (LightCardEditorPtr.IsValid())
 	{
-		LightCardEditorPtr.Pin()->AddNewDynamic(InClass);
+		const TSharedPtr<FScopedTransaction> Transaction = MakeShared<FScopedTransaction>(LOCTEXT("AddStageActorHereTransactionMessage", "Add Stage Actor Here"));
+		
+		AActor* NewActor = LightCardEditorPtr.Pin()->AddNewDynamic(InClass);
 
-		if (ViewportClient.IsValid())
+		if (ViewportClient.IsValid() && NewActor)
 		{
 			// Perform the positioning of the pasted light cards on the next scene refresh, since the light card proxies will have not been 
 			// regenerated until then
-			ViewportClient->GetOnNextSceneRefresh().AddLambda([this]()
+			AddActorHereTransactions.Add(Transaction);
+			ViewportClient->GetOnNextSceneRefresh().AddLambda([this, Transaction, NewActor]()
 			{
-				ViewportClient->MoveSelectedActorsToPixel(FIntPoint(PasteHerePos.X, PasteHerePos.Y));
+				ViewportClient->MoveActorsToPixel(FIntPoint(PasteHerePos.X, PasteHerePos.Y), { NewActor });
+				AddActorHereTransactions.Remove(Transaction);
 			});
 		}
 	}
diff --git a/Engine/Plugins/Runtime/nDisplay/Source/DisplayClusterLightCardEditor/Private/Viewport/DisplayClusterLightcardEditorViewport.h b/Engine/Plugins/Runtime/nDisplay/Source/DisplayClusterLightCardEditor/Private/Viewport/DisplayClusterLightcardEditorViewport.h
--- a/Engine/Plugins/Runtime/nDisplay/Source/DisplayClusterLightCardEditor/Private/Viewport/DisplayClusterLightcardEditorViewport.h
+++ b/Engine/Plugins/Runtime/nDisplay/Source/DisplayClusterLightCardEditor/Private/Viewport/DisplayClusterLightcardEditorViewport.h
@@ -13,6 +13,7 @@
 class FDisplayClusterLightCardEditor;
 class ADisplayClusterRootActor;
 class FDisplayClusterLightCardEditorViewportClient;
+class FScopedTransaction;
 
 /**
  * Slate widget which renders our view client.
@@ -132,6 +133,9 @@ private:
 private:
 	/** The preview scene to use by the viewport client to manage the preview 3D world */
 	TSharedPtr<FPreviewScene> PreviewScene;
+
+	/** Open transactions when using an 'Add Actor Here' method since they occur over multiple frames */
+	TSet<TSharedPtr<FScopedTransaction>> AddActorHereTransactions;
 	
 	/** Level viewport client */
 	TSharedPtr<FDisplayClusterLightCardEditorViewportClient> ViewportClient;
diff --git a/Engine/Plugins/Runtime/nDisplay/Source/DisplayClusterScenePreview/Private/DisplayClusterLightCardEditorHelper.cpp b/Engine/Plugins/Runtime/nDisplay/Source/DisplayClusterScenePreview/Private/DisplayClusterLightCardEditorHelper.cpp
--- a/Engine/Plugins/Runtime/nDisplay/Source/DisplayClusterScenePreview/Private/DisplayClusterLightCardEditorHelper.cpp
+++ b/Engine/Plugins/Runtime/nDisplay/Source/DisplayClusterScenePreview/Private/DisplayClusterLightCardEditorHelper.cpp
@@ -1456,8 +1456,9 @@ bool FDisplayClusterLightCardEditorHelper::UpdateNormalMaps()
 
 void FDisplayClusterLightCardEditorHelper::UpdateNormalMapMesh()
 {
-	if (!NormalMeshScene.IsValid())
+	if (!NormalMeshScene.IsValid() || !IsValid(NormalMeshScene->GetWorld()))
 	{
+		FWorldDelegates::OnWorldCleanup.RemoveAll(this);
 		FWorldDelegates::OnWorldCleanup.AddRaw(this, &FDisplayClusterLightCardEditorHelper::OnWorldCleanup);
 		NormalMeshScene = MakeShared<FPreviewScene>(FPreviewScene::ConstructionValues());
 	}
