[nDisplay] Renaming a local variable

diff --git a/Engine/Plugins/Runtime/nDisplay/Source/DisplayClusterMedia/Private/Input/DisplayClusterMediaInputViewport.cpp b/Engine/Plugins/Runtime/nDisplay/Source/DisplayClusterMedia/Private/Input/DisplayClusterMediaInputViewport.cpp
--- a/Engine/Plugins/Runtime/nDisplay/Source/DisplayClusterMedia/Private/Input/DisplayClusterMediaInputViewport.cpp
+++ b/Engine/Plugins/Runtime/nDisplay/Source/DisplayClusterMedia/Private/Input/DisplayClusterMediaInputViewport.cpp
@@ -43,16 +43,16 @@ void FDisplayClusterMediaInputViewport::PostCrossGpuTransfer_RenderThread(FRHICo
 {
 	checkSlow(ViewportManagerProxy);
 
-	if (const IDisplayClusterViewportProxy* const CaptureViewport = ViewportManagerProxy->FindViewport_RenderThread(GetViewportId()))
+	if (const IDisplayClusterViewportProxy* const PlaybackViewport = ViewportManagerProxy->FindViewport_RenderThread(GetViewportId()))
 	{
-		const bool bShouldImportMedia = !CaptureViewport->GetPostRenderSettings_RenderThread().Replace.IsEnabled();
+		const bool bShouldImportMedia = !PlaybackViewport->GetPostRenderSettings_RenderThread().Replace.IsEnabled();
 
 		if (bShouldImportMedia)
 		{
 			TArray<FRHITexture*> Textures;
 			TArray<FIntRect>     Regions;
 
-			if (CaptureViewport->GetResourcesWithRects_RenderThread(EDisplayClusterViewportResourceType::InternalRenderTargetResource, Textures, Regions))
+			if (PlaybackViewport->GetResourcesWithRects_RenderThread(EDisplayClusterViewportResourceType::InternalRenderTargetResource, Textures, Regions))
 			{
 				if (Textures.Num() > 0 && Regions.Num() > 0)
 				{
