[nDisplay] Late OCIO pass to allow custom per node OCIO profiles with media multicasting

diff --git a/Engine/Plugins/Runtime/nDisplay/Source/DisplayCluster/Private/Render/Viewport/Configuration/DisplayClusterViewportConfigurationHelpers.cpp b/Engine/Plugins/Runtime/nDisplay/Source/DisplayCluster/Private/Render/Viewport/Configuration/DisplayClusterViewportConfigurationHelpers.cpp
--- a/Engine/Plugins/Runtime/nDisplay/Source/DisplayCluster/Private/Render/Viewport/Configuration/DisplayClusterViewportConfigurationHelpers.cpp
+++ b/Engine/Plugins/Runtime/nDisplay/Source/DisplayCluster/Private/Render/Viewport/Configuration/DisplayClusterViewportConfigurationHelpers.cpp
@@ -107,11 +107,17 @@ void FDisplayClusterViewportConfigurationHelpers::UpdateBaseViewportSetting(FDis
 
 			if (MediaSettings.bEnable)
 			{
-				// Don't render the viewport if media input assigned
-				InOutRenderSettings.bSkipSceneRenderingButLeaveResourcesAvailable = MediaSettings.IsMediaInputAssigned();
+				const bool bMediaInputAssigned  = MediaSettings.IsMediaInputAssigned();
+				const bool bMediaOutputAssigned = MediaSettings.IsMediaOutputAssigned();
+
+				// Don't render this viewport if media input assigned
+				InOutRenderSettings.bSkipSceneRenderingButLeaveResourcesAvailable = bMediaInputAssigned;
 
 				// Mark this viewport is going to be captured by a capture device
-				InOutRenderSettings.bIsBeingCaptured = MediaSettings.IsMediaOutputAssigned();
+				InOutRenderSettings.bIsBeingCaptured = bMediaOutputAssigned;
+
+				// Late OCIO pass
+				InOutRenderSettings.bForceLateOCIOPass = (bMediaOutputAssigned || bMediaInputAssigned ? MediaSettings.bLateOCIOPass : false);
 			}
 		}
 	}
diff --git a/Engine/Plugins/Runtime/nDisplay/Source/DisplayCluster/Private/Render/Viewport/Configuration/DisplayClusterViewportConfigurationHelpers_ICVFX.cpp b/Engine/Plugins/Runtime/nDisplay/Source/DisplayCluster/Private/Render/Viewport/Configuration/DisplayClusterViewportConfigurationHelpers_ICVFX.cpp
--- a/Engine/Plugins/Runtime/nDisplay/Source/DisplayCluster/Private/Render/Viewport/Configuration/DisplayClusterViewportConfigurationHelpers_ICVFX.cpp
+++ b/Engine/Plugins/Runtime/nDisplay/Source/DisplayCluster/Private/Render/Viewport/Configuration/DisplayClusterViewportConfigurationHelpers_ICVFX.cpp
@@ -551,11 +551,17 @@ void FDisplayClusterViewportConfigurationHelpers_ICVFX::UpdateCameraViewportSett
 			{
 				const FString ThisClusterNodeId = DstViewport.GetClusterNodeId();
 
+				const bool bMediaInputAssigned  = MediaICVFXSettings.IsMediaInputAssigned(ThisClusterNodeId);
+				const bool bMediaOutputAssigned = MediaICVFXSettings.IsMediaOutputAssigned(ThisClusterNodeId);
+
 				// Don't render the viewport if media input assigned
-				InOutRenderSettings.bSkipSceneRenderingButLeaveResourcesAvailable = MediaICVFXSettings.IsMediaInputAssigned(ThisClusterNodeId);
+				InOutRenderSettings.bSkipSceneRenderingButLeaveResourcesAvailable = bMediaInputAssigned;
 
 				// Mark this viewport is going to be captured by a capture device
-				InOutRenderSettings.bIsBeingCaptured = MediaICVFXSettings.IsMediaOutputAssigned(ThisClusterNodeId);
+				InOutRenderSettings.bIsBeingCaptured = bMediaOutputAssigned;
+
+				// Late OCIO pass
+				InOutRenderSettings.bForceLateOCIOPass = (bMediaOutputAssigned || bMediaInputAssigned ? MediaICVFXSettings.bLateOCIOPass : false);
 			}
 		}
 	}
diff --git a/Engine/Plugins/Runtime/nDisplay/Source/DisplayCluster/Private/Render/Viewport/DisplayClusterViewportProxy.cpp b/Engine/Plugins/Runtime/nDisplay/Source/DisplayCluster/Private/Render/Viewport/DisplayClusterViewportProxy.cpp
--- a/Engine/Plugins/Runtime/nDisplay/Source/DisplayCluster/Private/Render/Viewport/DisplayClusterViewportProxy.cpp
+++ b/Engine/Plugins/Runtime/nDisplay/Source/DisplayCluster/Private/Render/Viewport/DisplayClusterViewportProxy.cpp
@@ -448,6 +448,12 @@ EDisplayClusterViewportOpenColorIOMode FDisplayClusterViewportProxy::GetOpenColo
 			// Rendering without post-processing, OCIO is applied last, to the RTT texture of the viewport
 			return EDisplayClusterViewportOpenColorIOMode::Resolved;
 		}
+		else if (RenderSettings.bForceLateOCIOPass)
+		{
+			// When capturing a viewport, it's possible that it's going to be shared within a cluster via the media pipeline.
+			// In this case we should postpone the OCIO step so every node can apply its own OCIO settings.
+			return EDisplayClusterViewportOpenColorIOMode::Resolved;
+		}
 
 		// By default, viewports render with a postprocess, OCIO must be done in between.
 		return EDisplayClusterViewportOpenColorIOMode::PostProcess;
diff --git a/Engine/Plugins/Runtime/nDisplay/Source/DisplayCluster/Public/Render/Viewport/Containers/DisplayClusterViewport_RenderSettings.h b/Engine/Plugins/Runtime/nDisplay/Source/DisplayCluster/Public/Render/Viewport/Containers/DisplayClusterViewport_RenderSettings.h
--- a/Engine/Plugins/Runtime/nDisplay/Source/DisplayCluster/Public/Render/Viewport/Containers/DisplayClusterViewport_RenderSettings.h
+++ b/Engine/Plugins/Runtime/nDisplay/Source/DisplayCluster/Public/Render/Viewport/Containers/DisplayClusterViewport_RenderSettings.h
@@ -59,6 +59,9 @@ public:
 	// Is this viewport being captured by a media capture device?
 	bool bIsBeingCaptured = false;
 
+	// Should OCIO be forcibly applied at later pass (custom nDisplay pass)?
+	bool bForceLateOCIOPass = false;
+
 	// Enable cross-GPU transfer for this viewport.
 	// It may be disabled in some configurations. For example, when using offscreen rendering with TextureShare,
 	// cross-gpu transfer can be disabled for this viewport to improve performance, because when transfer is called,
diff --git a/Engine/Plugins/Runtime/nDisplay/Source/DisplayClusterConfiguration/Public/DisplayClusterConfigurationTypes_Media.h b/Engine/Plugins/Runtime/nDisplay/Source/DisplayClusterConfiguration/Public/DisplayClusterConfigurationTypes_Media.h
--- a/Engine/Plugins/Runtime/nDisplay/Source/DisplayClusterConfiguration/Public/DisplayClusterConfigurationTypes_Media.h
+++ b/Engine/Plugins/Runtime/nDisplay/Source/DisplayClusterConfiguration/Public/DisplayClusterConfigurationTypes_Media.h
@@ -72,6 +72,10 @@ public:
 	UPROPERTY(EditAnywhere, BlueprintReadWrite, Category = "Media")
 	TArray<FDisplayClusterConfigurationMediaOutput> MediaOutputs;
 
+	/** Force late OCIO pass */
+	UPROPERTY(EditAnywhere, BlueprintReadWrite, Category = "Media", meta = (DisplayName = "Late OCIO Pass"))
+	bool bLateOCIOPass = false;
+
 public:
 	/** Returns true if a media source assigned */
 	bool IsMediaInputAssigned() const;
@@ -157,6 +161,10 @@ public:
 	UPROPERTY(EditAnywhere, BlueprintReadWrite, Category = "Media")
 	TArray<FDisplayClusterConfigurationMediaOutputGroup> MediaOutputGroups;
 
+	/** Force late OCIO pass */
+	UPROPERTY(EditAnywhere, BlueprintReadWrite, Category = "Media", meta = (DisplayName = "Late OCIO Pass"))
+	bool bLateOCIOPass = false;
+
 public:
 	/** Returns true if a specific cluster node has media source assigned */
 	bool IsMediaInputAssigned(const FString& NodeId) const;
