nDisplay: Fix ICVFX proxies being removed when changing a property on a DCRA component.

diff --git a/Engine/Plugins/Runtime/nDisplay/Source/DisplayClusterScenePreview/Private/DisplayClusterScenePreviewModule.cpp b/Engine/Plugins/Runtime/nDisplay/Source/DisplayClusterScenePreview/Private/DisplayClusterScenePreviewModule.cpp
--- a/Engine/Plugins/Runtime/nDisplay/Source/DisplayClusterScenePreview/Private/DisplayClusterScenePreviewModule.cpp
+++ b/Engine/Plugins/Runtime/nDisplay/Source/DisplayClusterScenePreview/Private/DisplayClusterScenePreviewModule.cpp
@@ -6,6 +6,7 @@
 #include "CanvasTypes.h"
 #include "Components/DisplayClusterScreenComponent.h"
 #include "Components/DisplayClusterPreviewComponent.h"
+#include "DisplayClusterChromakeyCardActor.h"
 #include "DisplayClusterLightCardActor.h"
 #include "DisplayClusterRootActor.h"
 #include "Engine/Blueprint.h"
@@ -440,10 +441,13 @@ void FDisplayClusterScenePreviewModule::RegisterRootActorEvents(ADisplayClusterR
 
 void FDisplayClusterScenePreviewModule::AutoPopulateScene(FRendererConfig& RendererConfig)
 {
-	RendererConfig.Renderer->ClearScene();
-	RendererConfig.AddedActors.Empty();
-	RendererConfig.AutoActors.Empty();
-
+	if (RendererConfig.bAutoUpdateLightcards)
+	{
+		RendererConfig.Renderer->ClearScene();
+		RendererConfig.AddedActors.Empty();
+		RendererConfig.AutoActors.Empty();
+	}
+	
 	if (ADisplayClusterRootActor* RootActor = InternalGetRendererRootActor(RendererConfig))
 	{
 		TArray<FString> ProjectionMeshNames;
@@ -466,7 +470,11 @@ void FDisplayClusterScenePreviewModule::AutoPopulateScene(FRendererConfig& Rende
 			// Automatically add the lightcards found on this actor
 			TSet<ADisplayClusterLightCardActor*> LightCards;
 			UDisplayClusterBlueprintLib::FindLightCardsForRootActor(RootActor, LightCards);
-
+			
+			TSet<ADisplayClusterChromakeyCardActor*> ChromaKeyCards;
+			UDisplayClusterBlueprintLib::FindChromakeyCardsForRootActor(RootActor, ChromaKeyCards);
+			LightCards.Append(reinterpret_cast<TSet<ADisplayClusterLightCardActor*>&>(ChromaKeyCards));
+			
 			TSet<AActor*> Actors;
 			Actors.Reserve(LightCards.Num());
 			for (ADisplayClusterLightCardActor* LightCard : LightCards)
