nDisplay: Implemented a new preview rendering architecture:
- Added preview interfaces IDisplayClusterViewportManagerPreview and IDisplayClusterViewportPreview
  * Redesign of DCRA preview rendering without using the UDisplayClusterPreviewComponent
  * All preview rendering has been moved to this interfaces (from RootActor and ViewportManager).
  * The preview configuration has been redesigned.
  * New preview architecture implemented
- Redesign of DCRA preview rendering without using the UDisplayClusterPreviewComponent.
- The preview rendering moved outside of RootActor, as a result the class FDisplayClusterRootActorPreviewRenderingManager has been rewoked to new one FDisplayClusterViewportManagerPreviewRenderingSingleton
- The UDisplayClusterDisplayDeviceBaseComponent and UDisplayClusterDisplayDeviceDeviceComponent components have been redesigned for the new architecture
   * From this component, the materials used to render the previews are provided.
   * Preview meshes and material instance parameters are customized from the handler functions.
   * Added IDisplayClusterDisplayDeviceProxy proxy object implementation that integrates into the rendering thread.
   * Reworked OCIO preview implementation from GT to RT.
   * Added implementation for Techvis.
   * Reworked assets with preview materials.
-  The UDisplayClusterPreviewShareComponent has been redesigned for the new architecture
  * The PullPreviewFromSourceActor() function has been redesigned for the new architecture. A new instance of IDisplayClusterViewportManager is used for rendering preview.
  * Standalone/Package support added
- The 'InFrustumFit' warp policy has been reworked for the new architecture
- A new OnUpdatePreviewEditableMesh() function has been added to IDisplayClusterWarpPolicy, which is used to configure dynamic material parameters for EditableMeshes.
- The IDisplayClusterViewportConfiguration has been redesigned.
- The implementation of preview meshes in projection policies has been reworked.
- The ADisplayClusterRootActor actor has been reworked for the new architecture.
- The FDisplayClusterViewportPreviewResource has been redesigned to use UTextureRenderTarget2D
- In many places code is added to block repeated output of the same message to the log (see CanShowLogMsgOnce())
- The code of the ICVFX panel and  has been reworked for the new architecture
- Added techvis materials to the preview mesh.
- ULineBatchComponent has been added to RootActor. It is used for special rendering, for example, frustum, which is not visible in viewports.
- Added preview rendering for Standalone/Package
- Fixed math in ICVFX shader: InCamera projection should ignore the RegionMatrix that is used to get the source image for the external viewport. (This matrix is used for RotateFrustumToFitContextSize).
- fixed several places in the code where a crash occurred when using ENQUEUE_RENDER_COMMAND with a pointer to an object that was deleted on the game thread as an argument.

diff --git a/Engine/Plugins/Runtime/nDisplay/Source/DisplayCluster/Private/Game/EngineClasses/Basics/DisplayClusterViewportClient.cpp b/Engine/Plugins/Runtime/nDisplay/Source/DisplayCluster/Private/Game/EngineClasses/Basics/DisplayClusterViewportClient.cpp
--- a/Engine/Plugins/Runtime/nDisplay/Source/DisplayCluster/Private/Game/EngineClasses/Basics/DisplayClusterViewportClient.cpp
+++ b/Engine/Plugins/Runtime/nDisplay/Source/DisplayCluster/Private/Game/EngineClasses/Basics/DisplayClusterViewportClient.cpp
@@ -561,7 +561,7 @@ void UDisplayClusterViewportClient::Draw(FViewport* InViewport, FCanvas* SceneCa
 				FRotator	ViewRotation;
 				FSceneView* View = RenderFrameViewportManager->CalcSceneView(LocalPlayer, &ViewFamily, ViewLocation, ViewRotation, InViewport, nullptr, ViewportContext.StereoViewIndex);
 
-				if (View && !DCView.IsViewportContextCanBeRendered())
+				if (View && (!DCView.IsViewportContextCanBeRendered() || ViewFamily.RenderTarget == nullptr))
 				{
 					ViewFamily.Views.Remove(View);
 
@@ -997,55 +997,43 @@ bool UDisplayClusterViewportClient::Draw_PIE(FViewport* InViewport, FCanvas* Sce
 		return false;
 	}
 
-	// Get root actor from viewport
+	// Obtaining the primary root vector that can be used for PIE mode
 	ADisplayClusterRootActor* const RootActor = GameMgr->GetRootActor();
-	if (RootActor == nullptr)
+	if (!RootActor || !RootActor->IsPrimaryRootActorForPIE())
 	{
 		return false;
 	}
 
-	//@todo Implement this logic inside function DisplayCluster.GetConfigMgr()->GetLocalNodeId()
-	//@todo: change local node selection for customers
-	FString LocalNodeId = RootActor->PreviewNodeId;
-	if (LocalNodeId == DisplayClusterConfigurationStrings::gui::preview::PreviewNodeAll || LocalNodeId == DisplayClusterConfigurationStrings::gui::preview::PreviewNodeNone)
-	{
-		return false;
-	}
-
-
-	//Get world for render
-	UWorld* const MyWorld = GetWorld();
+	check(SceneCanvas);
+	check(GEngine);
 
 	// When the PIE is used by this DCRA, we must create a new ViewportManager
-	IDisplayClusterViewportManager* ViewportManager = RootActor->GetOrCreateViewportManager();
-	if (ViewportManager == nullptr)
+	if(IDisplayClusterViewportManager* ViewportManager = RootActor->GetOrCreateViewportManager())
 	{
-		return false;
-	}
+		FDisplayClusterViewport_PreviewSettings NewPreviewSettings = RootActor->GetPreviewSettings();
+		{
+			// Disable frustum preview rendering in PIE
+			NewPreviewSettings.bPreviewICVFXFrustums = false;
+
+			// Note: Normally these settings are not used for previewing in PIE and are ignored.
+			ViewportManager->GetConfiguration().SetPreviewSettings(NewPreviewSettings);
+		}
 
-	// update current world
-	ViewportManager->GetConfiguration().SetCurrentWorld(MyWorld);
+		const EDisplayClusterRenderFrameMode RenderFrameMode = ViewportManager->GetConfiguration().GetRenderModeForPIE();
+		if (ViewportManager->GetViewportManagerPreview().InitializeClusterNodePreview(RenderFrameMode, GetWorld(), RootActor->PreviewNodeId, InViewport))
+		{
+			OnBeginDraw().Broadcast();
 
-	// Update local node viewports (update\create\delete) and build new render frame
-	const EDisplayClusterRenderFrameMode RenderModeForPIE = ViewportManager->GetConfiguration().GetRenderModeForPIE();
-	if (ViewportManager->GetConfiguration().UpdateConfigurationForClusterNode(RenderModeForPIE, LocalNodeId) == false)
-	{
-		return false;
-	}
+			ViewportManager->GetViewportManagerPreview().RenderClusterNodePreview(INDEX_NONE, InViewport, SceneCanvas);
+			//ensure canvas has been flushed before rendering UI
+			SceneCanvas->Flush_GameThread();
 
-	FDisplayClusterRenderFrame RenderFrame;
-	if (ViewportManager->BeginNewFrame(InViewport, RenderFrame) == false)
-	{
-		return false;
-	}
+			OnEndDraw().Broadcast();
 
-	bool bOutFrameRendered = false;
-	int32 RenderedViewportsAmount = 0;
-	if (ViewportManager->RenderInEditor(RenderFrame, InViewport, 0, -1, RenderedViewportsAmount, bOutFrameRendered) == false)
-	{
-		return false;
+			return true;
+		}
 	}
 
-	return true;
+	return false;
 }
 #endif /*WITH_EDITOR*/
diff --git a/Engine/Plugins/Runtime/nDisplay/Source/DisplayCluster/Private/Game/EngineClasses/Scene/DisplayClusterPreviewShareComponent.cpp b/Engine/Plugins/Runtime/nDisplay/Source/DisplayCluster/Private/Game/EngineClasses/Scene/DisplayClusterPreviewShareComponent.cpp
--- a/Engine/Plugins/Runtime/nDisplay/Source/DisplayCluster/Private/Game/EngineClasses/Scene/DisplayClusterPreviewShareComponent.cpp
+++ b/Engine/Plugins/Runtime/nDisplay/Source/DisplayCluster/Private/Game/EngineClasses/Scene/DisplayClusterPreviewShareComponent.cpp
@@ -1,10 +1,12 @@
 // Copyright Epic Games, Inc. All Rights Reserved.
 
 #include "Components/DisplayClusterPreviewShareComponent.h"
+#include "Render/Viewport/DisplayClusterViewportManager.h"
+#include "Render/Viewport/IDisplayClusterViewportManagerPreview.h"
+#include "Render/Viewport/Containers/DisplayClusterViewport_PreviewSettings.h"
 
 #include "CineCameraComponent.h"
 #include "Components/DisplayClusterICVFXCameraComponent.h"
-#include "Components/DisplayClusterPreviewComponent.h"
 #include "DisplayClusterRootActor.h"
 #include "Engine/TextureRenderTarget2D.h"
 #include "IDisplayCluster.h"
@@ -15,10 +17,8 @@
 #include "SharedMemoryMediaSource.h"
 
 #if WITH_EDITOR
-	#include "LevelEditor.h"
-#endif // WITH_EDITOR
-
-#if WITH_EDITOR
+#include "LevelEditor.h"
+#endif
 
 namespace UE::DisplayClusterPreviewShare
 {
@@ -160,50 +160,16 @@ namespace UE::DisplayClusterPreviewShare
 			}
 		}
 	}
-
-	/** Retrieves appropriate preview render target from preview component */
-	UTextureRenderTarget2D* GetPreviewRenderTarget(const UDisplayClusterPreviewComponent* PreviewComponent)
-	{
-		if (!PreviewComponent)
-		{
-			return nullptr;
-		}
-
-		// We're primarily interested in the post process render target.
-		UTextureRenderTarget2D* RenderTarget = PreviewComponent->GetRenderTargetTexturePostProcess();
-
-		// Fallback if unavailable.
-		if (!RenderTarget)
-		{
-			RenderTarget = PreviewComponent->GetRenderTargetTexture();
-		}
-
-		return RenderTarget;
-	}
-
 }
 
-#endif // WITH_EDITOR
-
 UDisplayClusterPreviewShareComponent::UDisplayClusterPreviewShareComponent(const FObjectInitializer& ObjectInitializer)
 	: Super(ObjectInitializer)
 {
-#if WITH_EDITOR
-
-	if (!AllowedToShare())
-	{
-		return;
-	}
-
 	PrimaryComponentTick.bCanEverTick = true;
 	PrimaryComponentTick.bStartWithTickEnabled = false;
 	SetTickEnable(false);
-
-#endif // WITH_EDITOR
 }
 
-#if WITH_EDITOR // Bulk wrap with WITH_EDITOR until preview is supported in other modes.
-
 bool UDisplayClusterPreviewShareComponent::AllowedToShare() const
 {
 	// This component should be inactive if it is a CDO.
@@ -240,75 +206,36 @@ bool UDisplayClusterPreviewShareComponent::AllowedToShare() const
 	return true;
 }
 
-#if WITH_EDITOR
-void UDisplayClusterPreviewShareComponent::PostEditChangeProperty(FPropertyChangedEvent& PropertyChangedEvent)
-{
-	Super::PostEditChangeProperty(PropertyChangedEvent);
 
-	if (PropertyChangedEvent.GetPropertyName() == GET_MEMBER_NAME_CHECKED(UDisplayClusterPreviewShareComponent, Mode))
-	{
-		ModeChanged();
-	}
-	else if (PropertyChangedEvent.GetPropertyName() == GET_MEMBER_NAME_CHECKED(UDisplayClusterPreviewShareComponent, UniqueName))
-	{
-		// Remove spaces to reduce chances of the user not realizing that there is a mismatch with its counterpart.
-		UniqueName = UniqueName.TrimStartAndEnd();
 
-		// All the names are now invalid, so we need to close all media.
-		CloseAllMedia();
-	}
-	else if (PropertyChangedEvent.GetPropertyName() == GET_MEMBER_NAME_CHECKED(UDisplayClusterPreviewShareComponent, SourceNDisplayActor))
+void UDisplayClusterPreviewShareComponent::ModeChanged()
+{
+	// We ignore the desired Mode if it is not allowed.
+	if (!AllowedToShare())
 	{
-		ModeChanged();
+		Mode = EDisplayClusterPreviewShareMode::None;
 	}
-}
-#endif // WITH_EDITOR
 
-void UDisplayClusterPreviewShareComponent::ModeChanged()
-{
 	// Close all media before restarting the sharing
 	CloseAllMedia();
 
 	// Restore original root actor settings (e.g. texture replace settings that we overwrote)
+	PullPreviewFromSourceActor(nullptr);
 	RestoreRootActorOriginalSettings();
 
-	// Unsubscribe from previews
-	UnsubscribeFromAllPreviews();
-
-	// Default to no tick
-	SetTickEnable(false);
-
 	switch (Mode)
 	{
 	case EDisplayClusterPreviewShareMode::PullActor:
-	{
-		// The source root actor should enable the preview since we're pulling from it.
-		SubscribeToPreview(SourceNDisplayActor.Get());
-
-		SetTickEnable(true);
-
-		break;
-	}
 	case EDisplayClusterPreviewShareMode::Send:
-	{
-		// If we're sending, we need to tell the root actor to generate the preview textures for us to send.
-		SubscribeToPreview(Cast<ADisplayClusterRootActor>(GetOwner()));
-
-		SetTickEnable(true);
-
-		break;
-	}
 	case EDisplayClusterPreviewShareMode::Receive:
-	{
 		SetTickEnable(true);
-
 		break;
-	}
+
 	default:
-	{
+		// Default to no tick
+		SetTickEnable(false);
 		break;
 	}
-	}
 }
 
 void UDisplayClusterPreviewShareComponent::SetTickEnable(const bool bEnable)
@@ -366,39 +293,9 @@ void UDisplayClusterPreviewShareComponent::RestoreRootActorOriginalSettings()
 	OriginalTextureReplaces.Empty();
 }
 
-void UDisplayClusterPreviewShareComponent::UnsubscribeFromAllPreviews()
-{
-	if (UWorld* World = GetWorld())
-	{
-		for (TActorIterator<ADisplayClusterRootActor> It(World); It; ++It)
-		{
-			UnsubscribeFromPreview(*It);
-		}
-	}
-}
-
-void UDisplayClusterPreviewShareComponent::SubscribeToPreview(ADisplayClusterRootActor* RootActor)
-{
-	if (IsValid(RootActor))
-	{
-		RootActor->SubscribeToPostProcessRenderTarget(reinterpret_cast<uint8*>(this));
-		RootActor->AddPreviewEnableOverride(reinterpret_cast<uint8*>(this));
-	}
-}
-
-void UDisplayClusterPreviewShareComponent::UnsubscribeFromPreview(ADisplayClusterRootActor* RootActor)
-{
-	if (IsValid(RootActor))
-	{
-		RootActor->UnsubscribeFromPostProcessRenderTarget(reinterpret_cast<uint8*>(this));
-		RootActor->RemovePreviewEnableOverride(reinterpret_cast<uint8*>(this));
-	}
-}
-
 void UDisplayClusterPreviewShareComponent::CloseAllMedia()
 {
 	// Stop all media
-
 	for (const TPair<FString, TObjectPtr<UMediaCapture>>& Pair : MediaCaptures)
 	{
 		UMediaCapture* MediaCapture = Pair.Value;
@@ -431,6 +328,9 @@ void UDisplayClusterPreviewShareComponent::CloseAllMedia()
 	MediaSources.Empty();
 	MediaPlayers.Empty();
 	MediaTextures.Empty();
+
+	// Release custom viewport manager
+	ReleaseCustomViewportManager();
 }
 
 FString UDisplayClusterPreviewShareComponent::GenerateViewportKey(const FString& NodeName, const FString& ViewportName) const
@@ -449,39 +349,32 @@ FString UDisplayClusterPreviewShareComponent::GenerateMediaUniqueName(const FStr
 	);
 }
 
-
 void UDisplayClusterPreviewShareComponent::TickSend()
 {
 	using namespace UE::DisplayClusterPreviewShare;
 
+	// Rendering RootActor on CustomViewportManager with custom settings and sending results
 	ADisplayClusterRootActor* RootActor = Cast<ADisplayClusterRootActor>(GetOwner());
-
-	if (!RootActor)
+	if(!UpdateCustomViewportManager(RootActor))
 	{
 		CloseAllMedia();
 		return;
 	}
 
 	// Make sure we're up to date with the viewports
-
-	const TArray<UActorComponent*> ActorComponents = RootActor->K2_GetComponentsByClass(UDisplayClusterPreviewComponent::StaticClass());
-
 	// We will use this array to close and remove unused media captures.
 	TSet<FString> LeftoverViewportKeys;
 	MediaOutputs.GetKeys(LeftoverViewportKeys);
 
-	for (UActorComponent* ActorComponent : ActorComponents)
+	for (const TSharedPtr<IDisplayClusterViewportPreview, ESPMode::ThreadSafe>& ViewportPreviewIt : CustomViewportManager->GetViewportManagerPreview().GetEntireClusterPreviewViewports())
 	{
-		UDisplayClusterPreviewComponent* PreviewComponent = Cast<UDisplayClusterPreviewComponent>(ActorComponent);
-
-		if (!PreviewComponent)
+		if (!ViewportPreviewIt.IsValid())
 		{
 			continue;
 		}
 
 		// Get the texture to share
-		UTextureRenderTarget2D* PreviewTexture = GetPreviewRenderTarget(PreviewComponent);
-
+		UTextureRenderTarget2D* PreviewTexture = ViewportPreviewIt->GetPreviewTextureRenderTarget2D();
 		if (!PreviewTexture)
 		{
 			continue;
@@ -490,8 +383,8 @@ void UDisplayClusterPreviewShareComponent::TickSend()
 		// Each viewport gets a unique name
 
 		const FString ViewportKey = GenerateViewportKey(
-			PreviewComponent->GetClusterNodeId(),
-			PreviewComponent->GetViewportId()
+			ViewportPreviewIt->GetClusterNodeId(),
+			ViewportPreviewIt->GetId()
 		);
 
 		LeftoverViewportKeys.Remove(ViewportKey);
@@ -593,22 +486,23 @@ void UDisplayClusterPreviewShareComponent::TickSend()
 void UDisplayClusterPreviewShareComponent::TickPullActor()
 {
 	ADisplayClusterRootActor* SourceRootActor = SourceNDisplayActor.Get();
-
 	if (!IsValid(SourceRootActor))
 	{
 		return;
 	}
 
 	ADisplayClusterRootActor* RootActor = Cast<ADisplayClusterRootActor>(GetOwner());
-
-	if (!RootActor)
+	if (RootActor == SourceRootActor || !RootActor)
 	{
+		// Disallow self referencing
 		return;
 	}
 
 	// Pull preview textures
 	PullPreviewFromSourceActor(SourceRootActor);
 
+	// Note: Did we really need to synchronize the ICVFX camera function for PullActor mode?
+	// 
 	// Sync Icvfx cameras
 	switch (IcvfxCamerasSyncType)
 	{
@@ -629,111 +523,41 @@ void UDisplayClusterPreviewShareComponent::PullPreviewFromSourceActor(const ADis
 {
 	using namespace UE::DisplayClusterPreviewShare;
 
-	if (!SourceRootActor)
-	{
-		return;
-	}
-
 	// Grab destination root actor.
-
 	ADisplayClusterRootActor* RootActor = Cast<ADisplayClusterRootActor>(GetOwner());
-
-	if (!RootActor)
+	if (RootActor == SourceRootActor || !RootActor)
 	{
+		// Disallow self referencing
 		return;
 	}
 
-	// Disallow self referencing
-
-	if (RootActor == SourceRootActor)
+	if (IDisplayClusterViewportManager* ViewportManager = RootActor->GetViewportManager())
 	{
-		return;
-	}
-
-	// Gather source preview components
-
-	TMap<FString, UDisplayClusterPreviewComponent*> SrcPreviewComponentWithUniqueName;
-
-	for (UActorComponent* SrcComponent : SourceRootActor->K2_GetComponentsByClass(UDisplayClusterPreviewComponent::StaticClass()))
-	{
-		UDisplayClusterPreviewComponent* SrcPreviewComponent = Cast<UDisplayClusterPreviewComponent>(SrcComponent);
-
-		if (!SrcPreviewComponent)
+		if (SourceRootActor)
 		{
-			continue;
+			// Disable local settings for owner root actor
+			RootActor->bUseLocalPreviewSetttings = false;
+
+			// reconfigure to rendering from external DCRA
+			ViewportManager->GetConfiguration().SetRootActor(EDisplayClusterRootActorType::Scene | EDisplayClusterRootActorType::Configuration, SourceRootActor);
+
+			// Override preview settings:
+			FDisplayClusterViewport_PreviewSettings CustomPreviewSettings = SourceRootActor->GetPreviewSettings();
+			CustomPreviewSettings.bPreviewEnable = RootActor->bPreviewEnable;
+			CustomPreviewSettings.bPreviewEnablePostProcess = bPreviewEnablePostProcess;
+			CustomPreviewSettings.bPreviewICVFXFrustums = false; // disable frustum rendering
+			ViewportManager->GetConfiguration().SetPreviewSettings(CustomPreviewSettings);
 		}
-
-		const FString SrcViewportKey = GenerateViewportKey(
-			SrcPreviewComponent->GetClusterNodeId(),
-			SrcPreviewComponent->GetViewportId()
-		);
-
-		SrcPreviewComponentWithUniqueName.Add(SrcViewportKey, SrcPreviewComponent);
-	}
-
-	// Iterate over destination preview components, find the matching source preview component, and assign the preview texture accordingly.
-
-	for (UActorComponent* DstComponent : RootActor->K2_GetComponentsByClass(UDisplayClusterPreviewComponent::StaticClass()))
-	{
-		UDisplayClusterPreviewComponent* DstPreviewComponent = Cast<UDisplayClusterPreviewComponent>(DstComponent);
-
-		if (!DstPreviewComponent)
-		{
-			continue;
-		}
-
-		// We will ultimately update the viewport so let's make sure it exists.
-
-		UDisplayClusterConfigurationViewport* DstViewport = GetViewportFromDCRA(RootActor, DstPreviewComponent->GetClusterNodeId(), DstPreviewComponent->GetViewportId());
-
-		if (!DstViewport)
-		{
-			continue;
-		}
-
-		// Each viewport gets associated with a uniquely name
-
-		const FString DstViewportKey = GenerateViewportKey(
-			DstPreviewComponent->GetClusterNodeId(),
-			DstPreviewComponent->GetViewportId()
-		);
-
-		// Find corresponding source preview component
-
-		UDisplayClusterPreviewComponent** SrcPreviewComponentPtr = SrcPreviewComponentWithUniqueName.Find(DstViewportKey);
-
-		if (!SrcPreviewComponentPtr || !*SrcPreviewComponentPtr)
-		{
-			continue;
-		}
-
-		// Assign destination preview texture from source preview texture.
-
-		const UDisplayClusterPreviewComponent* SrcPreviewComponent = *SrcPreviewComponentPtr;
-		check(SrcPreviewComponent);
-
-		UTextureRenderTarget2D* SrcPreviewTexture = GetPreviewRenderTarget(SrcPreviewComponent);
-
-		if (!SrcPreviewTexture)
+		else
 		{
-			continue;
-		}
+			// Restore default preview:
 
-		// Remember the original replacement textures, if they haven't been saved yet.
+			// use local settings for owner root actor
+			RootActor->bUseLocalPreviewSetttings = true;
 
-		if (!OriginalTextureReplaces.Find(DstViewportKey))
-		{
-			OriginalTextureReplaces.Add(DstViewportKey, DstViewport->RenderSettings.Replace.bAllowReplace);
-		}
-
-		if (!OriginalSourceTextures.Find(DstViewportKey))
-		{
-			OriginalSourceTextures.Add(DstViewportKey, DstViewport->RenderSettings.Replace.SourceTexture);
+			// reconfigure to rendering from owner root actor
+			ViewportManager->GetConfiguration().SetRootActor(EDisplayClusterRootActorType::Scene | EDisplayClusterRootActorType::Configuration, RootActor);
 		}
-
-		// Override the replace texture values
-		DstViewport->RenderSettings.Replace.bAllowReplace = true;
-		DstViewport->RenderSettings.Replace.SourceTexture = SrcPreviewTexture;
 	}
 }
 
@@ -849,33 +673,28 @@ void UDisplayClusterPreviewShareComponent::TickReceive()
 	using namespace UE::DisplayClusterPreviewShare;
 
 	ADisplayClusterRootActor* RootActor = Cast<ADisplayClusterRootActor>(GetOwner());
-
-	if (!RootActor)
+	IDisplayClusterViewportManager* ViewportManager = RootActor ? RootActor->GetViewportManager() : nullptr;
+	if (!ViewportManager)
 	{
 		CloseAllMedia();
 		return;
 	}
 
 	// Make sure we're up to date with the viewports
-
-	const TArray<UActorComponent*> ActorComponents = RootActor->K2_GetComponentsByClass(UDisplayClusterPreviewComponent::StaticClass());
-
 	// We will use this array to close and remove unused media sources.
 	TSet<FString> LeftoverViewportKeys;
 	MediaSources.GetKeys(LeftoverViewportKeys);
 
-	for (UActorComponent* ActorComponent : ActorComponents)
+	for (TSharedPtr<IDisplayClusterViewport, ESPMode::ThreadSafe>& ViewportIt : ViewportManager->GetEntireClusterViewports())
 	{
-		UDisplayClusterPreviewComponent* PreviewComponent = Cast<UDisplayClusterPreviewComponent>(ActorComponent);
-
-		if (!PreviewComponent)
+		if (!ViewportIt.IsValid())
 		{
 			continue;
 		}
 
 		// We will ultimately update the viewport so let's make sure it exists.
 
-		UDisplayClusterConfigurationViewport* Viewport = GetViewportFromDCRA(RootActor, PreviewComponent->GetClusterNodeId(), PreviewComponent->GetViewportId());
+		UDisplayClusterConfigurationViewport* Viewport = GetViewportFromDCRA(RootActor, ViewportIt->GetClusterNodeId(), ViewportIt->GetId());
 
 		if (!Viewport)
 		{
@@ -885,8 +704,8 @@ void UDisplayClusterPreviewShareComponent::TickReceive()
 		// Each viewport gets associated with a uniquely named shared texture.
 
 		const FString ViewportKey = GenerateViewportKey(
-			PreviewComponent->GetClusterNodeId(),
-			PreviewComponent->GetViewportId()
+			ViewportIt->GetClusterNodeId(),
+			ViewportIt->GetId()
 		);
 
 		LeftoverViewportKeys.Remove(ViewportKey);
@@ -1048,7 +867,6 @@ void UDisplayClusterPreviewShareComponent::OnComponentDestroyed(bool bDestroying
 	SetMode(EDisplayClusterPreviewShareMode::None);
 
 #if WITH_EDITOR
-
 	if (GIsEditor)
 	{
 		// Unregister Map Change Events
@@ -1057,7 +875,6 @@ void UDisplayClusterPreviewShareComponent::OnComponentDestroyed(bool bDestroying
 			LevelEditor->OnMapChanged().RemoveAll(this);
 		}
 	}
-
 #endif // WITH_EDITOR
 
 	Super::OnComponentDestroyed(bDestroyingHierarchy);
@@ -1072,7 +889,8 @@ void UDisplayClusterPreviewShareComponent::OnRegister()
 	ModeChanged();
 }
 
-void UDisplayClusterPreviewShareComponent::HandleMapChanged(UWorld* InWorld, EMapChangeType InMapChangeType)
+#if WITH_EDITOR
+void UDisplayClusterPreviewShareComponent::HandleMapChanged_Editor(UWorld* InWorld, EMapChangeType InMapChangeType)
 {
 	// We remove the reference to the source actor preview texture to avoid this being interpreted as a 
 	// memory leak when unloadin the source actor map while the current map is kept open.
@@ -1081,14 +899,10 @@ void UDisplayClusterPreviewShareComponent::HandleMapChanged(UWorld* InWorld, EMa
 		RestoreRootActorOriginalSettings();
 	}
 }
-
-#endif // WITH_EDITOR in bulk
-
+#endif
 
 void UDisplayClusterPreviewShareComponent::Serialize(FArchive& Ar)
 {
-#if WITH_EDITOR
-
 	using namespace UE::DisplayClusterPreviewShare;
 
 	if (Ar.IsSaving())
@@ -1102,7 +916,6 @@ void UDisplayClusterPreviewShareComponent::Serialize(FArchive& Ar)
 			RemoveExternalMapReferences(Cast<ADisplayClusterRootActor>(GetOwner()));
 		}
 	}
-#endif // WITH_EDITOR
 
 	Super::Serialize(Ar);
 }
@@ -1111,27 +924,24 @@ void UDisplayClusterPreviewShareComponent::PostLoad()
 {
 	Super::PostLoad();
 
-#if WITH_EDITOR
 	// Make sure we are in the right state after loading
 	ModeChanged();
 
+#if WITH_EDITOR
 	if (GIsEditor)
 	{
 		// Register for Map Change Events. We do this to remove texture external texture reference that
 		// prevents proper world GC. In particular, our owing actor reference to the source nDisplay actor preview texture.
 		if (FLevelEditorModule* LevelEditor = FModuleManager::GetModulePtr<FLevelEditorModule>("LevelEditor"))
 		{
-			LevelEditor->OnMapChanged().AddUObject(this, &UDisplayClusterPreviewShareComponent::HandleMapChanged);
+			LevelEditor->OnMapChanged().AddUObject(this, &UDisplayClusterPreviewShareComponent::HandleMapChanged_Editor);
 		}
 	}
-
 #endif // WITH_EDITOR
 }
 
 void UDisplayClusterPreviewShareComponent::SetMode(EDisplayClusterPreviewShareMode NewMode)
 {
-#if WITH_EDITORONLY_DATA
-
 	// Nothing to do if the mode is unchanged.
 	if (Mode == NewMode)
 	{
@@ -1140,32 +950,88 @@ void UDisplayClusterPreviewShareComponent::SetMode(EDisplayClusterPreviewShareMo
 
 	Mode = NewMode;
 
-#if WITH_EDITOR
+	ModeChanged();
+}
 
-	// We ignore the desired Mode if it is not allowed.
-	if (!AllowedToShare())
+void UDisplayClusterPreviewShareComponent::SetUniqueName(const FString& NewUniqueName)
+{
+	UniqueName = NewUniqueName.TrimStartAndEnd();
+
+	// All the names are now invalid, so we need to close all media.
+	CloseAllMedia();
+}
+
+bool UDisplayClusterPreviewShareComponent::UpdateCustomViewportManager(const ADisplayClusterRootActor* InSrcRootActor)
+{
+	check(IsInGameThread());
+
+	if (!InSrcRootActor)
 	{
-		Mode = EDisplayClusterPreviewShareMode::None;
+		// The source RootActor is required
+		return false;
 	}
 
-	ModeChanged();
+	if (!CustomViewportManager.IsValid())
+	{
+		CustomViewportManager = IDisplayClusterViewportManager::CreateViewportManager();
+	}
 
-#endif // WITH_EDITOR
+	IDisplayClusterViewportConfiguration& CustomConfiguration = CustomViewportManager->GetConfiguration();
+
+	// Assign DCRA to render, excluding PreviewRootActor since we are not using the preview mesh, only the viewport textures.
+	CustomConfiguration.SetRootActor(EDisplayClusterRootActorType::Scene | EDisplayClusterRootActorType::Configuration, InSrcRootActor);
 
-#endif // WITH_EDITORONLY_DATA
+	// Override preview settings:
+	FDisplayClusterViewport_PreviewSettings CustomPreviewSettings = InSrcRootActor->GetPreviewSettings();
+	CustomPreviewSettings.bPreviewEnable = true;
+	CustomPreviewSettings.bPreviewEnablePostProcess = bPreviewEnablePostProcess;
+	CustomPreviewSettings.bEnablePreviewMesh = false;
+	CustomPreviewSettings.bEnablePreviewEditableMesh = false;
+	CustomPreviewSettings.bPreviewICVFXFrustums = false; // disable frustum rendering
+	CustomConfiguration.SetPreviewSettings(CustomPreviewSettings);
 
+	// Request rendering of the entire cluster for our custom viewport manager
+	CustomViewportManager->GetViewportManagerPreview().UpdateEntireClusterPreviewRender(true);
+
+	return true;
 }
 
-void UDisplayClusterPreviewShareComponent::SetUniqueName(const FString& NewUniqueName)
+void UDisplayClusterPreviewShareComponent::ReleaseCustomViewportManager()
 {
-#if WITH_EDITORONLY_DATA
-	UniqueName = NewUniqueName.TrimStartAndEnd();
-#endif // WITH_EDITORONLY_DATA
+	// Note: shared resources refer to internal preview textures inside the ViewportManager, and all these resources will be freed.
+	// we must handle this correctly on both sides of the IPC
+
+	if (CustomViewportManager.IsValid())
+	{
+		// Reset all DCRA references
+		CustomViewportManager->GetConfiguration().SetRootActor(EDisplayClusterRootActorType::Any, nullptr);
+
+		// Immediately release the viewport manager with resources
+		CustomViewportManager.Reset();
+	}
+}
 
 #if WITH_EDITOR
 
-	// All the names are now invalid, so we need to close all media.
-	CloseAllMedia();
+void UDisplayClusterPreviewShareComponent::PostEditChangeProperty(FPropertyChangedEvent& PropertyChangedEvent)
+{
+	Super::PostEditChangeProperty(PropertyChangedEvent);
 
-#endif // WITH_EDITOR
+	if (PropertyChangedEvent.GetPropertyName() == GET_MEMBER_NAME_CHECKED(UDisplayClusterPreviewShareComponent, Mode))
+	{
+		ModeChanged();
+	}
+	else if (PropertyChangedEvent.GetPropertyName() == GET_MEMBER_NAME_CHECKED(UDisplayClusterPreviewShareComponent, UniqueName))
+	{
+		// Remove spaces to reduce chances of the user not realizing that there is a mismatch with its counterpart.
+		UniqueName = UniqueName.TrimStartAndEnd();
+
+		// All the names are now invalid, so we need to close all media.
+		CloseAllMedia();
+	}
+	else if (PropertyChangedEvent.GetPropertyName() == GET_MEMBER_NAME_CHECKED(UDisplayClusterPreviewShareComponent, SourceNDisplayActor))
+	{
+		ModeChanged();
+	}
 }
+#endif // WITH_EDITOR
\ No newline at end of file
diff --git a/Engine/Plugins/Runtime/nDisplay/Source/DisplayCluster/Private/Game/EngineClasses/Scene/DisplayClusterRootActor.cpp b/Engine/Plugins/Runtime/nDisplay/Source/DisplayCluster/Private/Game/EngineClasses/Scene/DisplayClusterRootActor.cpp
--- a/Engine/Plugins/Runtime/nDisplay/Source/DisplayCluster/Private/Game/EngineClasses/Scene/DisplayClusterRootActor.cpp
+++ b/Engine/Plugins/Runtime/nDisplay/Source/DisplayCluster/Private/Game/EngineClasses/Scene/DisplayClusterRootActor.cpp
@@ -1,7 +1,6 @@
 // Copyright Epic Games, Inc. All Rights Reserved.
 
 #include "DisplayClusterRootActor.h"
-#include "DisplayClusterRootActorPreviewRenderingManager.h"
 
 #include "Async/ParallelFor.h"
 #include "Components/SceneComponent.h"
@@ -9,10 +8,11 @@
 #include "Components/DisplayClusterCameraComponent.h"
 #include "Components/DisplayClusterScreenComponent.h"
 #include "Components/DisplayClusterSceneComponentSyncParent.h"
-#include "Components/DisplayClusterPreviewComponent.h"
 #include "Components/DisplayClusterSyncTickComponent.h"
 #include "Components/DisplayClusterICVFXCameraComponent.h"
 #include "Components/DisplayClusterSceneComponentSyncThis.h"
+#include "Render/DisplayDevice/Components/DisplayClusterDisplayDeviceBaseComponent.h"
+#include "Render/DisplayDevice/Components/DisplayClusterDisplayDeviceComponent.h"
 #include "CineCameraComponent.h"
 #include "DisplayClusterChromakeyCardActor.h"
 #include "ProceduralMeshComponent.h"
@@ -48,6 +48,7 @@
 #include "Components/DisplayClusterChromakeyCardStageActorComponent.h"
 #include "Components/DisplayClusterStageActorComponent.h"
 #include "Components/DisplayClusterStageGeometryComponent.h"
+#include "Components/LineBatchComponent.h"
 #include "UObject/Package.h"
 
 
@@ -59,6 +60,40 @@
 #include "AssetToolsModule.h"
 #endif
 
+namespace UE::DisplayCluster::RootActor
+{
+	template <typename TComp>
+	void CollectPrimitiveComponentsImpl(TSet<FPrimitiveComponentId>& OutPrimitives, TComp* pComp, bool bForceHide = false, const bool bCollectChildrenVisualizationComponent = true)
+	{
+		if (UPrimitiveComponent* PrimComp = Cast<UPrimitiveComponent>(pComp))
+		{
+			if (PrimComp->bHiddenInGame
+			|| bForceHide
+#if WITH_EDITOR
+			|| (GIsEditor && PrimComp->bHiddenInSceneCapture /* We are running as a scene capture for preview */)
+#endif
+				)
+			{
+				OutPrimitives.Add(PrimComp->GetPrimitiveSceneId());
+			}
+		}
+
+		if (bCollectChildrenVisualizationComponent)
+		{
+			if (USceneComponent* SceneComp = Cast<USceneComponent>(pComp))
+			{
+				TArray<USceneComponent*> ChildrenComponents;
+				SceneComp->GetChildrenComponents(false, ChildrenComponents);
+
+				for (USceneComponent* CompIt : ChildrenComponents)
+				{
+					CollectPrimitiveComponentsImpl(OutPrimitives, CompIt, bForceHide, bCollectChildrenVisualizationComponent);
+				}
+			}
+		}
+	}
+};
+
 //////////////////////////////////////////////////////////////////////////////////////////////
 // ADisplayClusterRootActor
 //////////////////////////////////////////////////////////////////////////////////////////////
@@ -107,6 +142,14 @@ ADisplayClusterRootActor::ADisplayClusterRootActor(const FObjectInitializer& Obj
 	bIsSpatiallyLoaded = false;
 #endif
 
+	// Set to internal default, user may change
+	DefaultDisplayDeviceName = GetInternalDisplayDeviceName();
+
+	// Our internal display device which always exists
+	BasicDisplayDeviceComponent = CreateDefaultSubobject<UDisplayClusterDisplayDeviceComponent>(GetInternalDisplayDeviceName());
+	LineBatcherComponent = CreateDefaultSubobject<ULineBatchComponent>(TEXT("LineBatcher"));
+
+	ResetEntireClusterPreviewRendering();
 
 #if WITH_EDITOR
 	Constructor_Editor();
@@ -120,6 +163,16 @@ ADisplayClusterRootActor::~ADisplayClusterRootActor()
 #endif
 }
 
+void ADisplayClusterRootActor::ResetEntireClusterPreviewRendering()
+{
+	if (IDisplayClusterViewportManager* ViewportManager = GetViewportManager())
+	{
+		// Update the preview settings as is from this DCRA
+		ViewportManager->GetViewportManagerPreview().ResetEntireClusterPreviewRendering();
+
+	}
+}
+
 IDisplayClusterViewportManager* ADisplayClusterRootActor::GetViewportManager() const
 {
 	return GetViewportManagerImpl();
@@ -131,16 +184,10 @@ IDisplayClusterViewportManager* ADisplayClusterRootActor::GetOrCreateViewportMan
 
 	if (!ViewportManagerPtr.IsValid())
 	{
-		ViewportManagerPtr = MakeShared<FDisplayClusterViewportManager, ESPMode::ThreadSafe>();
-
-		// After the constructor, we should always call this function to initialize internal references.
-		ViewportManagerPtr->Initialize();
+		ViewportManagerPtr = IDisplayClusterViewportManager::CreateViewportManager()->ToSharedRef();
 
 		// Set the owner's DCRA to the newly created viewport manager.
-		ViewportManagerPtr->GetConfiguration().SetRootActor(this, EDisplayClusterRootActorType::Any);
-
-		// Preview rendering depends on the DC VM
-		FDisplayClusterRootActorPreviewRenderingManager::HandleEvent(EDisplayClusterRootActorPreviewEvent::Create, this);
+		ViewportManagerPtr->GetConfiguration().SetRootActor(EDisplayClusterRootActorType::Any, this);
 	}
 
 	return GetViewportManagerImpl();
@@ -148,7 +195,7 @@ IDisplayClusterViewportManager* ADisplayClusterRootActor::GetOrCreateViewportMan
 
 FDisplayClusterViewportManager* ADisplayClusterRootActor::GetViewportManagerImpl() const
 {
-	return ViewportManagerPtr.IsValid() ? ViewportManagerPtr.Get() : nullptr;
+	return ViewportManagerPtr.Get();
 }
 
 IDisplayClusterViewportConfiguration* ADisplayClusterRootActor::GetViewportConfiguration() const
@@ -160,32 +207,61 @@ void ADisplayClusterRootActor::RemoveViewportManager()
 {
 	if (ViewportManagerPtr.IsValid())
 	{
-		// Preview rendering depends on the  DC VM
-		FDisplayClusterRootActorPreviewRenderingManager::HandleEvent(EDisplayClusterRootActorPreviewEvent::Remove, this);
-
 		// Reset all DCRA references
-		ViewportManagerPtr->GetConfiguration().SetRootActor(nullptr, EDisplayClusterRootActorType::Any);
+		ViewportManagerPtr->GetConfiguration().SetRootActor(EDisplayClusterRootActorType::Any, nullptr);
 
 		// Immediately release the viewport manager with resources
 		ViewportManagerPtr.Reset();
 	}
 }
 
-bool ADisplayClusterRootActor::IsRunningGameOrPIE() const
+bool ADisplayClusterRootActor::IsPrimaryRootActor() const
 {
-	if (!IsRunningGame())
+	if (IDisplayCluster::Get().GetGameMgr()->GetRootActor() == this)
+	{
+		return true;
+	}
+
+	return false;
+}
+
+bool ADisplayClusterRootActor::IsPrimaryRootActorForPIE() const
+{
+#if WITH_EDITOR
+	// A preview node must be selected in DCRA so that it can be used in PIE mode
+	if (PreviewNodeId == DisplayClusterConfigurationStrings::gui::preview::PreviewNodeAll || PreviewNodeId == DisplayClusterConfigurationStrings::gui::preview::PreviewNodeNone)
+	{
+		return false;
+	}
+
+	// Only PIE is currently supported
+	return IsPrimaryRootActor();
+#endif
+
+	return false;
+}
+
+bool ADisplayClusterRootActor::IsRunningPIE() const
 	{
 #if WITH_EDITOR
 		const UWorld* World = GetWorld();
 		return World && World->IsPlayInEditor();
 #else
-		return true;
+		// Without editor return false
+		return false;
 #endif
 	}
 
+bool ADisplayClusterRootActor::IsRunningDisplayCluster() const
+{
+	if (OperationMode == EDisplayClusterOperationMode::Cluster || OperationMode == EDisplayClusterOperationMode::Editor)
+	{
 	return true;
 }
 
+	return false;
+}
+
 const FDisplayClusterConfigurationICVFX_StageSettings& ADisplayClusterRootActor::GetStageSettings() const
 {
 	check(CurrentConfigData);
@@ -208,13 +284,6 @@ void ADisplayClusterRootActor::InitializeFromConfig(UDisplayClusterConfiguration
 		UpdateConfigDataInstance(ConfigData, true);
 
 		BuildHierarchy();
-
-#if WITH_EDITOR
-		if (GIsEditor && GetWorld())
-		{
-			UpdatePreviewComponents();
-		}
-#endif
 	}
 }
 
@@ -349,14 +418,6 @@ void ADisplayClusterRootActor::OverrideFromConfig(UDisplayClusterConfigurationDa
 
 	// There is no sense to call BuildHierarchy because it works for non-BP root actors.
 	// On the other hand, OverwriteFromConfig method is called for BP root actors only by nature.
-
-	// And update preview stuff in Editor
-#if WITH_EDITOR
-	if (GIsEditor && GetWorld())
-	{
-		UpdatePreviewComponents();
-	}
-#endif
 }
 
 void ADisplayClusterRootActor::UpdateConfigDataInstance(UDisplayClusterConfigurationData* ConfigDataTemplate, bool bForceRecreate)
@@ -465,34 +526,11 @@ int ADisplayClusterRootActor::GetInnerFrustumPriority(const FString& InnerFrustu
 	return -1;
 }
 
-template <typename TComp>
-void ImplCollectChildHiddenComponents(TSet<FPrimitiveComponentId>& OutPrimitives, TComp* pComp)
-{
-#if WITH_EDITOR
-
-	USceneComponent* SceneComp = Cast<USceneComponent>(pComp);
-	if (SceneComp)
-	{
-		TArray<USceneComponent*> Childrens;
-		SceneComp->GetChildrenComponents(false, Childrens);
-		for (USceneComponent* ChildIt : Childrens)
-		{
-			if (ChildIt->bHiddenInGame)
-			{
-				UPrimitiveComponent* PrimComp = Cast<UPrimitiveComponent>(ChildIt);
-				if (PrimComp)
-				{
-					OutPrimitives.Add(PrimComp->GetPrimitiveSceneId());
-				}
-			}
-		}
-	}
-#endif
-}
-
 template <typename TComp>
 void ADisplayClusterRootActor::GetTypedPrimitives(TSet<FPrimitiveComponentId>& OutPrimitives, const TArray<FString>* InCompNames, bool bCollectChildrenVisualizationComponent) const
 {
+	using namespace UE::DisplayCluster::RootActor;
+
 	TArray<TComp*> TypedComponents;
 	GetComponents<TComp>(TypedComponents, true);
 
@@ -502,39 +540,16 @@ void ADisplayClusterRootActor::GetTypedPrimitives(TSet<FPrimitiveComponentId>& O
 		{
 			if (InCompNames != nullptr)
 			{
-				// add only comp from names list
-				for (const FString& NameIt : (*InCompNames))
+				if(InCompNames->Find(CompIt->GetName()) != INDEX_NONE)
 				{
-					if (CompIt->GetName() == NameIt)
-					{
-						UPrimitiveComponent* PrimComp = Cast<UPrimitiveComponent>(CompIt);
-						if (PrimComp)
-						{
-							OutPrimitives.Add(PrimComp->GetPrimitiveSceneId());
-						}
-
-						if (bCollectChildrenVisualizationComponent)
-						{
-							ImplCollectChildHiddenComponents(OutPrimitives, CompIt);
-						}
-						break;
-					}
+					// add only comp from names list
+					CollectPrimitiveComponentsImpl(OutPrimitives, CompIt, bCollectChildrenVisualizationComponent);
 				}
 			}
 			else
 			{
-				UPrimitiveComponent* PrimComp = Cast<UPrimitiveComponent>(CompIt);
-				if (PrimComp)
-				{
-					OutPrimitives.Add(PrimComp->GetPrimitiveSceneId());
-				}
-
-				if (bCollectChildrenVisualizationComponent)
-				{
-					ImplCollectChildHiddenComponents(OutPrimitives, CompIt);
-				}
+				CollectPrimitiveComponentsImpl(OutPrimitives, CompIt, bCollectChildrenVisualizationComponent);
 			}
-
 		}
 	}
 }
@@ -549,6 +564,8 @@ bool ADisplayClusterRootActor::FindPrimitivesByName(const TArray<FString>& InNam
 // Gather components not rendered in game
 bool ADisplayClusterRootActor::GetHiddenInGamePrimitives(TSet<FPrimitiveComponentId>& OutPrimitives)
 {
+	using namespace UE::DisplayCluster::RootActor;
+
 	TRACE_CPUPROFILER_EVENT_SCOPE(DCRootActor_GetHiddenInGamePrimitives);
 
 	check(IsInGameThread());
@@ -570,23 +587,33 @@ bool ADisplayClusterRootActor::GetHiddenInGamePrimitives(TSet<FPrimitiveComponen
 	// Always hide DC Screen component in game
 	GetTypedPrimitives<UDisplayClusterScreenComponent>(OutPrimitives);
 
-#if WITH_EDITOR
-	// Always hide preview meshes for preview
+	if (FDisplayClusterViewportManager* ViewportManager = GetViewportManagerImpl())
 	{
-		TArray<UDisplayClusterPreviewComponent*> CurrentPreviewComponents;
-		GetComponents(CurrentPreviewComponents);
-		for (UDisplayClusterPreviewComponent* PreviewComponentIt : CurrentPreviewComponents)
+		// Always hide preview meshes for preview
+		for (TSharedPtr<IDisplayClusterViewport, ESPMode::ThreadSafe>& ViewportIt : ViewportManager->GetEntireClusterViewports())
 		{
-			if (UMeshComponent* PreviewMesh = PreviewComponentIt->GetPreviewMesh())
+			if (ViewportIt.IsValid())
 			{
-				if(UPrimitiveComponent* PrimComp = Cast<UPrimitiveComponent>(PreviewMesh))
+				if (UMeshComponent* PreviewMesh = ViewportIt->GetViewportPreview().GetPreviewMeshComponent())
+				{
+					if (UPrimitiveComponent* PrimComp = Cast<UPrimitiveComponent>(PreviewMesh))
+					{
+						OutPrimitives.Add(PrimComp->GetPrimitiveSceneId());
+					}
+				}
+
+				if (UMeshComponent* PreviewEditableMesh = ViewportIt->GetViewportPreview().GetPreviewEditableMeshComponent())
 				{
-					OutPrimitives.Add(PrimComp->GetPrimitiveSceneId());
+					if (UPrimitiveComponent* PrimComp = Cast<UPrimitiveComponent>(PreviewEditableMesh))
+					{
+						OutPrimitives.Add(PrimComp->GetPrimitiveSceneId());
+					}
 				}
 			}
 		}
 	}
 
+#if WITH_EDITOR
 	// Hide visualization and hidden components from RootActor
 	{
 		TArray<UPrimitiveComponent*> PrimitiveComponents;
@@ -594,12 +621,7 @@ bool ADisplayClusterRootActor::GetHiddenInGamePrimitives(TSet<FPrimitiveComponen
 
 		for (UPrimitiveComponent* CompIt : PrimitiveComponents)
 		{
-			if (CompIt->bHiddenInGame)
-			{
-				OutPrimitives.Add(CompIt->GetPrimitiveSceneId());
-			}
-
-			ImplCollectChildHiddenComponents(OutPrimitives, CompIt);
+			CollectPrimitiveComponentsImpl(OutPrimitives, CompIt);
 		}
 	}
 
@@ -671,19 +693,9 @@ bool ADisplayClusterRootActor::GetHiddenInGamePrimitives(TSet<FPrimitiveComponen
 					if (IsValid(Actor))
 					{
 						Actor->GetComponents(PrimitiveComponents);
-
 						for (UPrimitiveComponent* PrimComp : PrimitiveComponents)
 						{
-							if (((Actor->IsHidden() || PrimComp->bHiddenInGame) && !PrimComp->bCastHiddenShadow)
-#if WITH_EDITOR
-								|| (GIsEditor && PrimComp->bHiddenInSceneCapture /* We are running as a scene capture for preview */)
-#endif
-								)
-							{
-								PrimitiveComponentsArray[Index].Add(PrimComp->GetPrimitiveSceneId());
-							}
-
-							ImplCollectChildHiddenComponents(PrimitiveComponentsArray[Index], PrimComp);
+							CollectPrimitiveComponentsImpl(PrimitiveComponentsArray[Index], PrimComp, Actor->IsHidden());
 						}
 
 						// Empty w/o ever shrinking, for efficiency
@@ -704,6 +716,12 @@ bool ADisplayClusterRootActor::GetHiddenInGamePrimitives(TSet<FPrimitiveComponen
 
 #endif
 
+	if (ULineBatchComponent* LineBatch = GetLineBatchComponent())
+	{
+		// Always hide RootActor batches on viewports
+		CollectPrimitiveComponentsImpl(OutPrimitives, LineBatch, true);
+	}
+
 	return OutPrimitives.Num() > 0;
 }
 
@@ -727,15 +745,12 @@ void ADisplayClusterRootActor::InitializeRootActor()
 	StageGeometryComponent->Invalidate();
 
 	// Packaged, PIE and -game runtime
-	if (IsRunningGameOrPIE())
+	if (IsRunningGame() || IsRunningPIE())
 	{
 		if (CurrentConfigData)
 		{
 			BuildHierarchy();
 
-#if WITH_EDITOR
-			UpdatePreviewComponents();
-#endif
 			return;
 		}
 	}
@@ -746,7 +761,6 @@ void ADisplayClusterRootActor::InitializeRootActor()
 		if (CurrentConfigData)
 		{
 			BuildHierarchy();
-			UpdatePreviewComponents();
 			return;
 		}
 	}
@@ -773,8 +787,7 @@ void ADisplayClusterRootActor::SetPreviewEnablePostProcess(const bool bNewPrevie
 	{
 		bPreviewEnablePostProcess = bNewPreviewEnablePostProcess;
 
-		ResetPreviewComponents_Editor(true);
-		PreviewRenderFrame.Reset();
+		ResetEntireClusterPreviewRendering();
 	}
 #endif // WITH_EDITOR
 }
@@ -987,8 +1000,16 @@ void ADisplayClusterRootActor::Tick(float DeltaSeconds)
 	// Update saved DeltaSeconds for root actor
 	LastDeltaSecondsValue = DeltaSeconds;
 
-	if (OperationMode == EDisplayClusterOperationMode::Cluster ||
-		OperationMode == EDisplayClusterOperationMode::Editor)
+	if (ULineBatchComponent* LineBatch = GetLineBatchComponent())
+	{
+		LineBatch->Flush();
+	}
+	
+	// Support for DCRA preview in the scene for standalone\package
+	// Use settings only from active DCRA
+	const bool bIsPrimaryRootActor = IsPrimaryRootActor();
+	const bool bIsRunningDisplayCluster = IsRunningDisplayCluster();
+	if (bIsRunningDisplayCluster && bIsPrimaryRootActor)
 	{
 		UWorld* const CurWorld = GetWorld();
 		if (CurWorld && CurrentConfigData)
@@ -1018,7 +1039,7 @@ void ADisplayClusterRootActor::Tick(float DeltaSeconds)
 	}
 
 	// Show 'not supported' warning if instanced stereo is used
-	if (OperationMode != EDisplayClusterOperationMode::Disabled)
+	if (OperationMode != EDisplayClusterOperationMode::Disabled && bIsPrimaryRootActor)
 	{
 		static const TConsoleVariableData<int32>* const InstancedStereoCVar = IConsoleManager::Get().FindTConsoleVariableDataInt(TEXT("vr.InstancedStereo"));
 		if (InstancedStereoCVar)
@@ -1032,23 +1053,150 @@ void ADisplayClusterRootActor::Tick(float DeltaSeconds)
 		}
 	}
 
+	// Update Preview settings
+	FDisplayClusterViewport_PreviewSettings NewPreviewSettings = GetPreviewSettings();
+
+	const bool bIsRunningPIE = IsRunningPIE();
+	const bool bIsRunningGame = IsRunningGame();
+	
+	// RootActor can have a preview in the scene
+	bool bEnablePreviewInScene = NewPreviewSettings.bPreviewEnable;
+	bool bPreviewInGame = false;
+
+	if (bIsRunningPIE)
+	{
+		bPreviewInGame = true;
+
+		if (IsPrimaryRootActorForPIE())
+		{
+			bEnablePreviewInScene = false;
+		}
+		else
+		{
+			// When running in a game or PIE, a special flag is required to show a preview for PIE/Standalone/Package
+			bEnablePreviewInScene = NewPreviewSettings.bPreviewEnable && NewPreviewSettings.bPreviewInGameEnable;
+		}
+	}
+	else if (bIsRunningGame || bIsRunningDisplayCluster)
+	{
+		bPreviewInGame = true;
+
+		// If RootActor is used in a cluster/game , special rules must be used for previewing in the scene:
+		if (bIsPrimaryRootActor)
+		{
+			// Disable preview in scene rendering for the primary RootActor in the game
+			bEnablePreviewInScene = false;
+		}
+		else
+		{
+			// When running in a game or PIE, a special flag is required to show a preview for PIE/Standalone/Package
+			bEnablePreviewInScene = NewPreviewSettings.bPreviewEnable && NewPreviewSettings.bPreviewInGameEnable;
+		}
+	}
+
 #if WITH_EDITOR
-	// Tick editor preview
-	if (!IsPreviewEnabled())
+	if (bMoviePipelineRenderPass)
 	{
-		ResetPreviewInternals_Editor();
+		// Disable preview rendering for RootActor used in MRQ
+		bEnablePreviewInScene = false;
 	}
+#endif
+
+	if (GetGameInstance() && GetGameInstance()->IsDedicatedServerInstance())
+	{
+		// Skip rendering on dedicated server
+		bEnablePreviewInScene = false;
+	}
+
+	if (bPreviewInGame && !NewPreviewSettings.bPreviewInGameRenderFrustum)
+	{
+		// Disable frustum preview rendering in game
+		NewPreviewSettings.bPreviewICVFXFrustums = false;
+	}
+
+	if (!bEnablePreviewInScene)
+	{
+		// Disable preview rendering
+		NewPreviewSettings.bPreviewEnable = false;
+	}
+
+	// Preview in scene is a special render:
+	if (!bEnablePreviewInScene)
+	{
+		// Hide invisible DCRA
+		SetActorHiddenInGame(true);
+
+		// Stop preview rendering.
+		if (IDisplayClusterViewportManager* ViewportManager = GetViewportManager())
+		{
+			// Update preview settings to new
+			ViewportManager->GetConfiguration().SetPreviewSettings(NewPreviewSettings);
+
+			// Stop preview rendering
+			ViewportManager->GetViewportManagerPreview().UpdateEntireClusterPreviewRender(false);
+		}
+	}
+	// When the preview is used by this DCRA, we must create a new ViewportManager
 	else
 	{
-		FDisplayClusterRootActorPreviewRenderingManager::HandleEvent(EDisplayClusterRootActorPreviewEvent::Render, this);
+		// Show DCRA preview in the scene for standalone\package
+		SetActorHiddenInGame(false);
+
+		if (IDisplayClusterViewportManager* ViewportManager = GetOrCreateViewportManager())
+		{
+			// Update preview settings to new
+			ViewportManager->GetConfiguration().SetPreviewSettings(NewPreviewSettings);
+
+			// Request rendering
+			ViewportManager->GetViewportManagerPreview().UpdateEntireClusterPreviewRender(true);
+		}
 	}
-#endif
 
 	SetLightCardOwnership();
 
 	Super::Tick(DeltaSeconds);
 }
 
+FDisplayClusterViewport_PreviewSettings ADisplayClusterRootActor::GetPreviewSettings() const
+{
+	if (IDisplayClusterViewportManager* ViewportManager = GetViewportManager())
+	{
+		if (!bUseLocalPreviewSetttings)
+		{
+			return ViewportManager->GetConfiguration().GetPreviewSettings();
+		}
+	}
+
+	FDisplayClusterViewport_PreviewSettings OutPreviewSettings;
+
+	// By default RootActor renders in scene colors
+	// but
+	OutPreviewSettings.EntireClusterPreviewRenderMode = EDisplayClusterRenderFrameMode::PreviewInScene;
+
+	OutPreviewSettings.bPreviewEnable       = bPreviewEnable;
+	OutPreviewSettings.bFreezePreviewRender = bFreezePreviewRender;
+
+	OutPreviewSettings.bEnablePreviewTechvis     = bEnablePreviewTechvis;
+	OutPreviewSettings.bPreviewEnablePostProcess = bPreviewEnablePostProcess;
+
+	OutPreviewSettings.bPreviewICVFXFrustums = bPreviewICVFXFrustums;
+	OutPreviewSettings.PreviewICVFXFrustumsFarDistance = PreviewICVFXFrustumsFarDistance;
+
+	OutPreviewSettings.bEnablePreviewMesh        = bEnablePreviewMesh;
+	OutPreviewSettings.bEnablePreviewEditableMesh = bEnablePreviewEditableMesh;
+
+	OutPreviewSettings.PreviewRenderTargetRatioMult = PreviewRenderTargetRatioMult;
+	OutPreviewSettings.PreviewMaxTextureDimension   = PreviewMaxTextureDimension;
+
+	OutPreviewSettings.TickPerFrame      = TickPerFrame;
+	OutPreviewSettings.ViewportsPerFrame = ViewportsPerFrame;
+
+	OutPreviewSettings.bPreviewInGameEnable        = bPreviewInGameEnable;
+	OutPreviewSettings.bPreviewInGameRenderFrustum = bPreviewInGameRenderFrustum;
+
+	return OutPreviewSettings;
+}
+
 void ADisplayClusterRootActor::PostLoad()
 {
 	Super::PostLoad();
@@ -1329,6 +1477,49 @@ UDisplayClusterCameraComponent* ADisplayClusterRootActor::GetDefaultCamera() con
 	return DefaultViewPoint;
 }
 
+FName ADisplayClusterRootActor::GetInternalDisplayDeviceName() const
+{
+	return FName("BasicDisplayDevice");
+}
+
+UDisplayClusterDisplayDeviceBaseComponent* ADisplayClusterRootActor::GetDefaultDisplayDevice() const
+{
+	if (DefaultDisplayDeviceComponent)
+	{
+		// Check already assigned/created
+		if (DefaultDisplayDeviceComponent->GetFName() == DefaultDisplayDeviceName)
+		{
+			return DefaultDisplayDeviceComponent;
+		}
+	}
+
+	DefaultDisplayDeviceComponent = nullptr;
+
+	// User assigned default
+	if (!DefaultDisplayDeviceName.IsNone() && DefaultDisplayDeviceName != GetInternalDisplayDeviceName())
+	{
+		DefaultDisplayDeviceComponent = GetComponentByName<UDisplayClusterDisplayDeviceBaseComponent>(DefaultDisplayDeviceName.ToString());
+
+		if (!DefaultDisplayDeviceComponent)
+		{
+			UE_LOG(LogDisplayClusterGame, Warning, TEXT("Invalid default display device. Using internal nDisplay default device."));
+		}
+	}
+
+	// Fallback to our internal default
+	if (!DefaultDisplayDeviceComponent)
+	{
+		DefaultDisplayDeviceComponent = BasicDisplayDeviceComponent;
+	}
+
+	return DefaultDisplayDeviceComponent;
+}
+
+ULineBatchComponent* ADisplayClusterRootActor::GetLineBatchComponent() const
+{
+	return LineBatcherComponent;
+}
+
 USceneComponent* ADisplayClusterRootActor::GetCommonViewPoint() const
 {
 	if (UDisplayClusterConfigurationData* ConfigData = GetConfigData())
diff --git a/Engine/Plugins/Runtime/nDisplay/Source/DisplayCluster/Private/Game/EngineClasses/Scene/DisplayClusterRootActorEditor.cpp b/Engine/Plugins/Runtime/nDisplay/Source/DisplayCluster/Private/Game/EngineClasses/Scene/DisplayClusterRootActorEditor.cpp
--- a/Engine/Plugins/Runtime/nDisplay/Source/DisplayCluster/Private/Game/EngineClasses/Scene/DisplayClusterRootActorEditor.cpp
+++ b/Engine/Plugins/Runtime/nDisplay/Source/DisplayCluster/Private/Game/EngineClasses/Scene/DisplayClusterRootActorEditor.cpp
@@ -6,15 +6,11 @@
 #include "Components/DisplayClusterCameraComponent.h"
 #include "Components/DisplayClusterICVFXCameraComponent.h"
 #include "Components/DisplayClusterOriginComponent.h"
-#include "Components/DisplayClusterPreviewComponent.h"
 #include "Components/DisplayClusterSceneComponentSyncParent.h"
 #include "Components/DisplayClusterScreenComponent.h"
 #include "Components/DisplayClusterStageGeometryComponent.h"
-#include "Components/DisplayDevice/DisplayClusterDisplayDeviceComponent.h"
-#include "Components/LineBatchComponent.h"
 #include "Components/MeshComponent.h"
 #include "Components/SceneComponent.h"
-#include "DisplayDevice/DisplayClusterDisplayDeviceUtils.h"
 
 #include "Config/IPDisplayClusterConfigManager.h"
 #include "DisplayClusterConfigurationStrings.h"
@@ -35,7 +31,6 @@
 #include "Misc/DisplayClusterStrings.h"
 
 #include "Misc/TransactionObjectEvent.h"
-#include "Render/Viewport/Configuration/DisplayClusterViewportConfigurationHelpers_ICVFX.h"
 #include "Render/Viewport/DisplayClusterViewportStrings.h"
 #include "Render/Viewport/IDisplayClusterViewport.h"
 #include "Render/Viewport/IDisplayClusterViewportManager.h"
@@ -63,46 +58,12 @@
 //////////////////////////////////////////////////////////////////////////////////////////////
 // ADisplayClusterRootActor
 //////////////////////////////////////////////////////////////////////////////////////////////
-void ADisplayClusterRootActor::ResetPreviewInternals_Editor()
-{
-	// Reset preview components before DCRA rebuild
-	ResetPreviewComponents_Editor(true);
-
-	PreviewRenderFrame.Reset();
-
-	TickPerFrameCounter = 0;
-	PreviewClusterNodeIndex = 0;
-
-	ResetClusterNodePreviewRendering_Editor();
-}
-
-void ADisplayClusterRootActor::ConfigureTechvis_Editor()
-{
-	TArray<UDisplayClusterPreviewComponent*> Components;
-	GetComponents(Components);
-
-	for (UDisplayClusterPreviewComponent* PreviewComponent : Components)
-	{
-		if (UMeshComponent* Mesh = PreviewComponent->GetPreviewMesh())
-		{
-			UE::DisplayClusterDisplayDeviceUtils::ConfigureTechvisForMesh(Mesh, bEnablePreviewTechvis);
-		}
-	}
-}
 
 void ADisplayClusterRootActor::Constructor_Editor()
 {
 	// Allow tick in editor for preview rendering
 	PrimaryActorTick.bStartWithTickEnabled = true;
 
-	// Set to internal default, user may change
-	DefaultDisplayDeviceName = GetInternalDisplayDeviceName();
-
-	// Our internal display device which always exists
-	BasicDisplayDeviceComponent = CreateDefaultSubobject<UDisplayClusterDisplayDeviceComponent>(GetInternalDisplayDeviceName());
-
-	ResetPreviewInternals_Editor();
-
 	FCoreUObjectDelegates::OnPackageReloaded.AddUObject(this, &ADisplayClusterRootActor::HandleAssetReload);
 
 	if (GEditor)
@@ -124,60 +85,14 @@ void ADisplayClusterRootActor::Destructor_Editor()
 	}
 }
 
-void ADisplayClusterRootActor::RenderPreview_Editor()
-{
-	if (IsPreviewEnabled() && IsInGameThread())
-	{
-		// When the preview is used by this DCRA, we must create a new ViewportManager
-		if (GetOrCreateViewportManager())
-		{
-			// Render viewport for preview material RTTs
-			if (bDeferPreviewGeneration)
-			{
-				// Hack to generate preview components on instances during map load.
-				// TODO: See if we can move InitializeRootActor out of PostLoad.
-				bDeferPreviewGeneration = false;
-				UpdatePreviewComponents();
-			}
-
-			// Update preview RTTs correspond to 'TickPerFrame' value
-			if (++TickPerFrameCounter >= TickPerFrame)
-			{
-				TickPerFrameCounter = 0;
-
-				ImplRenderPreview_Editor();
-			}
-
-			// preview frustums on each tick
-			ImplRenderPreviewFrustums_Editor();
-		}
-	}
-	else
-	{
-		ResetPreviewInternals_Editor();
-		if (FDisplayClusterViewportManager* ViewportManager = GetViewportManagerImpl())
-		{
-			if(ViewportManager->GetConfiguration().IsPreviewRendering())
-			{
-				// Preview is no longer in use.
-				// Release viewport manager with resources immediatelly
-				RemoveViewportManager();
-			}
-		}
-	}
-}
-
 void ADisplayClusterRootActor::PostActorCreated_Editor()
 {
-	ResetPreviewInternals_Editor();
+	ResetEntireClusterPreviewRendering();
 }
 
 void ADisplayClusterRootActor::PostLoad_Editor()
 {
-	// Generating the preview on load for instances in the world can't be done on PostLoad, components may not have loaded flags present.
-	bDeferPreviewGeneration = true;
-
-	ResetPreviewInternals_Editor();
+	ResetEntireClusterPreviewRendering();
 }
 
 void ADisplayClusterRootActor::EndPlay_Editor(const EEndPlayReason::Type EndPlayReason)
@@ -186,21 +101,17 @@ void ADisplayClusterRootActor::EndPlay_Editor(const EEndPlayReason::Type EndPlay
 
 void ADisplayClusterRootActor::Destroyed_Editor()
 {
-	ResetPreviewInternals_Editor();
-	ReleasePreviewComponents();
+	ResetEntireClusterPreviewRendering();
 
 	MarkAsGarbage();
 }
 
 void ADisplayClusterRootActor::BeginDestroy_Editor()
 {
-	ResetPreviewInternals_Editor();
-
-	ReleasePreviewComponents();
+	ResetEntireClusterPreviewRendering();
 
 	OnPreviewGenerated.Unbind();
 	OnPreviewDestroyed.Unbind();
-	bDeferPreviewGeneration = true;
 }
 
 void ADisplayClusterRootActor::RerunConstructionScripts_Editor()
@@ -212,71 +123,9 @@ void ADisplayClusterRootActor::RerunConstructionScripts_Editor()
 
 	UpdateInnerFrustumPriority();
 
-	// Reset preview components before DCRA rebuild
-	ResetPreviewComponents_Editor(false);
-
 	StageGeometryComponent->Invalidate();
 }
 
-void ADisplayClusterRootActor::EnableEditorRender(bool bValue)
-{
-	bEnableEditorRender = bValue;
-}
-
-bool ADisplayClusterRootActor::IsPreviewEnabled() const
-{
-	if (bMoviePipelineRenderPass)
-	{
-		// Disable preview rendering for MRQ
-		return false;
-	}
-
-	// -game or PIE case
-	if (IsRunningGameOrPIE())
-	{
-		// Only PIE is currently supported
-		return bPreviewEnable && GIsPlayInEditorWorld && (PreviewNodeId == DisplayClusterConfigurationStrings::gui::preview::PreviewNodeNone);
-	}
-
-	// Editor case
-	return bPreviewEnable || !PreviewEnableOverriders.IsEmpty();
-}
-
-bool ADisplayClusterRootActor::IsPreviewDrawnToScreens() const
-{
-	// The preview output is drawn to the screen when preview rendering is enabled (either manually or through an override) AND when
-	// the bPreviewEnable flag is specifically set to true. 
-	return IsPreviewEnabled() && bPreviewEnable;
-}
-
-
-// Return all RTT RHI resources for preview
-void ADisplayClusterRootActor::GetPreviewRenderTargetableTextures(const TArray<FString>& InViewportNames, TArray<FTextureRHIRef>& OutTextures)
-{
-	check(IsInGameThread());
-
-	for(const TPair<FString, TObjectPtr<UDisplayClusterPreviewComponent>>& PreviewComponentIt : PreviewComponents)
-	{
-		if (PreviewComponentIt.Value)
-		{
-			const int32 OutTextureIndex = InViewportNames.Find(PreviewComponentIt.Value->GetViewportId());
-			if (OutTextureIndex != INDEX_NONE)
-			{
-				// Add scope for func GetRenderTargetTexture()
-				if (UTextureRenderTarget2D* RenderTarget2D = ShouldThisFrameOutputPreviewToPostProcessRenderTarget() ?
-					PreviewComponentIt.Value->GetRenderTargetTexturePostProcess() : PreviewComponentIt.Value->GetRenderTargetTexture())
-				{
-					FTextureRenderTargetResource* DstRenderTarget = RenderTarget2D->GameThread_GetRenderTargetResource();
-					if (DstRenderTarget != nullptr)
-					{
-						OutTextures[OutTextureIndex] = DstRenderTarget->TextureRHI;
-					}
-				}
-			}
-		}
-	}
-}
-
 void ADisplayClusterRootActor::UpdateInnerFrustumPriority()
 {
 	if (InnerFrustumPriority.Num() == 0)
@@ -350,406 +199,6 @@ void ADisplayClusterRootActor::SetIsSelectedInEditor(bool bValue)
 	bIsSelectedInEditor = bValue;
 }
 
-IDisplayClusterViewport* ADisplayClusterRootActor::FindPreviewViewport(const FString& InViewportId) const
-{
-	if (FDisplayClusterViewportManager* ViewportManager = GetViewportManagerImpl())
-	{
-		return ViewportManager->FindViewport(InViewportId);
-	}
-
-	return nullptr;
-}
-
-EDisplayClusterRenderFrameMode ADisplayClusterRootActor::GetPreviewRenderMode_Editor() const
-{
-	// Todo: we can add HitProxy rendering support to this function with EDisplayClusterRenderFrameMode::PreviewProxyHitInScene
-	
-	return EDisplayClusterRenderFrameMode::PreviewInScene;
-}
-
-bool ADisplayClusterRootActor::ImplUpdatePreviewConfiguration_Editor(const FString& InClusterNodeId)
-{
-	DECLARE_SCOPE_CYCLE_COUNTER(TEXT("ADisplayClusterRootActor::ImplUpdatePreviewConfiguration_Editor"), STAT_ImplUpdatePreviewConfiguration_Editor, STATGROUP_NDisplay);
-	
-	// The function UpdateConfiguration() rebuilds viewports and their internals, so rendering is no longer valid for now.
-	// Reset current preview rendering.
-	ResetClusterNodePreviewRendering_Editor();
-
-	if (IsPreviewEnabled())
-	{
-		if (FDisplayClusterViewportManager* ViewportManager = GetViewportManagerImpl())
-		{
-			// Now we render this node
-			PreviewRenderFrameClusterNodeId = InClusterNodeId;
-
-			// Update current world from scene DCRA
-			ADisplayClusterRootActor* SceneRootActor = ViewportManager->GetConfiguration().GetRootActor(EDisplayClusterRootActorType::Scene);
-			ViewportManager->GetConfiguration().SetCurrentWorld(SceneRootActor ? SceneRootActor->GetWorld() : GetWorld());
-
-			// Update local node viewports (update\create\delete) and build new render frame
-			return ViewportManager->GetConfiguration().UpdateConfigurationForClusterNode(GetPreviewRenderMode_Editor(), InClusterNodeId);
-		}
-	}
-
-	return false;
-}
-
-bool ADisplayClusterRootActor::IsActiveClusterNodePreviewRendering_Editor() const
-{
-	return PreviewRenderFrame.IsValid() && PreviewViewportIndex >= 0 && !PreviewRenderFrameClusterNodeId.IsEmpty();
-}
-
-void ADisplayClusterRootActor::ResetClusterNodePreviewRendering_Editor()
-{
-	PreviewViewportIndex = -1;
-	PreviewRenderFrame.Reset();
-	PreviewRenderFrameClusterNodeId.Empty();
-}
-
-bool ADisplayClusterRootActor::ImplUpdatePreviewRenderFrame_Editor(const FString& InClusterNodeId)
-{
-	// Skip rendering on dedicated server
-	if (GetGameInstance() && GetGameInstance()->IsDedicatedServerInstance())
-	{
-		return false;
-	}
-
-	// Update cluster node for render:
-	if (PreviewRenderFrameClusterNodeId != InClusterNodeId)
-	{
-		ResetClusterNodePreviewRendering_Editor();
-	}
-
-	if (!PreviewRenderFrame.IsValid())
-	{
-		if (!ImplUpdatePreviewConfiguration_Editor(InClusterNodeId))
-		{
-			return false;
-		}
-
-		// Update all preview components resources before render
-		for (const TTuple<FString, TObjectPtr<UDisplayClusterPreviewComponent>>& PreviewComponentIt : PreviewComponents)
-		{
-			if (PreviewComponentIt.Value && (InClusterNodeId.IsEmpty() || PreviewComponentIt.Value->GetClusterNodeId() == InClusterNodeId))
-			{
-				PreviewComponentIt.Value->UpdatePreviewResources();
-			}
-		}
-
-		PreviewRenderFrame = MakeUnique<FDisplayClusterRenderFrame>();
-
-		if (FDisplayClusterViewportManager* ViewportManager = GetViewportManagerImpl())
-		{
-			// Update preview viewports from settings
-			if (!ViewportManager->BeginNewFrame(nullptr, *PreviewRenderFrame))
-			{
-				PreviewRenderFrame.Reset();
-
-				return false;
-			}
-
-			// Begin Render Preview For Cluster Node
-			PreviewViewportIndex = 0;
-
-			// Initialize frame for render
-			ViewportManager->InitializeNewFrame();
-		}
-	}
-
-	return IsActiveClusterNodePreviewRendering_Editor();
-}
-
-bool ADisplayClusterRootActor::ImplRenderPassPreviewClusterNode_Editor(const FString& InClusterNodeId)
-{
-	if(!ImplUpdatePreviewRenderFrame_Editor(InClusterNodeId))
-	{
-		return false;
-	}
-
-	int32 ViewportsAmount = ViewportsPerFrame - PreviewViewportsRenderedInThisFrameCnt;
-	if (ViewportsAmount <= 0)
-	{
-		// All viewports for this pass is rendered
-		return false;
-	}
-
-	if (FDisplayClusterViewportManager* ViewportManager = GetViewportManagerImpl())
-	{
-		bool bFrameRendered = false;
-		int32 RenderedViewportsAmount = 0;
-
-		ViewportManager->RenderInEditor(*PreviewRenderFrame, nullptr, PreviewViewportIndex, ViewportsAmount, RenderedViewportsAmount, bFrameRendered);
-
-		// Increase viewport index
-		PreviewViewportIndex += ViewportsAmount;
-
-		// Count only rendered viewports
-		PreviewViewportsRenderedInThisFrameCnt += RenderedViewportsAmount;
-
-		if (bFrameRendered)
-		{
-			// current cluster node is composed
-			ResetClusterNodePreviewRendering_Editor();
-
-			// Send event about RTT changed
-			OnPreviewGenerated.ExecuteIfBound();
-
-			return true;
-		}
-	}
-
-	return false;
-}
-
-void ADisplayClusterRootActor::ImplRenderPreview_Editor()
-{
-	if (CurrentConfigData == nullptr || !IsPreviewEnabled())
-	{
-		// no preview
-		return;
-	}
-
-	if (PreviewClusterNodeIndex < 0)
-	{
-		// Allow preview render
-		PreviewClusterNodeIndex = 0;
-	}
-
-	// per-node render
-	TArray<FString> ExistClusterNodesIDs;
-	CurrentConfigData->Cluster->GetNodeIds(ExistClusterNodesIDs);
-
-	// When viewports count changed, reset render cycle
-	if (PreviewClusterNodeIndex >= ExistClusterNodesIDs.Num())
-	{
-		PreviewClusterNodeIndex = 0;
-	}
-
-	// Try render all nodes in one pass
-	int32 NumNodesForSceneMaterialPreview = ExistClusterNodesIDs.Num();
-	PreviewViewportsRenderedInThisFrameCnt = 0;
-
-	for (int32 NodeIt = 0; NodeIt < NumNodesForSceneMaterialPreview; NodeIt++)
-	{
-		int32 ViewportsAmount = ViewportsPerFrame - PreviewViewportsRenderedInThisFrameCnt;
-		if (ViewportsAmount <= 0)
-		{
-			// All viewports for this pass is rendered
-			return;
-		}
-
-		/**
-		 * Render this cluster node viewports
-		 * Note: ViewportManager should be used correctly:
-		 * 1. ViewportManager->UpdateConfiguration(ClusterNode, Configuration) - only when changing DCRA or rendering settings (or cluster node)
-		 * 2. ViewportManager->BeginNewFrame(PreviewWorld, RenderSettings) - once per frame for the entire cluster node
-		 * 3. The cluster node will be rendered using the RenderInEditor() function, which will be called multiple times until the node is completely rendered.
-		 * 3.1. ViewportManager->RenderInEditor(PreviewInfo)
-		 * ...
-		 * 3.N. ViewportManager->RenderInEditor(PreviewInfo)
-		 * 
-		 * Warning: When any function of this workflow is called in the wrong order, the entire workflow must be restarted.
-		 */
-		if (!ImplRenderPassPreviewClusterNode_Editor(ExistClusterNodesIDs[PreviewClusterNodeIndex]))
-		{
-			// Cluster node render still in progress..
-			return;
-		}
-
-		// Loop over cluster nodes
-		PreviewClusterNodeIndex++;
-		if (PreviewClusterNodeIndex >= ExistClusterNodesIDs.Num())
-		{
-			PreviewClusterNodeIndex = 0;
-
-			bOutputFrameToPostProcessRenderTarget = bOutputFrameToPostProcessRenderTarget ?
-				bPreviewEnablePostProcess :
-				bPreviewEnablePostProcess || DoObserversNeedPostProcessRenderTarget();
-		}
-
-		if (PreviewClusterNodeIndex < 0)
-		{
-			// Frame captured. stop render
-			break;
-		}
-	}
-}
-
-void ADisplayClusterRootActor::ImplRenderPreviewFrustums_Editor()
-{
-	FDisplayClusterViewportManager* ViewportManager = GetViewportManagerImpl();
-	if (CurrentConfigData == nullptr || !ViewportManager)
-	{
-		return;
-	}
-
-	// frustum preview viewports
-	TArray<IDisplayClusterViewport*> FrustumPreviewViewports;
-
-	for (const TPair<FString, TObjectPtr<UDisplayClusterConfigurationClusterNode>>& Node : CurrentConfigData->Cluster->Nodes)
-	{
-		if (Node.Value == nullptr)
-		{
-			continue;
-		}
-
-		// collect node viewports
-		for (const TPair<FString, TObjectPtr<UDisplayClusterConfigurationViewport>>& ViewportConfig : Node.Value->Viewports)
-		{
-			if (ViewportConfig.Value->bAllowPreviewFrustumRendering == false || ViewportConfig.Value->bIsVisible == false || ViewportConfig.Value == nullptr)
-			{
-				continue;
-			}
-
-			IDisplayClusterViewport* Viewport = ViewportManager->FindViewport(ViewportConfig.Key);
-			if (Viewport != nullptr && Viewport->GetContexts().Num() > 0 && Viewport->GetRenderSettings().bEnable)
-			{
-				FrustumPreviewViewports.Add(Viewport);
-			}
-		}
-	}
-
-	// collect incameras
-	if (bPreviewICVFXFrustums)
-	{
-		// Iterate over rendered inner camera viewports (whole cluster)
-		for (const TSharedPtr<FDisplayClusterViewport, ESPMode::ThreadSafe>& InnerCameraViewportIt : FDisplayClusterViewportConfigurationHelpers_ICVFX::GetAllVisibleInnerCameraViewports(*ViewportManager->Configuration))
-		{
-			if (InnerCameraViewportIt.IsValid())
-			{
-				FrustumPreviewViewports.Add(InnerCameraViewportIt.Get());
-			}
-		}
-	}
-
-	for (IDisplayClusterViewport* ViewportIt : FrustumPreviewViewports)
-	{
-		const TArray<FDisplayClusterViewport_Context>& Contexts = ViewportIt->GetContexts();
-
-		// Due to rendering optimizations for DCRA preview (node rendered over multiple frames).
-		// As a result, viewport math values (stored in contexts) may not be ready at the moment.
-		FFrustumPreviewViewportContextCache FrustumPreviewViewportContext;
-		bool bIsValidViewportContext = false;
-
-		// Preview rendered only in mono
-		if (Contexts.Num() == 1 && EnumHasAllFlags(Contexts[0].ContextState, EDisplayClusterViewportContextState::HasCalculatedProjectionMatrix | EDisplayClusterViewportContextState::HasCalculatedViewPoint))
-		{
-			FrustumPreviewViewportContext.ProjectionMatrix = Contexts[0].ProjectionMatrix;
-			FrustumPreviewViewportContext.ViewLocation = Contexts[0].ViewLocation;
-			FrustumPreviewViewportContext.ViewRotation = Contexts[0].ViewRotation;
-
-			if (ViewportIt->GetProjectionMatrix(0, FrustumPreviewViewportContext.ProjectionMatrix))
-			{
-				bIsValidViewportContext = true;
-			}
-		}
-
-		// Get cached value
-		if (!bIsValidViewportContext)
-		{
-			FFrustumPreviewViewportContextCache* const LastValidContext = FrustumPreviewViewportContextCache.Find(ViewportIt->GetId());
-			if (LastValidContext)
-			{
-				FrustumPreviewViewportContext = *LastValidContext;
-				bIsValidViewportContext = true;
-			}
-		}
-
-		if(bIsValidViewportContext)
-		{
-			// Update cache
-			FrustumPreviewViewportContextCache.Emplace(ViewportIt->GetId(), FrustumPreviewViewportContext);
-
-			// Render the frustum
-
-			FMatrix ViewRotationMatrix = FInverseRotationMatrix(FrustumPreviewViewportContext.ViewRotation) * FMatrix(
-				FPlane(0, 0, 1, 0),
-				FPlane(1, 0, 0, 0),
-				FPlane(0, 1, 0, 0),
-				FPlane(0, 0, 0, 1));
-			const FMatrix ViewMatrix = FTranslationMatrix(-FrustumPreviewViewportContext.ViewLocation) * ViewRotationMatrix;
-
-			ImplRenderPreviewViewportFrustum_Editor(FrustumPreviewViewportContext.ProjectionMatrix, ViewMatrix, FrustumPreviewViewportContext.ViewLocation);
-		}
-	}
-}
-
-void ADisplayClusterRootActor::ImplRenderPreviewViewportFrustum_Editor(const FMatrix ProjectionMatrix, const FMatrix ViewMatrix, const FVector ViewOrigin)
-{
-	const float FarPlane = PreviewICVFXFrustumsFarDistance;
-	const float NearPlane = GNearClippingPlane;
-	const FColor Color = FColor::Green;
-	const float Thickness = 1.0f;
-
-	const UWorld* World = GetWorld();
-	ULineBatchComponent* LineBatcher = World ? World->LineBatcher : nullptr;
-	if (!LineBatcher)
-	{
-		return;
-	}
-
-	// Get FOV and AspectRatio from the view's projection matrix.
-	const float AspectRatio = ProjectionMatrix.M[1][1] / ProjectionMatrix.M[0][0];
-	const bool bIsPerspectiveProjection = true;
-
-	// Build the camera frustum for this cascade
-	const float HalfHorizontalFOV = bIsPerspectiveProjection ? FMath::Atan(1.0f / ProjectionMatrix.M[0][0]) : PI / 4.0f;
-	const float HalfVerticalFOV = bIsPerspectiveProjection ? FMath::Atan(1.0f / ProjectionMatrix.M[1][1]) : FMath::Atan((FMath::Tan(PI / 4.0f) / AspectRatio));
-	const float AsymmetricFOVScaleX = ProjectionMatrix.M[2][0];
-	const float AsymmetricFOVScaleY = ProjectionMatrix.M[2][1];
-	
-	// Near plane
-	const float StartHorizontalTotalLength = NearPlane * FMath::Tan(HalfHorizontalFOV);
-	const float StartVerticalTotalLength = NearPlane * FMath::Tan(HalfVerticalFOV);
-	const FVector StartCameraLeftOffset = ViewMatrix.GetColumn(0) * -StartHorizontalTotalLength * (1 + AsymmetricFOVScaleX);
-	const FVector StartCameraRightOffset = ViewMatrix.GetColumn(0) * StartHorizontalTotalLength * (1 - AsymmetricFOVScaleX);
-	const FVector StartCameraBottomOffset = ViewMatrix.GetColumn(1) * -StartVerticalTotalLength * (1 + AsymmetricFOVScaleY);
-	const FVector StartCameraTopOffset = ViewMatrix.GetColumn(1) * StartVerticalTotalLength * (1 - AsymmetricFOVScaleY);
-	
-	// Far plane
-	const float EndHorizontalTotalLength = FarPlane * FMath::Tan(HalfHorizontalFOV);
-	const float EndVerticalTotalLength = FarPlane * FMath::Tan(HalfVerticalFOV);
-	const FVector EndCameraLeftOffset = ViewMatrix.GetColumn(0) * -EndHorizontalTotalLength * (1 + AsymmetricFOVScaleX);
-	const FVector EndCameraRightOffset = ViewMatrix.GetColumn(0) * EndHorizontalTotalLength * (1 - AsymmetricFOVScaleX);
-	const FVector EndCameraBottomOffset = ViewMatrix.GetColumn(1) * -EndVerticalTotalLength * (1 + AsymmetricFOVScaleY);
-	const FVector EndCameraTopOffset = ViewMatrix.GetColumn(1) * EndVerticalTotalLength * (1 - AsymmetricFOVScaleY);
-	
-	const FVector CameraDirection = ViewMatrix.GetColumn(2);
-
-	// Preview frustum vertices
-	FVector PreviewFrustumVerts[8];
-
-	// Get the 4 points of the camera frustum near plane, in world space
-	PreviewFrustumVerts[0] = ViewOrigin + CameraDirection * NearPlane + StartCameraRightOffset + StartCameraTopOffset;         // 0 Near  Top    Right
-	PreviewFrustumVerts[1] = ViewOrigin + CameraDirection * NearPlane + StartCameraRightOffset + StartCameraBottomOffset;      // 1 Near  Bottom Right
-	PreviewFrustumVerts[2] = ViewOrigin + CameraDirection * NearPlane + StartCameraLeftOffset + StartCameraTopOffset;          // 2 Near  Top    Left
-	PreviewFrustumVerts[3] = ViewOrigin + CameraDirection * NearPlane + StartCameraLeftOffset + StartCameraBottomOffset;       // 3 Near  Bottom Left
-
-	// Get the 4 points of the camera frustum far plane, in world space
-	PreviewFrustumVerts[4] = ViewOrigin + CameraDirection * FarPlane + EndCameraRightOffset + EndCameraTopOffset;         // 4 Far  Top    Right
-	PreviewFrustumVerts[5] = ViewOrigin + CameraDirection * FarPlane + EndCameraRightOffset + EndCameraBottomOffset;      // 5 Far  Bottom Right
-	PreviewFrustumVerts[6] = ViewOrigin + CameraDirection * FarPlane + EndCameraLeftOffset + EndCameraTopOffset;          // 6 Far  Top    Left
-	PreviewFrustumVerts[7] = ViewOrigin + CameraDirection * FarPlane + EndCameraLeftOffset + EndCameraBottomOffset;       // 7 Far  Bottom Left	
-	
-	// frustum lines
-	LineBatcher->DrawLine(PreviewFrustumVerts[0], PreviewFrustumVerts[4], Color, SDPG_World, Thickness, 0.f); // right top
-	LineBatcher->DrawLine(PreviewFrustumVerts[1], PreviewFrustumVerts[5], Color, SDPG_World, Thickness, 0.f); // right bottom
-	LineBatcher->DrawLine(PreviewFrustumVerts[2], PreviewFrustumVerts[6], Color, SDPG_World, Thickness, 0.f); // left top
-	LineBatcher->DrawLine(PreviewFrustumVerts[3], PreviewFrustumVerts[7], Color, SDPG_World, Thickness, 0.f); // left bottom
-
-	// near plane square
-	LineBatcher->DrawLine(PreviewFrustumVerts[0], PreviewFrustumVerts[1], Color, SDPG_World, Thickness, 0.f); // right top to right bottom
-	LineBatcher->DrawLine(PreviewFrustumVerts[1], PreviewFrustumVerts[3], Color, SDPG_World, Thickness, 0.f); // right bottom to left bottom
-	LineBatcher->DrawLine(PreviewFrustumVerts[3], PreviewFrustumVerts[2], Color, SDPG_World, Thickness, 0.f); // left bottom to left top
-	LineBatcher->DrawLine(PreviewFrustumVerts[2], PreviewFrustumVerts[0], Color, SDPG_World, Thickness, 0.f); // left top to right top
-
-	// far plane square
-	LineBatcher->DrawLine(PreviewFrustumVerts[4], PreviewFrustumVerts[5], Color, SDPG_World, Thickness, 0.f); // right top to right bottom
-	LineBatcher->DrawLine(PreviewFrustumVerts[5], PreviewFrustumVerts[7], Color, SDPG_World, Thickness, 0.f); // right bottom to left bottom
-	LineBatcher->DrawLine(PreviewFrustumVerts[7], PreviewFrustumVerts[6], Color, SDPG_World, Thickness, 0.f); // left bottom to left top
-	LineBatcher->DrawLine(PreviewFrustumVerts[6], PreviewFrustumVerts[4], Color, SDPG_World, Thickness, 0.f); // left top to right top
-}
-
 static FName Name_RelativeLocation = USceneComponent::GetRelativeLocationPropertyName();
 static FName Name_RelativeRotation = USceneComponent::GetRelativeRotationPropertyName();
 static FName Name_RelativeScale3D = USceneComponent::GetRelativeScale3DPropertyName();
@@ -846,8 +295,6 @@ void ADisplayClusterRootActor::PostEditChangeProperty(FPropertyChangedEvent& Pro
 	}
 	else
 	{
-		// any property update causes a DCRA rebuild. Restore preview material show
-		ResetPreviewComponents_Editor(true);
 		bResetPreviewComponents = true;
 		Super::PostEditChangeProperty(PropertyChangedEvent);
 	}
@@ -867,7 +314,7 @@ void ADisplayClusterRootActor::PostEditChangeProperty(FPropertyChangedEvent& Pro
 	else if (PropertyName == GET_MEMBER_NAME_CHECKED(ADisplayClusterRootActor, TickPerFrame)
 		|| PropertyName == GET_MEMBER_NAME_CHECKED(ADisplayClusterRootActor, ViewportsPerFrame))
 	{
-		ResetPreviewInternals_Editor();
+		ResetEntireClusterPreviewRendering();
 
 		bReinitializeActor = false;
 	}
@@ -878,14 +325,9 @@ void ADisplayClusterRootActor::PostEditChangeProperty(FPropertyChangedEvent& Pro
 	else if (PropertyName == GET_MEMBER_NAME_CHECKED(ADisplayClusterRootActor, bPreviewEnablePostProcess)
 		|| PropertyName == GET_MEMBER_NAME_CHECKED(ADisplayClusterRootActor, DefaultDisplayDeviceName))
 	{
-		ResetPreviewComponents_Editor(false);
-		PreviewRenderFrame.Reset();
+		ResetEntireClusterPreviewRendering();
 		bReinitializeActor = false;
 	}
-	else if (PropertyName == GET_MEMBER_NAME_CHECKED(ADisplayClusterRootActor, bEnablePreviewTechvis))
-	{
-		ConfigureTechvis_Editor();
-	}
 	
 	if (bReinitializeActor)
 	{
@@ -894,7 +336,6 @@ void ADisplayClusterRootActor::PostEditChangeProperty(FPropertyChangedEvent& Pro
 
 	if (bResetPreviewComponents)
 	{
-		ResetPreviewComponents_Editor(false);
 	}
 }
 
@@ -902,10 +343,6 @@ void ADisplayClusterRootActor::PostEditMove(bool bFinished)
 {
 	// Don't update the preview with the config data if we're just moving the actor.
 	Super::PostEditMove(bFinished);
-
-	FrustumPreviewViewportContextCache.Empty();
-
-	ResetPreviewComponents_Editor(false);
 }
 
 void ADisplayClusterRootActor::HandleAssetReload(const EPackageReloadPhase InPackageReloadPhase,
@@ -929,265 +366,4 @@ void ADisplayClusterRootActor::OnEndObjectMovement(UObject& InObject)
 	}
 }
 
-void ADisplayClusterRootActor::ResetPreviewComponents_Editor(bool bInRestoreSceneMaterial)
-{
-	TArray<UDisplayClusterPreviewComponent*> AllPreviewComponents;
-	GetComponents(AllPreviewComponents);
-	
-	if (AllPreviewComponents.Num() > 0)
-	{
-		// Preview components exist and can restore the scene material
-		for (UDisplayClusterPreviewComponent* ExistingComp : AllPreviewComponents)
-		{
-			ExistingComp->ResetPreviewComponent(bInRestoreSceneMaterial);
-		}
-	}
-	else if (bInRestoreSceneMaterial && !bPreviewEnable && CurrentConfigData && !IsTemplate())
-	{
-		// There are no preview components but the preview is disabled and we want to restore the scene material
-		CurrentConfigData->ForEachViewport([this](const TObjectPtr<UDisplayClusterConfigurationViewport>& Viewport)
-		{
-			// First locate the display device the viewport references
-			if (const UDisplayClusterDisplayDeviceBaseComponent* DisplayDevice = UE::DisplayClusterDisplayDeviceUtils::FindAndSyncDisplayDeviceFromViewport(Viewport))
-			{
-				// Next find the preview mesh component the viewport is assigned to. We can't just use FindPreviewViewport and its referenced mesh
-				// because the preview viewport won't exist if the editor preview is disabled when the root actor is initialized
-				FString ParameterKey;
-				if (Viewport->ProjectionPolicy.Type == DisplayClusterProjectionStrings::projection::Simple)
-				{
-					ParameterKey = DisplayClusterProjectionStrings::cfg::simple::Screen;
-				}
-				else if (Viewport->ProjectionPolicy.Type == DisplayClusterProjectionStrings::projection::Camera)
-				{
-					ParameterKey = DisplayClusterProjectionStrings::cfg::camera::Component;
-				}
-				else if (Viewport->ProjectionPolicy.Type == DisplayClusterProjectionStrings::projection::Mesh)
-				{
-					ParameterKey = DisplayClusterProjectionStrings::cfg::mesh::Component;
-				}
-
-				if (!ParameterKey.IsEmpty())
-				{
-					if (const FString* ComponentName = Viewport->ProjectionPolicy.Parameters.Find(ParameterKey))
-					{
-						if (UMeshComponent* MeshComponent = GetComponentByName<UMeshComponent>(*ComponentName))
-						{
-							MeshComponent->SetMaterial(0, DisplayDevice->GetMeshMaterial());
-						}
-					}
-				}
-			}
-		});
-	}
-}
-
-void ADisplayClusterRootActor::UpdatePreviewComponents()
-{
-	DECLARE_SCOPE_CYCLE_COUNTER(TEXT("ADisplayClusterRootActor::UpdatePreviewComponents"), STAT_UpdatePreviewComponents, STATGROUP_NDisplay);
-	
-	if (IsTemplate() || bDeferPreviewGeneration)
-	{
-		return;
-	}
-
-	// Do not updated inside preview rendering workflow
-	if (!IsActiveClusterNodePreviewRendering_Editor())
-	{
-		ImplUpdatePreviewConfiguration_Editor(DisplayClusterConfigurationStrings::gui::preview::PreviewNodeAll);
-	}
-
-	TArray<UDisplayClusterPreviewComponent*> IteratedPreviewComponents;
-
-	auto DestroyPreviewComponent = [this](UDisplayClusterPreviewComponent* PreviewComp)
-	{
-		PreviewComponents.Remove(PreviewComp->GetName());
-			
-		// Reset preview components before unregister
-		PreviewComp->ResetPreviewComponent(true);
-
-		PreviewComp->UnregisterComponent();
-		PreviewComp->DestroyComponent();
-	};
-	
-	if (CurrentConfigData != nullptr && IsPreviewEnabled())
-	{
-		for (const TPair<FString, TObjectPtr<UDisplayClusterConfigurationClusterNode>>& Node : CurrentConfigData->Cluster->Nodes)
-		{
-			if (Node.Value == nullptr)
-			{
-				continue;
-			}
-
-			for (const TPair<FString, TObjectPtr<UDisplayClusterConfigurationViewport>>& Viewport : Node.Value->Viewports)
-			{
-				const FString PreviewCompId = GeneratePreviewComponentName_Editor(Node.Key, Viewport.Key);
-				UDisplayClusterPreviewComponent* PreviewComp = PreviewComponents.FindRef(PreviewCompId);
-
-				if (PreviewComp && PreviewComp->GetOwner() != this)
-				{
-					// In this case a viewport was likely deleted, the blueprint compiled, then the deletion undone.
-					// Just destroy the preview component and start over.
-					ensure(!PreviewComp->GetOwner() || PreviewComp->GetOwner()->GetName().StartsWith(TEXT("REINST_")));
-					DestroyPreviewComponent(PreviewComp);
-					PreviewComp = nullptr;
-				}
-
-				const bool bMarkDirty = false;
-				Modify(bMarkDirty);
-				if (!PreviewComp)
-				{
-					PreviewComp = NewObject<UDisplayClusterPreviewComponent>(this, FName(*PreviewCompId), RF_DuplicateTransient | RF_Transactional | RF_NonPIEDuplicateTransient);
-					check(PreviewComp);
-
-					PreviewComponents.Emplace(PreviewCompId, PreviewComp);
-
-					// Refresh preview when new component created
-					PreviewClusterNodeIndex = 0;
-				}
-
-				// Make sure we're an owned component. Possible this can be lost on undo/redo without a recompile.
-				AddOwnedComponent(PreviewComp);
-				
-				if (GetWorld() && !PreviewComp->IsRegistered())
-				{
-					PreviewComp->RegisterComponent();
-				}
-
-				// Always reinitialize so changes impact the preview component.
-				PreviewComp->InitializePreviewComponent(this, Node.Key, Viewport.Key, Viewport.Value);
-				if (UMeshComponent* PreviewMesh = PreviewComp->GetPreviewMesh())
-				{
-					PreviewMesh->SetHiddenInGame(false);
-				}
-
-				IteratedPreviewComponents.Add(PreviewComp);
-			}
-		}
-	}
-
-	// Cleanup unused components.
-	TArray<UDisplayClusterPreviewComponent*> AllPreviewComponents;
-	GetComponents(AllPreviewComponents);
-	
-	for (UDisplayClusterPreviewComponent* ExistingComp : AllPreviewComponents)
-	{
-		if (!IteratedPreviewComponents.Contains(ExistingComp))
-		{
-			DestroyPreviewComponent(ExistingComp);
-		}
-	}
-}
-
-void ADisplayClusterRootActor::ReleasePreviewComponents()
-{
-	for (const TPair<FString, TObjectPtr<UDisplayClusterPreviewComponent>>& CompPair : PreviewComponents)
-	{
-		if (CompPair.Value)
-		{
-			CompPair.Value->ResetPreviewComponent(true);
-
-			CompPair.Value->UnregisterComponent();
-			CompPair.Value->DestroyComponent();
-		}
-	}
-
-	PreviewComponents.Reset();
-	OnPreviewDestroyed.ExecuteIfBound();
-}
-
-int32 ADisplayClusterRootActor::SubscribeToPostProcessRenderTarget(const uint8* Object)
-{
-	check(Object);
-	PostProcessRenderTargetObservers.Add(Object);
-	return PostProcessRenderTargetObservers.Num();
-}
-
-int32 ADisplayClusterRootActor::UnsubscribeFromPostProcessRenderTarget(const uint8* Object)
-{
-	check(Object);
-	PostProcessRenderTargetObservers.Remove(Object);
-	return PostProcessRenderTargetObservers.Num();
-}
-
-bool ADisplayClusterRootActor::DoObserversNeedPostProcessRenderTarget() const
-{
-	return PostProcessRenderTargetObservers.Num() > 0;
-}
-
-bool ADisplayClusterRootActor::ShouldThisFrameOutputPreviewToPostProcessRenderTarget() const
-{
-	return bOutputFrameToPostProcessRenderTarget;
-}
-
-void ADisplayClusterRootActor::AddPreviewEnableOverride(const uint8* Object)
-{
-	check(Object);
-	PreviewEnableOverriders.Add(Object);
-
-	UpdatePreviewComponents();
-}
-
-void ADisplayClusterRootActor::RemovePreviewEnableOverride(const uint8* Object)
-{
-	check(Object);
-	PreviewEnableOverriders.Remove(Object);
-
-	UpdatePreviewComponents();
-}
-
-FName ADisplayClusterRootActor::GetInternalDisplayDeviceName()
-{
-	return FName("BasicDisplayDevice");
-}
-
-UDisplayClusterDisplayDeviceBaseComponent* ADisplayClusterRootActor::GetDefaultDisplayDevice()
-{
-	if (DefaultDisplayDeviceComponent)
-	{
-		// Check already assigned/created
-		if (DefaultDisplayDeviceComponent->GetFName() == DefaultDisplayDeviceName)
-		{
-			return DefaultDisplayDeviceComponent;
-		}
-	}
-
-	DefaultDisplayDeviceComponent = nullptr;
-
-	// User assigned default
-	if (!DefaultDisplayDeviceName.IsNone() && DefaultDisplayDeviceName != GetInternalDisplayDeviceName())
-	{
-		DefaultDisplayDeviceComponent =
-			GetComponentByName<UDisplayClusterDisplayDeviceBaseComponent>(DefaultDisplayDeviceName.ToString());
-
-		if (!DefaultDisplayDeviceComponent)
-		{
-			UE_LOG(LogDisplayClusterGame, Warning, TEXT("Invalid default display device. Using internal nDisplay default device."));
-		}
-	}
-
-	// Fallback to our internal default
-	if (!DefaultDisplayDeviceComponent)
-	{
-		DefaultDisplayDeviceComponent = BasicDisplayDeviceComponent;
-	}
-
-	return DefaultDisplayDeviceComponent;
-}
-
-FString ADisplayClusterRootActor::GeneratePreviewComponentName_Editor(const FString& NodeId, const FString& ViewportId) const
-{
-	return FString::Printf(TEXT("%s_%s"), *NodeId, *ViewportId);
-}
-
-UDisplayClusterPreviewComponent* ADisplayClusterRootActor::GetPreviewComponent(const FString& NodeId, const FString& ViewportId) const
-{
-	const FString PreviewCompId = GeneratePreviewComponentName_Editor(NodeId, ViewportId);
-	if (PreviewComponents.Contains(PreviewCompId))
-	{
-		return PreviewComponents[PreviewCompId];
-	}
-
-	return nullptr;
-}
-
 #endif
diff --git a/Engine/Plugins/Runtime/nDisplay/Source/DisplayCluster/Private/Render/Containers/DisplayClusterRender_MeshComponent.cpp b/Engine/Plugins/Runtime/nDisplay/Source/DisplayCluster/Private/Render/Containers/DisplayClusterRender_MeshComponent.cpp
--- a/Engine/Plugins/Runtime/nDisplay/Source/DisplayCluster/Private/Render/Containers/DisplayClusterRender_MeshComponent.cpp
+++ b/Engine/Plugins/Runtime/nDisplay/Source/DisplayCluster/Private/Render/Containers/DisplayClusterRender_MeshComponent.cpp
@@ -11,12 +11,12 @@
 
 //-------------------------------------------------------------------------------
 //
-static void ImplUpdateMeshComponentProxyData(FDisplayClusterRender_MeshComponentProxy* MeshComponentProxy, FDisplayClusterRender_MeshComponentProxyData* NewProxyData)
+static void ImplUpdateMeshComponentProxyData(const TSharedPtr<FDisplayClusterRender_MeshComponentProxy, ESPMode::ThreadSafe>& InMeshComponentProxy, FDisplayClusterRender_MeshComponentProxyData* NewProxyData)
 {
-	if (MeshComponentProxy)
+	if (InMeshComponentProxy.IsValid())
 	{
 		ENQUEUE_RENDER_COMMAND(DisplayClusterRender_MeshComponentProxy_Update)(
-			[MeshComponentProxy, NewProxyData](FRHICommandListImmediate& RHICmdList)
+			[MeshComponentProxy = InMeshComponentProxy, NewProxyData](FRHICommandListImmediate& RHICmdList)
 		{
 			// Update RHI
 			if (NewProxyData != nullptr)
@@ -32,13 +32,16 @@ static void ImplUpdateMeshComponentProxyData(FDisplayClusterRender_MeshComponent
 	}
 }
 
-static void ImplDeleteMeshComponentProxy(FDisplayClusterRender_MeshComponentProxy* MeshComponentProxy)
+static void ImplDeleteMeshComponentProxy(const TSharedPtr<FDisplayClusterRender_MeshComponentProxy, ESPMode::ThreadSafe>& InMeshComponentProxy)
 {
-	ENQUEUE_RENDER_COMMAND(DisplayClusterRender_MeshComponentProxy_Delete)(
-		[MeshComponentProxy](FRHICommandListImmediate& RHICmdList)
+	if (InMeshComponentProxy.IsValid())
 	{
-		delete MeshComponentProxy;
-	});
+		ENQUEUE_RENDER_COMMAND(DisplayClusterRender_MeshComponentProxy_Delete)(
+			[MeshComponentProxy = InMeshComponentProxy](FRHICommandListImmediate& RHICmdList)
+			{
+				MeshComponentProxy->Release_RenderThread();
+			});
+	}
 }
 
 static const FStaticMeshLODResources* ImplGetStaticMeshLODResources(const UStaticMesh* InStaticMesh, int32 LODIndex = 0)
@@ -97,16 +100,16 @@ static const FProcMeshSection* ImplGetProceduralMeshComponentSection(UProcedural
 FDisplayClusterRender_MeshComponent::FDisplayClusterRender_MeshComponent()
 {
 	// Create render proxy object
-	MeshComponentProxy = new FDisplayClusterRender_MeshComponentProxy();
+	MeshComponentProxyPtr = MakeShared<FDisplayClusterRender_MeshComponentProxy, ESPMode::ThreadSafe>();
 }
 
 FDisplayClusterRender_MeshComponent::~FDisplayClusterRender_MeshComponent()
 {
 	// Delete proxy on RenderThread
-	ImplDeleteMeshComponentProxy(MeshComponentProxy);
+	ImplDeleteMeshComponentProxy(MeshComponentProxyPtr);
 
 	// Forget Ptr on GameThread
-	MeshComponentProxy = nullptr;
+	MeshComponentProxyPtr.Reset();
 }
 
 USceneComponent* FDisplayClusterRender_MeshComponent::GetOriginComponent() const
@@ -127,7 +130,7 @@ IDisplayClusterRender_MeshComponentProxy* FDisplayClusterRender_MeshComponent::G
 {
 	check(IsInRenderingThread());
 
-	return MeshComponentProxy;
+	return MeshComponentProxyPtr.Get();
 }
 
 UMeshComponent* FDisplayClusterRender_MeshComponent::GetMeshComponent() const
@@ -201,12 +204,12 @@ void FDisplayClusterRender_MeshComponent::AssignStaticMeshComponentRefs(UStaticM
 	{
 		const FString SourceGeometryName = InStaticMeshComponent->GetFName().ToString();
 		// Send geometry to proxy
-		ImplUpdateMeshComponentProxyData(MeshComponentProxy, new FDisplayClusterRender_MeshComponentProxyData(SourceGeometryName, DataFunc, *StaticMeshLODResources, InUVs));
+		ImplUpdateMeshComponentProxyData(MeshComponentProxyPtr, new FDisplayClusterRender_MeshComponentProxyData(SourceGeometryName, DataFunc, *StaticMeshLODResources, InUVs));
 	}
 	else
 	{
 		// no StaticMesh - release proxy geometry
-		ImplUpdateMeshComponentProxyData(MeshComponentProxy, nullptr);
+		ImplUpdateMeshComponentProxyData(MeshComponentProxyPtr, nullptr);
 	}
 }
 
@@ -254,12 +257,12 @@ void FDisplayClusterRender_MeshComponent::AssignProceduralMeshComponentRefs(UPro
 	{
 		const FString SourceGeometryName = InProceduralMeshComponent->GetFName().ToString();
 		// Send geometry to proxy
-		ImplUpdateMeshComponentProxyData(MeshComponentProxy, new FDisplayClusterRender_MeshComponentProxyData(SourceGeometryName, DataFunc, *ProcMeshSection, InUVs));
+		ImplUpdateMeshComponentProxyData(MeshComponentProxyPtr, new FDisplayClusterRender_MeshComponentProxyData(SourceGeometryName, DataFunc, *ProcMeshSection, InUVs));
 	}
 	else
 	{
 		// no ProceduralMesh - release proxy geometry
-		ImplUpdateMeshComponentProxyData(MeshComponentProxy, nullptr);
+		ImplUpdateMeshComponentProxyData(MeshComponentProxyPtr, nullptr);
 	}
 }
 
@@ -279,7 +282,7 @@ void FDisplayClusterRender_MeshComponent::AssignProceduralMeshSection(const FPro
 
 	const FString SourceGeometryName(TEXT("ProcMeshSection"));
 	// Send geometry to proxy
-	ImplUpdateMeshComponentProxyData(MeshComponentProxy, new FDisplayClusterRender_MeshComponentProxyData(SourceGeometryName, DataFunc, InProcMeshSection, InUVs));
+	ImplUpdateMeshComponentProxyData(MeshComponentProxyPtr, new FDisplayClusterRender_MeshComponentProxyData(SourceGeometryName, DataFunc, InProcMeshSection, InUVs));
 }
 
 void FDisplayClusterRender_MeshComponent::AssignStaticMesh(const UStaticMesh* InStaticMesh, const FDisplayClusterMeshUVs& InUVs, int32 InLODIndex)
@@ -302,12 +305,12 @@ void FDisplayClusterRender_MeshComponent::AssignStaticMesh(const UStaticMesh* In
 	{
 		const FString SourceGeometryName = InStaticMesh->GetFName().ToString();
 		// Send geometry to proxy
-		ImplUpdateMeshComponentProxyData(MeshComponentProxy, new FDisplayClusterRender_MeshComponentProxyData(SourceGeometryName, DataFunc, *StaticMeshLODResources, InUVs));
+		ImplUpdateMeshComponentProxyData(MeshComponentProxyPtr, new FDisplayClusterRender_MeshComponentProxyData(SourceGeometryName, DataFunc, *StaticMeshLODResources, InUVs));
 	}
 	else
 	{
 		// no StaticMesh - release proxy geometry
-		ImplUpdateMeshComponentProxyData(MeshComponentProxy, nullptr);
+		ImplUpdateMeshComponentProxyData(MeshComponentProxyPtr, nullptr);
 	}
 }
 
@@ -329,12 +332,12 @@ void FDisplayClusterRender_MeshComponent::AssignMeshGeometry(const FDisplayClust
 	{
 		const FString SourceGeometryName(TEXT("nDCRender_MeshGeometry"));
 		// Send geometry to proxy
-		ImplUpdateMeshComponentProxyData(MeshComponentProxy, new FDisplayClusterRender_MeshComponentProxyData(SourceGeometryName, DataFunc, *InMeshGeometry));
+		ImplUpdateMeshComponentProxyData(MeshComponentProxyPtr, new FDisplayClusterRender_MeshComponentProxyData(SourceGeometryName, DataFunc, *InMeshGeometry));
 	}
 	else
 	{
 		// no InMeshGeometry- release proxy geometry
-		ImplUpdateMeshComponentProxyData(MeshComponentProxy, nullptr);
+		ImplUpdateMeshComponentProxyData(MeshComponentProxyPtr, nullptr);
 	}
 }
 
@@ -346,7 +349,7 @@ void FDisplayClusterRender_MeshComponent::AssignMeshGeometry_RenderThread(const
 	{
 		const FString SourceGeometryName(TEXT("nDCRender_MeshGeometry_RenderThread"));
 		// Send geometry to proxy
-		ImplUpdateMeshComponentProxyData(MeshComponentProxy, new FDisplayClusterRender_MeshComponentProxyData(SourceGeometryName, InDataFunc, *InMeshGeometry));
+		ImplUpdateMeshComponentProxyData(MeshComponentProxyPtr, new FDisplayClusterRender_MeshComponentProxyData(SourceGeometryName, InDataFunc, *InMeshGeometry));
 	}
 }
 
@@ -364,7 +367,7 @@ void FDisplayClusterRender_MeshComponent::ReleaseMeshComponent()
 	// Set source geometry type
 	GeometrySource = EDisplayClusterRender_MeshComponentGeometrySource::Disabled;
 
-	ImplUpdateMeshComponentProxyData(MeshComponentProxy, nullptr);
+	ImplUpdateMeshComponentProxyData(MeshComponentProxyPtr, nullptr);
 }
 
 void FDisplayClusterRender_MeshComponent::ReleaseProxyGeometry()
@@ -372,7 +375,7 @@ void FDisplayClusterRender_MeshComponent::ReleaseProxyGeometry()
 	check(IsInGameThread());
 
 	// just release  geometry on proxy
-	ImplUpdateMeshComponentProxyData(MeshComponentProxy, nullptr);
+	ImplUpdateMeshComponentProxyData(MeshComponentProxyPtr, nullptr);
 }
 
 bool FDisplayClusterRender_MeshComponent::EqualsMeshComponentName(const FName& InMeshComponentName) const
diff --git a/Engine/Plugins/Runtime/nDisplay/Source/DisplayCluster/Private/Render/Containers/DisplayClusterRender_MeshComponent.h b/Engine/Plugins/Runtime/nDisplay/Source/DisplayCluster/Private/Render/Containers/DisplayClusterRender_MeshComponent.h
--- a/Engine/Plugins/Runtime/nDisplay/Source/DisplayCluster/Private/Render/Containers/DisplayClusterRender_MeshComponent.h
+++ b/Engine/Plugins/Runtime/nDisplay/Source/DisplayCluster/Private/Render/Containers/DisplayClusterRender_MeshComponent.h
@@ -61,5 +61,5 @@ private:
 	TWeakObjectPtr<const UStaticMesh> StaticMeshRef;
 
 private:
-	FDisplayClusterRender_MeshComponentProxy* MeshComponentProxy = nullptr;
+	TSharedPtr<FDisplayClusterRender_MeshComponentProxy, ESPMode::ThreadSafe> MeshComponentProxyPtr;
 };
diff --git a/Engine/Plugins/Runtime/nDisplay/Source/DisplayCluster/Private/Render/Device/DisplayClusterDeviceBase.cpp b/Engine/Plugins/Runtime/nDisplay/Source/DisplayCluster/Private/Render/Device/DisplayClusterDeviceBase.cpp
--- a/Engine/Plugins/Runtime/nDisplay/Source/DisplayCluster/Private/Render/Device/DisplayClusterDeviceBase.cpp
+++ b/Engine/Plugins/Runtime/nDisplay/Source/DisplayCluster/Private/Render/Device/DisplayClusterDeviceBase.cpp
@@ -328,11 +328,11 @@ bool FDisplayClusterDeviceBase::BeginNewFrame(FViewport* InViewport, UWorld* InW
 		{
 			const FString LocalNodeId = DeviceBaseHelpers::GetDisplayClusterAPI().GetConfigMgr()->GetLocalNodeId();
 
-			// update current world
-			ViewportManagerPtr->GetConfiguration().SetCurrentWorld(InWorld);
+			// Dont use preview setting on primary RootActor in game
+			ViewportManagerPtr->GetConfiguration().SetPreviewSettings(FDisplayClusterViewport_PreviewSettings());
 
 			// Update local node viewports (update\create\delete) and build new render frame
-			if (ViewportManagerPtr->GetConfiguration().UpdateConfigurationForClusterNode(RenderFrameMode, LocalNodeId))
+			if (ViewportManagerPtr->GetConfiguration().UpdateConfigurationForClusterNode(RenderFrameMode, InWorld, LocalNodeId))
 			{
 				if (ViewportManagerPtr->BeginNewFrame(InViewport, OutRenderFrame))
 				{
@@ -366,7 +366,7 @@ void FDisplayClusterDeviceBase::InitializeNewFrame()
 
 			// Send viewport manager proxy on render thread
 			ENQUEUE_RENDER_COMMAND(DisplayClusterDevice_SetViewportManagerProxy)(
-				[DCRenderDevice = this, NewViewportManagerProxy = ViewportManagerProxyPtr->AsShared()](FRHICommandListImmediate& RHICmdList)
+				[DCRenderDevice = SharedThis(this), NewViewportManagerProxy = ViewportManagerProxyPtr->AsShared()](FRHICommandListImmediate& RHICmdList)
 				{
 					DCRenderDevice->ViewportManagerProxyWeakPtr = NewViewportManagerProxy;
 				});
diff --git a/Engine/Plugins/Runtime/nDisplay/Source/DisplayCluster/Private/Render/Device/DisplayClusterDeviceBase.h b/Engine/Plugins/Runtime/nDisplay/Source/DisplayCluster/Private/Render/Device/DisplayClusterDeviceBase.h
--- a/Engine/Plugins/Runtime/nDisplay/Source/DisplayCluster/Private/Render/Device/DisplayClusterDeviceBase.h
+++ b/Engine/Plugins/Runtime/nDisplay/Source/DisplayCluster/Private/Render/Device/DisplayClusterDeviceBase.h
@@ -13,6 +13,7 @@
 #include "Render/IDisplayClusterRenderManager.h"
 
 #include "Containers/Queue.h"
+#include "Templates/SharedPointer.h"
 
 class IDisplayClusterPostProcess;
 class FDisplayClusterPresentationBase;
@@ -26,6 +27,7 @@ class FSceneView;
 class FDisplayClusterDeviceBase
 	: public IStereoRenderTargetManager
 	, public IDisplayClusterRenderDevice
+	, public TSharedFromThis<FDisplayClusterDeviceBase, ESPMode::ThreadSafe>
 {
 public:
 	FDisplayClusterDeviceBase() = delete;
