nDisplay: Implemented a new preview rendering architecture:

diff --git a/Engine/Plugins/Runtime/nDisplay/Source/DisplayCluster/Private/Render/DisplayDevice/Components/DisplayClusterDisplayDeviceBaseComponent.cpp b/Engine/Plugins/Runtime/nDisplay/Source/DisplayCluster/Private/Render/DisplayDevice/Components/DisplayClusterDisplayDeviceBaseComponent.cpp
new file mode 100644
--- /dev/null
+++ b/Engine/Plugins/Runtime/nDisplay/Source/DisplayCluster/Private/Render/DisplayDevice/Components/DisplayClusterDisplayDeviceBaseComponent.cpp
@@ -0,0 +1,151 @@
+// Copyright Epic Games, Inc. All Rights Reserved.
+
+#include "Render/DisplayDevice/Components/DisplayClusterDisplayDeviceBaseComponent.h"
+#include "Render/DisplayDevice/DisplayClusterDisplayDeviceStrings.h"
+
+#include "DisplayClusterRootActor.h"
+
+#include "Materials/Material.h"
+#include "UObject/ConstructorHelpers.h"
+#include "Engine/TextureRenderTarget2D.h"
+
+bool UDisplayClusterDisplayDeviceBaseComponent::ShouldUseDisplayDevice(IDisplayClusterViewportConfiguration& InConfiguration) const
+{
+	// Use this Display device only for preview
+	if (InConfiguration.IsPreviewRendering())
+	{
+		return true;
+	}
+
+	return false;
+}
+
+TSharedPtr<IDisplayClusterDisplayDeviceProxy, ESPMode::ThreadSafe> UDisplayClusterDisplayDeviceBaseComponent::GetDisplayDeviceProxy(IDisplayClusterViewportConfiguration& InConfiguration)
+{
+	if (!ShouldUseDisplayDevice(InConfiguration))
+	{
+		return nullptr;
+	}
+
+	UpdateDisplayDeviceProxyImpl(InConfiguration);
+
+	return DisplayDeviceProxy;
+}
+
+UDisplayClusterDisplayDeviceBaseComponent::UDisplayClusterDisplayDeviceBaseComponent()
+{
+	static ConstructorHelpers::FObjectFinder<UMaterial> MeshMaterialObj(UE::DisplayClusterDisplayDeviceStrings::material::asset::mesh);
+	check(MeshMaterialObj.Object);
+	MeshMaterial = MeshMaterialObj.Object;
+
+	static ConstructorHelpers::FObjectFinder<UMaterial> PreviewMeshMaterialObj(UE::DisplayClusterDisplayDeviceStrings::material::asset::preview_mesh);
+	check(PreviewMeshMaterialObj.Object);
+	PreviewMeshMaterial = PreviewMeshMaterialObj.Object;
+
+	static ConstructorHelpers::FObjectFinder<UMaterial> PreviewMeshTechvisMaterialObj(UE::DisplayClusterDisplayDeviceStrings::material::asset::preview_techvis_mesh);
+	check(PreviewMeshTechvisMaterialObj.Object);
+	PreviewMeshTechvisMaterial = PreviewMeshTechvisMaterialObj.Object;
+}
+
+TObjectPtr<UMaterial> UDisplayClusterDisplayDeviceBaseComponent::GetDisplayDeviceMaterial(const EDisplayClusterDisplayDeviceMaterialType InMaterialType) const
+{
+	switch (InMaterialType)
+	{
+	case EDisplayClusterDisplayDeviceMaterialType::DefaultPreviewMeshMaterial:
+		return MeshMaterial;
+
+	case EDisplayClusterDisplayDeviceMaterialType::PreviewMeshMaterial:
+		return PreviewMeshMaterial;
+
+	case EDisplayClusterDisplayDeviceMaterialType::PreviewMeshTechvisMaterial:
+		return PreviewMeshTechvisMaterial;
+
+	default:
+		break;
+	}
+
+	return nullptr;
+}
+
+void UDisplayClusterDisplayDeviceBaseComponent::OnUpdateDisplayDeviceMaterialInstance(IDisplayClusterViewportPreview& InViewportPreview, const EDisplayClusterDisplayDeviceMeshType InMeshType, const EDisplayClusterDisplayDeviceMaterialType InMaterialType, UMaterialInstanceDynamic* InMaterialInstance) const
+{
+	if (!InMaterialInstance || !ShouldUseDisplayDevice(InViewportPreview.GetConfiguration()))
+	{
+		return;
+	}
+
+	// Update preview RTT parameter values:
+	switch (InMaterialType)
+	{
+	case EDisplayClusterDisplayDeviceMaterialType::PreviewMeshMaterial:
+	case EDisplayClusterDisplayDeviceMaterialType::PreviewMeshTechvisMaterial:
+		if (InViewportPreview.HasAnyFlags(EDisplayClusterViewportPreviewFlags::HasChangedPreviewMeshMaterialInstance | EDisplayClusterViewportPreviewFlags::HasChangedPreviewEditableMeshMaterialInstance | EDisplayClusterViewportPreviewFlags::HasChangedPreviewRTT))
+		{
+			// Updates the RTT parameter for the material instance when the PreviewTexture or mesh material is changed.
+			InMaterialInstance->SetTextureParameterValue(UE::DisplayClusterDisplayDeviceStrings::material::attr::Preview, InViewportPreview.GetPreviewTextureRenderTarget2D());
+		}
+		break;
+	default:
+		break;
+	}
+}
+
+void UDisplayClusterDisplayDeviceBaseComponent::OnUpdateDisplayDeviceMeshComponent(IDisplayClusterViewportPreview& InViewportPreview, const EDisplayClusterDisplayDeviceMeshType InMeshType, UMeshComponent* InMeshComponent) const
+{
+	if (!InMeshComponent || !ShouldUseDisplayDevice(InViewportPreview.GetConfiguration()))
+	{
+		return;
+	}
+
+	// Only some of the preview meshes are supported by Techvis
+	switch (InMeshType)
+	{
+		case EDisplayClusterDisplayDeviceMeshType::PreviewMesh:
+		case EDisplayClusterDisplayDeviceMeshType::PreviewEditableMesh:
+			break;
+
+		default:
+			// Dont apply Techvis for other mesh types
+			return;
+	}
+
+	if (InViewportPreview.GetConfiguration().IsTechvisEnabled())
+	{
+		InMeshComponent->bAffectDynamicIndirectLighting = true;
+		InMeshComponent->bAffectIndirectLightingWhileHidden = true;
+	}
+	else
+	{
+		// When disabling just revert to the archetype values. The original values at the time of enabling
+		// techvis aren't saved, and techvis is enabled by default so they'll always be set to true on new meshes.
+		if (const UMeshComponent* MeshArchetype = Cast<UMeshComponent>(InMeshComponent->GetArchetype()))
+		{
+			InMeshComponent->bAffectDynamicIndirectLighting = MeshArchetype->bAffectDynamicIndirectLighting;
+			InMeshComponent->bAffectIndirectLightingWhileHidden = MeshArchetype->bAffectIndirectLightingWhileHidden;
+		}
+	}
+
+	// For all preview meshes:
+	InMeshComponent->SetCastShadow(false);
+	InMeshComponent->SetHiddenInGame(false);
+	InMeshComponent->SetVisibility(true);
+
+	InMeshComponent->bVisibleInReflectionCaptures = false;
+	InMeshComponent->bVisibleInRayTracing = false;
+	InMeshComponent->bVisibleInRealTimeSkyCaptures = false;
+}
+
+void UDisplayClusterDisplayDeviceBaseComponent::SetupSceneView(const IDisplayClusterViewportPreview& InViewportPreview, uint32 ContextNum, FSceneViewFamily& InOutViewFamily, FSceneView& InOutView) const
+{ }
+
+void UDisplayClusterDisplayDeviceBaseComponent::UpdateDisplayDeviceProxyImpl(IDisplayClusterViewportConfiguration& InConfiguration)
+{
+	// This component does not interact with the rendering thread
+}
+
+#if WITH_EDITOR
+void UDisplayClusterDisplayDeviceBaseComponent::PostEditChangeProperty(FPropertyChangedEvent& PropertyChangedEvent)
+{
+	Super::PostEditChangeProperty(PropertyChangedEvent);
+}
+#endif
diff --git a/Engine/Plugins/Runtime/nDisplay/Source/DisplayCluster/Private/Render/DisplayDevice/Components/DisplayClusterDisplayDeviceComponent.cpp b/Engine/Plugins/Runtime/nDisplay/Source/DisplayCluster/Private/Render/DisplayDevice/Components/DisplayClusterDisplayDeviceComponent.cpp
new file mode 100644
--- /dev/null
+++ b/Engine/Plugins/Runtime/nDisplay/Source/DisplayCluster/Private/Render/DisplayDevice/Components/DisplayClusterDisplayDeviceComponent.cpp
@@ -0,0 +1,77 @@
+// Copyright Epic Games, Inc. All Rights Reserved.
+
+#include "Render/DisplayDevice/Components/DisplayClusterDisplayDeviceComponent.h"
+#include "Render/Viewport/IDisplayClusterViewportConfiguration.h"
+
+#include "Render/DisplayDevice/Proxy/DisplayClusterDisplayDeviceProxy_OpenColorIO.h"
+#include "Render/DisplayDevice/DisplayClusterDisplayDeviceStrings.h"
+
+#include "DisplayClusterRootActor.h"
+#include "IDisplayCluster.h"
+
+#include "Engine/TextureRenderTarget2D.h"
+#include "Materials/MaterialInstanceDynamic.h"
+
+UDisplayClusterDisplayDeviceComponent::UDisplayClusterDisplayDeviceComponent()
+{ }
+
+void UDisplayClusterDisplayDeviceComponent::OnUpdateDisplayDeviceMaterialInstance(IDisplayClusterViewportPreview& InViewportPreview, const EDisplayClusterDisplayDeviceMeshType InMeshType, const EDisplayClusterDisplayDeviceMaterialType InMaterialType, UMaterialInstanceDynamic* InMaterialInstance) const
+{
+	if (!InMaterialInstance || !ShouldUseDisplayDevice(InViewportPreview.GetConfiguration()))
+	{
+		return;
+	}
+
+	switch (InMaterialType)
+	{
+	case EDisplayClusterDisplayDeviceMaterialType::PreviewMeshMaterial:
+		InMaterialInstance->SetScalarParameterValue(UE::DisplayClusterDisplayDeviceStrings::material::attr::Exposure, 0.f);
+		break;
+
+	case EDisplayClusterDisplayDeviceMaterialType::PreviewMeshTechvisMaterial:
+		InMaterialInstance->SetScalarParameterValue(UE::DisplayClusterDisplayDeviceStrings::material::attr::Exposure, Exposure);
+		break;
+
+	default:
+		break;
+	}
+
+	Super::OnUpdateDisplayDeviceMaterialInstance(InViewportPreview, InMeshType, InMaterialType, InMaterialInstance);
+}
+
+void UDisplayClusterDisplayDeviceComponent::UpdateDisplayDeviceProxyImpl(IDisplayClusterViewportConfiguration& InConfiguration)
+{
+	if (!ShouldUseDisplayDevice(InConfiguration))
+	{
+		// Use this only for preview rendering
+		DisplayDeviceProxy.Reset();
+
+		return;
+	}
+
+	if (!bEnableRenderPass || !ColorConversionSettings.IsValid() || !InConfiguration.IsTechvisEnabled())
+	{
+		// OCIO for preview is disabled
+		DisplayDeviceProxy.Reset();
+
+		return;
+	}
+
+	// Add a single OCIO render pass at the rendering thread
+
+	// If the settings have changed, create a new proxy object
+	if(FDisplayClusterDisplayDeviceProxy_OpenColorIO* Proxy = static_cast<FDisplayClusterDisplayDeviceProxy_OpenColorIO*>(DisplayDeviceProxy.Get()))
+	{
+		const bool bIsOCIOEquals = Proxy->OCIOPassId == ColorConversionSettings.ToString();
+		if (!bIsOCIOEquals)
+		{
+			DisplayDeviceProxy.Reset();
+		}
+	}
+
+	if (!DisplayDeviceProxy.IsValid())
+	{
+		// create a new proxy object for the current settings
+		DisplayDeviceProxy = MakeShared<FDisplayClusterDisplayDeviceProxy_OpenColorIO, ESPMode::ThreadSafe>(ColorConversionSettings);
+	}
+}
diff --git a/Engine/Plugins/Runtime/nDisplay/Source/DisplayCluster/Private/Render/DisplayDevice/Proxy/DisplayClusterDisplayDeviceProxy_OpenColorIO.cpp b/Engine/Plugins/Runtime/nDisplay/Source/DisplayCluster/Private/Render/DisplayDevice/Proxy/DisplayClusterDisplayDeviceProxy_OpenColorIO.cpp
new file mode 100644
--- /dev/null
+++ b/Engine/Plugins/Runtime/nDisplay/Source/DisplayCluster/Private/Render/DisplayDevice/Proxy/DisplayClusterDisplayDeviceProxy_OpenColorIO.cpp
@@ -0,0 +1,75 @@
+// Copyright Epic Games, Inc. All Rights Reserved.
+
+#include "Render/DisplayDevice/Proxy/DisplayClusterDisplayDeviceProxy_OpenColorIO.h"
+#include "Render/Viewport/Containers/DisplayClusterViewport_Context.h"
+
+#include "Shader.h"
+#include "SceneView.h"
+#include "ScreenPass.h"
+#include "CommonRenderResources.h"
+#include "PostProcess/PostProcessMaterialInputs.h"
+
+namespace UE::DisplayCluster
+{
+	static inline float GetDisplayGamma(const FDisplayClusterViewport_Context& InViewportContext)
+	{
+		const float DefaultDisplayGamma = FOpenColorIORendering::DefaultDisplayGamma;
+		const float EngineDisplayGamma = InViewportContext.RenderThreadData.EngineDisplayGamma;
+
+		// There is a special case where post processing and tonemapper are disabled. In this case tonemapper applies a static display Inverse of Gamma which defaults to 2.2.
+		// In the case when Both PostProcessing and ToneMapper are disabled we apply gamma manually. In every other case we apply inverse gamma before applying OCIO.
+		const float DisplayGamma = (InViewportContext.RenderThreadData.EngineShowFlags.Tonemapper == 0) || (InViewportContext.RenderThreadData.EngineShowFlags.PostProcessing == 0) ? DefaultDisplayGamma : DefaultDisplayGamma / EngineDisplayGamma;
+
+		return DisplayGamma;
+	}
+};
+
+FDisplayClusterDisplayDeviceProxy_OpenColorIO::FDisplayClusterDisplayDeviceProxy_OpenColorIO(FOpenColorIOColorConversionSettings& InColorConversionSettings)
+	: OCIOPassId(InColorConversionSettings.ToString())
+	, OCIOPassResources(FOpenColorIORendering::GetRenderPassResources(InColorConversionSettings, GMaxRHIFeatureLevel))
+{ }
+
+bool FDisplayClusterDisplayDeviceProxy_OpenColorIO::HasFinalPass_RenderThread() const
+{
+	return true;
+}
+
+bool FDisplayClusterDisplayDeviceProxy_OpenColorIO::AddFinalPass_RenderThread(FRDGBuilder& GraphBuilder, const FDisplayClusterViewport_Context& InViewportContext,
+	FRHITexture2D* InputTextureRHI, const FIntRect& InputRect, FRHITexture2D* OutputTextureRHI, const FIntRect& OutputRect) const
+{
+	FRDGTextureRef InputTexture = GraphBuilder.RegisterExternalTexture(CreateRenderTarget(InputTextureRHI, TEXT("FDisplayClusterDisplayDeviceOCIOInputTexture")));
+	FRDGTextureRef OutputTexture = GraphBuilder.RegisterExternalTexture(CreateRenderTarget(OutputTextureRHI, TEXT("FDisplayClusterDisplayDeviceOCIORenderTargetTexture")));
+
+	FScreenPassRenderTarget Output = FScreenPassRenderTarget(OutputTexture, OutputRect, ERenderTargetLoadAction::EClear);
+	EOpenColorIOTransformAlpha TransformAlpha = EOpenColorIOTransformAlpha::None;
+
+	FOpenColorIORendering::AddPass_RenderThread(
+		GraphBuilder,
+		FScreenPassViewInfo(),
+		GMaxRHIFeatureLevel,
+		FScreenPassTexture(InputTexture),
+		Output,
+		OCIOPassResources,
+		UE::DisplayCluster::GetDisplayGamma(InViewportContext),
+		TransformAlpha
+	);
+
+	return true;
+}
+
+/**
+* A helper function that extracts the right scene color texture, untouched, to be used further in post processing.
+*/
+FScreenPassTexture IDisplayClusterDisplayDeviceProxy::ReturnUntouchedSceneColorForPostProcessing(const FPostProcessMaterialInputs& InOutInputs)
+{
+	if (InOutInputs.OverrideOutput.IsValid())
+	{
+		return InOutInputs.OverrideOutput;
+	}
+	else
+	{
+		/** We don't want to modify scene texture in any way. We just want it to be passed back onto the next stage. */
+		FScreenPassTexture SceneTexture = const_cast<FScreenPassTexture&>(InOutInputs.Textures[(uint32)EPostProcessMaterialInput::SceneColor]);
+		return SceneTexture;
+	}
+}
diff --git a/Engine/Plugins/Runtime/nDisplay/Source/DisplayCluster/Private/Render/DisplayDevice/Proxy/DisplayClusterDisplayDeviceProxy_OpenColorIO.h b/Engine/Plugins/Runtime/nDisplay/Source/DisplayCluster/Private/Render/DisplayDevice/Proxy/DisplayClusterDisplayDeviceProxy_OpenColorIO.h
new file mode 100644
--- /dev/null
+++ b/Engine/Plugins/Runtime/nDisplay/Source/DisplayCluster/Private/Render/DisplayDevice/Proxy/DisplayClusterDisplayDeviceProxy_OpenColorIO.h
@@ -0,0 +1,31 @@
+// Copyright Epic Games, Inc. All Rights Reserved.
+
+#pragma once
+
+#include "Render/DisplayDevice/IDisplayClusterDisplayDeviceProxy.h"
+
+#include "OpenColorIORendering.h"
+#include "OpenColorIOColorSpace.h"
+
+/**
+ * Display Device Proxy object (OCIO render pass)
+ * [rendering thread]
+ */
+class FDisplayClusterDisplayDeviceProxy_OpenColorIO
+	: public IDisplayClusterDisplayDeviceProxy
+{
+public:
+	FDisplayClusterDisplayDeviceProxy_OpenColorIO(FOpenColorIOColorConversionSettings& InColorConversionSettings);
+	virtual ~FDisplayClusterDisplayDeviceProxy_OpenColorIO() = default;
+
+public:
+	virtual bool HasFinalPass_RenderThread() const override;
+	virtual bool AddFinalPass_RenderThread(FRDGBuilder& GraphBuilder, const FDisplayClusterViewport_Context& InViewportContext,
+		FRHITexture2D* InputTextureRHI, const FIntRect& InputRect, FRHITexture2D* OutputTextureRHI, const FIntRect& OutputRect) const override;
+
+public:
+	const FString OCIOPassId;
+
+private:
+	FOpenColorIORenderPassResources OCIOPassResources;
+};
diff --git a/Engine/Plugins/Runtime/nDisplay/Source/DisplayCluster/Private/Render/Viewport/Configuration/DisplayClusterViewportConfiguration.cpp b/Engine/Plugins/Runtime/nDisplay/Source/DisplayCluster/Private/Render/Viewport/Configuration/DisplayClusterViewportConfiguration.cpp
--- a/Engine/Plugins/Runtime/nDisplay/Source/DisplayCluster/Private/Render/Viewport/Configuration/DisplayClusterViewportConfiguration.cpp
+++ b/Engine/Plugins/Runtime/nDisplay/Source/DisplayCluster/Private/Render/Viewport/Configuration/DisplayClusterViewportConfiguration.cpp
@@ -45,9 +45,17 @@ namespace UE::DisplayCluster::ConfigurationHelpers
 		return nullptr;
 	}
 
-	static inline bool ImplSetRootActor(ADisplayClusterRootActor* InRootActor, FDisplayClusterActorRef& InOutConfigurationRootActorRef)
+	static inline bool ImplSetRootActor(const ADisplayClusterRootActor* InRootActor, FDisplayClusterActorRef& InOutConfigurationRootActorRef)
 	{
-		if (ImplGetRootActor(InOutConfigurationRootActorRef) != InRootActor)
+		const bool bIsDefined = InOutConfigurationRootActorRef.IsDefinedSceneActor();
+		if (InRootActor == nullptr)
+		{
+			InOutConfigurationRootActorRef.ResetSceneActor();
+
+			return bIsDefined;
+		}
+
+		if (!bIsDefined || ImplGetRootActor(InOutConfigurationRootActorRef) != InRootActor)
 		{
 			InOutConfigurationRootActorRef.SetSceneActor(InRootActor);
 
@@ -76,87 +84,89 @@ void FDisplayClusterViewportConfiguration::Initialize(FDisplayClusterViewportMan
 	Proxy->Initialize_GameThread(InViewportManager.GetViewportManagerProxy());
 }
 
-void FDisplayClusterViewportConfiguration::SetRootActor(ADisplayClusterRootActor* InRootActor, const EDisplayClusterRootActorType InRootActorType)
+void FDisplayClusterViewportConfiguration::SetRootActor(const EDisplayClusterRootActorType InRootActorType, const ADisplayClusterRootActor* InRootActor)
 {
 	check(IsInGameThread());
 
-	switch (InRootActorType)
+	if (EnumHasAnyFlags(InRootActorType, EDisplayClusterRootActorType::Preview))
 	{
-	case EDisplayClusterRootActorType::Preview:
-		bCurrentSceneNeedsToBeUpdated |= ConfigurationHelpers::ImplSetRootActor(InRootActor, PreviewRootActorRef);
-		break;
-
-	case EDisplayClusterRootActorType::Scene:
-		bCurrentSceneNeedsToBeUpdated |= ConfigurationHelpers::ImplSetRootActor(InRootActor, SceneRootActorRef);
-		break;
-
-	case EDisplayClusterRootActorType::Configuration:
-		ConfigurationHelpers::ImplSetRootActor(InRootActor, ConfigurationRootActorRef);
-		break;
+		ConfigurationHelpers::ImplSetRootActor(InRootActor, PreviewRootActorRef);
+	}
 
-	case EDisplayClusterRootActorType::Any:
-		bCurrentSceneNeedsToBeUpdated |= ConfigurationHelpers::ImplSetRootActor(InRootActor, PreviewRootActorRef);
-		bCurrentSceneNeedsToBeUpdated |= ConfigurationHelpers::ImplSetRootActor(InRootActor, SceneRootActorRef);
+	if (EnumHasAnyFlags(InRootActorType, EDisplayClusterRootActorType::Scene))
+	{
+		ConfigurationHelpers::ImplSetRootActor(InRootActor, SceneRootActorRef);
+	}
 
+	if (EnumHasAnyFlags(InRootActorType, EDisplayClusterRootActorType::Configuration))
+	{
 		ConfigurationHelpers::ImplSetRootActor(InRootActor, ConfigurationRootActorRef);
-		break;
-
-	default:
-		break;
 	}
 }
 
 ADisplayClusterRootActor* FDisplayClusterViewportConfiguration::GetRootActor(const EDisplayClusterRootActorType InRootActorType) const
 {
-	switch (InRootActorType)
+	if (EnumHasAnyFlags(InRootActorType, EDisplayClusterRootActorType::Preview))
 	{
-	case EDisplayClusterRootActorType::Preview:
-		return ConfigurationHelpers::ImplGetRootActor(PreviewRootActorRef);
-
-	case EDisplayClusterRootActorType::Scene:
-		return ConfigurationHelpers::ImplGetRootActor(SceneRootActorRef);
-
-	case EDisplayClusterRootActorType::Configuration:
-		return ConfigurationHelpers::ImplGetRootActor(ConfigurationRootActorRef);
+		if (ADisplayClusterRootActor* OutRootActor = ConfigurationHelpers::ImplGetRootActor(PreviewRootActorRef))
+		{
+			return OutRootActor;
+		}
+	}
 
-	case EDisplayClusterRootActorType::Any:
+	if (EnumHasAnyFlags(InRootActorType, EDisplayClusterRootActorType::Scene))
 	{
-		ADisplayClusterRootActor* OutRootActor = ConfigurationHelpers::ImplGetRootActor(PreviewRootActorRef);
-		if (!OutRootActor)
+		if (ADisplayClusterRootActor* OutRootActor = ConfigurationHelpers::ImplGetRootActor(SceneRootActorRef))
 		{
-			OutRootActor = ConfigurationHelpers::ImplGetRootActor(SceneRootActorRef);
-			if (!OutRootActor)
-			{
-				OutRootActor = ConfigurationHelpers::ImplGetRootActor(ConfigurationRootActorRef);
-			}
+			return OutRootActor;
 		}
-
-		return OutRootActor;
 	}
 
-	default:
-		break;
+	if (EnumHasAnyFlags(InRootActorType, EDisplayClusterRootActorType::Configuration))
+	{
+		if (ADisplayClusterRootActor* OutRootActor = ConfigurationHelpers::ImplGetRootActor(ConfigurationRootActorRef))
+		{
+			return OutRootActor;
+		}
 	}
 
 	return nullptr;
 }
 
-void FDisplayClusterViewportConfiguration::SetCurrentWorld(UWorld* InWorld)
+void FDisplayClusterViewportConfiguration::OnHandleStartScene()
+{
+	bCurrentSceneActive = true;
+}
+
+void FDisplayClusterViewportConfiguration::OnHandleEndScene()
+{
+	bCurrentSceneActive = false;
+	bCurrentRenderFrameViewportsNeedsToBeUpdated = true;
+
+	// Always reset world ptr at the end of the scene
+	CurrentWorldRef.Reset();
+}
+
+void FDisplayClusterViewportConfiguration::SetCurrentWorldImpl(const UWorld* InWorld)
 {
 	if (GetCurrentWorld() != InWorld)
 	{
 		if (FDisplayClusterViewportManager* ViewportManager = GetViewportManagerImpl())
 		{
-			const bool bNeedRestartScene = bCurrentSceneActive;
+			const bool bIsSceneOpened = IsSceneOpened();
 
-			ViewportManager->HandleEndScene();
-			CurrentWorldRef.Reset();
+			if (bIsSceneOpened)
+			{
+				ViewportManager->HandleEndScene();
+			}
 
 			if (InWorld)
 			{
-				CurrentWorldRef = TWeakObjectPtr<UWorld>(InWorld);
+				// Ignore const UWorld
+				CurrentWorldRef = TWeakObjectPtr<UWorld>((UWorld*)InWorld);
 
-				if (bNeedRestartScene)
+				// Restore scene status
+				if (bIsSceneOpened)
 				{
 					ViewportManager->HandleStartScene();
 				}
@@ -164,7 +174,7 @@ void FDisplayClusterViewportConfiguration::SetCurrentWorld(UWorld* InWorld)
 		}
 		else
 		{
-			CurrentWorldRef.Reset();
+			ViewportManager->HandleEndScene();
 		}
 	}
 }
@@ -235,6 +245,21 @@ bool FDisplayClusterViewportConfiguration::IsRootActorWorldHasAnyType(const EDis
 	return false;
 }
 
+const float FDisplayClusterViewportConfiguration::GetWorldToMeters() const
+{
+	// Get world scale
+	float OutWorldToMeters = 100.f;
+	if (UWorld* World = GetCurrentWorld())
+	{
+		if (AWorldSettings* WorldSettings = World->GetWorldSettings())
+		{
+			OutWorldToMeters = WorldSettings->WorldToMeters;
+		}
+	}
+
+	return OutWorldToMeters;
+}
+
 EDisplayClusterRenderFrameMode FDisplayClusterViewportConfiguration::GetRenderModeForPIE() const
 {
 #if WITH_EDITOR
@@ -243,10 +268,10 @@ EDisplayClusterRenderFrameMode FDisplayClusterViewportConfiguration::GetRenderMo
 		switch (PreviewRootActor->RenderMode)
 		{
 		case EDisplayClusterConfigurationRenderMode::SideBySide:
-			return EDisplayClusterRenderFrameMode::SideBySide;
+			return EDisplayClusterRenderFrameMode::PIE_SideBySide;
 
 		case EDisplayClusterConfigurationRenderMode::TopBottom:
-			return EDisplayClusterRenderFrameMode::TopBottom;
+			return EDisplayClusterRenderFrameMode::PIE_TopBottom;
 
 		default:
 			break;
@@ -254,7 +279,7 @@ EDisplayClusterRenderFrameMode FDisplayClusterViewportConfiguration::GetRenderMo
 	}
 #endif
 
-	return EDisplayClusterRenderFrameMode::Mono;
+	return EDisplayClusterRenderFrameMode::PIE_Mono;
 }
 
 IDisplayClusterViewportManager* FDisplayClusterViewportConfiguration::GetViewportManager() const
@@ -262,20 +287,39 @@ IDisplayClusterViewportManager* FDisplayClusterViewportConfiguration::GetViewpor
 	return GetViewportManagerImpl();
 }
 
-bool FDisplayClusterViewportConfiguration::ImplUpdateConfiguration(EDisplayClusterRenderFrameMode InRenderMode, const FString& InClusterNodeId, const TArray<FString>* InViewportNames)
+void FDisplayClusterViewportConfiguration::ReleaseConfiguration()
+{
+	if (FDisplayClusterViewportManager* ViewportManager = GetViewportManagerImpl())
+	{
+		ViewportManager->HandleEndScene();
+		ViewportManager->ReleaseTextures();
+	}
+
+	RenderFrameSettings = FDisplayClusterRenderFrameSettings();
+}
+
+bool FDisplayClusterViewportConfiguration::ImplUpdateConfiguration(EDisplayClusterRenderFrameMode InRenderMode, const UWorld* InWorld, const FString& InClusterNodeId, const TArray<FString>* InViewportNames)
 {
 	check(IsInGameThread());
 	
 	FDisplayClusterViewportManager* ViewportManager = GetViewportManagerImpl();
-	if (!ViewportManager || !FDisplayClusterViewportConfigurationHelpers_RenderFrameSettings::UpdateRenderFrameConfiguration(InRenderMode, *this))
+	if (!ViewportManager || !FDisplayClusterViewportConfigurationHelpers_RenderFrameSettings::UpdateRenderFrameConfiguration(ViewportManager, InRenderMode, *this))
 	{
 		return false;
 	}
 
-	if (bCurrentSceneNeedsToBeUpdated)
+	if (!InWorld)
 	{
-		// When the root actor changes, we have to ResetScene() to reinitialize the internal references of the projection policy.
-		ViewportManager->HandleEndScene();
+		// The world is required
+		return false;
+	}
+
+	// The current world should always be initialized immediately before updating the configuration, since there is a dependency (e.g., visibility lists, etc.)
+	SetCurrentWorldImpl(InWorld);
+
+	if (!IsSceneOpened())
+	{
+		// Before render we need to start scene
 		ViewportManager->HandleStartScene();
 	}
 
@@ -319,16 +363,6 @@ bool FDisplayClusterViewportConfiguration::ImplUpdateConfiguration(EDisplayClust
 	return true;
 }
 
-bool FDisplayClusterViewportConfiguration::UpdateConfigurationForClusterNode(EDisplayClusterRenderFrameMode InRenderMode, const FString& InClusterNodeId)
-{
-	return ImplUpdateConfiguration(InRenderMode, InClusterNodeId, nullptr);
-}
-
-bool FDisplayClusterViewportConfiguration::UpdateConfigurationForViewportsList(EDisplayClusterRenderFrameMode InRenderMode, const TArray<FString>& InViewportNames)
-{
-	return ImplUpdateConfiguration(InRenderMode, TEXT(""), &InViewportNames);
-}
-
 void FDisplayClusterViewportConfiguration::ImplUpdateConfigurationVisibility() const
 {
 	ADisplayClusterRootActor* SceneRootActor = GetRootActor(EDisplayClusterRootActorType::Scene);
diff --git a/Engine/Plugins/Runtime/nDisplay/Source/DisplayCluster/Private/Render/Viewport/Configuration/DisplayClusterViewportConfiguration.h b/Engine/Plugins/Runtime/nDisplay/Source/DisplayCluster/Private/Render/Viewport/Configuration/DisplayClusterViewportConfiguration.h
--- a/Engine/Plugins/Runtime/nDisplay/Source/DisplayCluster/Private/Render/Viewport/Configuration/DisplayClusterViewportConfiguration.h
+++ b/Engine/Plugins/Runtime/nDisplay/Source/DisplayCluster/Private/Render/Viewport/Configuration/DisplayClusterViewportConfiguration.h
@@ -31,10 +31,29 @@ public:
 
 public:
 	//~ BEGIN IDisplayClusterViewportConfiguration
-	virtual void SetCurrentWorld(UWorld* InWorld) override;
-	virtual void SetRootActor(ADisplayClusterRootActor* InRootActor, const EDisplayClusterRootActorType InRootActorType) override;
-	virtual bool UpdateConfigurationForClusterNode(EDisplayClusterRenderFrameMode InRenderMode, const FString& InClusterNodeId) override;
-	virtual bool UpdateConfigurationForViewportsList(EDisplayClusterRenderFrameMode InRenderMode, const TArray<FString>& InViewportNames) override;
+	virtual void SetRootActor(const EDisplayClusterRootActorType InRootActorType, const ADisplayClusterRootActor* InRootActor) override;
+	
+	virtual void SetPreviewSettings(const FDisplayClusterViewport_PreviewSettings& InPreviewSettings) override
+	{
+		RenderFrameSettings.PreviewSettings = InPreviewSettings;
+	}
+
+	virtual const FDisplayClusterViewport_PreviewSettings& GetPreviewSettings() const override
+	{
+		return RenderFrameSettings.PreviewSettings;
+	}
+
+	virtual bool UpdateConfigurationForClusterNode(EDisplayClusterRenderFrameMode InRenderMode, const UWorld* InWorld, const FString& InClusterNodeId) override
+	{
+		return ImplUpdateConfiguration(InRenderMode, InWorld, InClusterNodeId, nullptr);
+	}
+
+	virtual bool UpdateConfigurationForViewportsList(EDisplayClusterRenderFrameMode InRenderMode, const UWorld* InWorld, const TArray<FString>& InViewportNames) override
+	{
+		return ImplUpdateConfiguration(InRenderMode, InWorld, TEXT(""), &InViewportNames);
+	}
+
+	virtual void ReleaseConfiguration() override;
 
 	virtual UWorld* GetCurrentWorld() const override;
 	virtual ADisplayClusterRootActor* GetRootActor(const EDisplayClusterRootActorType InRootActorType) const override;
@@ -66,6 +85,18 @@ public:
 	{
 		return RenderFrameSettings.IsPreviewRendering();
 	}
+
+	virtual bool IsTechvisEnabled() const override
+	{
+		return RenderFrameSettings.IsTechvisEnabled();
+	}
+
+	virtual bool IsPreviewInGameEnabled() const override
+	{
+		return RenderFrameSettings.IsPreviewInGameEnabled();
+	}
+
+	virtual const float GetWorldToMeters() const override;
 	// ~~END IDisplayClusterViewportConfiguration
 
 public:
@@ -104,39 +135,42 @@ public:
 		}
 	}
 
+	void OnHandleStartScene();
+	void OnHandleEndScene();
+
 private:
 	/** Update configuration implementation.
 	* 
-	* @param InRenderMode - rendering mode.
+	* @param InRenderMode    - rendering mode.
+	* @param InWorld         - ptr to the world to be rendered
 	* @param InClusterNodeId - (opt) Configuring rendering for a cluster node.
 	* @param InViewportNames - (opt) Configuring rendering for a list of viewports.
 	*/
-	bool ImplUpdateConfiguration(EDisplayClusterRenderFrameMode InRenderMode, const FString& InClusterNodeId, const TArray<FString>* InViewportNames);
+	bool ImplUpdateConfiguration(EDisplayClusterRenderFrameMode InRenderMode, const UWorld* InWorld, const FString& InClusterNodeId, const TArray<FString>* InViewportNames);
 
 	/** Hide DCRA components for nDisplay rendering.*/
 	void ImplUpdateConfigurationVisibility() const;
 
+	/** Set current world.*/
+	void SetCurrentWorldImpl(const UWorld* InWorld);
+
 public:
 	// Reference to configuration proxy object
 	const TSharedRef<FDisplayClusterViewportConfigurationProxy, ESPMode::ThreadSafe> Proxy;
 
-	// Is current scene started
-	bool bCurrentSceneActive = false;
-
-	// Is scene update is required
-	bool bCurrentSceneNeedsToBeUpdated = false;
-
 	// Whether the list of viewports of the current frame needs to be updated
 	bool bCurrentRenderFrameViewportsNeedsToBeUpdated = false;
 
 private:
+	// Is current scene started
+	bool bCurrentSceneActive = false;
+
 	// Current render frame settings
 	FDisplayClusterRenderFrameSettings RenderFrameSettings;
 
 	// A reference to the owning viewport manager
 	TWeakPtr<FDisplayClusterViewportManager, ESPMode::ThreadSafe> ViewportManagerWeakPtr;
 
-private:
 	// This DCRA will be used to render previews. The meshes and preview materials are created at runtime.
 	FDisplayClusterActorRef PreviewRootActorRef;
 
diff --git a/Engine/Plugins/Runtime/nDisplay/Source/DisplayCluster/Private/Render/Viewport/Configuration/DisplayClusterViewportConfigurationHelpers.cpp b/Engine/Plugins/Runtime/nDisplay/Source/DisplayCluster/Private/Render/Viewport/Configuration/DisplayClusterViewportConfigurationHelpers.cpp
--- a/Engine/Plugins/Runtime/nDisplay/Source/DisplayCluster/Private/Render/Viewport/Configuration/DisplayClusterViewportConfigurationHelpers.cpp
+++ b/Engine/Plugins/Runtime/nDisplay/Source/DisplayCluster/Private/Render/Viewport/Configuration/DisplayClusterViewportConfigurationHelpers.cpp
@@ -60,6 +60,8 @@ void FDisplayClusterViewportConfigurationHelpers::UpdateBaseViewportSetting(FDis
 			InOutRenderSettings.bEnable = false;
 		}
 
+		InOutRenderSettings.DisplayDeviceId = InConfigurationViewport.DisplayDeviceName;
+
 		InOutRenderSettings.CameraId = InConfigurationViewport.Camera;
 		InOutRenderSettings.Rect = FDisplayClusterViewportHelpers::GetValidViewportRect(InConfigurationViewport.Region.ToRect(), DstViewport.GetId(), TEXT("Configuration Region"));
 
diff --git a/Engine/Plugins/Runtime/nDisplay/Source/DisplayCluster/Private/Render/Viewport/Configuration/DisplayClusterViewportConfigurationHelpers_RenderFrameSettings.cpp b/Engine/Plugins/Runtime/nDisplay/Source/DisplayCluster/Private/Render/Viewport/Configuration/DisplayClusterViewportConfigurationHelpers_RenderFrameSettings.cpp
--- a/Engine/Plugins/Runtime/nDisplay/Source/DisplayCluster/Private/Render/Viewport/Configuration/DisplayClusterViewportConfigurationHelpers_RenderFrameSettings.cpp
+++ b/Engine/Plugins/Runtime/nDisplay/Source/DisplayCluster/Private/Render/Viewport/Configuration/DisplayClusterViewportConfigurationHelpers_RenderFrameSettings.cpp
@@ -4,6 +4,7 @@
 #include "Render/Viewport/Configuration/DisplayClusterViewportConfiguration.h"
 #include "Render/Viewport/RenderFrame/DisplayClusterRenderFrameSettings.h"
 #include "Render/Viewport/DisplayClusterViewportManager.h"
+#include "Render/Viewport/Preview/DisplayClusterViewportManagerPreview.h"
 
 #include "DisplayClusterRootActor.h"
 #include "IDisplayCluster.h"
@@ -112,11 +113,20 @@ EDisplayClusterRenderFrameAlphaChannelCaptureMode FDisplayClusterViewportConfigu
 	return EDisplayClusterRenderFrameAlphaChannelCaptureMode::None;
 }
 
-bool FDisplayClusterViewportConfigurationHelpers_RenderFrameSettings::UpdateRenderFrameConfiguration(EDisplayClusterRenderFrameMode InRenderMode, FDisplayClusterViewportConfiguration& InOutConfiguration)
+bool FDisplayClusterViewportConfigurationHelpers_RenderFrameSettings::UpdateRenderFrameConfiguration(FDisplayClusterViewportManager* ViewportManager, EDisplayClusterRenderFrameMode InRenderMode, FDisplayClusterViewportConfiguration& InOutConfiguration)
 {
+	if (InRenderMode == EDisplayClusterRenderFrameMode::Unknown)
+	{
+		// Do not initialize for unknown rendering type
+		return false;
+	}
+
+	check(ViewportManager);
+
 	ADisplayClusterRootActor* ConfigurationRootActor = InOutConfiguration.GetRootActor(EDisplayClusterRootActorType::Configuration);
 	if (!(ConfigurationRootActor))
 	{
+		// If the ConfigurationRootActor is not defined, initialization cannot be performed.
 		return false;
 	}
 
@@ -125,26 +135,15 @@ bool FDisplayClusterViewportConfigurationHelpers_RenderFrameSettings::UpdateRend
 	// Set current rendering mode
 	NewRenderFrameSettings.RenderMode = InRenderMode;
 
-#if WITH_EDITOR
-	// Update preview settings:
-	if(ADisplayClusterRootActor* PreviewRootActor = InOutConfiguration.GetRootActor(EDisplayClusterRootActorType::Preview))
+	// Preview Settings : experimental mGPU feature
+	NewRenderFrameSettings.PreviewMultiGPURendering.Reset();
+	if (GDisplayClusterPreviewAllowMultiGPURendering)
 	{
-		NewRenderFrameSettings.PreviewSettings.RenderTargetRatioMult = PreviewRootActor->PreviewRenderTargetRatioMult;
-		NewRenderFrameSettings.PreviewSettings.bFreezeRender = PreviewRootActor->bFreezePreviewRender;
-		NewRenderFrameSettings.PreviewSettings.bEnablePostProcess = PreviewRootActor->ShouldThisFrameOutputPreviewToPostProcessRenderTarget();
-		NewRenderFrameSettings.PreviewSettings.MaxTextureDimension = PreviewRootActor->PreviewMaxTextureDimension;
-
-		// Preview Settings : experimental mGPU feature
-		NewRenderFrameSettings.PreviewSettings.MultiGPURendering.Reset();
-		if (GDisplayClusterPreviewAllowMultiGPURendering)
-		{
-			const int32 MinGPUIndex = FMath::Max(0, GDisplayClusterPreviewMultiGPURenderingMinIndex);
-			const int32 MaxGPUIndex = FMath::Max(MinGPUIndex, GDisplayClusterPreviewMultiGPURenderingMaxIndex);
-
-			NewRenderFrameSettings.PreviewSettings.MultiGPURendering = FIntPoint(MinGPUIndex, MaxGPUIndex);
-		}
+		const int32 MinGPUIndex = FMath::Max(0, GDisplayClusterPreviewMultiGPURenderingMinIndex);
+		const int32 MaxGPUIndex = FMath::Max(MinGPUIndex, GDisplayClusterPreviewMultiGPURenderingMaxIndex);
+
+		NewRenderFrameSettings.PreviewMultiGPURendering = FIntPoint(MinGPUIndex, MaxGPUIndex);
 	}
-#endif
 
 	// Support alpha channel capture
 	NewRenderFrameSettings.AlphaChannelCaptureMode = GetAlphaChannelCaptureMode();
@@ -200,7 +199,8 @@ void FDisplayClusterViewportConfigurationHelpers_RenderFrameSettings::PostUpdate
 	{
 		FDisplayClusterRenderFrameSettings NewRenderFrameSettings = InOutConfiguration.GetRenderFrameSettings();
 
-		// Some frame postprocess require additional render targetable resources
+		// Update global flags that control viewport resources
+		NewRenderFrameSettings.bShouldUseOutputTargetableResources = ViewportManager->ShouldUseOutputTargetableResources();
 		NewRenderFrameSettings.bShouldUseAdditionalFrameTargetableResource = ViewportManager->ShouldUseAdditionalFrameTargetableResource();
 		NewRenderFrameSettings.bShouldUseFullSizeFrameTargetableResource = ViewportManager->ShouldUseFullSizeFrameTargetableResource();
 
diff --git a/Engine/Plugins/Runtime/nDisplay/Source/DisplayCluster/Private/Render/Viewport/Configuration/DisplayClusterViewportConfigurationHelpers_RenderFrameSettings.h b/Engine/Plugins/Runtime/nDisplay/Source/DisplayCluster/Private/Render/Viewport/Configuration/DisplayClusterViewportConfigurationHelpers_RenderFrameSettings.h
--- a/Engine/Plugins/Runtime/nDisplay/Source/DisplayCluster/Private/Render/Viewport/Configuration/DisplayClusterViewportConfigurationHelpers_RenderFrameSettings.h
+++ b/Engine/Plugins/Runtime/nDisplay/Source/DisplayCluster/Private/Render/Viewport/Configuration/DisplayClusterViewportConfigurationHelpers_RenderFrameSettings.h
@@ -7,6 +7,7 @@
 #include "Render/Viewport/RenderFrame/DisplayClusterRenderFrameSettingsEnums.h"
 
 class FDisplayClusterViewportConfiguration;
+class FDisplayClusterViewportManager;
 
 /**
 * This enum is for CVar values only and is used to process the logic that converts the values to the runtime enum in GetAlphaChannelCaptureMode().
@@ -55,7 +56,7 @@ class FDisplayClusterViewportConfigurationHelpers_RenderFrameSettings
 {
 public:
 	// Update RenderFrameSettings
-	static bool UpdateRenderFrameConfiguration(EDisplayClusterRenderFrameMode InRenderMode, FDisplayClusterViewportConfiguration& InOutConfiguration);
+	static bool UpdateRenderFrameConfiguration(FDisplayClusterViewportManager* ViewportManager, EDisplayClusterRenderFrameMode InRenderMode, FDisplayClusterViewportConfiguration& InOutConfiguration);
 	static void PostUpdateRenderFrameConfiguration(FDisplayClusterViewportConfiguration& InOutConfiguration);
 
 private:
diff --git a/Engine/Plugins/Runtime/nDisplay/Source/DisplayCluster/Private/Render/Viewport/Configuration/DisplayClusterViewportConfiguration_ICVFXCamera.cpp b/Engine/Plugins/Runtime/nDisplay/Source/DisplayCluster/Private/Render/Viewport/Configuration/DisplayClusterViewportConfiguration_ICVFXCamera.cpp
--- a/Engine/Plugins/Runtime/nDisplay/Source/DisplayCluster/Private/Render/Viewport/Configuration/DisplayClusterViewportConfiguration_ICVFXCamera.cpp
+++ b/Engine/Plugins/Runtime/nDisplay/Source/DisplayCluster/Private/Render/Viewport/Configuration/DisplayClusterViewportConfiguration_ICVFXCamera.cpp
@@ -124,15 +124,7 @@ bool FDisplayClusterViewportConfiguration_ICVFXCamera::Initialize()
 	const FVector ViewOffset = FVector::ZeroVector;
 
 	// Get world scale
-	float WorldToMeters = 100.f;
-	if (UWorld* World = SceneRootActor->GetWorld())
-	{
-		if (AWorldSettings* WorldSettings = World->GetWorldSettings())
-		{
-			WorldToMeters = WorldSettings->WorldToMeters;
-		}
-	}
-
+	const float WorldToMeters = Configuration.GetWorldToMeters();
 	// Supports custom near clipping plane
 	const float NCP = (CustomNearClippingPlane >= 0) ? CustomNearClippingPlane : GNearClippingPlane;
 
diff --git a/Engine/Plugins/Runtime/nDisplay/Source/DisplayCluster/Private/Render/Viewport/Configuration/DisplayClusterViewportConfiguration_ProjectionPolicy.cpp b/Engine/Plugins/Runtime/nDisplay/Source/DisplayCluster/Private/Render/Viewport/Configuration/DisplayClusterViewportConfiguration_ProjectionPolicy.cpp
--- a/Engine/Plugins/Runtime/nDisplay/Source/DisplayCluster/Private/Render/Viewport/Configuration/DisplayClusterViewportConfiguration_ProjectionPolicy.cpp
+++ b/Engine/Plugins/Runtime/nDisplay/Source/DisplayCluster/Private/Render/Viewport/Configuration/DisplayClusterViewportConfiguration_ProjectionPolicy.cpp
@@ -80,28 +80,41 @@ bool FDisplayClusterViewportConfiguration_ProjectionPolicy::UpdateCameraPolicy(F
 	if (!DisplayClusterHelpers::map::template ExtractValue(CameraPolicyParameters, DisplayClusterProjectionStrings::cfg::camera::Component, CameraComponentId))
 	{
 		// use default cameras
+		DstViewport.ResetShowLogMsgOnce(EDisplayClusterViewportShowLogMsgOnce::UpdateCameraPolicy);
+
 		return true;
 	}
 
 	if (CameraComponentId.IsEmpty())
 	{
-		UE_LOG(LogDisplayClusterViewport, Verbose, TEXT("Viewport '%s': referenced camera '' (empty name)."), *DstViewport.GetId());
+		if (DstViewport.CanShowLogMsgOnce(EDisplayClusterViewportShowLogMsgOnce::UpdateCameraPolicy_ReferencedCameraNameIsEmpty))
+		{
+			UE_LOG(LogDisplayClusterViewport, Verbose, TEXT("Viewport '%s': referenced camera '' (empty name)."), *DstViewport.GetId());
+		}
+
 		return false;
 	}
 
 	// Get ICVFX camera component
 	if (UpdateCameraPolicy_ICVFX(DstViewport, CameraComponentId))
 	{
+		DstViewport.ResetShowLogMsgOnce(EDisplayClusterViewportShowLogMsgOnce::UpdateCameraPolicy);
+
 		return true;
 	}
 
 	// Get camera component
 	if (UpdateCameraPolicy_Base(DstViewport, CameraComponentId))
 	{
+		DstViewport.ResetShowLogMsgOnce(EDisplayClusterViewportShowLogMsgOnce::UpdateCameraPolicy);
+
 		return true;
 	}
 
-	UE_LOG(LogDisplayClusterViewport, Error, TEXT("Viewport '%s': referenced camera '%s' not found."), *DstViewport.GetId(), *CameraComponentId);
+	if (DstViewport.CanShowLogMsgOnce(EDisplayClusterViewportShowLogMsgOnce::UpdateCameraPolicy_ReferencedCameraNotFound))
+	{
+		UE_LOG(LogDisplayClusterViewport, Error, TEXT("Viewport '%s': referenced camera '%s' not found."), *DstViewport.GetId(), *CameraComponentId);
+	}
 
 	return false;
 }
diff --git a/Engine/Plugins/Runtime/nDisplay/Source/DisplayCluster/Private/Render/Viewport/Containers/DisplayClusterViewportProxyData.h b/Engine/Plugins/Runtime/nDisplay/Source/DisplayCluster/Private/Render/Viewport/Containers/DisplayClusterViewportProxyData.h
--- a/Engine/Plugins/Runtime/nDisplay/Source/DisplayCluster/Private/Render/Viewport/Containers/DisplayClusterViewportProxyData.h
+++ b/Engine/Plugins/Runtime/nDisplay/Source/DisplayCluster/Private/Render/Viewport/Containers/DisplayClusterViewportProxyData.h
@@ -46,6 +46,9 @@ public:
 	// OpenColorIO
 	TSharedPtr<FDisplayClusterViewport_OpenColorIO, ESPMode::ThreadSafe> OpenColorIO;
 
+	// Display Device
+	TSharedPtr<IDisplayClusterDisplayDeviceProxy, ESPMode::ThreadSafe> DisplayDeviceProxy;
+
 	// Projection policy instance that serves this viewport
 	TSharedPtr<IDisplayClusterProjectionPolicy, ESPMode::ThreadSafe> ProjectionPolicy;
 
diff --git a/Engine/Plugins/Runtime/nDisplay/Source/DisplayCluster/Private/Render/Viewport/Containers/DisplayClusterViewport_InternalEnums.h b/Engine/Plugins/Runtime/nDisplay/Source/DisplayCluster/Private/Render/Viewport/Containers/DisplayClusterViewport_InternalEnums.h
new file mode 100644
--- /dev/null
+++ b/Engine/Plugins/Runtime/nDisplay/Source/DisplayCluster/Private/Render/Viewport/Containers/DisplayClusterViewport_InternalEnums.h
@@ -0,0 +1,30 @@
+// Copyright Epic Games, Inc. All Rights Reserved.
+
+#pragma once
+
+#include "CoreMinimal.h"
+
+/**
+ * Special flags for nDisplay viewport log
+ */
+enum class EDisplayClusterViewportShowLogMsgOnce : uint8
+{
+	None = 0,
+
+	// No projection policy assigned for Viewports
+	HandleStartScene_InvalidProjectionPolicy = 1 << 0,
+
+	UpdateFrameContexts_FrameTargetRectHasZeroSize = 1 << 1,
+	UpdateFrameContexts_RenderTargetRectHasZeroSize = 1 << 2,
+	UpdateFrameContexts = UpdateFrameContexts_FrameTargetRectHasZeroSize | UpdateFrameContexts_RenderTargetRectHasZeroSize,
+
+	UpdateCameraPolicy_ReferencedCameraNameIsEmpty = 1 << 3,
+	UpdateCameraPolicy_ReferencedCameraNotFound = 1 << 4,
+	UpdateCameraPolicy = UpdateCameraPolicy_ReferencedCameraNameIsEmpty | UpdateCameraPolicy_ReferencedCameraNotFound,
+
+	GetViewPointCameraComponent_NoRootActorFound = 1 << 5,
+	GetViewPointCameraComponent_HasAssignedViewPoint = 1 << 6,
+	GetViewPointCameraComponent_NotFound = 1 << 7,
+	GetViewPointCameraComponent = GetViewPointCameraComponent_NoRootActorFound | GetViewPointCameraComponent_NotFound,
+};
+ENUM_CLASS_FLAGS(EDisplayClusterViewportShowLogMsgOnce);
diff --git a/Engine/Plugins/Runtime/nDisplay/Source/DisplayCluster/Private/Render/Viewport/DisplayClusterViewport.cpp b/Engine/Plugins/Runtime/nDisplay/Source/DisplayCluster/Private/Render/Viewport/DisplayClusterViewport.cpp
--- a/Engine/Plugins/Runtime/nDisplay/Source/DisplayCluster/Private/Render/Viewport/DisplayClusterViewport.cpp
+++ b/Engine/Plugins/Runtime/nDisplay/Source/DisplayCluster/Private/Render/Viewport/DisplayClusterViewport.cpp
@@ -22,6 +22,10 @@
 #include "Render/Viewport/Containers/DisplayClusterViewport_PostRenderSettings.h"
 #include "Render/Viewport/Containers/DisplayClusterViewportProxyData.h"
 
+#include "Render/Viewport/Preview/DisplayClusterViewportPreview.h"
+
+#include "Render/DisplayDevice/Components/DisplayClusterDisplayDeviceBaseComponent.h"
+
 #include "EngineUtils.h"
 #include "SceneManagement.h"
 #include "SceneView.h"
@@ -46,6 +50,7 @@ static FAutoConsoleVariableRef CVarDisplayClusterMultiGPUEnable(
 ///////////////////////////////////////////////////////////////////////////////////////
 FDisplayClusterViewport::FDisplayClusterViewport(const TSharedRef<FDisplayClusterViewportConfiguration, ESPMode::ThreadSafe>& InConfiguration, const FString& InViewportId, const TSharedPtr<IDisplayClusterProjectionPolicy, ESPMode::ThreadSafe>& InProjectionPolicy)
 	: Configuration(InConfiguration)
+	, ViewportPreview(MakeShared< FDisplayClusterViewportPreview, ESPMode::ThreadSafe>(InConfiguration, InViewportId))
 	, ViewportProxy(MakeShared<FDisplayClusterViewportProxy, ESPMode::ThreadSafe>(InConfiguration, InViewportId, InProjectionPolicy))
 	, ViewportId(InViewportId)
 	, ClusterNodeId(InConfiguration->GetClusterNodeId())
@@ -87,7 +92,7 @@ FDisplayClusterViewport::~FDisplayClusterViewport()
 	OpenColorIO.Reset();
 
 	// Handle projection policy EndScene event
-	HandleEndScene();
+	OnHandleEndScene();
 
 	// Handle projection policy event
 	ProjectionPolicy.Reset();
@@ -100,12 +105,23 @@ FDisplayClusterViewport::~FDisplayClusterViewport()
 	}
 }
 
+IDisplayClusterViewportPreview& FDisplayClusterViewport::GetViewportPreview() const
+{
+	return ViewportPreview.Get();
+}
+
 FDisplayClusterViewportProxyData* FDisplayClusterViewport::CreateViewportProxyData()
 {
 	FDisplayClusterViewportProxyData* OutViewportProxyData = new FDisplayClusterViewportProxyData(ViewportProxy);
 
 	OutViewportProxyData->OpenColorIO = OpenColorIO;
 
+	// Get Display Device proxy object
+	if (UDisplayClusterDisplayDeviceBaseComponent* DisplayDevice = GetDisplayDeviceComponent(EDisplayClusterRootActorType::Configuration))
+	{
+		OutViewportProxyData->DisplayDeviceProxy = DisplayDevice->GetDisplayDeviceProxy(GetConfiguration());
+	}
+
 	OutViewportProxyData->RenderSettings = RenderSettings;
 	OutViewportProxyData->RenderSettingsICVFX.SetParameters(RenderSettingsICVFX);
 	OutViewportProxyData->PostRenderSettings.SetParameters(PostRenderSettings);
@@ -229,7 +245,18 @@ const TArray<FSceneViewExtensionRef> FDisplayClusterViewport::GatherActiveExtens
 	return OutCustomExtensions;
 }
 
-void FDisplayClusterViewport::HandleStartScene()
+void FDisplayClusterViewport::Initialize()
+{
+	// Initialize a reference to this viewport for the preview API
+	ViewportPreview->Initialize(*this);
+}
+
+void FDisplayClusterViewport::ReleaseTextures()
+{
+	Resources.ReleaseAllResources();
+}
+
+void FDisplayClusterViewport::OnHandleStartScene()
 {
 	if (UninitializedProjectionPolicy.IsValid())
 	{
@@ -237,20 +264,25 @@ void FDisplayClusterViewport::HandleStartScene()
 		{
 			ProjectionPolicy = UninitializedProjectionPolicy;
 			UninitializedProjectionPolicy.Reset();
+
+			ResetShowLogMsgOnce(EDisplayClusterViewportShowLogMsgOnce::HandleStartScene_InvalidProjectionPolicy);
 		}
 	}
-	else 
+	else
 	{
 		// Already Initialized
 		if (!ProjectionPolicy.IsValid())
 		{
 			// No projection policy for this viewport
-			UE_LOG(LogDisplayClusterViewport, Error, TEXT("No projection policy assigned for Viewports '%s'."), *GetId());
+			if (CanShowLogMsgOnce(EDisplayClusterViewportShowLogMsgOnce::HandleStartScene_InvalidProjectionPolicy))
+			{
+				UE_LOG(LogDisplayClusterViewport, Error, TEXT("No projection policy assigned for Viewports '%s'."), *GetId());
+			}
 		}
 	}
 }
 
-void FDisplayClusterViewport::HandleEndScene()
+void FDisplayClusterViewport::OnHandleEndScene()
 {
 	if (ProjectionPolicy.IsValid())
 	{
@@ -286,11 +318,27 @@ bool FDisplayClusterViewport::ShouldUseAdditionalTargetableResource() const
 	return false;
 }
 
+bool FDisplayClusterViewport::ShouldUseOutputTargetableResources() const
+{
+	// Do not create output RTTs for internal ICVFX resources (invisible viewports that are only used to compose the output of other viewports)
+	if (EnumHasAnyFlags(RenderSettingsICVFX.RuntimeFlags, EDisplayClusterViewportRuntimeICVFXFlags::InternalResource))
+	{
+		return false;
+	}
+
+	// Only if this viewport is enabled and visible on the final frame.
+	return RenderSettings.bEnable && RenderSettings.bVisible;
+}
+
 bool FDisplayClusterViewport::ShouldUseAdditionalFrameTargetableResource() const
 {
-	if (ViewportRemap.IsUsed())
+	if (ShouldUseOutputTargetableResources())
 	{
-		return true;
+		// OutputFrameTargetableResources must be used for AdditionalFrameTargetableResource
+		if (ViewportRemap.IsUsed())
+		{
+			return true;
+		}
 	}
 
 	return false;
@@ -372,6 +420,12 @@ void FDisplayClusterViewport::SetupSceneView(uint32 ContextNum, class UWorld* Wo
 
 	// Handle Motion blur parameters
 	CameraMotionBlur.SetupSceneView(Contexts[ContextNum], InOutView);
+
+	// Handle DisplayDevice
+	if (UDisplayClusterDisplayDeviceBaseComponent* InDisplayDeviceComponent = GetDisplayDeviceComponent(EDisplayClusterRootActorType::Configuration))
+	{
+		InDisplayDeviceComponent->SetupSceneView(*ViewportPreview, ContextNum, InOutViewFamily, InOutView);
+	}
 }
 
 inline void AdjustRect(FIntRect& InOutRect, const float multX, const float multY)
@@ -441,7 +495,7 @@ float FDisplayClusterViewport::GetCustomBufferRatio(const FDisplayClusterRenderF
 
 void FDisplayClusterViewport::ResetFrameContexts()
 {
-	Resources.Release();
+	Resources.ReleaseAllResources();
 }
 
 bool FDisplayClusterViewport::UpdateFrameContexts(const uint32 InStereoViewIndex)
@@ -473,9 +527,6 @@ bool FDisplayClusterViewport::UpdateFrameContexts(const uint32 InStereoViewIndex
 	// Special case mono->stereo
 	const uint32 ViewportContextAmount = RenderSettings.bForceMono ? 1 : FrameTargetsAmount;
 
-	// Release only part of the resources, leaving resources that can be used by other viewports (viewport override feature)
-	Resources.Release(false);
-
 	// Freeze the image in the viewport only after the frame has been rendered
 	if (RenderSettings.bFreezeRendering && RenderSettings.bEnable)
 	{
@@ -495,6 +546,9 @@ bool FDisplayClusterViewport::UpdateFrameContexts(const uint32 InStereoViewIndex
 				ContextIt.FrameTargetRect = FDisplayClusterViewportHelpers::GetValidViewportRect(DesiredFrameTargetRect, GetId(), TEXT("Context Frame"));
 			}
 
+			// Release only part of the resources, leaving resources that can be used by other viewports (viewport override feature)
+			Resources.ReleaseNotSharedResources();
+
 			return true;
 		}
 	}
@@ -503,7 +557,7 @@ bool FDisplayClusterViewport::UpdateFrameContexts(const uint32 InStereoViewIndex
 	Contexts.Empty();
 
 	// Free all resources
-	Resources.Release();
+	Resources.ReleaseAllResources();
 
 	if (RenderSettings.bEnable == false)
 	{
@@ -527,7 +581,11 @@ bool FDisplayClusterViewport::UpdateFrameContexts(const uint32 InStereoViewIndex
 	// Exclude zero-size viewports from render
 	if (FrameTargetRect.Size().GetMin() <= 0)
 	{
-		UE_LOG(LogDisplayClusterViewport, Error, TEXT("The viewport '%s' FrameTarget rect has zero size %dx%d: Disabled"), *GetId(), FrameTargetRect.Size().X, FrameTargetRect.Size().Y);
+		if (CanShowLogMsgOnce(EDisplayClusterViewportShowLogMsgOnce::UpdateFrameContexts_FrameTargetRectHasZeroSize))
+		{
+			UE_LOG(LogDisplayClusterViewport, Error, TEXT("The viewport '%s' FrameTarget rect has zero size %dx%d: Disabled"), *GetId(), FrameTargetRect.Size().X, FrameTargetRect.Size().Y);
+		}
+
 		return false;
 	}
 
@@ -536,7 +594,7 @@ bool FDisplayClusterViewport::UpdateFrameContexts(const uint32 InStereoViewIndex
 
 	// Apply restrictions on the maximum size of the viewport texture.
 	const int32 ViewportTextureMaxSize = InFrameSettings.GetViewportTextureMaxSize();
-	if(ViewportTextureMaxSize > 0)
+	if (ViewportTextureMaxSize > 0)
 	{
 		DesiredContextSize = FDisplayClusterViewportHelpers::GetTextureSizeLessThanMax(DesiredContextSize, ViewportTextureMaxSize);
 	}
@@ -544,7 +602,11 @@ bool FDisplayClusterViewport::UpdateFrameContexts(const uint32 InStereoViewIndex
 	// Exclude zero-size viewports from render
 	if (DesiredContextSize.GetMin() <= 0)
 	{
-		UE_LOG(LogDisplayClusterViewport, Error, TEXT("The viewport '%s' RenderTarget rect has zero size %dx%d: Disabled"), *GetId(), DesiredContextSize.X, DesiredContextSize.Y);
+		if (CanShowLogMsgOnce(EDisplayClusterViewportShowLogMsgOnce::UpdateFrameContexts_RenderTargetRectHasZeroSize))
+		{
+			UE_LOG(LogDisplayClusterViewport, Error, TEXT("The viewport '%s' RenderTarget rect has zero size %dx%d: Disabled"), *GetId(), DesiredContextSize.X, DesiredContextSize.Y);
+		}
+
 		return false;
 	}
 
@@ -570,7 +632,6 @@ bool FDisplayClusterViewport::UpdateFrameContexts(const uint32 InStereoViewIndex
 	// Fix buffer ratio value vs MaxTextureSize:
 	const float CustomBufferRatio = FDisplayClusterViewportHelpers::GetValidSizeMultiplier(RenderTargetRect.Size(), BaseCustomBufferRatio, 1.f);
 
-	// Setup resource usage logic:
 	bool bDisableRender = false;
 	if (PostRenderSettings.Replace.IsEnabled())
 	{
@@ -629,7 +690,7 @@ bool FDisplayClusterViewport::UpdateFrameContexts(const uint32 InStereoViewIndex
 			}
 		}
 
-		if(!bUseUVLightCardViewport)
+		if (!bUseUVLightCardViewport)
 		{
 			// do not use UV LightCard viewport
 			return false;
@@ -654,18 +715,18 @@ bool FDisplayClusterViewport::UpdateFrameContexts(const uint32 InStereoViewIndex
 
 		const int32 MaxExplicitGPUIndex = GDisplayClusterMultiGPUEnable ? GNumExplicitGPUsForRendering - 1 : 0;
 		if (MaxExplicitGPUIndex > 0 && !bDisableRender)
-			{
-				// Experimental: allow mGPU for preview rendering:
+		{
+			// Experimental: allow mGPU for preview rendering:
 			if (const FIntPoint* GPURange = InFrameSettings.GetPreviewMultiGPURendering())
-				{
+			{
 				int32 MinGPUIndex = FMath::Min(GPURange->X, MaxExplicitGPUIndex);
 				int32 MaxGPUIndex = FMath::Min(GPURange->Y, MaxExplicitGPUIndex);
 
-					static int32 PreviewGPUIndex = MinGPUIndex;
-					if (PreviewGPUIndex > MaxGPUIndex)
-					{
-						PreviewGPUIndex = MinGPUIndex;
-					}
+				static int32 PreviewGPUIndex = MinGPUIndex;
+				if (PreviewGPUIndex > MaxGPUIndex)
+				{
+					PreviewGPUIndex = MinGPUIndex;
+				}
 
 				Context.GPUIndex = PreviewGPUIndex++;
 			}
@@ -700,12 +761,6 @@ bool FDisplayClusterViewport::UpdateFrameContexts(const uint32 InStereoViewIndex
 	}
 
 	// Reserve for resources
-	if (InFrameSettings.IsPreviewRendering())
-	{
-		// reserve preview texture resource
-		Resources[EDisplayClusterViewportResource::OutputPreviewTargetableResources].AddZeroed(FrameTargetsAmount);
-	}
-
 	if (!bDisableRender)
 	{
 		Resources[EDisplayClusterViewportResource::RenderTargets].AddZeroed(FrameTargetsAmount);
@@ -731,5 +786,13 @@ bool FDisplayClusterViewport::UpdateFrameContexts(const uint32 InStereoViewIndex
 		}
 	}
 
+	if (InFrameSettings.IsPreviewRendering() && ShouldUseOutputTargetableResources())
+	{
+		// reserve preview texture resource for all visible viewports
+		Resources[EDisplayClusterViewportResource::OutputPreviewTargetableResources].AddZeroed(FrameTargetsAmount);
+	}
+
+	ResetShowLogMsgOnce(EDisplayClusterViewportShowLogMsgOnce::UpdateFrameContexts);
+
 	return true;
 }
diff --git a/Engine/Plugins/Runtime/nDisplay/Source/DisplayCluster/Private/Render/Viewport/DisplayClusterViewport.h b/Engine/Plugins/Runtime/nDisplay/Source/DisplayCluster/Private/Render/Viewport/DisplayClusterViewport.h
--- a/Engine/Plugins/Runtime/nDisplay/Source/DisplayCluster/Private/Render/Viewport/DisplayClusterViewport.h
+++ b/Engine/Plugins/Runtime/nDisplay/Source/DisplayCluster/Private/Render/Viewport/DisplayClusterViewport.h
@@ -10,6 +10,7 @@
 #include "Render/Viewport/Containers/DisplayClusterViewport_OverscanRuntimeSettings.h"
 #include "Render/Viewport/Containers/DisplayClusterViewportRemap.h"
 #include "Render/Viewport/Containers/DisplayClusterViewport_Enums.h"
+#include "Render/Viewport/Containers/DisplayClusterViewport_InternalEnums.h"
 
 #include "Render/Viewport/DisplayClusterViewport_CustomPostProcessSettings.h"
 #include "Render/Viewport/DisplayClusterViewport_VisibilitySettings.h"
@@ -24,6 +25,7 @@
 class FDisplayClusterViewportManager;
 class FDisplayClusterViewportManagerProxy;
 class FDisplayClusterViewportProxy;
+class FDisplayClusterViewportPreview;
 
 struct FDisplayClusterRenderFrameSettings;
 
@@ -52,6 +54,16 @@ public:
 		return AsShared();
 	}
 
+	virtual TSharedRef<class FDisplayClusterViewport, ESPMode::ThreadSafe> ToSharedRef() override
+	{
+		return AsShared();
+	}
+
+	virtual TSharedRef<const class FDisplayClusterViewport, ESPMode::ThreadSafe> ToSharedRef() const override
+	{
+		return AsShared();
+	}
+
 	virtual IDisplayClusterViewportConfiguration& GetConfiguration() override
 	{
 		return Configuration.Get();
@@ -62,15 +74,16 @@ public:
 		return Configuration.Get();
 	}
 
+	/** Get viewport preview API */
+	virtual IDisplayClusterViewportPreview& GetViewportPreview() const override;
+
 	virtual FString GetId() const override
 	{ 
-		check(IsInGameThread());
 		return ViewportId; 
 	}
 
 	virtual FString GetClusterNodeId() const override
 	{
-		check(IsInGameThread());
 		return ClusterNodeId;
 	}
 
@@ -103,6 +116,7 @@ public:
 	virtual bool SetupViewPoint(FMinimalViewInfo& InOutViewInfo) override;
 	virtual float GetStereoEyeOffsetDistance(const uint32 InContextNum) override;
 	virtual class UDisplayClusterCameraComponent* GetViewPointCameraComponent(const EDisplayClusterRootActorType InRootActorType) const override;
+	virtual class UDisplayClusterDisplayDeviceBaseComponent* GetDisplayDeviceComponent(const EDisplayClusterRootActorType InRootActorType) const override;
 	virtual bool GetViewPointCameraEye(const uint32 InContextNum, FVector& OutViewLocation, FRotator& OutViewRotation, FVector& OutViewOffset) override;
 
 	virtual const FDisplayClusterViewport_RenderSettingsICVFX& GetRenderSettingsICVFX() const override
@@ -148,14 +162,17 @@ public:
 	/// ~IDisplayClusterViewport
 	//////////////////////////////////////////////////////
 
-	FSceneView* ImplCalcScenePreview(class FSceneViewFamilyContext& InOutViewFamily, uint32 ContextNum);
-	bool    ImplPreview_CalculateStereoViewOffset(const uint32 InContextNum, FRotator& ViewRotation, const float WorldToMeters, FVector& ViewLocation);
-	FMatrix ImplPreview_GetStereoProjectionMatrix(const uint32 InContextNum);
+	void Initialize();
+	void ReleaseTextures();
+
 
 	// Get from logic request for additional targetable resource
 	bool ShouldUseAdditionalTargetableResource() const;
 	bool ShouldUseAdditionalFrameTargetableResource() const;
 	bool ShouldUseFullSizeFrameTargetableResource() const;
+	
+	// Return true if this viewport requires to use any of output resources (OutputPreviewTargetableResources or OutputFrameTargetableResources)
+	bool ShouldUseOutputTargetableResources() const;
 
 	void SetViewportBufferRatio(const float InBufferRatio);
 
@@ -179,8 +196,8 @@ public:
 		return false;
 	}
 
-	void HandleStartScene();
-	void HandleEndScene();
+	void OnHandleStartScene();
+	void OnHandleEndScene();
 
 	void AddReferencedObjects(FReferenceCollector& Collector);
 
@@ -318,6 +335,25 @@ public:
 	/** Cleanup view states. */
 	void CleanupViewState();
 
+	/** Returns true once if this type of log message can be displayed for the first time.*/
+	bool CanShowLogMsgOnce(const EDisplayClusterViewportShowLogMsgOnce& InLogState) const
+	{
+		if (!EnumHasAnyFlags(ShowLogMsgOnceFlags, InLogState))
+		{
+			EnumAddFlags(ShowLogMsgOnceFlags, InLogState);
+
+			return true;
+		}
+
+		return false;
+	}
+
+	/** Reset viewport log states. */
+	void ResetShowLogMsgOnce(const EDisplayClusterViewportShowLogMsgOnce& InLogState) const
+	{
+		EnumRemoveFlags(ShowLogMsgOnceFlags, InLogState);
+	}
+
 private:
 	float GetClusterRenderTargetRatioMult(const FDisplayClusterRenderFrameSettings& InFrameSettings) const;
 	FIntPoint GetDesiredContextSize(const FIntPoint& InSize, const FDisplayClusterRenderFrameSettings& InFrameSettings) const;
@@ -327,6 +363,9 @@ public:
 	// Configuration of the current cluster node
 	const TSharedRef<FDisplayClusterViewportConfiguration, ESPMode::ThreadSafe> Configuration;
 
+	// Viewport preview
+	const TSharedRef<FDisplayClusterViewportPreview, ESPMode::ThreadSafe> ViewportPreview;
+
 	// viewport proxy (render thread data)
 	const TSharedRef<FDisplayClusterViewportProxy, ESPMode::ThreadSafe> ViewportProxy;
 
@@ -376,8 +415,8 @@ private:
 	// View states (preview only)
 	TArray<TSharedPtr<FSceneViewStateReference, ESPMode::ThreadSafe>> ViewStates;
 
-	// Auxiliary variable for the log
-	bool bProjectionPolicyCalculateViewWarningOnce = false;
+	// A recurring message in the log will be shown only once
+	mutable EDisplayClusterViewportShowLogMsgOnce ShowLogMsgOnceFlags = EDisplayClusterViewportShowLogMsgOnce::None;
 
 	// Near clipping plane value (obtained from the GetDesiredView() functions).
 	// If the value is less than zero, it does not apply to this viewport.
diff --git a/Engine/Plugins/Runtime/nDisplay/Source/DisplayCluster/Private/Render/Viewport/DisplayClusterViewportManager.cpp b/Engine/Plugins/Runtime/nDisplay/Source/DisplayCluster/Private/Render/Viewport/DisplayClusterViewportManager.cpp
--- a/Engine/Plugins/Runtime/nDisplay/Source/DisplayCluster/Private/Render/Viewport/DisplayClusterViewportManager.cpp
+++ b/Engine/Plugins/Runtime/nDisplay/Source/DisplayCluster/Private/Render/Viewport/DisplayClusterViewportManager.cpp
@@ -6,6 +6,9 @@
 #include "Render/Viewport/DisplayClusterViewportManagerViewPointExtension.h"
 #include "Render/Viewport/DisplayClusterViewportFrameStatsViewExtension.h"
 
+#include "Render/Viewport/Preview/DisplayClusterViewportManagerPreview.h"
+#include "Render/Viewport/Preview/DisplayClusterViewportPreview.h"
+
 #include "IDisplayCluster.h"
 #include "Render/IDisplayClusterRenderManager.h"
 #include "Render/Projection/IDisplayClusterProjectionPolicyFactory.h"
@@ -24,6 +27,7 @@
 #include "Render/Viewport/Configuration/DisplayClusterViewportConfiguration.h"
 #include "Render/Viewport/Configuration/DisplayClusterViewportConfigurationProxy.h"
 #include "Render/Viewport/Configuration/DisplayClusterViewportConfiguration_Viewport.h"
+
 #include "Render/Viewport/DisplayClusterViewportStrings.h"
 
 #include "IDisplayClusterWarpBlend.h"
@@ -131,9 +135,19 @@ using namespace UE::DisplayCluster::ViewportManager;
 ///////////////////////////////////////////////////////////////////////////////////////
 //          FDisplayClusterViewportManager
 ///////////////////////////////////////////////////////////////////////////////////////
+TSharedRef<IDisplayClusterViewportManager, ESPMode::ThreadSafe> IDisplayClusterViewportManager::CreateViewportManager()
+{
+	TSharedRef<FDisplayClusterViewportManager, ESPMode::ThreadSafe> ViewportManager = MakeShared<FDisplayClusterViewportManager, ESPMode::ThreadSafe>();
+
+	// After the constructor, we should always call this function to initialize internal references.
+	ViewportManager->Initialize();
+
+	return ViewportManager;
+}
 
 FDisplayClusterViewportManager::FDisplayClusterViewportManager()
 	: Configuration(MakeShared<FDisplayClusterViewportConfiguration, ESPMode::ThreadSafe>())
+	, ViewportManagerPreview(MakeShared<FDisplayClusterViewportManagerPreview, ESPMode::ThreadSafe>(Configuration))
 	, RenderTargetManager(MakeShared<FDisplayClusterRenderTargetManager, ESPMode::ThreadSafe>(Configuration))
 	, PostProcessManager(MakeShared<FDisplayClusterViewportPostProcessManager, ESPMode::ThreadSafe>(Configuration))
 	, LightCardManager(MakeShared<FDisplayClusterViewportLightCardManager, ESPMode::ThreadSafe>(Configuration))
@@ -152,11 +166,17 @@ void FDisplayClusterViewportManager::Initialize()
 
 	// Initialize configuration for exists viewport managers/proxy objects
 	Configuration->Initialize(*this);
+
+	// Register this viewport manager with the preview rendering pipeline.
+	ViewportManagerPreview->RegisterPreviewRendering();
 }
 
 FDisplayClusterViewportManager::~FDisplayClusterViewportManager()
 {
 	UnregisterCallbacks();
+
+	// Remove this viewport manager from the preview rendering pipeline.
+	ViewportManagerPreview->UnregisterPreviewRendering();
 	
 	// Remove viewports
 	EntireClusterViewports.Reset();
@@ -168,11 +188,13 @@ FDisplayClusterViewportManager::~FDisplayClusterViewportManager()
 	FrameStatsViewExtension.Reset();
 
 	LightCardManager->Release();
+	RenderTargetManager->Release();
+	PostProcessManager->Release_GameThread();
 
 	if (ViewportManagerProxy.IsValid())
 	{
 		// Remove viewport manager proxy on render_thread
-		ENQUEUE_RENDER_COMMAND(DeleteDisplayClusterViewportManagerProxy)(
+		ENQUEUE_RENDER_COMMAND(DeleteDisplayClusterViewportManagerProxy_Release)(
 			[ViewportManagerProxy = ViewportManagerProxy](FRHICommandListImmediate& RHICmdList)
 			{
 				ViewportManagerProxy->Release_RenderThread();
@@ -182,14 +204,38 @@ FDisplayClusterViewportManager::~FDisplayClusterViewportManager()
 	}
 }
 
-IDisplayClusterViewportConfiguration& FDisplayClusterViewportManager::GetConfiguration()
+void FDisplayClusterViewportManager::ReleaseTextures()
 {
-	return Configuration.Get();
+	for (const TSharedPtr<FDisplayClusterViewport, ESPMode::ThreadSafe>& ViewportIt : ImplGetEntireClusterViewports())
+	{
+		if (ViewportIt.IsValid())
+		{
+			ViewportIt->ReleaseTextures();
+		}
+	}
+
+	if (ViewportManagerProxy.IsValid())
+	{
+		// Remove viewport manager proxy on render_thread
+		ENQUEUE_RENDER_COMMAND(DisplayClusterViewportManagerProxy_ReleaseTextures)(
+			[ViewportManagerProxy = ViewportManagerProxy](FRHICommandListImmediate& RHICmdList)
+			{
+				ViewportManagerProxy->ReleaseTextures_RenderThread();
+			});
+	}
+
+	// Release RTT manager caches
+	RenderTargetManager->Release();
+}
+
+IDisplayClusterViewportManagerPreview& FDisplayClusterViewportManager::GetViewportManagerPreview()
+{
+	return ViewportManagerPreview.Get();
 }
 
-const IDisplayClusterViewportConfiguration& FDisplayClusterViewportManager::GetConfiguration() const
+const IDisplayClusterViewportManagerPreview& FDisplayClusterViewportManager::GetViewportManagerPreview() const
 {
-	return Configuration.Get();
+	return ViewportManagerPreview.Get();
 }
 
 const IDisplayClusterViewportManagerProxy* FDisplayClusterViewportManager::GetProxy() const
@@ -205,9 +251,8 @@ IDisplayClusterViewportManagerProxy* FDisplayClusterViewportManager::GetProxy()
 void FDisplayClusterViewportManager::HandleStartScene()
 {
 	check(IsInGameThread());
-	check(!Configuration->bCurrentSceneActive);
 
-	if (Configuration->GetCurrentWorld() == nullptr)
+	if (Configuration->IsSceneOpened())
 	{
 		// Do not start scene, if world not defined
 		return;
@@ -217,15 +262,13 @@ void FDisplayClusterViewportManager::HandleStartScene()
 	{
 		if (ViewportIt.IsValid())
 		{
-			ViewportIt->HandleStartScene();
+			ViewportIt->OnHandleStartScene();
 		}
 	}
 
-	PostProcessManager->HandleStartScene();
-	LightCardManager->HandleStartScene();
-
-	Configuration->bCurrentSceneActive = true;
-	Configuration->bCurrentSceneNeedsToBeUpdated = false;
+	PostProcessManager->OnHandleStartScene();
+	LightCardManager->OnHandleStartScene();
+	Configuration->OnHandleStartScene();
 }
 
 void FDisplayClusterViewportManager::HandleEndScene()
@@ -236,14 +279,17 @@ void FDisplayClusterViewportManager::HandleEndScene()
 	{
 		if (ViewportIt.IsValid())
 		{
-			ViewportIt->HandleEndScene();
+			ViewportIt->OnHandleEndScene();
 		}
 	}
 
-	PostProcessManager->HandleEndScene();
-	LightCardManager->HandleEndScene();
+	PostProcessManager->OnHandleEndScene();
+	LightCardManager->OnHandleEndScene();
+
+	// Release preview from prev scene
+	ViewportManagerPreview->Release();
 
-	Configuration->bCurrentSceneActive = false;
+	Configuration->OnHandleEndScene();
 }
 
 void FDisplayClusterViewportManager::HandleViewportRTTChanges(const TArray<FDisplayClusterViewport_Context>& InPrevContexts, const TArray<FDisplayClusterViewport_Context>& InContexts)
@@ -331,6 +377,21 @@ bool FDisplayClusterViewportManager::ShouldUseFullSizeFrameTargetableResource()
 	return false;
 }
 
+bool FDisplayClusterViewportManager::ShouldUseOutputTargetableResources() const
+{
+	check(IsInGameThread());
+
+	for (const TSharedPtr<FDisplayClusterViewport, ESPMode::ThreadSafe>& Viewport : ImplGetCurrentRenderFrameViewports())
+	{
+		if (Viewport.IsValid() && Viewport->ShouldUseOutputTargetableResources())
+		{
+			return true;
+		}
+	}
+
+	return false;
+}
+
 bool FDisplayClusterViewportManager::ShouldUseAdditionalFrameTargetableResource() const
 {
 	check(IsInGameThread());
@@ -577,13 +638,14 @@ bool FDisplayClusterViewportManager::BeginNewFrame(FViewport* InViewport, FDispl
 	// Update desired views number
 	OutRenderFrame.DesiredNumberOfViews = (ImplGetEntireClusterViewports().Num() * Configuration->GetRenderFrameSettings().GetViewPerViewportAmount()) + 1;
 
-#if WITH_EDITOR
-	// Get preview resources from root actor
-	ImplUpdatePreviewRTTResources();
-#endif /*WITH_EDITOR*/
-
 	PostProcessManager->HandleBeginNewFrame(OutRenderFrame);
 
+	// Update viewport preview instances if preview is used and DCRA supports previews:
+	if (Configuration->IsPreviewRendering() && Configuration->GetRootActor(EDisplayClusterRootActorType::Preview) != nullptr)
+	{
+		ViewportManagerPreview->Update();
+	}
+
 	return true;
 }
 
@@ -734,7 +796,8 @@ FSceneViewFamily::ConstructionValues FDisplayClusterViewportManager::CreateViewF
 		break;
 	}
 
-	return FSceneViewFamily::ConstructionValues(InFrameTarget.RenderTargetPtr, InScene, InEngineShowFlags)
+	FRenderTarget* RenderTarget = InFrameTarget.RenderTargetResource.IsValid() ? InFrameTarget.RenderTargetResource->GetViewportResourceRenderTarget() : nullptr;
+	return FSceneViewFamily::ConstructionValues(RenderTarget, InScene, InEngineShowFlags)
 		.SetResolveScene(bResolveScene)
 		.SetRealtimeUpdate(true)
 		.SetAdditionalViewFamily(bInAdditionalViewFamily);
@@ -870,12 +933,17 @@ TSharedPtr<FDisplayClusterViewport, ESPMode::ThreadSafe> FDisplayClusterViewport
 	TSharedPtr<FDisplayClusterViewport, ESPMode::ThreadSafe> NewViewport = MakeShared<FDisplayClusterViewport>(Configuration, ViewportId, InProjectionPolicy);
 	if (NewViewport.IsValid())
 	{
+		NewViewport->Initialize();
+
 		// Add viewport on gamethread
 		EntireClusterViewports.Add(NewViewport);
 		Configuration->bCurrentRenderFrameViewportsNeedsToBeUpdated = true;
 
-		// Handle start scene for viewport
-		NewViewport->HandleStartScene();
+		if (Configuration->IsSceneOpened())
+		{
+			// Handle start scene for viewport
+			NewViewport->OnHandleStartScene();
+		}
 	}
 
 	return NewViewport;
@@ -1033,3 +1101,15 @@ FSceneView* FDisplayClusterViewportManager::CalcSceneView(ULocalPlayer* LocalPla
 
 	return View;
 }
+
+void FDisplayClusterViewportManager::OnPreGarbageCollect()
+{
+	// The view state can reference materials from the world being cleaned up. (Example post process materials)
+	for (const TSharedPtr<FDisplayClusterViewport, ESPMode::ThreadSafe>& ViewportIt : ImplGetEntireClusterViewports())
+	{
+		if (ViewportIt.IsValid())
+		{
+			ViewportIt->CleanupViewState();
+		}
+	}
+}
\ No newline at end of file
diff --git a/Engine/Plugins/Runtime/nDisplay/Source/DisplayCluster/Private/Render/Viewport/DisplayClusterViewportManager.h b/Engine/Plugins/Runtime/nDisplay/Source/DisplayCluster/Private/Render/Viewport/DisplayClusterViewportManager.h
--- a/Engine/Plugins/Runtime/nDisplay/Source/DisplayCluster/Private/Render/Viewport/DisplayClusterViewportManager.h
+++ b/Engine/Plugins/Runtime/nDisplay/Source/DisplayCluster/Private/Render/Viewport/DisplayClusterViewportManager.h
@@ -37,6 +37,10 @@ public:
 
 	void Initialize();
 
+	/** Release all texture resources used for viewports.
+	* After that, the BeginNewFrame() function should be called for a new resource allocation.*/
+	void ReleaseTextures();
+
 	//~ Begin IDisplayClusterViewportManager
 	virtual TSharedPtr<IDisplayClusterViewportManager, ESPMode::ThreadSafe> ToSharedPtr() override
 	{
@@ -48,11 +52,32 @@ public:
 		return AsShared();
 	}
 
+	virtual TSharedRef<FDisplayClusterViewportManager, ESPMode::ThreadSafe> ToSharedRef() override
+	{
+		return AsShared();
+	}
+
+	virtual TSharedRef<const FDisplayClusterViewportManager, ESPMode::ThreadSafe> ToSharedRef() const override
+	{
+		return AsShared();
+	}
+
+	/** Get viewport manager preview API */
+	virtual IDisplayClusterViewportManagerPreview& GetViewportManagerPreview() override;
+	virtual const IDisplayClusterViewportManagerPreview& GetViewportManagerPreview() const override;
+
+
 	virtual const IDisplayClusterViewportManagerProxy* GetProxy() const override;
 	virtual       IDisplayClusterViewportManagerProxy* GetProxy() override;
 
-	virtual IDisplayClusterViewportConfiguration& GetConfiguration() override;
-	virtual const IDisplayClusterViewportConfiguration& GetConfiguration() const override;
+	virtual IDisplayClusterViewportConfiguration& GetConfiguration() override
+	{
+		return Configuration.Get();
+	}
+	virtual const IDisplayClusterViewportConfiguration& GetConfiguration() const override
+	{
+		return Configuration.Get();
+	}
 
 	virtual bool BeginNewFrame(FViewport* InViewport, FDisplayClusterRenderFrame& OutRenderFrame) override;
 	virtual void InitializeNewFrame() override;
@@ -67,11 +92,6 @@ public:
 
 	virtual void ConfigureViewFamily(const FDisplayClusterRenderFrameTarget& InFrameTarget, const FDisplayClusterRenderFrameTargetViewFamily& InFrameViewFamily, FSceneViewFamilyContext& InOutViewFamily) override;
 	virtual void RenderFrame(FViewport* InViewport) override;
-	virtual bool RenderInEditor(FDisplayClusterRenderFrame& InRenderFrame, FViewport* InViewport, const uint32 InFirstViewportNum, const int32 InViewportsAmount, int32& OutViewportsAmount, bool& bOutFrameRendered) override;
-
-#if WITH_EDITOR
-	void ImplUpdatePreviewRTTResources();
-#endif
 
 private:
 	/** Called before garbage collection is run */
@@ -145,6 +165,9 @@ public:
 	TSharedPtr<class FDisplayClusterViewportManagerViewExtension, ESPMode::ThreadSafe> GetViewportManagerViewExtension() const
 	{ return ViewportManagerViewExtension; }
 
+	/** Returns true if any of the viewports are visible and should use the output RTT resources. */
+	bool ShouldUseOutputTargetableResources() const;
+
 	bool ShouldUseAdditionalFrameTargetableResource() const;
 	bool ShouldUseFullSizeFrameTargetableResource() const;
 
@@ -167,6 +190,9 @@ public:
 	// Configuration of the current cluster node
 	const TSharedRef<class FDisplayClusterViewportConfiguration, ESPMode::ThreadSafe> Configuration;
 
+	// Viewport preview
+	const TSharedRef<class FDisplayClusterViewportManagerPreview, ESPMode::ThreadSafe> ViewportManagerPreview;
+
 	// Resource manager
 	const TSharedRef<class FDisplayClusterRenderTargetManager, ESPMode::ThreadSafe> RenderTargetManager;
 
diff --git a/Engine/Plugins/Runtime/nDisplay/Source/DisplayCluster/Private/Render/Viewport/DisplayClusterViewportManagerProxy.cpp b/Engine/Plugins/Runtime/nDisplay/Source/DisplayCluster/Private/Render/Viewport/DisplayClusterViewportManagerProxy.cpp
--- a/Engine/Plugins/Runtime/nDisplay/Source/DisplayCluster/Private/Render/Viewport/DisplayClusterViewportManagerProxy.cpp
+++ b/Engine/Plugins/Runtime/nDisplay/Source/DisplayCluster/Private/Render/Viewport/DisplayClusterViewportManagerProxy.cpp
@@ -63,9 +63,6 @@ void FDisplayClusterViewportManagerProxy::Release_RenderThread()
 	EntireClusterViewportProxies.Empty();
 	CurrentRenderFrameViewportProxies.Empty();
 
-	RenderTargetManager->Release();
-	PostProcessManager->Release();
-
 	if (ViewportManagerViewExtension.IsValid())
 	{
 		// Force release of VE data since this TSharedPtr<> may also be held by other resources at this time.
@@ -188,7 +185,7 @@ DECLARE_GPU_STAT_NAMED(nDisplay_ViewportManager_WarpBlend, TEXT("nDisplay Viewpo
 void FDisplayClusterViewportManagerProxy::ImplRenderFrame_GameThread(FViewport* InViewport)
 {
 	ENQUEUE_RENDER_COMMAND(DisplayClusterRenderFrame_Setup)(
-		[InViewportManagerProxy = SharedThis(this), InViewport](FRHICommandListImmediate& RHICmdList)
+		[InViewportManagerProxy = SharedThis(this)](FRHICommandListImmediate& RHICmdList)
 	{
 		SCOPED_GPU_STAT(RHICmdList, nDisplay_ViewportManager_RenderFrame);
 		SCOPED_DRAW_EVENT(RHICmdList, nDisplay_ViewportManager_RenderFrame);
@@ -200,7 +197,7 @@ void FDisplayClusterViewportManagerProxy::ImplRenderFrame_GameThread(FViewport*
 	});
 
 	ENQUEUE_RENDER_COMMAND(DisplayClusterRenderFrame_CrossGPUTransfer)(
-		[InViewportManagerProxy = SharedThis(this), InViewport](FRHICommandListImmediate& RHICmdList)
+		[InViewportManagerProxy = SharedThis(this), OutputViewport = InViewport](FRHICommandListImmediate& RHICmdList)
 	{
 		SCOPED_GPU_STAT(RHICmdList, nDisplay_ViewportManager_CrossGPUTransfer);
 		SCOPED_DRAW_EVENT(RHICmdList, nDisplay_ViewportManager_CrossGPUTransfer);
@@ -216,13 +213,13 @@ void FDisplayClusterViewportManagerProxy::ImplRenderFrame_GameThread(FViewport*
 		}
 
 		// PostCrossGpuTransfer notification
-		IDisplayCluster::Get().GetCallbacks().OnDisplayClusterPostCrossGpuTransfer_RenderThread().Broadcast(RHICmdList, ViewportManagerProxy, InViewport);
+		IDisplayCluster::Get().GetCallbacks().OnDisplayClusterPostCrossGpuTransfer_RenderThread().Broadcast(RHICmdList, ViewportManagerProxy, OutputViewport);
 		// Latency processing
-		IDisplayCluster::Get().GetCallbacks().OnDisplayClusterProcessLatency_RenderThread().Broadcast(RHICmdList, ViewportManagerProxy, InViewport);
+		IDisplayCluster::Get().GetCallbacks().OnDisplayClusterProcessLatency_RenderThread().Broadcast(RHICmdList, ViewportManagerProxy, OutputViewport);
 	});
 
 	ENQUEUE_RENDER_COMMAND(DisplayClusterRenderFrame_UpdateDeferredResources)(
-		[ViewportManagerProxy = SharedThis(this), InViewport](FRHICommandListImmediate& RHICmdList)
+		[ViewportManagerProxy = SharedThis(this)](FRHICommandListImmediate& RHICmdList)
 	{
 		SCOPED_GPU_STAT(RHICmdList, nDisplay_ViewportManager_UpdateDeferredResources);
 		SCOPED_DRAW_EVENT(RHICmdList, nDisplay_ViewportManager_UpdateDeferredResources);
@@ -232,7 +229,7 @@ void FDisplayClusterViewportManagerProxy::ImplRenderFrame_GameThread(FViewport*
 	});
 
 	ENQUEUE_RENDER_COMMAND(DisplayClusterRenderFrame_WarpBlend)(
-		[InViewportManagerProxy = SharedThis(this), InViewport](FRHICommandListImmediate& RHICmdList)
+		[InViewportManagerProxy = SharedThis(this), OutputViewport = InViewport](FRHICommandListImmediate& RHICmdList)
 	{
 		SCOPED_GPU_STAT(RHICmdList, nDisplay_ViewportManager_WarpBlend);
 		SCOPED_DRAW_EVENT(RHICmdList, nDisplay_ViewportManager_WarpBlend);
@@ -247,31 +244,26 @@ void FDisplayClusterViewportManagerProxy::ImplRenderFrame_GameThread(FViewport*
 		ViewportManagerProxy->PostProcessManager->HandleEndUpdateFrameResources_RenderThread(RHICmdList, ViewportManagerProxy);
 
 		// Postrender notification before copying final image to the backbuffer
-		IDisplayCluster::Get().GetCallbacks().OnDisplayClusterPostFrameRender_RenderThread().Broadcast(RHICmdList, ViewportManagerProxy, InViewport);
+		IDisplayCluster::Get().GetCallbacks().OnDisplayClusterPostFrameRender_RenderThread().Broadcast(RHICmdList, ViewportManagerProxy, OutputViewport);
 
-		if (InViewport)
+		if (OutputViewport)
 		{
-			if (FRHITexture2D* FrameOutputRTT = InViewport->GetRenderTargetTexture())
+			if (FRHITexture2D* FrameOutputRTT = OutputViewport->GetRenderTargetTexture())
 			{
 				// For quadbuf stereo copy only left eye, right copy from OutputFrameTarget
 				//@todo Copy QuadBuf_LeftEye/(mono,sbs,tp) to separate rtt, before UI and debug rendering
 				//@todo QuadBuf_LeftEye copied latter, before present
-				switch (ViewportManagerProxy->GetConfigurationProxy().GetRenderMode_RenderThread())
+				if(ViewportManagerProxy->ConfigurationProxy->GetRenderFrameSettings().ShouldUseStereoRenderingOnMonoscopicDisplay())
 				{
-				case EDisplayClusterRenderFrameMode::SideBySide:
-				case EDisplayClusterRenderFrameMode::TopBottom:
 					ViewportManagerProxy->ResolveFrameTargetToBackBuffer_RenderThread(RHICmdList, 1, 0, FrameOutputRTT, FrameOutputRTT->GetSizeXY());
-					break;
-				default:
-					break;
 				}
 
 				ViewportManagerProxy->ResolveFrameTargetToBackBuffer_RenderThread(RHICmdList, 0, 0, FrameOutputRTT, FrameOutputRTT->GetSizeXY());
 
 				// Finally, notify about backbuffer update
-				IDisplayCluster::Get().GetCallbacks().OnDisplayClusterPostBackbufferUpdated_RenderThread().Broadcast(RHICmdList, InViewport);
+				IDisplayCluster::Get().GetCallbacks().OnDisplayClusterPostBackbufferUpdated_RenderThread().Broadcast(RHICmdList, OutputViewport);
 PRAGMA_DISABLE_DEPRECATION_WARNINGS
-				IDisplayCluster::Get().GetCallbacks().OnDisplayClusterPostBackbufferUpdate_RenderThread().Broadcast(RHICmdList, ViewportManagerProxy, InViewport);
+				IDisplayCluster::Get().GetCallbacks().OnDisplayClusterPostBackbufferUpdate_RenderThread().Broadcast(RHICmdList, ViewportManagerProxy, OutputViewport);
 PRAGMA_ENABLE_DEPRECATION_WARNINGS
 			}
 		}
@@ -418,10 +410,6 @@ void FDisplayClusterViewportManagerProxy::UpdateFrameResources_RenderThread(FRHI
 		// Iterate over visible viewports:
 		if (ViewportProxyIt.IsValid() && ViewportProxyIt->GetRenderSettings_RenderThread().bVisible)
 		{
-			// resolve warped viewport resource to the output texture
-			ViewportProxyIt->ResolveResources_RenderThread(RHICmdList, EDisplayClusterViewportResourceType::AfterWarpBlendTargetableResource, EDisplayClusterViewportResourceType::OutputTargetableResource);
-
-			// Apply post-warp (viewport remap, etc)
 			ViewportProxyIt->PostResolveViewport_RenderThread(RHICmdList);
 		}
 	}
@@ -496,7 +484,7 @@ void FDisplayClusterViewportManagerProxy::DoCrossGPUTransfers_RenderThread(FRHIC
 			{
 				if (FRenderTarget* RenderTarget = ViewportRenderTargetResource.IsValid() ? ViewportRenderTargetResource->GetViewportResourceRenderTarget() : nullptr)
 				{
-					if (FRHITexture2D* TextureRHI = ViewportRenderTargetResource->GetViewportResourceRHI())
+					if (FRHITexture2D* TextureRHI = ViewportRenderTargetResource->GetViewportResourceRHI_RenderThread())
 					{
 						const FRHIGPUMask RenderTargetGPUMask = RenderTarget->GetGPUMask(RHICmdList);
 
@@ -541,14 +529,14 @@ bool FDisplayClusterViewportManagerProxy::GetFrameTargets_RenderThread(TArray<FR
 
 			for (int32 FrameIt = 0; FrameIt < Frames.Num(); FrameIt++)
 			{
-				if (FRHITexture2D* FrameTexture = Frames[FrameIt].IsValid() ? Frames[FrameIt]->GetViewportResourceRHI() : nullptr)
+				if (FRHITexture2D* FrameTexture = Frames[FrameIt].IsValid() ? Frames[FrameIt]->GetViewportResourceRHI_RenderThread() : nullptr)
 				{
 					OutFrameResources.Add(FrameTexture);
 					OutTargetOffsets.Add(Frames[FrameIt]->GetBackbufferFrameOffset());
 
 					if (OutAdditionalFrameResources && AdditionalFrames.IsValidIndex(FrameIt))
 					{
-						if (FRHITexture2D* AdditionalFrameTexture = AdditionalFrames[FrameIt].IsValid() ? AdditionalFrames[FrameIt]->GetViewportResourceRHI() : nullptr)
+						if (FRHITexture2D* AdditionalFrameTexture = AdditionalFrames[FrameIt].IsValid() ? AdditionalFrames[FrameIt]->GetViewportResourceRHI_RenderThread() : nullptr)
 						{
 							OutAdditionalFrameResources->Add(AdditionalFrameTexture);
 						}
@@ -622,6 +610,17 @@ bool FDisplayClusterViewportManagerProxy::ResolveFrameTargetToBackBuffer_RenderT
 	return false;
 }
 
+void FDisplayClusterViewportManagerProxy::ReleaseTextures_RenderThread()
+{
+	for (const TSharedPtr<FDisplayClusterViewportProxy, ESPMode::ThreadSafe>& ViewportProxyIt : ImplGetEntireClusterViewportProxies_RenderThread())
+	{
+		if (ViewportProxyIt.IsValid())
+		{
+			ViewportProxyIt->ReleaseTextures_RenderThread();
+		}
+	}
+}
+
 FDisplayClusterViewportProxy* FDisplayClusterViewportManagerProxy::ImplFindViewportProxy_RenderThread(const FString& ViewportId) const
 {
 	check(IsInRenderingThread());
diff --git a/Engine/Plugins/Runtime/nDisplay/Source/DisplayCluster/Private/Render/Viewport/DisplayClusterViewportManagerProxy.h b/Engine/Plugins/Runtime/nDisplay/Source/DisplayCluster/Private/Render/Viewport/DisplayClusterViewportManagerProxy.h
--- a/Engine/Plugins/Runtime/nDisplay/Source/DisplayCluster/Private/Render/Viewport/DisplayClusterViewportManagerProxy.h
+++ b/Engine/Plugins/Runtime/nDisplay/Source/DisplayCluster/Private/Render/Viewport/DisplayClusterViewportManagerProxy.h
@@ -68,6 +68,10 @@ public:
 	virtual bool ResolveFrameTargetToBackBuffer_RenderThread(FRHICommandListImmediate& RHICmdList, const uint32 InContextNum, const int32 DestArrayIndex, FRHITexture2D* DestTexture, FVector2D WindowSize) const override;
 	//~~IDisplayClusterViewportManagerProxy
 
+
+	/** Release all textures on render thread. */
+	void ReleaseTextures_RenderThread();
+
 	/** Custom cross-GPU implementation for mGPU. Applyed for viewports with bOverrideCrossGPUTransfer=true. */
 	void DoCrossGPUTransfers_RenderThread(FRHICommandListImmediate& RHICmdList) const;
 
