nDisplay: Implemented a new preview rendering architecture:

diff --git a/Engine/Plugins/Runtime/nDisplay/Source/DisplayCluster/Private/Render/Viewport/RenderTarget/DisplayClusterRenderTargetManager.cpp b/Engine/Plugins/Runtime/nDisplay/Source/DisplayCluster/Private/Render/Viewport/RenderTarget/DisplayClusterRenderTargetManager.cpp
--- a/Engine/Plugins/Runtime/nDisplay/Source/DisplayCluster/Private/Render/Viewport/RenderTarget/DisplayClusterRenderTargetManager.cpp
+++ b/Engine/Plugins/Runtime/nDisplay/Source/DisplayCluster/Private/Render/Viewport/RenderTarget/DisplayClusterRenderTargetManager.cpp
@@ -84,6 +84,8 @@ namespace UE::DisplayCluster::RenderTargetManager
 			UE_LOG(LogDisplayClusterViewport, Verbose, TEXT("Created new %s resource (%dx%d) for viewport '%s':%d"), *ResourceId, InTextureResource->GetResourceSettings().GetSizeXY().X, InTextureResource->GetResourceSettings().GetSizeXY().Y, *InViewport->GetId(), ContextNum);
 		}
 	}
+
+	static FString EmptyString;
 };
 using namespace UE::DisplayCluster::RenderTargetManager;
 
@@ -118,12 +120,11 @@ bool FDisplayClusterRenderTargetManager::AllocateRenderFrameResources(FViewport*
 			if (FrameRenderTargetIt.bShouldUseRenderTarget)
 			{
 				// reallocate
-				TSharedPtr<FDisplayClusterViewportResource, ESPMode::ThreadSafe> NewResource = ResourcesPool->AllocateResource(FrameRenderTargetIt.RenderTargetSize, ImplGetCustomFormat(FrameRenderTargetIt.CaptureMode), EDisplayClusterViewportResourceSettingsFlags::RenderTarget);
+				TSharedPtr<FDisplayClusterViewportResource, ESPMode::ThreadSafe> NewResource = ResourcesPool->AllocateResource(EmptyString, FrameRenderTargetIt.RenderTargetSize, ImplGetCustomFormat(FrameRenderTargetIt.CaptureMode), EDisplayClusterViewportResourceSettingsFlags::RenderTarget);
 				if (NewResource.IsValid())
 				{
 					// Set RenderFrame resource
-					FrameRenderTargetIt.RenderTargetPtr = NewResource->GetViewportResourceRenderTarget();
-					check(FrameRenderTargetIt.RenderTargetPtr);
+					FrameRenderTargetIt.RenderTargetResource = NewResource;
 
 					// Assign for all views in render target families
 					for (FDisplayClusterRenderFrameTargetViewFamily& ViewFamily : FrameRenderTargetIt.ViewFamilies)
@@ -164,28 +165,28 @@ bool FDisplayClusterRenderTargetManager::AllocateRenderFrameResources(FViewport*
 					const uint32& ContextNum = ContextIt.ContextNum;
 
 					// Allocate per-viewport preview output textures:
-					if (ViewportIt->GetViewportResources(EDisplayClusterViewportResource::OutputPreviewTargetableResources).Num() > (int32)ContextNum)
+					if (ViewportIt->GetViewportResources(EDisplayClusterViewportResource::OutputPreviewTargetableResources).IsValidIndex(ContextNum))
 					{
-						ViewportIt->GetViewportResourcesImpl(EDisplayClusterViewportResource::OutputPreviewTargetableResources)[ContextNum] = ResourcesPool->AllocateResource(ContextSize, CustomFormat, EDisplayClusterViewportResourceSettingsFlags::PreviewTargetableTexture);
+						ViewportIt->GetViewportResourcesImpl(EDisplayClusterViewportResource::OutputPreviewTargetableResources)[ContextNum] = ResourcesPool->AllocateResource(ViewportIt->GetId(), ContextSize, CustomFormat, EDisplayClusterViewportResourceSettingsFlags::PreviewTargetableTexture);
 						ImplViewportTextureResourceLogging(ViewportIt, ContextNum, TEXT("InputShader"), ViewportIt->GetViewportResources(EDisplayClusterViewportResource::OutputPreviewTargetableResources)[ContextNum]);
 					}
 
-					if (ViewportIt->GetViewportResources(EDisplayClusterViewportResource::InputShaderResources).Num() > (int32)ContextNum)
+					if (ViewportIt->GetViewportResources(EDisplayClusterViewportResource::InputShaderResources).IsValidIndex(ContextNum))
 					{
-						ViewportIt->GetViewportResourcesImpl(EDisplayClusterViewportResource::InputShaderResources)[ContextNum] = ResourcesPool->AllocateResource(ContextSize, CustomFormat, EDisplayClusterViewportResourceSettingsFlags::ResolveTargetableTexture);
+						ViewportIt->GetViewportResourcesImpl(EDisplayClusterViewportResource::InputShaderResources)[ContextNum] = ResourcesPool->AllocateResource(EmptyString, ContextSize, CustomFormat, EDisplayClusterViewportResourceSettingsFlags::ResolveTargetableTexture);
 						ImplViewportTextureResourceLogging(ViewportIt, ContextNum, TEXT("InputShader"), ViewportIt->GetViewportResources(EDisplayClusterViewportResource::InputShaderResources)[ContextNum]);
 					}
 
 					// Allocate custom resources:
-					if (ViewportIt->GetViewportResources(EDisplayClusterViewportResource::AdditionalTargetableResources).Num() > (int32)ContextNum)
+					if (ViewportIt->GetViewportResources(EDisplayClusterViewportResource::AdditionalTargetableResources).IsValidIndex(ContextNum))
 					{
-						ViewportIt->GetViewportResourcesImpl(EDisplayClusterViewportResource::AdditionalTargetableResources)[ContextNum] = ResourcesPool->AllocateResource(ContextSize, CustomFormat, EDisplayClusterViewportResourceSettingsFlags::RenderTargetableTexture);
+						ViewportIt->GetViewportResourcesImpl(EDisplayClusterViewportResource::AdditionalTargetableResources)[ContextNum] = ResourcesPool->AllocateResource(EmptyString, ContextSize, CustomFormat, EDisplayClusterViewportResourceSettingsFlags::RenderTargetableTexture);
 						ImplViewportTextureResourceLogging(ViewportIt, ContextNum, TEXT("AdditionalRTT"), ViewportIt->GetViewportResources(EDisplayClusterViewportResource::AdditionalTargetableResources)[ContextNum]);
 					}
 
-					if (ViewportIt->GetViewportResources(EDisplayClusterViewportResource::MipsShaderResources).Num() > (int32)ContextNum)
+					if (ViewportIt->GetViewportResources(EDisplayClusterViewportResource::MipsShaderResources).IsValidIndex(ContextNum))
 					{
-						ViewportIt->GetViewportResourcesImpl(EDisplayClusterViewportResource::MipsShaderResources)[ContextNum] = ResourcesPool->AllocateResource(ContextSize, CustomFormat, EDisplayClusterViewportResourceSettingsFlags::ResolveTargetableTexture, ContextIt.NumMips);
+						ViewportIt->GetViewportResourcesImpl(EDisplayClusterViewportResource::MipsShaderResources)[ContextNum] = ResourcesPool->AllocateResource(EmptyString, ContextSize, CustomFormat, EDisplayClusterViewportResourceSettingsFlags::ResolveTargetableTexture, ContextIt.NumMips);
 						ImplViewportTextureResourceLogging(ViewportIt, ContextNum, TEXT("Mips"), ViewportIt->GetViewportResources(EDisplayClusterViewportResource::MipsShaderResources)[ContextNum]);
 					}
 				}
@@ -209,43 +210,17 @@ bool FDisplayClusterRenderTargetManager::AllocateRenderFrameResources(FViewport*
 
 bool FDisplayClusterRenderTargetManager::AllocateFrameTargets(const FIntPoint& InViewportSize, FDisplayClusterRenderFrame& InOutRenderFrame)
 {
-	// Preview render to external RTT resource
-	if (Configuration->IsPreviewRendering())
+	// ICVFX internal viewports are invisible and do not use output frame resources.
+	// In case of off-screen rendering this texture is not used.
+	if (!Configuration->GetRenderFrameSettings().ShouldUseOutputFrameTargetableResources())
 	{
+		// Skip creating output resources if there are no visible viewports. This will save GPU memory.
 		return true;
 	}
 
-	uint32 FrameTargetsAmount = 0;
-
-	// Support side_by_side and top_bottom eye offset, aligned to Viewport size:
+	const uint32 FrameTargetsAmount = Configuration->GetRenderFrameSettings().GetViewPerViewportAmount();
+	const FIntPoint TargetOffset = Configuration->GetRenderFrameSettings().GetDesiredRTTSize(InViewportSize);
 	FIntPoint TargetLocation(ForceInitToZero);
-	FIntPoint TargetOffset(ForceInitToZero);
-
-	//Re-allocate frame targets
-	switch (Configuration->GetRenderFrameSettings().RenderMode)
-	{
-	case EDisplayClusterRenderFrameMode::Mono:
-		FrameTargetsAmount = 1;
-		break;
-
-	case EDisplayClusterRenderFrameMode::Stereo:
-		FrameTargetsAmount = 2;
-		break;
-
-	case EDisplayClusterRenderFrameMode::SideBySide:
-		FrameTargetsAmount = 2;
-		TargetOffset.X = InViewportSize.X / 2;
-		break;
-
-	case EDisplayClusterRenderFrameMode::TopBottom:
-		FrameTargetsAmount = 2;
-		TargetOffset.Y = InViewportSize.Y / 2;
-		break;
-
-	default:
-		// skip not implemented cases
-		return false;
-	}
 
 	// Reallocate frame target resources
 	TArray<TSharedPtr<FDisplayClusterViewportResource, ESPMode::ThreadSafe>> NewFrameTargetResources;
@@ -253,7 +228,7 @@ bool FDisplayClusterRenderTargetManager::AllocateFrameTargets(const FIntPoint& I
 
 	for (uint32 FrameTargetsIt = 0; FrameTargetsIt < FrameTargetsAmount; FrameTargetsIt++)
 	{
-		TSharedPtr<FDisplayClusterViewportResource, ESPMode::ThreadSafe> NewResource = ResourcesPool->AllocateResource(InOutRenderFrame.FrameRect.Size(), PF_Unknown, EDisplayClusterViewportResourceSettingsFlags::RenderTargetableTexture);
+		TSharedPtr<FDisplayClusterViewportResource, ESPMode::ThreadSafe> NewResource = ResourcesPool->AllocateResource(EmptyString, InOutRenderFrame.FrameRect.Size(), PF_Unknown, EDisplayClusterViewportResourceSettingsFlags::RenderTargetableTexture);
 		if (NewResource.IsValid())
 		{
 			if (EnumHasAnyFlags(NewResource->GetResourceState(), EDisplayClusterViewportResourceState::Initialized) == false)
@@ -270,7 +245,7 @@ bool FDisplayClusterRenderTargetManager::AllocateFrameTargets(const FIntPoint& I
 
 		if (Configuration->GetRenderFrameSettings().bShouldUseAdditionalFrameTargetableResource)
 		{
-			TSharedPtr<FDisplayClusterViewportResource, ESPMode::ThreadSafe> NewAdditionalResource = ResourcesPool->AllocateResource(InOutRenderFrame.FrameRect.Size(), PF_Unknown, EDisplayClusterViewportResourceSettingsFlags::RenderTargetableTexture);
+			TSharedPtr<FDisplayClusterViewportResource, ESPMode::ThreadSafe> NewAdditionalResource = ResourcesPool->AllocateResource(EmptyString, InOutRenderFrame.FrameRect.Size(), PF_Unknown, EDisplayClusterViewportResourceSettingsFlags::RenderTargetableTexture);
 			if (NewAdditionalResource.IsValid())
 			{
 				if (EnumHasAnyFlags(NewAdditionalResource->GetResourceState(), EDisplayClusterViewportResourceState::Initialized) == false)
diff --git a/Engine/Plugins/Runtime/nDisplay/Source/DisplayCluster/Private/Render/Viewport/RenderTarget/DisplayClusterRenderTargetResource.h b/Engine/Plugins/Runtime/nDisplay/Source/DisplayCluster/Private/Render/Viewport/RenderTarget/DisplayClusterRenderTargetResource.h
--- a/Engine/Plugins/Runtime/nDisplay/Source/DisplayCluster/Private/Render/Viewport/RenderTarget/DisplayClusterRenderTargetResource.h
+++ b/Engine/Plugins/Runtime/nDisplay/Source/DisplayCluster/Private/Render/Viewport/RenderTarget/DisplayClusterRenderTargetResource.h
@@ -23,8 +23,10 @@ public:
 
 public:
 	//~ Begin FDisplayClusterViewportResource
-	virtual FRHITexture2D* GetViewportResourceRHI() const override
+	virtual FRHITexture2D* GetViewportResourceRHI_RenderThread() const override
 	{
+		check(IsInRenderingThread());
+
 		return TextureRHI.IsValid() ? TextureRHI->GetTexture2D() : nullptr;
 	}
 
@@ -92,7 +94,7 @@ public:
 
 public:
 	//~ Begin FDisplayClusterViewportResource
-	virtual FRenderTarget* GetViewportResourceRenderTarget()
+	virtual FRenderTarget* GetViewportResourceRenderTarget() override
 	{
 		return this;
 	}
diff --git a/Engine/Plugins/Runtime/nDisplay/Source/DisplayCluster/Private/Render/Viewport/RenderTarget/DisplayClusterRenderTargetResourcesPool.cpp b/Engine/Plugins/Runtime/nDisplay/Source/DisplayCluster/Private/Render/Viewport/RenderTarget/DisplayClusterRenderTargetResourcesPool.cpp
--- a/Engine/Plugins/Runtime/nDisplay/Source/DisplayCluster/Private/Render/Viewport/RenderTarget/DisplayClusterRenderTargetResourcesPool.cpp
+++ b/Engine/Plugins/Runtime/nDisplay/Source/DisplayCluster/Private/Render/Viewport/RenderTarget/DisplayClusterRenderTargetResourcesPool.cpp
@@ -42,12 +42,12 @@ void FDisplayClusterRenderTargetResourcesPool::ImplUpdateResources(TArray<TShare
 			switch (InUpdateMode)
 			{
 			case EResourceUpdateMode::Release:
-				EnumRemoveFlags(ResourceIt->GetResourceState(), EDisplayClusterViewportResourceState::Initialized);
+				ResourceIt->ReleaseViewportResource();
 				break;
 
 			default:
 			case EResourceUpdateMode::Initialize:
-				EnumAddFlags(ResourceIt->GetResourceState(), EDisplayClusterViewportResourceState::Initialized);
+				ResourceIt->InitializeViewportResource();
 				break;
 			}
 		}
@@ -158,11 +158,11 @@ void FDisplayClusterRenderTargetResourcesPool::EndReallocateResources()
 	ResourceSettings = nullptr;
 }
 
-TSharedPtr<FDisplayClusterViewportResource, ESPMode::ThreadSafe> FDisplayClusterRenderTargetResourcesPool::AllocateResource(const FIntPoint& InSize, EPixelFormat CustomPixelFormat, const EDisplayClusterViewportResourceSettingsFlags InResourceFlags, int32 InNumMips)
+TSharedPtr<FDisplayClusterViewportResource, ESPMode::ThreadSafe> FDisplayClusterRenderTargetResourcesPool::AllocateResource(const FString InViewportId, const FIntPoint& InSize, EPixelFormat CustomPixelFormat, const EDisplayClusterViewportResourceSettingsFlags InResourceFlags, int32 InNumMips)
 {
 	check(ResourceSettings != nullptr);
 
-	FDisplayClusterViewportResourceSettings InSettings(*ResourceSettings, InSize, CustomPixelFormat, InResourceFlags, InNumMips);
+	FDisplayClusterViewportResourceSettings InSettings(*ResourceSettings, InViewportId, InSize, CustomPixelFormat, InResourceFlags, InNumMips);
 
 	if (!FDisplayClusterViewportHelpers::IsValidTextureSize(InSettings.GetSizeXY()))
 	{
diff --git a/Engine/Plugins/Runtime/nDisplay/Source/DisplayCluster/Private/Render/Viewport/RenderTarget/DisplayClusterRenderTargetResourcesPool.h b/Engine/Plugins/Runtime/nDisplay/Source/DisplayCluster/Private/Render/Viewport/RenderTarget/DisplayClusterRenderTargetResourcesPool.h
--- a/Engine/Plugins/Runtime/nDisplay/Source/DisplayCluster/Private/Render/Viewport/RenderTarget/DisplayClusterRenderTargetResourcesPool.h
+++ b/Engine/Plugins/Runtime/nDisplay/Source/DisplayCluster/Private/Render/Viewport/RenderTarget/DisplayClusterRenderTargetResourcesPool.h
@@ -27,6 +27,7 @@ public:
 	
 	/** Allocate a new resource or reuse exists
 	* 
+	* @param InViewportId      - This resource belongs to this viewport only. Empty string if not used
 	* @param InSize            - resource dimensions
 	* @param CustomPixelFormat - resource pixel format
 	* @param InResourceFlags   - These flags define the behavior of the resource.
@@ -34,7 +35,7 @@ public:
 	* 
 	* @return ref to the resource instance
 	*/
-	TSharedPtr<FDisplayClusterViewportResource, ESPMode::ThreadSafe> AllocateResource(const FIntPoint& InSize, EPixelFormat CustomPixelFormat, const EDisplayClusterViewportResourceSettingsFlags InResourceFlags, int32 NumMips = 1);
+	TSharedPtr<FDisplayClusterViewportResource, ESPMode::ThreadSafe> AllocateResource(const FString InViewportId, const FIntPoint& InSize, EPixelFormat CustomPixelFormat, const EDisplayClusterViewportResourceSettingsFlags InResourceFlags, int32 NumMips = 1);
 
 	/** End reallocate resources for the specified cluster node. */
 	void EndReallocateResources();
diff --git a/Engine/Plugins/Runtime/nDisplay/Source/DisplayCluster/Private/Render/Viewport/Resource/DisplayClusterViewportPreviewResource.cpp b/Engine/Plugins/Runtime/nDisplay/Source/DisplayCluster/Private/Render/Viewport/Resource/DisplayClusterViewportPreviewResource.cpp
--- a/Engine/Plugins/Runtime/nDisplay/Source/DisplayCluster/Private/Render/Viewport/Resource/DisplayClusterViewportPreviewResource.cpp
+++ b/Engine/Plugins/Runtime/nDisplay/Source/DisplayCluster/Private/Render/Viewport/Resource/DisplayClusterViewportPreviewResource.cpp
@@ -2,9 +2,14 @@
 
 #include "Render/Viewport/Resource/DisplayClusterViewportPreviewResource.h"
 
+#include "Engine/TextureRenderTarget2D.h"
+#include "UObject/UObjectGlobals.h"
+#include "UObject/Package.h"
+
 #include "RHI.h"
 #include "RenderResource.h"
 #include "UnrealClient.h"
+#include "TextureResource.h"
 
 ////////////////////////////////////////////////////////////////////////
 // FDisplayClusterViewportPreviewResource
@@ -12,3 +17,58 @@
 FDisplayClusterViewportPreviewResource::FDisplayClusterViewportPreviewResource(const FDisplayClusterViewportResourceSettings& InResourceSettings)
 	: FDisplayClusterViewportResource(InResourceSettings)
 { }
+
+FDisplayClusterViewportPreviewResource::~FDisplayClusterViewportPreviewResource()
+{ }
+
+void FDisplayClusterViewportPreviewResource::InitializeViewportResource()
+{
+	check(IsInGameThread());
+
+	if (EnumHasAnyFlags(GetResourceState(), EDisplayClusterViewportResourceState::Initialized))
+	{
+		// Already initialized
+		return;
+	}
+
+	RenderTargetTexture = NewObject<UTextureRenderTarget2D>(GetTransientPackage());
+	if (RenderTargetTexture)
+	{
+		RenderTargetTexture->AddToRoot();
+
+		RenderTargetTexture->ClearColor = FLinearColor::Black;
+		RenderTargetTexture->TargetGamma = ResourceSettings.GetDisplayGamma();
+		RenderTargetTexture->SRGB = EnumHasAnyFlags(ResourceSettings.GetResourceFlags(), EDisplayClusterViewportResourceSettingsFlags::ShouldUseSRGB);
+		RenderTargetTexture->InitCustomFormat(ResourceSettings.GetSizeX(), ResourceSettings.GetSizeY(), ResourceSettings.GetFormat(), false);
+
+		EnumAddFlags(GetResourceState(), EDisplayClusterViewportResourceState::Initialized);
+	}
+}
+
+void FDisplayClusterViewportPreviewResource::ReleaseViewportResource()
+{
+	check(IsInGameThread());
+
+	if (RenderTargetTexture)
+	{
+		RenderTargetTexture->RemoveFromRoot();
+		RenderTargetTexture = nullptr;
+	}
+
+	EnumRemoveFlags(GetResourceState(), EDisplayClusterViewportResourceState::Initialized);
+}
+
+FRHITexture2D* FDisplayClusterViewportPreviewResource::GetViewportResourceRHI_RenderThread() const
+{
+	check(IsInRenderingThread());
+
+	if (RenderTargetTexture)
+	{
+		if (FTextureRenderTargetResource* RenderTargetResource = RenderTargetTexture->GetRenderTargetResource())
+		{
+			return RenderTargetResource->TextureRHI;
+		}
+	}
+
+	return nullptr;
+}
diff --git a/Engine/Plugins/Runtime/nDisplay/Source/DisplayCluster/Private/Render/Viewport/Resource/DisplayClusterViewportPreviewResource.h b/Engine/Plugins/Runtime/nDisplay/Source/DisplayCluster/Private/Render/Viewport/Resource/DisplayClusterViewportPreviewResource.h
--- a/Engine/Plugins/Runtime/nDisplay/Source/DisplayCluster/Private/Render/Viewport/Resource/DisplayClusterViewportPreviewResource.h
+++ b/Engine/Plugins/Runtime/nDisplay/Source/DisplayCluster/Private/Render/Viewport/Resource/DisplayClusterViewportPreviewResource.h
@@ -6,6 +6,7 @@
 #include "RHI.h"
 #include "RHIResources.h"
 #include "RenderResource.h"
+#include "Engine/TextureRenderTarget2D.h"
 
 /**
  * Viewport preview texture resource
@@ -15,28 +16,22 @@ class FDisplayClusterViewportPreviewResource
 {
 public:
 	FDisplayClusterViewportPreviewResource(const FDisplayClusterViewportResourceSettings& InResourceSettings);
-	virtual ~FDisplayClusterViewportPreviewResource() = default;
+	virtual ~FDisplayClusterViewportPreviewResource();
 
 public:
 	//~ Begin FDisplayClusterViewportResource
-	virtual FRHITexture2D* GetViewportResourceRHI() const override
-	{
-		return OutputPreviewTargetableResource.IsValid() ? OutputPreviewTargetableResource->GetTexture2D() : nullptr;
-	}
+	virtual void InitializeViewportResource() override;
+	virtual void ReleaseViewportResource() override;
 
-	virtual void SetExternalViewportResourceRHI(FTextureRHIRef& InExternalViewportResourceRHI) override
+	virtual UTextureRenderTarget2D* GetTextureRenderTarget2D() const override
 	{
-		OutputPreviewTargetableResource = InExternalViewportResourceRHI;
+		return RenderTargetTexture;
 	}
 
-	virtual void ReleaseViewportResource_RenderThread(FRHICommandListBase& RHICmdList) override
-	{
-		OutputPreviewTargetableResource.SafeRelease();
-	}
+	virtual FRHITexture2D* GetViewportResourceRHI_RenderThread() const override;
 	//~~ End FDisplayClusterViewportResource
 
 public:
-	// here we implement the preview resources the old way,
-	// but when we remove the DCPreviewComponent, this class will be redesigned using UTexture.
-	FTextureRHIRef OutputPreviewTargetableResource;
+	// Preview RTT resource
+	TObjectPtr<UTextureRenderTarget2D> RenderTargetTexture = nullptr;
 };
diff --git a/Engine/Plugins/Runtime/nDisplay/Source/DisplayCluster/Private/Render/Viewport/Resource/DisplayClusterViewportResource.h b/Engine/Plugins/Runtime/nDisplay/Source/DisplayCluster/Private/Render/Viewport/Resource/DisplayClusterViewportResource.h
--- a/Engine/Plugins/Runtime/nDisplay/Source/DisplayCluster/Private/Render/Viewport/Resource/DisplayClusterViewportResource.h
+++ b/Engine/Plugins/Runtime/nDisplay/Source/DisplayCluster/Private/Render/Viewport/Resource/DisplayClusterViewportResource.h
@@ -10,6 +10,7 @@
 #include "Templates/SharedPointer.h"
 
 class FRenderTarget;
+class UTextureRenderTarget2D;
 
 /**
  * Runtime state of the viewport resource
@@ -68,7 +69,7 @@ public:
 	{ }
 
 	/** Get current RHI resource. */
-	virtual FRHITexture2D* GetViewportResourceRHI() const
+	virtual FRHITexture2D* GetViewportResourceRHI_RenderThread() const
 	{
 		return nullptr;
 	}
@@ -79,6 +80,11 @@ public:
 		return nullptr;
 	}
 
+	virtual UTextureRenderTarget2D* GetTextureRenderTarget2D() const
+	{
+		return nullptr;
+	}
+
 	/** Getting the viewport rect offset in the backbuffer. */
 	virtual FIntPoint GetBackbufferFrameOffset() const
 	{
@@ -89,6 +95,18 @@ public:
 	virtual void SetBackbufferFrameOffset(const FIntPoint& InBackbufferFrameOffset)
 	{ }
 
+	/** Initialize this resource in the game thread. */
+	virtual void InitializeViewportResource()
+	{
+		EnumAddFlags(ResourceState, EDisplayClusterViewportResourceState::Initialized);
+	}
+
+	/** Release this resource in the game thread*/
+	virtual void ReleaseViewportResource()
+	{
+		EnumRemoveFlags(ResourceState, EDisplayClusterViewportResourceState::Initialized);
+	}
+
 	/** RHI initialization for this resource in the rendering thread. */
 	virtual void InitializeViewportResource_RenderThread(FRHICommandListBase& RHICmdList)
 	{ }
@@ -104,7 +122,7 @@ protected:
 	/** Helper function: create 2D texture. */
 	void ImplInitDynamicRHI_TextureResource2D(FTexture2DRHIRef& OutTextureRHI);
 
-private:
+protected:
 	// Settings used by this resource
 	const FDisplayClusterViewportResourceSettings ResourceSettings;
 
diff --git a/Engine/Plugins/Runtime/nDisplay/Source/DisplayCluster/Private/Render/Viewport/Resource/DisplayClusterViewportResourceSettings.cpp b/Engine/Plugins/Runtime/nDisplay/Source/DisplayCluster/Private/Render/Viewport/Resource/DisplayClusterViewportResourceSettings.cpp
--- a/Engine/Plugins/Runtime/nDisplay/Source/DisplayCluster/Private/Render/Viewport/Resource/DisplayClusterViewportResourceSettings.cpp
+++ b/Engine/Plugins/Runtime/nDisplay/Source/DisplayCluster/Private/Render/Viewport/Resource/DisplayClusterViewportResourceSettings.cpp
@@ -38,7 +38,7 @@ FDisplayClusterViewportResourceSettings::FDisplayClusterViewportResourceSettings
 		if (InRenderFrameSettings.ShouldUseLinearGamma())
 		{
 				DisplayGamma = 1.f;
-			}
+		}
 
 		if(InRenderFrameSettings.IsPreviewRendering())
 		{
@@ -48,8 +48,9 @@ FDisplayClusterViewportResourceSettings::FDisplayClusterViewportResourceSettings
 	}
 }
 
-FDisplayClusterViewportResourceSettings::FDisplayClusterViewportResourceSettings(const FDisplayClusterViewportResourceSettings& InBaseSettings, const FIntPoint& InSize, const EPixelFormat InFormat, const EDisplayClusterViewportResourceSettingsFlags InResourceFlags, const int32 InNumMips)
+FDisplayClusterViewportResourceSettings::FDisplayClusterViewportResourceSettings(const FDisplayClusterViewportResourceSettings& InBaseSettings, const FString InViewportId, const FIntPoint& InSize, const EPixelFormat InFormat, const EDisplayClusterViewportResourceSettingsFlags InResourceFlags, const int32 InNumMips)
 	: ClusterNodeId(InBaseSettings.ClusterNodeId)
+	, ViewportId(InViewportId)
 	, Size(InSize)
 	, Format((InFormat == PF_Unknown) ? InBaseSettings.Format : InFormat)
 	, DisplayGamma(InBaseSettings.DisplayGamma)
diff --git a/Engine/Plugins/Runtime/nDisplay/Source/DisplayCluster/Private/Render/Viewport/Resource/DisplayClusterViewportResourceSettings.h b/Engine/Plugins/Runtime/nDisplay/Source/DisplayCluster/Private/Render/Viewport/Resource/DisplayClusterViewportResourceSettings.h
--- a/Engine/Plugins/Runtime/nDisplay/Source/DisplayCluster/Private/Render/Viewport/Resource/DisplayClusterViewportResourceSettings.h
+++ b/Engine/Plugins/Runtime/nDisplay/Source/DisplayCluster/Private/Render/Viewport/Resource/DisplayClusterViewportResourceSettings.h
@@ -56,16 +56,21 @@ struct FDisplayClusterViewportResourceSettings
 	* @param InResourceFlags - these flags determine what type of resource will be created.
 	* @param InNumMips       - the number of mips that will be created. This value is only supported by a certain type of resource.
 	*/
-	FDisplayClusterViewportResourceSettings(const FDisplayClusterViewportResourceSettings& InBaseSettings, const FIntPoint& InSize, const EPixelFormat InFormat, const EDisplayClusterViewportResourceSettingsFlags InResourceFlags = EDisplayClusterViewportResourceSettingsFlags::None, const int32 InNumMips = 1);
+	FDisplayClusterViewportResourceSettings(const FDisplayClusterViewportResourceSettings& InBaseSettings, const FString InViewportId, const FIntPoint& InSize, const EPixelFormat InFormat, const EDisplayClusterViewportResourceSettingsFlags InResourceFlags = EDisplayClusterViewportResourceSettingsFlags::None, const int32 InNumMips = 1);
 
 public:
 	/** Returns true if the input resource settings match the local values. */
 	inline bool IsResourceSettingsEqual(const FDisplayClusterViewportResourceSettings& In) const
 	{
-		// This is a temporary rule, because the preview texture uses a link to an external resource
-		if (EnumHasAnyFlags(ResourceFlags | In.ResourceFlags, EDisplayClusterViewportResourceSettingsFlags::PreviewTargetableTexture))
+		// Special rules for preview RTT resource
+		if (EnumHasAnyFlags(ResourceFlags, EDisplayClusterViewportResourceSettingsFlags::PreviewTargetableTexture))
 		{
-			return false;
+			// A cluster node can have multiple viewports using RTT previews
+			// But since these textures are linked to the preview mesh material, we will also add this rule to link to a specific viewport.
+			if (In.ViewportId != ViewportId)
+			{
+				return false;
+			}
 		}
 
 		return (In.Size == Size)
@@ -127,6 +132,9 @@ private:
 	// This is a resource belonging to the cluster node
 	const FString ClusterNodeId;
 
+	// The preview resource is always used by the same Outer viewport
+	const FString ViewportId;
+
 	// Resource width and height
 	FIntPoint    Size;
 
diff --git a/Engine/Plugins/Runtime/nDisplay/Source/DisplayCluster/Public/Components/DisplayClusterPreviewShareComponent.h b/Engine/Plugins/Runtime/nDisplay/Source/DisplayCluster/Public/Components/DisplayClusterPreviewShareComponent.h
--- a/Engine/Plugins/Runtime/nDisplay/Source/DisplayCluster/Public/Components/DisplayClusterPreviewShareComponent.h
+++ b/Engine/Plugins/Runtime/nDisplay/Source/DisplayCluster/Public/Components/DisplayClusterPreviewShareComponent.h
@@ -52,6 +52,7 @@ enum class EDisplayClusterPreviewShareIcvfxSync
 	PushActor,
 };
 
+class IDisplayClusterViewportManager;
 
 /**
  * nDisplay Viewport preview share component
@@ -74,15 +75,11 @@ public:
 	/** Constructor */
 	UDisplayClusterPreviewShareComponent(const FObjectInitializer& ObjectInitializer);
 
-#if WITH_EDITOR // Bulk wrap with WITH_EDITOR until preview is supported in other modes.
-
 	//~ UActorComponent interface begin
 	virtual void TickComponent(float DeltaTime, ELevelTick TickType, FActorComponentTickFunction* ThisTickFunction) override;
 	virtual void OnComponentDestroyed(bool bDestroyingHierarchy) override;
 	//~ UActorComponent interface end
 
-#endif // WITH_EDITOR
-
 	//~ UObject interface begin
 
 	virtual void Serialize(FArchive& Ar) override;
@@ -104,19 +101,22 @@ public:
 
 
 protected:
+	/**
+	* Creates or updates a custom rendering for the specified RootActor.
+	* Returns nullptr if rendering is not possible, or a reference to a custom viewport manager storing rendered resources.
+	* 
+	.*/
+	bool UpdateCustomViewportManager(const ADisplayClusterRootActor* InSrcRootActor);
+
+	/** Release the custom viewport manager instance, if it exists. */
+	void ReleaseCustomViewportManager();
 
-#if WITH_EDITOR // Bulk wrap with WITH_EDITOR until preview is supported in other modes.
 
 	//~ UActorComponent interface begin
 	virtual void OnRegister() override;
 	//~ UActorComponent interface end
 
-#endif // WITH_EDITOR
-
 public:
-
-#if WITH_EDITORONLY_DATA
-
 	/** Current sharing mode of this component */
 	UPROPERTY(EditAnywhere, Setter=SetMode, BlueprintSetter=SetMode, Category=Sharing)
 	EDisplayClusterPreviewShareMode Mode = EDisplayClusterPreviewShareMode::None;
@@ -133,11 +133,14 @@ public:
 	UPROPERTY(EditAnywhere, Category = Sharing, meta = (EditCondition = "Mode == EDisplayClusterPreviewShareMode::PullActor"))
 	EDisplayClusterPreviewShareIcvfxSync IcvfxCamerasSyncType = EDisplayClusterPreviewShareIcvfxSync::PullActor;
 
-#endif // WITH_EDITORONLY_DATA
+	/** Override the Enable PostProcess value. */
+	UPROPERTY(EditAnywhere, BlueprintReadWrite, Category = Sharing, meta = (DisplayName = "Enable Post Process"))
+	bool bPreviewEnablePostProcess = true;
 
 private:
-
-#if WITH_EDITOR
+	// DC ViewportManager instance for this component
+	// Used to perform rendering with custom settings from an external DCRA
+	TSharedPtr<IDisplayClusterViewportManager, ESPMode::ThreadSafe> CustomViewportManager;
 
 	/** Closes all media related objects (i.e. media captures and media players) */
 	void CloseAllMedia();
@@ -175,24 +178,24 @@ private:
 	/** Enables/Disables component ticking */
 	void SetTickEnable(const bool bEnable);
 
+#if WITH_EDITOR
 	/** Called when the editor map is changed. Used to remove unwanted references to external maps. */
-	void HandleMapChanged(UWorld* InWorld, EMapChangeType InMapChangeType);
+	void HandleMapChanged_Editor(UWorld* InWorld, EMapChangeType InMapChangeType);
+#endif
 
 	/** Unsubcribe from preview of all root actors */
-	void UnsubscribeFromAllPreviews();
+	UE_DEPRECATED(5.4, "This function has been deprecated.")
+	void UnsubscribeFromAllPreviews() { };
 
 	/** Subscribe for preview of given root actor */
-	void SubscribeToPreview(ADisplayClusterRootActor* RootActor);
+	UE_DEPRECATED(5.4, "This function has been deprecated.")
+	void SubscribeToPreview(ADisplayClusterRootActor* RootActor) { };
 
 	/** Unsubscribe from preview of given root actor */
-	void UnsubscribeFromPreview(ADisplayClusterRootActor* RootActor);
-
-#endif // WITH_EDITOR
+	UE_DEPRECATED(5.4, "This function has been deprecated.")
+	void UnsubscribeFromPreview(ADisplayClusterRootActor* RootActor) { };
 
 private:
-
-#if WITH_EDITORONLY_DATA
-
 	/** Media Outputs associated with the given viewport unique names */
 	UPROPERTY(Transient)
 	TMap<FString, TObjectPtr<UMediaOutput>> MediaOutputs;
@@ -220,7 +223,4 @@ private:
 	/** Cache of original Texture Replace enable boolean associated with the given viewport unique names. Used when restoring the original state */
 	UPROPERTY(Transient)
 	TMap<FString, bool> OriginalTextureReplaces;
-
-#endif // WITH_EDITORONLY_DATA
-
 };
diff --git a/Engine/Plugins/Runtime/nDisplay/Source/DisplayCluster/Public/DisplayClusterRootActor.h b/Engine/Plugins/Runtime/nDisplay/Source/DisplayCluster/Public/DisplayClusterRootActor.h
--- a/Engine/Plugins/Runtime/nDisplay/Source/DisplayCluster/Public/DisplayClusterRootActor.h
+++ b/Engine/Plugins/Runtime/nDisplay/Source/DisplayCluster/Public/DisplayClusterRootActor.h
@@ -27,7 +27,6 @@
 #include "DisplayClusterRootActor.generated.h"
 
 #if WITH_EDITOR
-class UDisplayClusterDisplayDeviceBaseComponent;
 class IDisplayClusterConfiguratorBlueprintEditor;
 class FTransactionObjectEvent;
 #endif
@@ -35,10 +34,11 @@ class FTransactionObjectEvent;
 class FDisplayClusterViewportManager;
 class IDisplayClusterStageActor;
 class USceneComponent;
+class UDisplayClusterDisplayDeviceBaseComponent;
+class ULineBatchComponent;
 class UDisplayClusterConfigurationData;
 class UDisplayClusterCameraComponent;
 class UDisplayClusterOriginComponent;
-class UDisplayClusterPreviewComponent;
 class UDisplayClusterStageGeometryComponent;
 class UDisplayClusterSyncTickComponent;
 class UProceduralMeshComponent;
@@ -83,8 +83,22 @@ public:
 	 */
 	void UpdateConfigDataInstance(UDisplayClusterConfigurationData* ConfigDataTemplate, bool bForceRecreate = false);
 
+	/** Returns true if this RootActor is used as the primary RootActor in the DC GameManager. */
+	bool IsPrimaryRootActor() const;
+
+	/** Returns true if this RootActor is primary and should be displayed in PIE mode. */
+	bool IsPrimaryRootActorForPIE() const;
+
+	/** Returns true if this RootActor is running in PIE mode. */
+	bool IsRunningPIE() const;
+
+	/** Returns true if this RootActor is running in game or in PIE mode. */
 	bool IsRunningGameOrPIE() const;
 
+	/** Returns true if this RootActor is running in DC mode. */
+	bool IsRunningDisplayCluster() const;
+
+
 	UDisplayClusterConfigurationData* GetDefaultConfigDataFromAsset() const;
 	UDisplayClusterConfigurationData* GetConfigData() const;
 
@@ -155,6 +169,13 @@ public:
 	UFUNCTION(BlueprintGetter)
 	UDisplayClusterCameraComponent* GetDefaultCamera() const;
 
+	/** Retrieve the default display device, creating it if it doesn't exist */
+	UFUNCTION(BlueprintGetter)
+	UDisplayClusterDisplayDeviceBaseComponent* GetDefaultDisplayDevice() const;
+
+	/** Retrieve the line batch component. */
+	ULineBatchComponent* GetLineBatchComponent() const;
+
 	/**
 	 * Get the view origin most commonly used by viewports in this cluster.
 	 * If no viewports override the camera, this returns the default camera, or if there isn't one, the actor's root component.
@@ -227,6 +248,9 @@ private:
 	/** Get ViewportManager API. */
 	FDisplayClusterViewportManager* GetViewportManagerImpl() const;
 
+	/** Reset preview rendering. */
+	void ResetEntireClusterPreviewRendering();
+
 private:
 	// DC ViewportManager instance for this DCRA
 	TSharedPtr<FDisplayClusterViewportManager, ESPMode::ThreadSafe> ViewportManagerPtr;
@@ -378,17 +402,19 @@ public:
 		return LastDeltaSecondsValue;
 	}
 
+	/** Get current settings for preview rendering. */
+	FDisplayClusterViewport_PreviewSettings GetPreviewSettings() const;
+
 //////////////////////////////////////////////////////////////////////////////////////////////
 // EDITOR RELATED SETTINGS
 //////////////////////////////////////////////////////////////////////////////////////////////
-#if WITH_EDITORONLY_DATA
 public:
-	/** When the MRQ is rendered, this flag is raised. */
-	UPROPERTY()
-	bool bMoviePipelineRenderPass = false;
+	/** Render this DCRA in game for Standalone/Package builds. */
+	UPROPERTY(EditAnywhere, BlueprintReadWrite, Category = "Preview In Game", meta = (DisplayName = "Enable Preview in Game"))
+	bool bPreviewInGameEnable = false;
 
-#endif
-	// Some preview settings have been opened be available for Standalone/Package builds:
+	UPROPERTY(EditAnywhere, BlueprintReadWrite, Category = "Preview In Game", meta = (DisplayName = "Render Preview Frustum in Game"))
+	bool bPreviewInGameRenderFrustum = false;
 
 	/** Render the scene and display it as a preview on the nDisplay root actor in the editor.  This will impact editor performance. */
 	UPROPERTY(EditAnywhere, BlueprintReadWrite, Category = "Editor Preview", meta = (DisplayName = "Enable Editor Preview"))
@@ -402,11 +428,22 @@ public:
 	UPROPERTY(EditAnywhere, BlueprintReadWrite, Category = "Editor Preview", meta = (DisplayName = "Enable Post Process"), BlueprintSetter = SetPreviewEnablePostProcess)
 	bool bPreviewEnablePostProcess = false;
 
-#if WITH_EDITORONLY_DATA
 	/** Configure the root actor for Techvis rendering with preview components. */
 	UPROPERTY(EditAnywhere, BlueprintReadWrite, Category = "Editor Preview")
-	bool bEnablePreviewTechvis = true;
-	
+	bool bEnablePreviewTechvis = false;
+
+	/** Enable the use of a preview mesh for the preview for this DCRA. */
+	UPROPERTY()
+	bool bEnablePreviewMesh = true;
+
+	/** Enable the use of a preview editable mesh for the preview for this DCRA. */
+	UPROPERTY()
+	bool bEnablePreviewEditableMesh = true;
+
+	/** Always use preview settings defined in this RootActor. */
+	UPROPERTY()
+	bool bUseLocalPreviewSetttings = true;
+
 	/** Freeze preview render.  This will impact editor performance. */
 	UPROPERTY(EditAnywhere, BlueprintReadWrite, Category = "Editor Preview", meta = (DisplayName = "Freeze Editor Preview"))
 	bool bFreezePreviewRender = false;
@@ -419,6 +456,11 @@ public:
 	UPROPERTY(EditAnywhere, BlueprintReadWrite, Category = "Editor Preview", meta = (DisplayName = "Camera Frustum Distance"))
 	float PreviewICVFXFrustumsFarDistance = 1000.0f;
 
+#if WITH_EDITORONLY_DATA
+	/** When the MRQ is rendered, this flag is raised. */
+	UPROPERTY()
+	bool bMoviePipelineRenderPass = false;
+
 	/** Selectively preview a specific viewport or show all/none. */
 	UPROPERTY(EditAnywhere, BlueprintReadWrite, Category = "Editor Preview", meta = (DisplayName = "Preview Node"))
 	FString PreviewNodeId = DisplayClusterConfigurationStrings::gui::preview::PreviewNodeNone;
@@ -428,8 +470,6 @@ public:
 	EDisplayClusterConfigurationRenderMode RenderMode = EDisplayClusterConfigurationRenderMode::Mono;
 #endif
 
-	// Some preview settings have been opened be available for Standalone/Package builds:
-
 	/** Tick Per Frame */
 	UPROPERTY(EditAnywhere, BlueprintReadWrite, Category = "Editor Preview", AdvancedDisplay, meta = (DisplayName = "Tick Per Frame", ClampMin = "1", UIMin = "1", ClampMax = "200", UIMax = "200"))
 	int TickPerFrame = 1;
@@ -442,18 +482,6 @@ public:
 	UPROPERTY(EditAnywhere, BlueprintReadWrite, Category = "Editor Preview", AdvancedDisplay, meta = (DisplayName = "Preview Texture Max Size", ClampMin = "64", UIMin = "64", ClampMax = "4096", UIMax = "4096"))
 	int PreviewMaxTextureDimension = 2048;
 
-//////////////////////////////////////////////////////////////////////////////////////////////
-// EDITOR RELATED SETTINGS
-//////////////////////////////////////////////////////////////////////////////////////////////
-#if WITH_EDITORONLY_DATA
-private:
-	UPROPERTY(Transient)
-	TMap<FString, TObjectPtr<UDisplayClusterPreviewComponent>> PreviewComponents;
-
-	/** The default display device to use for preview rendering */
-	UPROPERTY(Transient, NonTransactional)
-	TObjectPtr<UDisplayClusterDisplayDeviceBaseComponent> DefaultDisplayDeviceComponent;
-
 	/** The included display device nDisplay provides by default */
 	UPROPERTY(VisibleDefaultsOnly, Category = "Editor Preview", DisplayName = "Basic Display Device")
 	TObjectPtr<UDisplayClusterDisplayDeviceBaseComponent> BasicDisplayDeviceComponent;
@@ -461,30 +489,34 @@ private:
 	/** Select the default display device class to use when a viewport doesn't have one assigned */
 	UPROPERTY(EditDefaultsOnly, Category = "Editor Preview", DisplayName = "Default Display Device")
 	FName DefaultDisplayDeviceName;
-	
-	UPROPERTY(Transient)
-	bool bDeferPreviewGeneration = false;
-#endif
+
+protected:
+	/** The default display device to use for preview rendering */
+	UPROPERTY(Transient, NonTransactional)
+	mutable TObjectPtr<UDisplayClusterDisplayDeviceBaseComponent> DefaultDisplayDeviceComponent;
+
+	/** Line Batchers. All lines to be drawn in the world, but not inside viewports. */
+	UPROPERTY(Transient, NonTransactional)
+	TObjectPtr<ULineBatchComponent> LineBatcherComponent;
+
+protected:
+	/** The name the internal default display device uses. */
+	virtual FName GetInternalDisplayDeviceName() const;
+
+//////////////////////////////////////////////////////////////////////////////////////////////
+// EDITOR RELATED SETTINGS
+//////////////////////////////////////////////////////////////////////////////////////////////
 
 #if WITH_EDITOR
 public:
 	/** Enable or disable editor render. Preview components may need to be on for texture overrides, but capture and rendering disabled. */
-	void EnableEditorRender(bool bValue);
+	UE_DEPRECATED(5.4, "This function has been deprecated.")
+	void EnableEditorRender(bool bValue) { }
 
 	/** If editor rendering is enabled. */
-	bool IsEditorRenderEnabled() const { return bEnableEditorRender; }
-
-protected:
-	/** Reset preview rendering for cluster node. */
-	void ResetClusterNodePreviewRendering_Editor();
-
-	/** Is preview rendering for cluster node in progress. */
-	bool IsActiveClusterNodePreviewRendering_Editor() const;
+	UE_DEPRECATED(5.4, "This function has been deprecated.")
+	bool IsEditorRenderEnabled() const { return false; }
 
-private:
-	/** Is editor rendering enabled? This can be false and the preview still enabled. */
-	bool bEnableEditorRender = true;
-	
 public:
 	DECLARE_DELEGATE(FOnPreviewUpdated);
 
@@ -492,20 +524,26 @@ public:
 	// We need tick in Editor
 	virtual bool ShouldTickIfViewportsOnly() const override { return true; }
 
+	UE_DEPRECATED(5.4, "This function has been deprecated.")
 	FOnPreviewUpdated& GetOnPreviewGenerated() { return OnPreviewGenerated; }
+
+	UE_DEPRECATED(5.4, "This function has been deprecated.")
 	FOnPreviewUpdated& GetOnPreviewDestroyed() { return OnPreviewDestroyed; }
 
 	// return true, if preview enabled for this actor
-	bool IsPreviewEnabled() const;
+	UE_DEPRECATED(5.4, "This function has been deprecated.")
+	bool IsPreviewEnabled() const { return false; }
 
 	/** Gets whether the preview output is displayed onto the stage actor's screen meshes */
-	bool IsPreviewDrawnToScreens() const;
+	UE_DEPRECATED(5.4, "This function has been deprecated.")
+	bool IsPreviewDrawnToScreens() const { return false; }
 
 	void Constructor_Editor();
 	void Destructor_Editor();
 
 	/** Perform a rendering of the DCRA preview. */
-	void RenderPreview_Editor();
+	UE_DEPRECATED(5.4, "This function has been deprecated.")
+	void RenderPreview_Editor() { };
 
 	void PostLoad_Editor();
 	void PostActorCreated_Editor();
@@ -515,12 +553,16 @@ public:
 	void RerunConstructionScripts_Editor();
 
 	// Preview components free referenced meshes and materials
-	void ResetPreviewComponents_Editor(bool bInRestoreSceneMaterial);
+	UE_DEPRECATED(5.4, "This function has been deprecated.")
+	void ResetPreviewComponents_Editor(bool bInRestoreSceneMaterial) { }
 
-	UDisplayClusterPreviewComponent* GetPreviewComponent(const FString& NodeId, const FString& ViewportId) const;
+	UE_DEPRECATED(5.4, "This function has been deprecated.")
+	class UDisplayClusterPreviewComponent* GetPreviewComponent(const FString& NodeId, const FString& ViewportId) const { return nullptr; }
 
-	void UpdatePreviewComponents();
-	void ReleasePreviewComponents();
+	UE_DEPRECATED(5.4, "This function has been deprecated.")
+	void UpdatePreviewComponents() { }
+	UE_DEPRECATED(5.4, "This function has been deprecated.")
+	void ReleasePreviewComponents() { }
 
 	/**
 	 * Enable the use of a post process render target when bPreviewEnablePostProcess is disabled on the actor. The root actor
@@ -531,7 +573,8 @@ public:
 	 * @param Object The object subscribing to updates
 	 * @return The number of subscribers to use post process.
 	 */
-	int32 SubscribeToPostProcessRenderTarget(const uint8* Object);
+	UE_DEPRECATED(5.4, "This function has been deprecated.")
+	int32 SubscribeToPostProcessRenderTarget(const uint8* Object) { return INDEX_NONE; };
 
 	/**
 	 * Unsubscribe a registered object from requiring post process render target updates.
@@ -539,31 +582,43 @@ public:
 	 * @param Object The object subscribing to updates. When the counter is zero post process render targets will not be used.
 	 * @return The number of subscribers to use post process.
 	 */
-	int32 UnsubscribeFromPostProcessRenderTarget(const uint8* Object);
+	UE_DEPRECATED(5.4, "This function has been deprecated.")
+	int32 UnsubscribeFromPostProcessRenderTarget(const uint8* Object) { return INDEX_NONE; };
 
 	/** If one or more observers are subscribed to receive post process preview targets. */
-	bool DoObserversNeedPostProcessRenderTarget() const;
+	UE_DEPRECATED(5.4, "This function has been deprecated.")
+	bool DoObserversNeedPostProcessRenderTarget() const { return false; };
 	
 	/** When rendering the preview determine which render target should be used for the current frame. */
-	bool ShouldThisFrameOutputPreviewToPostProcessRenderTarget() const;
+	UE_DEPRECATED(5.4, "This function has been deprecated.")
+	bool ShouldThisFrameOutputPreviewToPostProcessRenderTarget() const { return false; }
 
 	/** Force preview rendering to be enabled regardless of the user's setting until a matching RemovePreviewEnableOverride call is made. */
-	void AddPreviewEnableOverride(const uint8* Object);
+	UE_DEPRECATED(5.4, "This function has been deprecated.")
+	void AddPreviewEnableOverride(const uint8* Object) { };
 
 	/**
 	 * Stop forcing preview rendering to be enabled for this caller. If other objects have called AddPreviewEnableOverride, it will remain
 	 * forced until they have also removed their overrides.
 	 */
-	void RemovePreviewEnableOverride(const uint8* Object);
+	UE_DEPRECATED(5.4, "This function has been deprecated.")
+	void RemovePreviewEnableOverride(const uint8* Object) { };
 
+	UE_DEPRECATED(5.4, "This function has been deprecated.")
 	float GetPreviewRenderTargetRatioMult() const
 	{
 		return PreviewRenderTargetRatioMult;
 	};
 
-	IDisplayClusterViewport* FindPreviewViewport(const FString& InViewportId) const;
+	UE_DEPRECATED(5.4, "This function has been deprecated.")
+	IDisplayClusterViewport* FindPreviewViewport(const FString& InViewportId) const
+	{
+		return nullptr;
+	}
 
-	void GetPreviewRenderTargetableTextures(const TArray<FString>& InViewportNames, TArray<FTextureRHIRef>& OutTextures);
+	UE_DEPRECATED(5.4, "This function has been deprecated.")
+	void GetPreviewRenderTargetableTextures(const TArray<FString>& InViewportNames, TArray<FTextureRHIRef>& OutTextures)
+	{ }
 
 	void UpdateInnerFrustumPriority();
 	void ResetInnerFrustumPriority();
@@ -572,40 +627,13 @@ public:
 	void SetIsSelectedInEditor(bool bValue);
 
 	// Don't show actor preview in the level viewport when DCRA actor is selected, but none of its children are.
+	UE_DEPRECATED(5.4, "This function has been deprecated.")
 	virtual bool IsDefaultPreviewEnabled() const override
 	{
 		return false;
 	}
 
-	/** The name the internal default display device uses. */
-	static FName GetInternalDisplayDeviceName();
-
-	/** Retrieve the default display device, creating it if it doesn't exist */
-	UDisplayClusterDisplayDeviceBaseComponent* GetDefaultDisplayDevice();
-
-	/** If the preview should use Techvis rendering. */
-	bool IsTechvisPreviewRenderingEnabled() const { return bEnablePreviewTechvis; }
-	
 protected:
-	FString GeneratePreviewComponentName_Editor(const FString& NodeId, const FString& ViewportId) const;
-	void ResetPreviewInternals_Editor();
-
-	/** Configure preview components for Techvis rendering. */
-	void ConfigureTechvis_Editor();
-
-	bool ImplUpdatePreviewConfiguration_Editor(const FString& InClusterNodeId);
-
-	/** Gets a preview rendering mode for this root actor. */
-	EDisplayClusterRenderFrameMode GetPreviewRenderMode_Editor() const;
-
-	void ImplRenderPreview_Editor();
-	bool ImplRenderPassPreviewClusterNode_Editor(const FString& InClusterNodeId);
-
-	bool ImplUpdatePreviewRenderFrame_Editor(const FString& InClusterNodeId);
-
-	void ImplRenderPreviewFrustums_Editor();
-	void ImplRenderPreviewViewportFrustum_Editor(const FMatrix ProjectionMatrix, const FMatrix ViewMatrix, const FVector ViewOrigin);
-
 	virtual void PostEditChangeChainProperty(FPropertyChangedChainEvent& PropertyChainEvent) override;
 	virtual void PostEditChangeProperty(FPropertyChangedEvent& PropertyChangedEvent) override;
 	virtual void PostEditMove(bool bFinished) override;
@@ -627,38 +655,9 @@ private:
 
 	bool bIsSelectedInEditor = false;
 
-	/** When the preview render should be directed to the post process render target. */
-	bool bOutputFrameToPostProcessRenderTarget;
-
-	/** Enables preview components to output to the post process render target when bPreviewEnablePostProcess is disabled. Contains all subscribed objects. */
-	TSet<const uint8*> PostProcessRenderTargetObservers;
-
-	/* Addresses of callers to AddPreviewEnableOverride that haven't removed their overrides yet. */
-	TSet<const uint8*> PreviewEnableOverriders;
-	
 	TWeakPtr<IDisplayClusterConfiguratorBlueprintEditor> ToolkitPtr;
 
-	int32 TickPerFrameCounter = 0;
-
-	int32 PreviewClusterNodeIndex = 0;
-	
-	int32 PreviewViewportIndex = -1;
-	TUniquePtr<FDisplayClusterRenderFrame> PreviewRenderFrame;
-	FString PreviewRenderFrameClusterNodeId;
-
-	int32 PreviewViewportsRenderedInThisFrameCnt = 0;
-
 	FOnPreviewUpdated OnPreviewGenerated;
 	FOnPreviewUpdated OnPreviewDestroyed;
-
-	struct FFrustumPreviewViewportContextCache
-	{
-		FVector  ViewLocation;
-		FRotator ViewRotation;
-		FMatrix  ProjectionMatrix;
-	};
-	// Cache the last valid viewport context
-	TMap<FString, FFrustumPreviewViewportContextCache> FrustumPreviewViewportContextCache;
-
 #endif
 };
diff --git a/Engine/Plugins/Runtime/nDisplay/Source/DisplayCluster/Public/Misc/DisplayClusterObjectRef.h b/Engine/Plugins/Runtime/nDisplay/Source/DisplayCluster/Public/Misc/DisplayClusterObjectRef.h
--- a/Engine/Plugins/Runtime/nDisplay/Source/DisplayCluster/Public/Misc/DisplayClusterObjectRef.h
+++ b/Engine/Plugins/Runtime/nDisplay/Source/DisplayCluster/Public/Misc/DisplayClusterObjectRef.h
@@ -71,7 +71,7 @@ public:
 		return ActorPtr.Get();
 	}
 
-	bool SetSceneActor(AActor* InActor)
+	bool SetSceneActor(const AActor* InActor)
 	{
 		FScopeLock lock(&DataGuard);
 
@@ -87,7 +87,9 @@ public:
 
 				WorldPtr = TWeakObjectPtr<UWorld>(InActor->GetWorld());
 				ActorClassPtr = TWeakObjectPtr<UClass>(ActorClass);
-				ActorPtr = TWeakObjectPtr<AActor>(InActor);
+
+				// Ignore const
+				ActorPtr = TWeakObjectPtr<AActor>((AActor*)InActor);
 
 				return true;
 			}
diff --git a/Engine/Plugins/Runtime/nDisplay/Source/DisplayCluster/Public/Render/Projection/IDisplayClusterProjectionPolicy.h b/Engine/Plugins/Runtime/nDisplay/Source/DisplayCluster/Public/Render/Projection/IDisplayClusterProjectionPolicy.h
--- a/Engine/Plugins/Runtime/nDisplay/Source/DisplayCluster/Public/Render/Projection/IDisplayClusterProjectionPolicy.h
+++ b/Engine/Plugins/Runtime/nDisplay/Source/DisplayCluster/Public/Render/Projection/IDisplayClusterProjectionPolicy.h
@@ -324,33 +324,33 @@ public:
 	}
 
 	/**
-	* Ask projection policy instance if it has any movable mesh based preview
+	* Ask projection policy instance if it has any Editable mesh based preview
 	* 
 	* @param InViewport - a owner viewport
 	* @return - True if mesh based preview is available
 	*/
-	virtual bool HasPreviewMovableMesh(IDisplayClusterViewport* InViewport)
+	virtual bool HasPreviewEditableMesh(IDisplayClusterViewport* InViewport)
 	{
 		return false;
 	}
 
 	/**
-	* Build preview movable mesh
+	* Build preview Editable mesh
 	* This MeshComponent is a copy of the preview mesh and can be moved freely with the UI visualization.
 	*
 	* @param InViewport - a owner viewport
 	*/
-	virtual UMeshComponent* GetOrCreatePreviewMovableMeshComponent(IDisplayClusterViewport* InViewport)
+	virtual UMeshComponent* GetOrCreatePreviewEditableMeshComponent(IDisplayClusterViewport* InViewport)
 	{
 		return nullptr;
 	}
 
 	/**
-	 * Return Origin point component used by preview movable mesh
+	 * Return Origin point component used by preview Editable mesh
 	 * 
 	 * @param InViewport - a owner viewport
 	 */
-	virtual USceneComponent* const GetPreviewMovableMeshOriginComponent(IDisplayClusterViewport* InViewport) const
+	virtual USceneComponent* const GetPreviewEditableMeshOriginComponent(IDisplayClusterViewport* InViewport) const
 	{
 		return nullptr;
 	}
diff --git a/Engine/Plugins/Runtime/nDisplay/Source/DisplayCluster/Public/Render/Viewport/Containers/DisplayClusterViewport_PreviewSettings.h b/Engine/Plugins/Runtime/nDisplay/Source/DisplayCluster/Public/Render/Viewport/Containers/DisplayClusterViewport_PreviewSettings.h
--- a/Engine/Plugins/Runtime/nDisplay/Source/DisplayCluster/Public/Render/Viewport/Containers/DisplayClusterViewport_PreviewSettings.h
+++ b/Engine/Plugins/Runtime/nDisplay/Source/DisplayCluster/Public/Render/Viewport/Containers/DisplayClusterViewport_PreviewSettings.h
@@ -9,22 +9,48 @@
 */
 struct FDisplayClusterViewport_PreviewSettings
 {
-	// Preview RTT size multiplier
-	float RenderTargetRatioMult = 1.f;
+	// The IDisplayClusterViewportManagerPreview::UpdateEntireClusterPreviewRender() function will use this rendering mode
+	// For special rendering cases, set a different value, such as EDisplayClusterRenderFrameMode::PreviewProxyHitInScene
+	EDisplayClusterRenderFrameMode EntireClusterPreviewRenderMode = EDisplayClusterRenderFrameMode::PreviewInScene;
+
+	// Render the scene and display it as a preview on the nDisplay root actor in the editor.  This will impact editor performance.
+	bool bPreviewEnable = false;
+
+	// Render this DCRA in game for Standalone/Package builds.
+	bool bPreviewInGameEnable = false;
+	bool bPreviewInGameRenderFrustum = false;
+
+	// Preview uses techvis mode for rendering
+	bool bEnablePreviewTechvis = false;
 
 	// Enable/Disable preview rendering. When disabled preview image freeze
-	bool bFreezeRender = false;
+	bool bFreezePreviewRender = false;
 
 	// Hack preview gamma.
 	// In a scene, PostProcess always renders on top of the preview textures.
 	// But in it, PostProcess is also rendered with the flag turned off.
-	bool bEnablePostProcess = false;
+	bool bPreviewEnablePostProcess = false;
+
+	// Allows you to process preview meshes inside DCRA (get or create a mesh from a projection policy, update materials on the preview mesh, etc.).
+	bool bEnablePreviewMesh = false;
+
+	// Allows you to process preview editable meshes inside DCRA (get or create a mesh from a projection policy, update materials on the preview editable mesh, etc.).
+	bool bEnablePreviewEditableMesh = false;
+
+	// Render ICVFX Frustums
+	bool bPreviewICVFXFrustums = false;
+	float PreviewICVFXFrustumsFarDistance = 1000.0f;
+
+	// Preview RTT size multiplier
+	float PreviewRenderTargetRatioMult = 1.f;
 
 	// The maximum dimension of any texture for preview
 	// Limit preview textures max size
-	int32 MaxTextureDimension = 2048;
+	int32 PreviewMaxTextureDimension = 2048;
+
+	// Tick Per Frame
+	int TickPerFrame = 1;
 
-	// [Experimental] Render preview in multi-GPU
-	// Specifies the mGPU index range for rendering the DCRA preview.
-	TOptional<FIntPoint> MultiGPURendering;
+	// Max amount of Viewports Per Frame
+	int ViewportsPerFrame = 1;
 };
diff --git a/Engine/Plugins/Runtime/nDisplay/Source/DisplayCluster/Public/Render/Viewport/Containers/DisplayClusterViewport_RenderSettings.h b/Engine/Plugins/Runtime/nDisplay/Source/DisplayCluster/Public/Render/Viewport/Containers/DisplayClusterViewport_RenderSettings.h
--- a/Engine/Plugins/Runtime/nDisplay/Source/DisplayCluster/Public/Render/Viewport/Containers/DisplayClusterViewport_RenderSettings.h
+++ b/Engine/Plugins/Runtime/nDisplay/Source/DisplayCluster/Public/Render/Viewport/Containers/DisplayClusterViewport_RenderSettings.h
@@ -17,6 +17,9 @@ public:
 	// Assigned camera. If empty, the currently active camera must be used
 	FString CameraId;
 
+	// Assigned display device. If empty, the currently active display device must be used
+	FString DisplayDeviceId;
+
 	// Location and size on a backbuffer.
 	FIntRect Rect;
 
diff --git a/Engine/Plugins/Runtime/nDisplay/Source/DisplayCluster/Public/Render/Viewport/IDisplayClusterViewport.h b/Engine/Plugins/Runtime/nDisplay/Source/DisplayCluster/Public/Render/Viewport/IDisplayClusterViewport.h
--- a/Engine/Plugins/Runtime/nDisplay/Source/DisplayCluster/Public/Render/Viewport/IDisplayClusterViewport.h
+++ b/Engine/Plugins/Runtime/nDisplay/Source/DisplayCluster/Public/Render/Viewport/IDisplayClusterViewport.h
@@ -4,6 +4,7 @@
 
 #include "Render/Viewport/IDisplayClusterViewport_CustomPostProcessSettings.h"
 #include "Render/Viewport/IDisplayClusterViewportConfiguration.h"
+#include "Render/Viewport/IDisplayClusterViewportPreview.h"
 
 #include "Render/Viewport/Containers/DisplayClusterViewport_Context.h"
 #include "Render/Viewport/Containers/DisplayClusterViewport_RenderSettings.h"
@@ -32,6 +33,10 @@ public:
 	virtual TSharedPtr<IDisplayClusterViewport, ESPMode::ThreadSafe> ToSharedPtr() = 0;
 	virtual TSharedPtr<const IDisplayClusterViewport, ESPMode::ThreadSafe> ToSharedPtr() const = 0;
 
+	/** Internal functions. Get TSharedRef from Self. */
+	virtual TSharedRef<class FDisplayClusterViewport, ESPMode::ThreadSafe> ToSharedRef() = 0;
+	virtual TSharedRef<const class FDisplayClusterViewport, ESPMode::ThreadSafe> ToSharedRef() const = 0;
+
 	virtual FString GetId() const = 0;
 	virtual FString GetClusterNodeId() const = 0;
 
@@ -41,6 +46,9 @@ public:
 	/** [const] Get viewport manager configuration interface. */
 	virtual const IDisplayClusterViewportConfiguration& GetConfiguration() const = 0;
 
+	/** Get viewport preview API */
+	virtual IDisplayClusterViewportPreview& GetViewportPreview() const = 0;
+
 	/** Get main viewport render settings. */
 	virtual const FDisplayClusterViewport_RenderSettings&      GetRenderSettings() const = 0;
 
@@ -98,6 +106,11 @@ public:
 	*/
 	virtual class UDisplayClusterCameraComponent* GetViewPointCameraComponent(const EDisplayClusterRootActorType InRootActorType) const = 0;
 
+	/** Return a DisplayDevice component for this viewport.
+	* @param InRootActorType - the root type
+	*/
+	virtual class UDisplayClusterDisplayDeviceBaseComponent* GetDisplayDeviceComponent(const EDisplayClusterRootActorType InRootActorType) const = 0;
+
 	/** Retrieves the view position from the ViewPoint component used by this viewport.
 	 * The CalculateView() function receives arguments that may not match this position because they can be overridden.
 	 * Therefore, projection policies that expect values from the ViewPoint component must use this function to get the correct values.
diff --git a/Engine/Plugins/Runtime/nDisplay/Source/DisplayCluster/Public/Render/Viewport/IDisplayClusterViewportConfiguration.h b/Engine/Plugins/Runtime/nDisplay/Source/DisplayCluster/Public/Render/Viewport/IDisplayClusterViewportConfiguration.h
--- a/Engine/Plugins/Runtime/nDisplay/Source/DisplayCluster/Public/Render/Viewport/IDisplayClusterViewportConfiguration.h
+++ b/Engine/Plugins/Runtime/nDisplay/Source/DisplayCluster/Public/Render/Viewport/IDisplayClusterViewportConfiguration.h
@@ -5,6 +5,7 @@
 #include "CoreMinimal.h"
 #include "Engine/EngineTypes.h"
 #include "Render/Viewport/RenderFrame/DisplayClusterRenderFrameEnums.h"
+#include "Render/Viewport/Containers/DisplayClusterViewport_PreviewSettings.h"
 
 class UWorld;
 class UDisplayClusterConfigurationData;
@@ -19,20 +20,21 @@ struct FDisplayClusterConfigurationRenderFrame;
 enum class EDisplayClusterRootActorType : uint8
 {
 	// This DCRA will be used to render previews. The meshes and preview materials are created at runtime.
-	Preview = 0,
+	Preview = 1 << 0,
 
 	// A reference to DCRA in the scene, used as a source for math calculations and references.
 	// Locations in the scene and math data are taken from this DCRA.
-	Scene,
+	Scene = 1 << 1,
 
 	// Reference to DCRA, used as a source of configuration data from DCRA and its components.
-	Configuration,
+	Configuration = 1 << 2,
 
 	// This value can only be used in very specific cases:
 	// For function GetRootActor() : Return any of the DCRAs that are not nullptr, in ascending order of type: Preview, Scene, Configuration.
 	// For function SetRootActor() : Sets all references to DRCA to the specified value.
-	Any
+	Any = Preview | Scene | Configuration,
 };
+ENUM_CLASS_FLAGS(EDisplayClusterRootActorType);
 
 /**
  * Viewport manager configuration.
@@ -43,36 +45,45 @@ public:
 	virtual ~IDisplayClusterViewportConfiguration() = default;
 
 public:
-	/**
-	 *  Sets a reference to the current world to be rendered in DCRA
-	 * @param InWorld - ptr to the world to be rendered
-	 */
-	virtual void SetCurrentWorld(UWorld* InWorld) = 0;
-
 	/**
 	 *  Sets a reference to the DCRA's
-	 * @param InRootActor     - a ref to DCRA
 	 * @param InRootActorType - Type of the DCRA
+	 * @param InRootActor     - a ref to DCRA
 	 */
-	virtual void SetRootActor(ADisplayClusterRootActor* InRootActor, const EDisplayClusterRootActorType InRootActorType) = 0;
+	virtual void SetRootActor(const EDisplayClusterRootActorType InRootActorType, const ADisplayClusterRootActor* InRootActor) = 0;
+
+	/**
+	* Assign new preview settings for rendering previews.
+	*
+	* @param InPreviewSettings - a new preview settings
+	*/
+	virtual void SetPreviewSettings(const FDisplayClusterViewport_PreviewSettings& InPreviewSettings) = 0;
+
+	/** Gets the current preview settings. */
+	virtual const FDisplayClusterViewport_PreviewSettings& GetPreviewSettings() const = 0;
 
 	/**
 	* Update\Create\Delete local node viewports
 	* Updating the configuration to render a ClusterNode in the specified mode
 	*
 	* @param InRenderMode     - Render mode
+	* @param InWorld          - ptr to the world to be rendered
 	* @param InClusterNodeId  - cluster node for rendering
 	*/
-	virtual bool UpdateConfigurationForClusterNode(EDisplayClusterRenderFrameMode InRenderMode, const FString& InClusterNodeId) = 0;
+	virtual bool UpdateConfigurationForClusterNode(EDisplayClusterRenderFrameMode InRenderMode, const UWorld* InWorld, const FString& InClusterNodeId) = 0;
 
 	/**
 	* Update\Create\Delete local node viewports
 	* Updating the configuration to render a list of viewports in a given mode
 	*
 	* @param InRenderMode    - Render mode
+	* @param InWorld         - ptr to the world to be rendered
 	* @param InViewportNames - Viewports names for next frame
 	*/
-	virtual bool UpdateConfigurationForViewportsList(EDisplayClusterRenderFrameMode InRenderMode, const TArray<FString>& InViewportNames) = 0;
+	virtual bool UpdateConfigurationForViewportsList(EDisplayClusterRenderFrameMode InRenderMode, const UWorld* InWorld, const TArray<FString>& InViewportNames) = 0;
+
+	/** Release the current configuration and free resources. */
+	virtual void ReleaseConfiguration() = 0;
 
 public:
 	/** Return the configuration proxy object. */
@@ -136,9 +147,19 @@ public:
 	/** Returns true if preview rendering mode is used. */
 	virtual bool IsPreviewRendering() const = 0;
 
+	/** Returns true, if Techvis is used. */
+	virtual bool IsTechvisEnabled() const = 0;
+
+	/** Returns true if the DCRA preview feature in Standalone/Package builds is used. */
+	virtual bool IsPreviewInGameEnabled() const = 0;
+	
+
 	/** Returns the rendering mode for PIE. */
 	virtual EDisplayClusterRenderFrameMode GetRenderModeForPIE() const = 0;
 
 	/** Return current cluster node id. */
 	virtual const FString& GetClusterNodeId() const = 0;
+
+	/** Return current value for WorldToMeters. */
+	virtual const float GetWorldToMeters() const = 0;
 };
diff --git a/Engine/Plugins/Runtime/nDisplay/Source/DisplayCluster/Public/Render/Viewport/IDisplayClusterViewportManager.h b/Engine/Plugins/Runtime/nDisplay/Source/DisplayCluster/Public/Render/Viewport/IDisplayClusterViewportManager.h
--- a/Engine/Plugins/Runtime/nDisplay/Source/DisplayCluster/Public/Render/Viewport/IDisplayClusterViewportManager.h
+++ b/Engine/Plugins/Runtime/nDisplay/Source/DisplayCluster/Public/Render/Viewport/IDisplayClusterViewportManager.h
@@ -2,6 +2,7 @@
 
 #pragma once
 #include "Render/Viewport/IDisplayClusterViewport.h"
+#include "Render/Viewport/IDisplayClusterViewportManagerPreview.h"
 #include "Render/Viewport/RenderFrame/DisplayClusterRenderFrameEnums.h"
 
 #include "SceneView.h"
@@ -27,10 +28,20 @@ public:
 	virtual ~IDisplayClusterViewportManager() = default;
 
 public:
+	/** Static viewport manager instance constructor.
+	* This function is needed for external modules that may need their own instance of ViewportManager, 
+	* which will be used to generate preview images with their own rendering settings.
+	*/
+	static TSharedRef<IDisplayClusterViewportManager, ESPMode::ThreadSafe> CreateViewportManager();
+
 	/** Get TSharedPtr from self. */
 	virtual TSharedPtr<IDisplayClusterViewportManager, ESPMode::ThreadSafe> ToSharedPtr() = 0;
 	virtual TSharedPtr<const IDisplayClusterViewportManager, ESPMode::ThreadSafe> ToSharedPtr() const = 0;
 
+	/** Internal functions. Get TSharedRef from Self. */
+	virtual TSharedRef<class FDisplayClusterViewportManager, ESPMode::ThreadSafe> ToSharedRef() = 0;
+	virtual TSharedRef<const class FDisplayClusterViewportManager, ESPMode::ThreadSafe> ToSharedRef() const = 0;
+
 	/** Get viewport manager proxy interface. */
 	virtual const IDisplayClusterViewportManagerProxy* GetProxy() const = 0;
 	virtual       IDisplayClusterViewportManagerProxy* GetProxy() = 0;
@@ -39,6 +50,10 @@ public:
 	virtual IDisplayClusterViewportConfiguration& GetConfiguration() = 0;
 	virtual const IDisplayClusterViewportConfiguration& GetConfiguration() const = 0;
 
+	/** Get viewport manager preview API */
+	virtual IDisplayClusterViewportManagerPreview& GetViewportManagerPreview() = 0;
+	virtual const IDisplayClusterViewportManagerPreview& GetViewportManagerPreview() const = 0;
+
 	/**
 	* Initialize new frame for all viewports on game thread, and update context, render resources with viewport new settings
 	* And finally build render frame structure and send to render thread proxy viewport objects
@@ -99,21 +114,6 @@ public:
 	/** Add internal DCVM objects to the reference collector. */
 	virtual void AddReferencedObjects(FReferenceCollector& Collector) = 0;
 
-	/**
-	* Render in editor (preview)
-	* [Game thread func]
-	*
-	* @param InRenderFrame - render frame setup
-	* @param InViewport
-	* @param InFirstViewportNum - begin render from this viewport in frame
-	* @param InViewportsAmount - max viewports for render
-	* @param OutViewportsAmount - total viewport rendered
-	* @param bOutFrameRendered - true, if cluster node composition pass done (additional pass after last viewport is rendered)
-	*
-	* @return - true, if render success
-	*/
-	virtual bool RenderInEditor(FDisplayClusterRenderFrame& InRenderFrame, FViewport* InViewport, const uint32 InFirstViewportNum, const int32 InViewportsAmount, int32& OutViewportsAmount, bool& bOutFrameRendered) = 0;
-
 	///////////////////////////////////////////////////////////////////////////////////////////////////////////////////////////////////////////////////////////////////////////
 	/**
 	* Find viewport object by name
@@ -279,4 +279,24 @@ public:
 	{
 		return false;
 	}
+
+
+	/**
+	* Render in editor (preview)
+	* [Game thread func]
+	*
+	* @param InRenderFrame - render frame setup
+	* @param InViewport
+	* @param InFirstViewportNum - begin render from this viewport in frame
+	* @param InViewportsAmount - max viewports for render
+	* @param OutViewportsAmount - total viewport rendered
+	* @param bOutFrameRendered - true, if cluster node composition pass done (additional pass after last viewport is rendered)
+	*
+	* @return - true, if render success
+	*/
+	UE_DEPRECATED(5.4, "This function has been deprecated. Please use 'GetViewportPreview().Render()'.")
+	virtual bool RenderInEditor(FDisplayClusterRenderFrame& InRenderFrame, FViewport* InViewport, const uint32 InFirstViewportNum, const int32 InViewportsAmount, int32& OutViewportsAmount, bool& bOutFrameRendered)
+	{
+		return false;
+	}
 };
diff --git a/Engine/Plugins/Runtime/nDisplay/Source/DisplayCluster/Public/Render/Viewport/RenderFrame/DisplayClusterRenderFrame.h b/Engine/Plugins/Runtime/nDisplay/Source/DisplayCluster/Public/Render/Viewport/RenderFrame/DisplayClusterRenderFrame.h
--- a/Engine/Plugins/Runtime/nDisplay/Source/DisplayCluster/Public/Render/Viewport/RenderFrame/DisplayClusterRenderFrame.h
+++ b/Engine/Plugins/Runtime/nDisplay/Source/DisplayCluster/Public/Render/Viewport/RenderFrame/DisplayClusterRenderFrame.h
@@ -11,7 +11,7 @@
 
 class IDisplayClusterViewport;
 class IDisplayClusterViewportManager;
-class FRenderTarget;
+class FDisplayClusterViewportResource;
 
 /**
  * nDisplay: DCViewport context for render in UE View
@@ -71,7 +71,7 @@ struct FDisplayClusterRenderFrameTarget
 	EDisplayClusterViewportCaptureMode CaptureMode = EDisplayClusterViewportCaptureMode::Default;
 
 	// Render target resource ref
-	FRenderTarget* RenderTargetPtr = nullptr;
+	TSharedPtr<FDisplayClusterViewportResource, ESPMode::ThreadSafe> RenderTargetResource;
 
 	// Families, rendered on this target
 	TArray<FDisplayClusterRenderFrameTargetViewFamily> ViewFamilies;
diff --git a/Engine/Plugins/Runtime/nDisplay/Source/DisplayCluster/Public/Render/Viewport/RenderFrame/DisplayClusterRenderFrameEnums.h b/Engine/Plugins/Runtime/nDisplay/Source/DisplayCluster/Public/Render/Viewport/RenderFrame/DisplayClusterRenderFrameEnums.h
--- a/Engine/Plugins/Runtime/nDisplay/Source/DisplayCluster/Public/Render/Viewport/RenderFrame/DisplayClusterRenderFrameEnums.h
+++ b/Engine/Plugins/Runtime/nDisplay/Source/DisplayCluster/Public/Render/Viewport/RenderFrame/DisplayClusterRenderFrameEnums.h
@@ -23,6 +23,20 @@ enum class EDisplayClusterRenderFrameMode: uint8
 	// Special render: render to two separated backbuffer textures, with half size Y
 	TopBottom,
 
+
+	// Monoscopic rendering in PIE
+	PIE_Mono,
+
+	// Stereoscopic side-by-side rendering in PIE
+	PIE_SideBySide,
+
+	// Stereoscopic top-bottom rendering in PIE
+	PIE_TopBottom,
+
+	// Monoscopic rendering for MRQ
+	MRQ_Mono,
+
+
 	// Special render for preview in scene
 	PreviewInScene,
 
diff --git a/Engine/Plugins/Runtime/nDisplay/Source/DisplayCluster/Public/Render/Warp/IDisplayClusterWarpPolicy.h b/Engine/Plugins/Runtime/nDisplay/Source/DisplayCluster/Public/Render/Warp/IDisplayClusterWarpPolicy.h
--- a/Engine/Plugins/Runtime/nDisplay/Source/DisplayCluster/Public/Render/Warp/IDisplayClusterWarpPolicy.h
+++ b/Engine/Plugins/Runtime/nDisplay/Source/DisplayCluster/Public/Render/Warp/IDisplayClusterWarpPolicy.h
@@ -3,9 +3,13 @@
 #pragma once
 
 #include "CoreMinimal.h"
+#include "Render/DisplayDevice/Containers/DisplayClusterDisplayDevice_Enums.h"
 
 class IDisplayClusterViewport;
 class IDisplayClusterViewportManager;
+class IDisplayClusterViewportPreview;
+class UMaterialInstanceDynamic;
+class UMeshComponent;
 
 /**
  * Warp policy interface
@@ -55,7 +59,7 @@ public:
 
 	/**
 	* Tick function
-	* The positions of movable preview components may be updated on each frame.
+	* The positions of editable preview components may be updated on each frame.
 	* Use this function to handle these updates.
 	* 
 	* @param InViewportManager - Viewport manager interface
@@ -106,12 +110,23 @@ public:
 	{ }
 
 	/**
-	* Ask warp  policy instance if it has any movable mesh based preview
+	* Ask warp  policy instance if it has any Editable mesh based preview
 	* @param InViewport - a owner viewport
 	* @return - True if mesh based preview is available
 	*/
-	virtual bool HasPreviewMovableMesh(IDisplayClusterViewport* InViewport)
+	virtual bool HasPreviewEditableMesh(IDisplayClusterViewport* InViewport)
 	{
 		return false;
 	}
+
+	/** Update Editable mesh material instance parameters.
+	* (this function is called every tick)
+	* 
+	* @param InViewport              - the current viewport
+	* @param InEditableMeshComponent - editable mesh component
+	* @param InMaterialType          - the type of material being requested
+	* @param InMaterialInstance      - material instance
+	*/
+	virtual void OnUpdatePreviewEditableMesh(IDisplayClusterViewportPreview& InViewportPreview, UMeshComponent* InEditableMeshComponent, const EDisplayClusterDisplayDeviceMaterialType InMaterialType, UMaterialInstanceDynamic* InMaterialInstance) const
+	{ }
 };
diff --git a/Engine/Plugins/Runtime/nDisplay/Source/DisplayClusterConfiguration/Public/DisplayClusterConfigurationStrings.h b/Engine/Plugins/Runtime/nDisplay/Source/DisplayClusterConfiguration/Public/DisplayClusterConfigurationStrings.h
--- a/Engine/Plugins/Runtime/nDisplay/Source/DisplayClusterConfiguration/Public/DisplayClusterConfigurationStrings.h
+++ b/Engine/Plugins/Runtime/nDisplay/Source/DisplayClusterConfiguration/Public/DisplayClusterConfigurationStrings.h
@@ -34,6 +34,7 @@ namespace DisplayClusterConfigurationStrings
 		static constexpr const TCHAR* ICVFXCategory                = TEXT("In Camera VFX");
 		static constexpr const TCHAR* ConfigurationCategory        = TEXT("Configuration");
 		static constexpr const TCHAR* PreviewCategory              = TEXT("Editor Preview");
+		static constexpr const TCHAR* PreviewInGameCategory        = TEXT("Preview In Game");
 		static constexpr const TCHAR* AdvancedCategory             = TEXT("Advanced");
 		static constexpr const TCHAR* TextureShareCategory         = TEXT("Texture Share");
 	}
diff --git a/Engine/Plugins/Runtime/nDisplay/Source/DisplayClusterConfiguration/Public/DisplayClusterConfigurationTypes_Viewport.h b/Engine/Plugins/Runtime/nDisplay/Source/DisplayClusterConfiguration/Public/DisplayClusterConfigurationTypes_Viewport.h
--- a/Engine/Plugins/Runtime/nDisplay/Source/DisplayClusterConfiguration/Public/DisplayClusterConfigurationTypes_Viewport.h
+++ b/Engine/Plugins/Runtime/nDisplay/Source/DisplayClusterConfiguration/Public/DisplayClusterConfigurationTypes_Viewport.h
@@ -22,7 +22,6 @@
 
 #include "DisplayClusterConfigurationTypes_Viewport.generated.h"
 
-class UDisplayClusterDisplayDeviceBaseComponent;
 struct FDisplayClusterConfigurationICVFX_StageSettings;
 
 USTRUCT(Blueprintable)
@@ -219,14 +218,12 @@ public:
 	UPROPERTY(EditDefaultsOnly, BlueprintReadWrite, Category = "Configuration")
 	FDisplayClusterConfigurationProjection ProjectionPolicy;
 
-#if WITH_EDITORONLY_DATA
 	/** Select a display device to use during preview. An empty string will use the default display device */
 	UPROPERTY(EditAnywhere, Category = "Preview", meta = (DisplayName = "Display Device"))
 	FString DisplayDeviceName;
-	
+
 	UPROPERTY(EditAnywhere, BlueprintReadWrite, Category = "Preview", meta = (DisplayName = "Preview Frustum"))
 	bool bAllowPreviewFrustumRendering = false;
-#endif
 	
 	/** Define the Viewport 2D coordinates */
 	UPROPERTY(EditAnywhere, BlueprintReadWrite, Category = "Configuration", meta = (DisplayMode = "Compound", FixedAspectRatioProperty = "bFixedAspectRatio"))
diff --git a/Engine/Plugins/Runtime/nDisplay/Source/DisplayClusterConfigurator/Private/DisplayClusterConfiguratorBlueprintEditor.cpp b/Engine/Plugins/Runtime/nDisplay/Source/DisplayClusterConfigurator/Private/DisplayClusterConfiguratorBlueprintEditor.cpp
--- a/Engine/Plugins/Runtime/nDisplay/Source/DisplayClusterConfigurator/Private/DisplayClusterConfiguratorBlueprintEditor.cpp
+++ b/Engine/Plugins/Runtime/nDisplay/Source/DisplayClusterConfigurator/Private/DisplayClusterConfiguratorBlueprintEditor.cpp
@@ -18,7 +18,6 @@
 #include "Blueprints/DisplayClusterBlueprint.h"
 #include "Components/DisplayClusterCameraComponent.h"
 #include "Components/DisplayClusterOriginComponent.h"
-#include "Components/DisplayClusterPreviewComponent.h"
 #include "Components/DisplayClusterScreenComponent.h"
 #include "Components/DisplayClusterXformComponent.h"
 
@@ -161,12 +160,6 @@ FDisplayClusterConfiguratorBlueprintEditor::~FDisplayClusterConfiguratorBlueprin
 		}
 	}
 
-	if (ADisplayClusterRootActor* RootActor = Cast<ADisplayClusterRootActor>(GetPreviewActor()))
-	{
-		RootActor->GetOnPreviewGenerated().Unbind();
-		RootActor->GetOnPreviewDestroyed().Unbind();
-	}
-	
 	ShutdownDCSCSEditors();
 
 	FBlueprintEditorUtils::OnRenameVariableReferencesEvent.Remove(RenameVariableHandle);
@@ -496,7 +489,6 @@ void FDisplayClusterConfiguratorBlueprintEditor::ClusterChanged(bool bStructureC
 
 	if (ADisplayClusterRootActor* Actor = Cast<ADisplayClusterRootActor>(GetPreviewActor()))
 	{
-		Actor->UpdatePreviewComponents();
 		FEditorSupportDelegates::ForcePropertyWindowRebuild.Broadcast(Actor->GetClass());
 	}
 
diff --git a/Engine/Plugins/Runtime/nDisplay/Source/DisplayClusterConfigurator/Private/DisplayClusterConfiguratorCompiler.cpp b/Engine/Plugins/Runtime/nDisplay/Source/DisplayClusterConfigurator/Private/DisplayClusterConfiguratorCompiler.cpp
--- a/Engine/Plugins/Runtime/nDisplay/Source/DisplayClusterConfigurator/Private/DisplayClusterConfiguratorCompiler.cpp
+++ b/Engine/Plugins/Runtime/nDisplay/Source/DisplayClusterConfigurator/Private/DisplayClusterConfiguratorCompiler.cpp
@@ -9,7 +9,7 @@
 #include "Blueprints/DisplayClusterBlueprint.h"
 #include "Blueprints/DisplayClusterBlueprintGeneratedClass.h"
 #include "Components/DisplayClusterCameraComponent.h"
-#include "Components/DisplayDevice/DisplayClusterDisplayDeviceBaseComponent.h"
+#include "Render/DisplayDevice/Components/DisplayClusterDisplayDeviceBaseComponent.h"
 #include "Camera/CameraComponent.h"
 
 #include "Engine/SCS_Node.h"
diff --git a/Engine/Plugins/Runtime/nDisplay/Source/DisplayClusterConfigurator/Private/DisplayClusterConfiguratorModule.cpp b/Engine/Plugins/Runtime/nDisplay/Source/DisplayClusterConfigurator/Private/DisplayClusterConfiguratorModule.cpp
--- a/Engine/Plugins/Runtime/nDisplay/Source/DisplayClusterConfigurator/Private/DisplayClusterConfiguratorModule.cpp
+++ b/Engine/Plugins/Runtime/nDisplay/Source/DisplayClusterConfigurator/Private/DisplayClusterConfiguratorModule.cpp
@@ -277,6 +277,11 @@ void FDisplayClusterConfiguratorModule::RegisterSectionMappings()
 			DisplayClusterConfigurationStrings::categories::PreviewCategory, LOCTEXT("Editor Preview", "Editor Preview"));
 		Section->AddCategory(DisplayClusterConfigurationStrings::categories::PreviewCategory);
 	}
+	{
+		const TSharedRef<FPropertySection> Section = PropertyModule.FindOrCreateSection(ADisplayClusterRootActor::StaticClass()->GetFName(),
+			DisplayClusterConfigurationStrings::categories::PreviewInGameCategory, LOCTEXT("Preview In Game", "Preview In Game"));
+		Section->AddCategory(DisplayClusterConfigurationStrings::categories::PreviewInGameCategory);
+	}
 
 	// ICVFX Component
 	{
diff --git a/Engine/Plugins/Runtime/nDisplay/Source/DisplayClusterConfigurator/Private/Views/Details/Cluster/DisplayClusterConfiguratorViewportDetailsCustomization.cpp b/Engine/Plugins/Runtime/nDisplay/Source/DisplayClusterConfigurator/Private/Views/Details/Cluster/DisplayClusterConfiguratorViewportDetailsCustomization.cpp
--- a/Engine/Plugins/Runtime/nDisplay/Source/DisplayClusterConfigurator/Private/Views/Details/Cluster/DisplayClusterConfiguratorViewportDetailsCustomization.cpp
+++ b/Engine/Plugins/Runtime/nDisplay/Source/DisplayClusterConfigurator/Private/Views/Details/Cluster/DisplayClusterConfiguratorViewportDetailsCustomization.cpp
@@ -8,7 +8,7 @@
 #include "DisplayClusterProjectionStrings.h"
 #include "DisplayClusterConfiguratorLog.h"
 #include "Components/DisplayClusterCameraComponent.h"
-#include "Components/DisplayDevice/DisplayClusterDisplayDeviceBaseComponent.h"
+#include "Render/DisplayDevice/Components/DisplayClusterDisplayDeviceBaseComponent.h"
 #include "Views/Details/Widgets/SDisplayClusterConfiguratorComponentPicker.h"
 #include "Views/Details/Widgets/SDisplayClusterConfigurationSearchableComboBox.h"
 
diff --git a/Engine/Plugins/Runtime/nDisplay/Source/DisplayClusterConfigurator/Private/Views/OutputMapping/EdNodes/DisplayClusterConfiguratorViewportNode.cpp b/Engine/Plugins/Runtime/nDisplay/Source/DisplayClusterConfigurator/Private/Views/OutputMapping/EdNodes/DisplayClusterConfiguratorViewportNode.cpp
--- a/Engine/Plugins/Runtime/nDisplay/Source/DisplayClusterConfigurator/Private/Views/OutputMapping/EdNodes/DisplayClusterConfiguratorViewportNode.cpp
+++ b/Engine/Plugins/Runtime/nDisplay/Source/DisplayClusterConfigurator/Private/Views/OutputMapping/EdNodes/DisplayClusterConfiguratorViewportNode.cpp
@@ -13,7 +13,10 @@
 #include "DisplayClusterConfigurationTypes_ViewportRemap.h"
 #include "DisplayClusterConfiguratorBlueprintEditor.h"
 #include "DisplayClusterRootActor.h"
-#include "Components/DisplayClusterPreviewComponent.h"
+
+#include "Render/Viewport/IDisplayClusterViewportPreview.h"
+#include "Engine/TextureRenderTarget2D.h"
+
 
 void UDisplayClusterConfiguratorViewportNode::Initialize(const FString& InNodeName, int32 InNodeZIndex, UObject* InObject, const TSharedRef<FDisplayClusterConfiguratorBlueprintEditor>& InToolkit)
 {
@@ -225,10 +228,13 @@ UTexture* UDisplayClusterConfiguratorViewportNode::GetPreviewTexture() const
 
 	if (ADisplayClusterRootActor* RootActor = Cast<ADisplayClusterRootActor>(Toolkit->GetPreviewActor()))
 	{
-		UDisplayClusterConfiguratorWindowNode* ParentWindow = GetParentChecked<UDisplayClusterConfiguratorWindowNode>();
-		if (UDisplayClusterPreviewComponent* PreviewComp = RootActor->GetPreviewComponent(ParentWindow->GetNodeName(), GetNodeName()))
+		if (IDisplayClusterViewportManager* ViewportManager = RootActor->GetViewportManager())
 		{
-			return PreviewComp->GetViewportPreviewTexture2D();
+			const FString ViewportId = GetNodeName();
+			if (IDisplayClusterViewport* Viewport = ViewportManager->FindViewport(ViewportId))
+			{
+				return Viewport->GetViewportPreview().GetPreviewTextureRenderTarget2D();
+			}
 		}
 	}
 
diff --git a/Engine/Plugins/Runtime/nDisplay/Source/DisplayClusterConfigurator/Private/Views/Viewport/DisplayClusterConfiguratorSCSEditorViewportClient.cpp b/Engine/Plugins/Runtime/nDisplay/Source/DisplayClusterConfigurator/Private/Views/Viewport/DisplayClusterConfiguratorSCSEditorViewportClient.cpp
--- a/Engine/Plugins/Runtime/nDisplay/Source/DisplayClusterConfigurator/Private/Views/Viewport/DisplayClusterConfiguratorSCSEditorViewportClient.cpp
+++ b/Engine/Plugins/Runtime/nDisplay/Source/DisplayClusterConfigurator/Private/Views/Viewport/DisplayClusterConfiguratorSCSEditorViewportClient.cpp
@@ -12,7 +12,7 @@
 #include "DisplayClusterConfiguratorPropertyUtils.h"
 
 #include "DisplayClusterRootActor.h"
-#include "Components/DisplayClusterPreviewComponent.h"
+#include "ProceduralMeshComponent.h"
 #include "Components/DisplayClusterCameraComponent.h"
 #include "Components/LineBatchComponent.h"
 
@@ -213,7 +213,7 @@ namespace
 		}
 	}
 
-	// Determine whether or not the given node has a parent node that is not the root node, is movable and is selected
+	// Determine whether or not the given node has a parent node that is not the root node, is editable and is selected
 	bool IsMovableParentNodeSelected(const FSubobjectEditorTreeNodePtrType& NodePtr, const TArray<FSubobjectEditorTreeNodePtrType>& SelectedNodes)
 	{
 		if (NodePtr.IsValid())
@@ -560,7 +560,7 @@ bool FDisplayClusterConfiguratorSCSEditorViewportClient::InputWidgetDelta(FViewp
 					{
 						continue;
 					}
-					// Don't allow editing of a root node, inherited SCS node or child node that also has a movable (non-root) parent node selected
+					// Don't allow editing of a root node, inherited SCS node or child node that also has a editable (non-root) parent node selected
 					const bool bCanEdit = GUnrealEd->ComponentVisManager.IsActive() ||
 						(!Data->IsRootComponent() && !IsMovableParentNodeSelected(SelectedNodePtr, SelectedNodes));
 
@@ -1116,7 +1116,6 @@ void FDisplayClusterConfiguratorSCSEditorViewportClient::SyncShowPreview()
 				if (Actor->PreviewNodeId != DisplayClusterConfigurationStrings::gui::preview::PreviewNodeNone)
 				{
 					Actor->PreviewNodeId = DisplayClusterConfigurationStrings::gui::preview::PreviewNodeNone;
-					Actor->UpdatePreviewComponents();
 				}
 			}
 		}
@@ -1306,51 +1305,54 @@ void FDisplayClusterConfiguratorSCSEditorViewportClient::DisplayViewportInformat
 		return;
 	}
 
-	if (ADisplayClusterRootActor* RootActor = Cast<ADisplayClusterRootActor>(BlueprintEditorPtr.Pin()->GetPreviewActor()))
+	ADisplayClusterRootActor* RootActor = Cast<ADisplayClusterRootActor>(BlueprintEditorPtr.Pin()->GetPreviewActor());
+	if (IDisplayClusterViewportManager* ViewportManager = RootActor ? RootActor->GetViewportManager() : nullptr)
 	{
-		TArray<UDisplayClusterPreviewComponent*> PreviewComponents;
-		RootActor->GetComponents(PreviewComponents);
-
-		for (UDisplayClusterPreviewComponent* PreviewComp : PreviewComponents)
+		for (TSharedPtr<IDisplayClusterViewport, ESPMode::ThreadSafe>& ViewportIt : ViewportManager->GetEntireClusterViewports())
 		{
-			if (UDisplayClusterConfigurationViewport* ConfigViewport = PreviewComp->GetViewportConfig())
+			if (UMeshComponent* PreviewMesh = ViewportIt.IsValid() ? ViewportIt->GetViewportPreview().GetPreviewMeshComponent() : nullptr)
 			{
-				if (UMeshComponent* PreviewMesh = PreviewComp->GetPreviewMesh())
+				UFont* Font = GEngine->GetSmallFont();
+				const FLinearColor Color = FLinearColor::Red;
+				const FString DisplayText = ViewportIt->GetId();
+
+				// Display at center of bounding box of mesh with support for various pivot points.
+				FVector WorldLocation = PreviewMesh->GetComponentLocation();
+				if (UStaticMeshComponent* StaticMeshComponent = Cast<UStaticMeshComponent>(PreviewMesh))
 				{
-					UFont* Font = GEngine->GetSmallFont();
-					const FLinearColor Color = FLinearColor::Red;
-					const FString DisplayText = ConfigViewport->GetName();
-					
-					// Display at center of bounding box of mesh with support for various pivot points.
-					FVector WorldLocation = PreviewMesh->GetComponentLocation();
-					if (UStaticMeshComponent* StaticMeshComponent = Cast<UStaticMeshComponent>(PreviewMesh))
-					{
-						if(UStaticMesh* StaticMesh = StaticMeshComponent->GetStaticMesh())
-						{
-							FVector MeshOrigin = StaticMesh->GetBounds().Origin * StaticMeshComponent->GetComponentScale();
-							FRotator ComponentRotation = StaticMeshComponent->GetComponentRotation();
-							WorldLocation += ComponentRotation.RotateVector(MeshOrigin);
-						}
-					}
-					
-					const FPlane Proj = SceneView.Project(WorldLocation);
-					if (Proj.W > 0.0f)
+					if (UStaticMesh* StaticMesh = StaticMeshComponent->GetStaticMesh())
 					{
-						int32 TextWidth, TextHeight;
-						StringSize(Font, TextWidth, TextHeight, *DisplayText);
-						
-						const int32 HalfX = 0.5f * Viewport->GetSizeXY().X;
-						const int32 HalfY = 0.5f * Viewport->GetSizeXY().Y;
-						const int32 XPos = HalfX + (HalfX * Proj.X) - TextWidth / 2;
-						const int32 YPos = HalfY + (HalfY * (Proj.Y * -1));
-						
-						const float DPI = UpdateViewportClientWindowDPIScale();
-						FVector2D AdjustedPosition(XPos / DPI, YPos / DPI);
-						
-						FCanvasTextItem TextItem(AdjustedPosition, FText::FromString(DisplayText), Font, Color);
-						TextItem.Draw(&Canvas);
+						FVector MeshOrigin = StaticMesh->GetBounds().Origin * StaticMeshComponent->GetComponentScale();
+						FRotator ComponentRotation = StaticMeshComponent->GetComponentRotation();
+						WorldLocation += ComponentRotation.RotateVector(MeshOrigin);
 					}
 				}
+				else if (UProceduralMeshComponent* ProceduralMeshComponent = Cast<UProceduralMeshComponent>(PreviewMesh))
+				{
+					// All projection policies except "mesh" and "simple" use ProceduralMeshComponent to create the mesh at runtime.
+					const FBoxSphereBounds& Bounds = ProceduralMeshComponent->Bounds;
+					FVector MeshOrigin = Bounds.Origin * ProceduralMeshComponent->GetComponentScale();
+					FRotator ComponentRotation = ProceduralMeshComponent->GetComponentRotation();
+					WorldLocation += ComponentRotation.RotateVector(MeshOrigin);
+				}
+
+				const FPlane Proj = SceneView.Project(WorldLocation);
+				if (Proj.W > 0.0f)
+				{
+					int32 TextWidth, TextHeight;
+					StringSize(Font, TextWidth, TextHeight, *DisplayText);
+
+					const int32 HalfX = 0.5f * Viewport->GetSizeXY().X;
+					const int32 HalfY = 0.5f * Viewport->GetSizeXY().Y;
+					const int32 XPos = HalfX + (HalfX * Proj.X) - TextWidth / 2;
+					const int32 YPos = HalfY + (HalfY * (Proj.Y * -1));
+
+					const float DPI = UpdateViewportClientWindowDPIScale();
+					FVector2D AdjustedPosition(XPos / DPI, YPos / DPI);
+
+					FCanvasTextItem TextItem(AdjustedPosition, FText::FromString(DisplayText), Font, Color);
+					TextItem.Draw(&Canvas);
+				}
 			}
 		}
 	}
diff --git a/Engine/Plugins/Runtime/nDisplay/Source/DisplayClusterLightCardEditor/Private/Viewport/DisplayClusterLightCardEditorViewportClient.cpp b/Engine/Plugins/Runtime/nDisplay/Source/DisplayClusterLightCardEditor/Private/Viewport/DisplayClusterLightCardEditorViewportClient.cpp
--- a/Engine/Plugins/Runtime/nDisplay/Source/DisplayClusterLightCardEditor/Private/Viewport/DisplayClusterLightCardEditorViewportClient.cpp
+++ b/Engine/Plugins/Runtime/nDisplay/Source/DisplayClusterLightCardEditor/Private/Viewport/DisplayClusterLightCardEditorViewportClient.cpp
@@ -8,7 +8,6 @@
 #include "LightCardTemplates/DisplayClusterLightCardTemplate.h"
 
 #include "Components/DisplayClusterCameraComponent.h"
-#include "Components/DisplayClusterPreviewComponent.h"
 #include "Components/DisplayClusterScreenComponent.h"
 #include "DisplayClusterConfigurationTypes.h"
 #include "DisplayClusterLightCardActor.h"
@@ -24,6 +23,8 @@
 
 #include "IDisplayClusterLightCardExtenderModule.h"
 
+#include "Render/Viewport/IDisplayClusterViewportManager.h"
+
 #include "AudioDevice.h"
 #include "CameraController.h"
 #include "Components/BillboardComponent.h"
@@ -144,7 +145,6 @@ FDisplayClusterLightCardEditorViewportClient::~FDisplayClusterLightCardEditorVie
 	IDisplayClusterScenePreview::Get().DestroyRenderer(PreviewRendererId);
 
 	EndTransaction();
-	UnsubscribeFromRootActor();
 
 	IDisplayClusterLightCardExtenderModule& LightCardExtenderModule = IDisplayClusterLightCardExtenderModule::Get();
 	LightCardExtenderModule.GetOnSequencerTimeChanged().RemoveAll(this);
@@ -175,28 +175,6 @@ void FDisplayClusterLightCardEditorViewportClient::Tick(float DeltaSeconds)
 		PreviewScene->GetWorld()->Tick(IsRealtime() ? LEVELTICK_ViewportsOnly : LEVELTICK_TimeOnly, DeltaSeconds);
 	}
 
-	if (RootActorProxy.IsValid() && RootActorLevelInstance.IsValid())
-	{
-		// Pass the preview render targets from the level instance root actor to the preview root actor
-		UDisplayClusterConfigurationData* Config = RootActorLevelInstance->GetConfigData();
-
-		for (const TPair<FString, TObjectPtr<UDisplayClusterConfigurationClusterNode>>& NodePair : Config->Cluster->Nodes)
-		{
-			const UDisplayClusterConfigurationClusterNode* Node = NodePair.Value;
-			for (const TPair<FString, TObjectPtr<UDisplayClusterConfigurationViewport>>& ViewportPair : Node->Viewports)
-			{
-				UDisplayClusterPreviewComponent* LevelInstancePreviewComp = RootActorLevelInstance->GetPreviewComponent(NodePair.Key, ViewportPair.Key);
-				UDisplayClusterPreviewComponent* PreviewComp = RootActorProxy->GetPreviewComponent(NodePair.Key, ViewportPair.Key);
-
-				if (PreviewComp && LevelInstancePreviewComp)
-				{
-					PreviewComp->SetUseDisplayDevice(false);
-					PreviewComp->SetOverrideTexture(LevelInstancePreviewComp->GetRenderTargetTexturePostProcess());
-				}
-			}
-		}
-	}
-
 	// EditorViewportClient sets the cursor settings based on the state of the built in FWidget, which isn't being used here, so
 	// force a hardware cursor if we are dragging an actor so that the correct mouse cursor shows up
 	switch (InputMode)
@@ -1235,26 +1213,6 @@ void FDisplayClusterLightCardEditorViewportClient::UpdateProxyTransforms(const F
 	}
 }
 
-void FDisplayClusterLightCardEditorViewportClient::SubscribeToRootActor()
-{
-	if (RootActorLevelInstance.IsValid())
-	{
-		const uint8* GenericThis = reinterpret_cast<uint8*>(this);
-		RootActorLevelInstance->SubscribeToPostProcessRenderTarget(GenericThis);
-		RootActorLevelInstance->AddPreviewEnableOverride(GenericThis);
-	}
-}
-
-void FDisplayClusterLightCardEditorViewportClient::UnsubscribeFromRootActor()
-{
-	if (RootActorLevelInstance.IsValid())
-	{
-		const uint8* GenericThis = reinterpret_cast<uint8*>(this);
-		RootActorLevelInstance->UnsubscribeFromPostProcessRenderTarget(GenericThis);
-		RootActorLevelInstance->RemovePreviewEnableOverride(GenericThis);
-	}
-}
-
 void FDisplayClusterLightCardEditorViewportClient::OnSequencerTimeChanged(TWeakPtr<ISequencer> InSequencer)
 {
 	UpdateProxyTransforms();
@@ -1564,6 +1522,8 @@ void FDisplayClusterLightCardEditorViewportClient::UpdatePreviewActor(ADisplayCl
 {
 	DECLARE_SCOPE_CYCLE_COUNTER(TEXT("FDisplayClusterLightCardEditorViewportClient::UpdatePreviewActor"), STAT_UpdatePreviewActor, STATGROUP_NDisplayLightCardEditor);
 	
+	check(IsInGameThread());
+
 	if (!StageActor && ((!bForce && RootActor == RootActorLevelInstance.GetEvenIfUnreachable()) ||
 		(ProxyTypesRefreshing.Contains(ProxyType) || ProxyTypesRefreshing.Contains(EDisplayClusterLightCardEditorProxyType::All))))
 	{
@@ -1601,7 +1561,6 @@ void FDisplayClusterLightCardEditorViewportClient::UpdatePreviewActor(ADisplayCl
 	if (RootActor == nullptr)
 	{
 		DestroyProxies(ProxyType);
-		UnsubscribeFromRootActor();
 		RootActorLevelInstance.Reset();
 		
 		Finalize();
@@ -1642,28 +1601,29 @@ void FDisplayClusterLightCardEditorViewportClient::UpdatePreviewActor(ADisplayCl
 				return;
 			}
 
-			if (RootActorLevelInstance.IsValid() && RootActorLevelInstance.Get() != RootActor)
-			{
-				UnsubscribeFromRootActor();
-			}
-			
 			RootActorLevelInstance = RootActorPtr;
 
-			SubscribeToRootActor();
-
 			TArray<TObjectPtr<AActor>> ActorProxiesCreated;
 			
 			if (ProxyType == EDisplayClusterLightCardEditorProxyType::All ||
 				ProxyType == EDisplayClusterLightCardEditorProxyType::RootActor)
 			{
+				if (!RootActorProxy.IsValid())
 				{
 					FObjectDuplicationParameters DupeActorParameters(RootActorPtr.Get(), PreviewWorld->GetCurrentLevel());
 					DupeActorParameters.FlagMask = RF_AllFlags & ~(RF_ArchetypeObject | RF_Transactional); // Keeps archetypes correct in config data.
 					DupeActorParameters.PortFlags = PPF_DuplicateVerbatim;
-			
+
 					RootActorProxy = CastChecked<ADisplayClusterRootActor>(StaticDuplicateObjectEx(DupeActorParameters));
 				}
-			
+
+				// Use root actor from scene to render
+				if (IDisplayClusterViewportManager* ViewportManager = RootActorProxy.IsValid() ? RootActorProxy->GetOrCreateViewportManager() : nullptr)
+				{
+					// Using DCRA from the scene for rendering
+					ViewportManager->GetConfiguration().SetRootActor(EDisplayClusterRootActorType::Scene | EDisplayClusterRootActorType::Configuration, RootActorPtr.Get());
+				}
+
 				PreviewWorld->GetCurrentLevel()->AddLoadedActor(RootActorProxy.Get());
 
 				// Draw the geometry map for the proxy stage actor immediately to avoid a race condition where the geometry map could render
@@ -1676,12 +1636,18 @@ void FDisplayClusterLightCardEditorViewportClient::UpdatePreviewActor(ADisplayCl
 
 				ProjectionOriginComponent = FindProjectionOriginComponent(RootActorProxy.Get());
 
-				RootActorProxy->UpdatePreviewComponents();
-				RootActorProxy->EnableEditorRender(false);
-
 				ProjectionHelper->SetLevelInstanceRootActor(*RootActorLevelInstance);
 				ProjectionHelper->SetEditorViewportClient(AsShared());
 
+				// Setup custom preview settings for proxy DCRA:
+				RootActorProxy->bPreviewICVFXFrustums = false;
+				RootActorProxy->bEnablePreviewTechvis = false;
+				RootActorProxy->bPreviewEnablePostProcess = true;
+				RootActorProxy->bPreviewEnable = true;
+				RootActorProxy->bPreviewICVFXFrustums = false;
+
+				//@Todo: In the case when RootActorProxy is not recreated each time it is changed:
+				// the DCRA  Copy cfg from RootActorLevelInstance to RootActorProxy every tick (do not recreate local DCRA)
 				if (UDisplayClusterConfigurationData* ProxyConfig = RootActorProxy->GetConfigData())
 				{
 					// Disable lightcards so that it doesn't try to update the ones in the level instance world.
diff --git a/Engine/Plugins/Runtime/nDisplay/Source/DisplayClusterLightCardEditor/Private/Viewport/DisplayClusterLightCardEditorViewportClient.h b/Engine/Plugins/Runtime/nDisplay/Source/DisplayClusterLightCardEditor/Private/Viewport/DisplayClusterLightCardEditorViewportClient.h
--- a/Engine/Plugins/Runtime/nDisplay/Source/DisplayClusterLightCardEditor/Private/Viewport/DisplayClusterLightCardEditorViewportClient.h
+++ b/Engine/Plugins/Runtime/nDisplay/Source/DisplayClusterLightCardEditor/Private/Viewport/DisplayClusterLightCardEditorViewportClient.h
@@ -300,12 +300,6 @@ private:
 	/** Update the transforms for the specific proxy */
 	void UpdateProxyTransforms(const FActorProxy& InActorProxy);
 
-	/** Subscribe post process and preview hooks to the root actor */
-	void SubscribeToRootActor();
-	
-	/** Unsubscribe post process and preview hooks from the root actor */
-	void UnsubscribeFromRootActor();
-
 	/** When the time in a sequencer has changed */
 	void OnSequencerTimeChanged(TWeakPtr<class ISequencer> InSequencer);
 	
diff --git a/Engine/Plugins/Runtime/nDisplay/Source/DisplayClusterMoviePipeline/Private/DisplayClusterMoviePipelineViewportPass.cpp b/Engine/Plugins/Runtime/nDisplay/Source/DisplayClusterMoviePipeline/Private/DisplayClusterMoviePipelineViewportPass.cpp
--- a/Engine/Plugins/Runtime/nDisplay/Source/DisplayClusterMoviePipeline/Private/DisplayClusterMoviePipelineViewportPass.cpp
+++ b/Engine/Plugins/Runtime/nDisplay/Source/DisplayClusterMoviePipeline/Private/DisplayClusterMoviePipelineViewportPass.cpp
@@ -485,9 +485,11 @@ bool UDisplayClusterMoviePipelineViewportPassBase::InitializeDisplayCluster()
 
 			if (IDisplayClusterViewportManager* ViewportManager = DCRootActor->GetOrCreateViewportManager())
 			{
+				// Don't use preview settings for MRQ rendering (same as for game)
+				ViewportManager->GetConfiguration().SetPreviewSettings(FDisplayClusterViewport_PreviewSettings());
+
 				// Update local node viewports (update\create\delete) and build new render frame
-				const EDisplayClusterRenderFrameMode RenderFrameMode = EDisplayClusterRenderFrameMode::Mono;
-				if (ViewportManager->GetConfiguration().UpdateConfigurationForViewportsList(RenderFrameMode, DisplayClusterViewports))
+				if (ViewportManager->GetConfiguration().UpdateConfigurationForViewportsList(EDisplayClusterRenderFrameMode::MRQ_Mono, GetWorld(), DisplayClusterViewports))
 				{
 					return true;
 				}
@@ -501,7 +503,6 @@ bool UDisplayClusterMoviePipelineViewportPassBase::InitializeDisplayCluster()
 	return false;
 }
 
-
 bool UDisplayClusterMoviePipelineViewportPassBase::GetViewportId(int32 InViewportIndex, FString& OutViewportId) const
 {
 	if (DCRootActor && InViewportIndex >= 0 && InViewportIndex < DisplayClusterViewports.Num())
@@ -554,8 +555,6 @@ IDisplayClusterViewport* UDisplayClusterMoviePipelineViewportPassBase::GetAndCal
 			DCViewport->SetRenderSettings(RenderSettings);
 		}
 
-		ViewportManager->GetConfiguration().SetCurrentWorld(CurrentWorld);
-
 		FDisplayClusterRenderFrame RenderFrame;
 		if (ViewportManager->BeginNewFrame(GameViewportClient->Viewport, RenderFrame))
 		{
@@ -569,13 +568,10 @@ IDisplayClusterViewport* UDisplayClusterMoviePipelineViewportPassBase::GetAndCal
 			OutView.ViewRotation = ViewInfo.Rotation;
 			OutView.ViewLocation = ViewInfo.Location;
 			
-			const AWorldSettings* WorldSettings = CurrentWorld->GetWorldSettings();
-			const float InWorldToMeters = (WorldSettings) ? WorldSettings->WorldToMeters : 100.f;
-
 			bool bResult = false;
 			// Obtaining the offset of the stereo eye and the values of the projection clipping plane for the given viewport was moved inside CalculateView().
 			// Perform view calculations on a policy side
-			if (DCViewport->CalculateView(InContextNum, OutView.ViewLocation, OutView.ViewRotation, InWorldToMeters))
+			if (DCViewport->CalculateView(InContextNum, OutView.ViewLocation, OutView.ViewRotation, ViewportManager->GetConfiguration().GetWorldToMeters()))
 			{
 				OutView.ProjectionMatrix = FMatrix::Identity;
 
