nDisplay: Implemented a new preview rendering architecture:

diff --git a/Engine/Plugins/Runtime/nDisplay/Source/DisplayClusterProjection/Private/Policy/MPCDI/DisplayClusterProjectionMPCDIPolicy.cpp b/Engine/Plugins/Runtime/nDisplay/Source/DisplayClusterProjection/Private/Policy/MPCDI/DisplayClusterProjectionMPCDIPolicy.cpp
--- a/Engine/Plugins/Runtime/nDisplay/Source/DisplayClusterProjection/Private/Policy/MPCDI/DisplayClusterProjectionMPCDIPolicy.cpp
+++ b/Engine/Plugins/Runtime/nDisplay/Source/DisplayClusterProjection/Private/Policy/MPCDI/DisplayClusterProjectionMPCDIPolicy.cpp
@@ -163,6 +163,11 @@ bool FDisplayClusterProjectionMPCDIPolicy::HandleStartScene(IDisplayClusterViewp
 	// Finally, initialize internal views data container
 	WarpBlendContexts.AddDefaulted(2);
 
+	if (WarpBlendInterface.IsValid())
+	{
+		WarpBlendInterface->HandleStartScene(InViewport);
+	}
+
 	return true;
 }
 
@@ -170,6 +175,11 @@ void FDisplayClusterProjectionMPCDIPolicy::HandleEndScene(IDisplayClusterViewpor
 {
 	check(IsInGameThread());
 
+	if (WarpBlendInterface.IsValid())
+	{
+		WarpBlendInterface->HandleEndScene(InViewport);
+	}
+
 	ImplRelease();
 }
 
@@ -179,6 +189,9 @@ void FDisplayClusterProjectionMPCDIPolicy::ImplRelease()
 
 	WarpBlendInterface.Reset();
 	WarpBlendContexts.Empty();
+
+	PreviewMeshComponentRef.ResetSceneComponent();
+	PreviewEditableMeshComponentRef.ResetSceneComponent();
 }
 
 bool FDisplayClusterProjectionMPCDIPolicy::GetWarpBlendInterface(TSharedPtr<IDisplayClusterWarpBlend, ESPMode::ThreadSafe>& OutWarpBlendInterface) const
@@ -463,6 +476,7 @@ UMeshComponent* FDisplayClusterProjectionMPCDIPolicy::GetOrCreatePreviewMeshComp
 	// If we have already created a preview mesh component before, return that component
 	if (UMeshComponent* ExistsPreviewMeshComp = Cast<UMeshComponent>(PreviewMeshComponentRef.GetOrFindSceneComponent()))
 	{
+		bOutIsRootActorComponent = bIsRootActorHasPreviewMeshComponent;
 		return ExistsPreviewMeshComp;
 	}
 
@@ -471,6 +485,7 @@ UMeshComponent* FDisplayClusterProjectionMPCDIPolicy::GetOrCreatePreviewMeshComp
 	{
 		// Store reference to mesh component
 		PreviewMeshComponentRef.SetSceneComponent(PreviewMeshComp);
+		bIsRootActorHasPreviewMeshComponent = bOutIsRootActorComponent;
 
 		return PreviewMeshComp;
 	}
@@ -478,52 +493,51 @@ UMeshComponent* FDisplayClusterProjectionMPCDIPolicy::GetOrCreatePreviewMeshComp
 	return nullptr;
 }
 
-USceneComponent* const FDisplayClusterProjectionMPCDIPolicy::GetPreviewMovableMeshOriginComponent(IDisplayClusterViewport* InViewport) const
+USceneComponent* const FDisplayClusterProjectionMPCDIPolicy::GetPreviewEditableMeshOriginComponent(IDisplayClusterViewport* InViewport) const
 {
-	// Note: currently for the movable mesh component we expect it to be used only in the scene,
+	// Note: currently for the Editable mesh component we expect it to be used only in the scene,
 	// so we use the root component from the scene all the time.
 	// But if other use cases are found, we need to refine this logic.
 
 	return GetOriginComponent();
 }
 
-bool FDisplayClusterProjectionMPCDIPolicy::HasPreviewMovableMesh(IDisplayClusterViewport* InViewport)
+bool FDisplayClusterProjectionMPCDIPolicy::HasPreviewEditableMesh(IDisplayClusterViewport* InViewport)
 {
-	if (InViewport)
+	if (bIsPreviewMeshEnabled && InViewport)
 	{
-		// The movable preview grid is a feature for the warp policy, so we must request permission to use it.
-		const bool bWarpPolicyUseMovableMesh = WarpPolicyInterface.IsValid() && WarpPolicyInterface->HasPreviewMovableMesh(InViewport);
-		if (bIsPreviewMeshEnabled && bWarpPolicyUseMovableMesh)
+		// The editable preview Editable mesh is a feature for the warp policy, so we must request permission to use it.
+		if (WarpPolicyInterface.IsValid() && WarpPolicyInterface->HasPreviewEditableMesh(InViewport))
 		{
 			return true;
 		}
 	}
 
-	PreviewMovableMeshComponentRef.ResetSceneComponent();
+	PreviewEditableMeshComponentRef.ResetSceneComponent();
 
 	return false;
 }
 
-UMeshComponent* FDisplayClusterProjectionMPCDIPolicy::GetOrCreatePreviewMovableMeshComponent(IDisplayClusterViewport* InViewport)
+UMeshComponent* FDisplayClusterProjectionMPCDIPolicy::GetOrCreatePreviewEditableMeshComponent(IDisplayClusterViewport* InViewport)
 {
-	if (!HasPreviewMovableMesh(InViewport))
+	if (!HasPreviewEditableMesh(InViewport))
 	{
 		return nullptr;
 	}
 
 	// If we have already created a preview mesh component before, return that component
-	if (UMeshComponent* ExistsPreviewMovableMeshComp = Cast<UMeshComponent>(PreviewMovableMeshComponentRef.GetOrFindSceneComponent()))
+	if (UMeshComponent* ExistsPreviewEditableMeshComp = Cast<UMeshComponent>(PreviewEditableMeshComponentRef.GetOrFindSceneComponent()))
 	{
-		return ExistsPreviewMovableMeshComp;
+		return ExistsPreviewEditableMeshComp;
 	}
 
 	// Get a new one
-	if (UMeshComponent* PreviewMovableMeshComp = WarpBlendInterface.IsValid() ? WarpBlendInterface->GetOrCreatePreviewMovableMeshComponent(InViewport) : nullptr)
+	if (UMeshComponent* PreviewEditableMeshComp = WarpBlendInterface.IsValid() ? WarpBlendInterface->GetOrCreatePreviewEditableMeshComponent(InViewport) : nullptr)
 	{
 		// Store reference to mesh component
-		PreviewMovableMeshComponentRef.SetSceneComponent(PreviewMovableMeshComp);
+		PreviewEditableMeshComponentRef.SetSceneComponent(PreviewEditableMeshComp);
 
-		return PreviewMovableMeshComp;
+		return PreviewEditableMeshComp;
 	}
 
 	return nullptr;
diff --git a/Engine/Plugins/Runtime/nDisplay/Source/DisplayClusterProjection/Private/Policy/MPCDI/DisplayClusterProjectionMPCDIPolicy.h b/Engine/Plugins/Runtime/nDisplay/Source/DisplayClusterProjection/Private/Policy/MPCDI/DisplayClusterProjectionMPCDIPolicy.h
--- a/Engine/Plugins/Runtime/nDisplay/Source/DisplayClusterProjection/Private/Policy/MPCDI/DisplayClusterProjectionMPCDIPolicy.h
+++ b/Engine/Plugins/Runtime/nDisplay/Source/DisplayClusterProjection/Private/Policy/MPCDI/DisplayClusterProjectionMPCDIPolicy.h
@@ -60,9 +60,9 @@ public:
 	virtual bool HasPreviewMesh(IDisplayClusterViewport* InViewport) override;
 	virtual UMeshComponent* GetOrCreatePreviewMeshComponent(IDisplayClusterViewport* InViewport, bool& bOutIsRootActorComponent) override;
 
-	virtual bool HasPreviewMovableMesh(IDisplayClusterViewport* InViewport) override;
-	virtual UMeshComponent* GetOrCreatePreviewMovableMeshComponent(IDisplayClusterViewport* InViewport) override;
-	virtual USceneComponent* const GetPreviewMovableMeshOriginComponent(IDisplayClusterViewport* InViewport) const override;
+	virtual bool HasPreviewEditableMesh(IDisplayClusterViewport* InViewport) override;
+	virtual UMeshComponent* GetOrCreatePreviewEditableMeshComponent(IDisplayClusterViewport* InViewport) override;
+	virtual USceneComponent* const GetPreviewEditableMeshOriginComponent(IDisplayClusterViewport* InViewport) const override;
 
 	//~~End IDisplayClusterProjectionPolicy
 
@@ -87,6 +87,12 @@ protected:
 	bool bIsPreviewMeshEnabled = false;
 
 private:
+	// Stored value of the preview mesh
 	FDisplayClusterSceneComponentRef PreviewMeshComponentRef;
-	FDisplayClusterSceneComponentRef PreviewMovableMeshComponentRef;
+
+	// Stored value of the preview meshes belonging flag. True if this component exists in DCRA and cannot be deleted with preview.
+	bool bIsRootActorHasPreviewMeshComponent = false;
+
+	// Stored value of the editable preview mesh
+	FDisplayClusterSceneComponentRef PreviewEditableMeshComponentRef;
 };
diff --git a/Engine/Plugins/Runtime/nDisplay/Source/DisplayClusterProjection/Private/Policy/Mesh/DisplayClusterProjectionMeshPolicy.cpp b/Engine/Plugins/Runtime/nDisplay/Source/DisplayClusterProjection/Private/Policy/Mesh/DisplayClusterProjectionMeshPolicy.cpp
--- a/Engine/Plugins/Runtime/nDisplay/Source/DisplayClusterProjection/Private/Policy/Mesh/DisplayClusterProjectionMeshPolicy.cpp
+++ b/Engine/Plugins/Runtime/nDisplay/Source/DisplayClusterProjection/Private/Policy/Mesh/DisplayClusterProjectionMeshPolicy.cpp
@@ -144,6 +144,11 @@ bool FDisplayClusterProjectionMeshPolicy::HandleStartScene(IDisplayClusterViewpo
 	// Finally, initialize internal views data container
 	WarpBlendContexts.AddDefaulted(2);
 
+	if (WarpBlendInterface.IsValid())
+	{
+		WarpBlendInterface->HandleStartScene(InViewport);
+	}
+
 	return true;
 }
 
diff --git a/Engine/Plugins/Runtime/nDisplay/Source/DisplayClusterProjection/Private/Policy/VIOSO/Windows/DisplayClusterProjectionVIOSOPolicy.cpp b/Engine/Plugins/Runtime/nDisplay/Source/DisplayClusterProjection/Private/Policy/VIOSO/Windows/DisplayClusterProjectionVIOSOPolicy.cpp
--- a/Engine/Plugins/Runtime/nDisplay/Source/DisplayClusterProjection/Private/Policy/VIOSO/Windows/DisplayClusterProjectionVIOSOPolicy.cpp
+++ b/Engine/Plugins/Runtime/nDisplay/Source/DisplayClusterProjection/Private/Policy/VIOSO/Windows/DisplayClusterProjectionVIOSOPolicy.cpp
@@ -84,6 +84,8 @@ void FDisplayClusterProjectionVIOSOPolicy::ImplRelease()
 {
 	ReleaseOriginComponent();
 
+	PreviewMeshComponentRef.ResetSceneComponent();
+
 	// Destroy VIOSO for all views
 	FScopeLock lock(&DllAccessCS);
 	Views.Reset();
@@ -210,11 +212,21 @@ bool FDisplayClusterProjectionVIOSOPolicy::ImplApplyWarpBlend_RenderThread(FRHIC
 	return true;
 }
 
-UMeshComponent* FDisplayClusterProjectionVIOSOPolicy::GetOrCreatePreviewMeshComponent(IDisplayClusterViewport* InViewport, bool& bOutIsRootActorComponent)
+bool FDisplayClusterProjectionVIOSOPolicy::HasPreviewMesh(IDisplayClusterViewport* InViewport)
 {
-	check(IsInGameThread());
-
 	if (!ViosoConfigData.bIsPreviewMeshEnabled || Views.IsEmpty() || !Views[0].IsValid() || !Views[0]->IsWarperInterfaceValid())
+	{
+		PreviewMeshComponentRef.ResetSceneActor();
+
+		return false;
+	}
+
+	return true;
+}
+
+UMeshComponent* FDisplayClusterProjectionVIOSOPolicy::GetOrCreatePreviewMeshComponent(IDisplayClusterViewport* InViewport, bool& bOutIsRootActorComponent)
+{
+	if (!HasPreviewMesh(InViewport))
 	{
 		return nullptr;
 	}
@@ -222,6 +234,12 @@ UMeshComponent* FDisplayClusterProjectionVIOSOPolicy::GetOrCreatePreviewMeshComp
 	// Create a new DCRA mesh component
 	bOutIsRootActorComponent = false;
 
+	// If we have already created a preview mesh component before, return that component
+	if (UMeshComponent* ExistsPreviewMeshComp = Cast<UMeshComponent>(PreviewMeshComponentRef.GetOrFindSceneComponent()))
+	{
+		return ExistsPreviewMeshComp;
+	}
+
 	USceneComponent* OriginComp = GetPreviewMeshOriginComponent(InViewport);
 	TSharedPtr<FDisplayClusterProjectionVIOSOGeometryExportData, ESPMode::ThreadSafe> GeometryExportData = FDisplayClusterProjectionVIOSOGeometryExportData::Create(VIOSOLibrary, ViosoConfigData);
 	if (OriginComp && GeometryExportData.IsValid())
@@ -244,6 +262,9 @@ UMeshComponent* FDisplayClusterProjectionVIOSOPolicy::GetOrCreatePreviewMeshComp
 			// Because of "nDisplay.render.show.visualizationcomponents" we need extra flag to exclude this geometry from render
 			MeshComp->SetHiddenInGame(true);
 
+			// Store reference to mesh component
+			PreviewMeshComponentRef.SetSceneComponent(MeshComp);
+
 			return MeshComp;
 		}
 	}
diff --git a/Engine/Plugins/Runtime/nDisplay/Source/DisplayClusterProjection/Private/Policy/VIOSO/Windows/DisplayClusterProjectionVIOSOPolicy.h b/Engine/Plugins/Runtime/nDisplay/Source/DisplayClusterProjection/Private/Policy/VIOSO/Windows/DisplayClusterProjectionVIOSOPolicy.h
--- a/Engine/Plugins/Runtime/nDisplay/Source/DisplayClusterProjection/Private/Policy/VIOSO/Windows/DisplayClusterProjectionVIOSOPolicy.h
+++ b/Engine/Plugins/Runtime/nDisplay/Source/DisplayClusterProjection/Private/Policy/VIOSO/Windows/DisplayClusterProjectionVIOSOPolicy.h
@@ -16,9 +16,7 @@ public:
 	virtual ~FDisplayClusterProjectionVIOSOPolicy();
 
 public:
-	//////////////////////////////////////////////////////////////////////////////////////////////
-	// IDisplayClusterProjectionPolicy
-	//////////////////////////////////////////////////////////////////////////////////////////////
+	//~~BEGIN IDisplayClusterProjectionPolicy
 	virtual const FString& GetType() const override;
 
 	virtual bool HandleStartScene(IDisplayClusterViewport* InViewport) override;
@@ -41,6 +39,11 @@ public:
 		return true;
 	}
 
+	virtual bool HasPreviewMesh(IDisplayClusterViewport* InViewport) override;
+	virtual class UMeshComponent* GetOrCreatePreviewMeshComponent(IDisplayClusterViewport* InViewport, bool& bOutIsRootActorComponent) override;
+
+	//~~END IDisplayClusterProjectionPolicy
+
 protected:
 	bool ImplApplyWarpBlend_RenderThread(FRHICommandListImmediate& RHICmdList, const IDisplayClusterViewportProxy* InViewportProxy);
 	void ImplRelease();
@@ -53,16 +56,7 @@ protected:
 
 	TArray<TSharedPtr<FDisplayClusterProjectionVIOSOPolicyViewData, ESPMode::ThreadSafe>> Views;
 
-	FCriticalSection DllAccessCS;
+	FDisplayClusterSceneComponentRef PreviewMeshComponentRef;
 
-protected:
-	//////////////////////////////////////////////////////////////////////////////////////////////
-	// IDisplayClusterProjectionPolicyPreview
-	//////////////////////////////////////////////////////////////////////////////////////////////
-	virtual bool HasPreviewMesh(IDisplayClusterViewport* InViewport) override
-	{
-		return true;
-	}
-
-	virtual class UMeshComponent* GetOrCreatePreviewMeshComponent(IDisplayClusterViewport* InViewport, bool& bOutIsRootActorComponent) override;
+	FCriticalSection DllAccessCS;
 };
diff --git a/Engine/Plugins/Runtime/nDisplay/Source/DisplayClusterScenePreview/Private/DisplayClusterScenePreviewModule.cpp b/Engine/Plugins/Runtime/nDisplay/Source/DisplayClusterScenePreview/Private/DisplayClusterScenePreviewModule.cpp
--- a/Engine/Plugins/Runtime/nDisplay/Source/DisplayClusterScenePreview/Private/DisplayClusterScenePreviewModule.cpp
+++ b/Engine/Plugins/Runtime/nDisplay/Source/DisplayClusterScenePreview/Private/DisplayClusterScenePreviewModule.cpp
@@ -5,7 +5,6 @@
 #include "Blueprints/DisplayClusterBlueprintLib.h"
 #include "CanvasTypes.h"
 #include "Components/DisplayClusterScreenComponent.h"
-#include "Components/DisplayClusterPreviewComponent.h"
 #include "DisplayClusterChromakeyCardActor.h"
 #include "DisplayClusterLightCardActor.h"
 #include "DisplayClusterRootActor.h"
@@ -320,47 +319,11 @@ bool FDisplayClusterScenePreviewModule::InternalRenderImmediate(FRendererConfig&
 		AutoPopulateScene(RendererConfig);
 	}
 
-#if WITH_EDITOR
-	TMap<UDisplayClusterPreviewComponent*, UTexture*> PreviewComponentOverrideTextures;
-	const bool bUsePostProcessTexture = RendererConfig.bUsePostProcessTexture && (RenderSettings.RenderType == EDisplayClusterMeshProjectionOutput::Color);
-
-	// Update the preview components with the post-processed texture
-	if (bUsePostProcessTexture)
-	{
-		UDisplayClusterConfigurationData* Config = RendererConfig.RootActor->GetConfigData();
-
-		for (const TPair<FString, TObjectPtr<UDisplayClusterConfigurationClusterNode>>& NodePair : Config->Cluster->Nodes)
-		{
-			const UDisplayClusterConfigurationClusterNode* Node = NodePair.Value;
-			for (const TPair<FString, TObjectPtr<UDisplayClusterConfigurationViewport>>& ViewportPair : Node->Viewports)
-			{
-				UDisplayClusterPreviewComponent* PreviewComp = RendererConfig.RootActor->GetPreviewComponent(NodePair.Key, ViewportPair.Key);
-				if (PreviewComp)
-				{
-					PreviewComponentOverrideTextures.Add(PreviewComp, PreviewComp->GetOverrideTexture());
-					PreviewComp->SetOverrideTexture(PreviewComp->GetRenderTargetTexturePostProcess());
-				}
-			}
-		}
-	}
-#endif
-
 	// Push any deferred render state updates to ensure that light card positions, preview meshes modified above, etc. are up to date
 	World->SendAllEndOfFrameUpdates();
 
 	RendererConfig.Renderer->Render(&Canvas, World->Scene, RenderSettings);
 
-#if WITH_EDITOR
-	// Remove the override texture so the actor renders normally
-	if (bUsePostProcessTexture)
-	{
-		for (const TPair<UDisplayClusterPreviewComponent*, UTexture*>& PreviewPair : PreviewComponentOverrideTextures)
-		{
-			PreviewPair.Key->SetOverrideTexture(PreviewPair.Value);
-		}
-	}
-#endif
-
 	return true;
 }
 
@@ -415,17 +378,6 @@ void FDisplayClusterScenePreviewModule::RegisterRootActorEvents(ADisplayClusterR
 
 	const uint8* GenericThis = reinterpret_cast<uint8*>(this);
 
-	// Register/unregister actors to use post-processing for preview renders so their lighting looks correct,
-	// and add/remove preview override so our preview renders have the latest nDisplay preview texture.
-	Actor->UnsubscribeFromPostProcessRenderTarget(GenericThis);
-	Actor->RemovePreviewEnableOverride(GenericThis);
-
-	if (bShouldRegister)
-	{
-		Actor->SubscribeToPostProcessRenderTarget(GenericThis);
-		Actor->AddPreviewEnableOverride(GenericThis);
-	}
-
 	// Register/unregister for Blueprint events
 	if (UBlueprint* Blueprint = UBlueprint::GetBlueprintFromClass(Actor->GetClass()))
 	{
diff --git a/Engine/Plugins/Runtime/nDisplay/Source/DisplayClusterShaders/Private/Shaders/DisplayClusterShadersWarpblend_ICVFX.cpp b/Engine/Plugins/Runtime/nDisplay/Source/DisplayClusterShaders/Private/Shaders/DisplayClusterShadersWarpblend_ICVFX.cpp
--- a/Engine/Plugins/Runtime/nDisplay/Source/DisplayClusterShaders/Private/Shaders/DisplayClusterShadersWarpblend_ICVFX.cpp
+++ b/Engine/Plugins/Runtime/nDisplay/Source/DisplayClusterShaders/Private/Shaders/DisplayClusterShadersWarpblend_ICVFX.cpp
@@ -698,7 +698,7 @@ public:
 
 			const FMatrix WorldToCamera = CameraTranslationMatrix.Inverse();
 			const FMatrix UVMatrix = WorldToCamera * Game2Render * Camera.ViewProjection.PrjMatrix;
-			const FMatrix InnerCameraProjectionMatrix = UVMatrix * WarpBlendParameters.Context.TextureMatrix * WarpBlendParameters.Context.RegionMatrix;
+			const FMatrix InnerCameraProjectionMatrix = UVMatrix * WarpBlendParameters.Context.TextureMatrix;
 
 			RenderPassData.PSParameters.OverlappedInnerCamerasProjectionMatrices[CameraIndex] = FMatrix44f(InnerCameraProjectionMatrix);
 			RenderPassData.PSParameters.OverlappedInnerCameraSoftEdges[CameraIndex] = GDisplayClusterShadersICVFXOverlapInnerFrustumSoftEdges ? FVector4f(Camera.SoftEdge) : FVector4f(0, 0, 0, 0);
@@ -772,7 +772,7 @@ public:
 
 		const FMatrix WorldToCamera = CameraTranslationMatrix.Inverse();
 		const FMatrix UVMatrix = WorldToCamera * Game2Render * Camera.ViewProjection.PrjMatrix;
-		const FMatrix InnerCameraProjectionMatrix = UVMatrix * WarpBlendParameters.Context.TextureMatrix * WarpBlendParameters.Context.RegionMatrix;
+		const FMatrix InnerCameraProjectionMatrix = UVMatrix * WarpBlendParameters.Context.TextureMatrix;
 
 		RenderPassData.PSParameters.InnerCameraTexture = Camera.Resource.Texture;
 		RenderPassData.PSParameters.InnerCameraSampler = TStaticSamplerState<SF_Trilinear, AM_Clamp, AM_Clamp, AM_Clamp>::GetRHI();
diff --git a/Engine/Plugins/Runtime/nDisplay/Source/DisplayClusterWarp/Private/Components/DisplayClusterInFrustumFitCameraComponent.cpp b/Engine/Plugins/Runtime/nDisplay/Source/DisplayClusterWarp/Private/Components/DisplayClusterInFrustumFitCameraComponent.cpp
--- a/Engine/Plugins/Runtime/nDisplay/Source/DisplayClusterWarp/Private/Components/DisplayClusterInFrustumFitCameraComponent.cpp
+++ b/Engine/Plugins/Runtime/nDisplay/Source/DisplayClusterWarp/Private/Components/DisplayClusterInFrustumFitCameraComponent.cpp
@@ -22,7 +22,6 @@
 #include "Render/Warp/IDisplayClusterWarpPolicy.h"
 
 #include "DisplayClusterRootActor.h"
-#include "Components/DisplayClusterPreviewComponent.h"
 
 namespace UE::DisplayClusterWarp::ViewPointComponent
 {
diff --git a/Engine/Plugins/Runtime/nDisplay/Source/DisplayClusterWarp/Private/Policy/InFrustumFit/DisplayClusterWarpInFrustumFitPolicy.cpp b/Engine/Plugins/Runtime/nDisplay/Source/DisplayClusterWarp/Private/Policy/InFrustumFit/DisplayClusterWarpInFrustumFitPolicy.cpp
--- a/Engine/Plugins/Runtime/nDisplay/Source/DisplayClusterWarp/Private/Policy/InFrustumFit/DisplayClusterWarpInFrustumFitPolicy.cpp
+++ b/Engine/Plugins/Runtime/nDisplay/Source/DisplayClusterWarp/Private/Policy/InFrustumFit/DisplayClusterWarpInFrustumFitPolicy.cpp
@@ -5,16 +5,15 @@
 #include "DisplayClusterRootActor.h"
 #include "Containers/DisplayClusterWarpEye.h"
 
+#include "Render/Viewport/IDisplayClusterViewportPreview.h"
 #include "Render/Viewport/IDisplayClusterViewport.h"
 #include "Render/Projection/IDisplayClusterProjectionPolicy.h"
 
 #include "IDisplayClusterWarpBlend.h"
-
 #include "DisplayClusterWarpStrings.h"
 
 #include "Components/DisplayClusterCameraComponent.h"
 #include "Components/DisplayClusterInFrustumFitCameraComponent.h"
-#include "Components/DisplayClusterPreviewComponent.h"
 #include "Components/MeshComponent.h"
 #include "Materials/MaterialInstanceDynamic.h"
 
@@ -161,78 +160,32 @@ void FDisplayClusterWarpInFrustumFitPolicy::HandleNewFrame(const TArray<TSharedP
 void FDisplayClusterWarpInFrustumFitPolicy::Tick(IDisplayClusterViewportManager* InViewportManager, float DeltaSeconds)
 {
 #if WITH_EDITOR
-	UDisplayClusterInFrustumFitCameraComponent* SceneCameraComponent = nullptr;
-	for (const TSharedPtr<IDisplayClusterViewport, ESPMode::ThreadSafe>& Viewport : InViewportManager->GetEntireClusterViewportsForWarpPolicy(SharedThis(this)))
+	if (GDisplayClusterWarpInFrustumFitPolicyDrawFrustum)
 	{
-		// Process only viewports with a projection policy based on the warpblend interface.
-		TSharedPtr<IDisplayClusterWarpBlend, ESPMode::ThreadSafe> WarpBlend;
-		if (HasPreviewMovableMesh(Viewport.Get()) && Viewport->GetProjectionPolicy()->GetWarpBlendInterface(WarpBlend))
+		for (const TSharedPtr<IDisplayClusterViewport, ESPMode::ThreadSafe>& Viewport : InViewportManager->GetEntireClusterViewportsForWarpPolicy(SharedThis(this)))
 		{
-			FDisplayClusterWarpData& WarpData = WarpBlend->GetWarpData(0);
-			UMeshComponent* MovableMeshComponent = Viewport->GetProjectionPolicy()->GetOrCreatePreviewMovableMeshComponent(Viewport.Get());
-
-			// Not all projection policies support a movable mesh.
-			if (MovableMeshComponent && WarpData.bValid && WarpData.bHasWarpPolicyChanges)
+			// Getting data from the first viewport, since all viewports use the same ViewPoint component
+			if (UDisplayClusterInFrustumFitCameraComponent* SceneCameraComponent = Cast<UDisplayClusterInFrustumFitCameraComponent>(Viewport->GetViewPointCameraComponent(EDisplayClusterRootActorType::Scene)))
 			{
-				const FTransform CameraTransform(WarpData.WarpProjection.CameraRotation.Quaternion(), WarpData.WarpProjection.CameraLocation);
-
-				const float HScale = (WarpData.WarpProjection.Left - WarpData.WarpProjection.Right) / (WarpData.GeometryWarpProjection.Left - WarpData.GeometryWarpProjection.Right);
-				const float VScale = (WarpData.WarpProjection.Top - WarpData.WarpProjection.Bottom) / (WarpData.GeometryWarpProjection.Top - WarpData.GeometryWarpProjection.Bottom);
-
-				checkf(FMath::IsNearlyEqual(HScale, VScale), TEXT("Streching the stage geometry to fit a different aspect ratio is not supported!"));
-
-				const FVector Scale = FVector(1, HScale, HScale);
-
-				// Compute the relative transform from the view origin to the geometry
-				FTransform RelativeTransform(WarpData.WarpContext.MeshToStageMatrix * WarpData.Local2World.Inverse());
-				RelativeTransform.ScaleTranslation(Scale);
-
-				// Final transform is computed from the relative transform of the geometry to the view point, the frustum fit transform
-				// which will scale and position the geometry based on the fitted frustum, and the camera transform
-				const FTransform FinalTransform = RelativeTransform * CameraTransform;
-
-				MovableMeshComponent->SetRelativeTransform(FinalTransform);
-
-				// Since the mesh needs to be skewed to scale appropriately, and since Unreal Engine does not support a skew transform
-				// through FTransform, the mesh needs to be skewed through the vertex shader using WorldPositionOffset,
-				// so pass in the "global" scale to the preview mesh's material instance
-				UMaterialInterface* MovableMeshMaterial = MovableMeshComponent->GetMaterial(0);
-				if (UMaterialInstanceDynamic* MaterialInstance = (MovableMeshMaterial && MovableMeshMaterial->IsA<UMaterialInstanceDynamic>()) ? Cast<UMaterialInstanceDynamic>(MovableMeshMaterial) : nullptr)
+				if (ADisplayClusterRootActor* SceneRootActor = InViewportManager->GetConfiguration().GetRootActor(EDisplayClusterRootActorType::Scene))
 				{
-					FMatrix CameraBasis = FRotationMatrix::Make(WarpData.WarpProjection.CameraRotation).Inverse();
-
-					MaterialInstance->SetVectorParameterValue(TEXT("GlobalScale"), Scale);
-					MaterialInstance->SetVectorParameterValue(TEXT("GlobalForward"), CameraBasis.GetUnitAxis(EAxis::X));
-					MaterialInstance->SetVectorParameterValue(TEXT("GlobalRight"), CameraBasis.GetUnitAxis(EAxis::Y));
-					MaterialInstance->SetVectorParameterValue(TEXT("GlobalUp"), CameraBasis.GetUnitAxis(EAxis::Z));
-				}
+					if (SceneCameraComponent)
+					{
+						DrawDebugGroupFrustum(SceneRootActor, SceneCameraComponent, FColor::Blue);
+					}
 
-				if (!SceneCameraComponent)
-				{
-					SceneCameraComponent = Cast<UDisplayClusterInFrustumFitCameraComponent>(Viewport->GetViewPointCameraComponent(EDisplayClusterRootActorType::Scene));
+					DrawDebugGroupBoundingBox(SceneRootActor, FColor::Red);
 				}
 			}
-		}
-	}
-
-	if (GDisplayClusterWarpInFrustumFitPolicyDrawFrustum)
-	{
-		if (ADisplayClusterRootActor* SceneRootActor = InViewportManager->GetConfiguration().GetRootActor(EDisplayClusterRootActorType::Scene))
-		{
-			if (SceneCameraComponent)
-			{
-				DrawDebugGroupFrustum(SceneRootActor, SceneCameraComponent, FColor::Blue);
-			}
 
-			DrawDebugGroupBoundingBox(SceneRootActor, FColor::Red);
+			break;
 		}
 	}
 #endif
 }
 
-bool FDisplayClusterWarpInFrustumFitPolicy::HasPreviewMovableMesh(IDisplayClusterViewport* InViewport)
+bool FDisplayClusterWarpInFrustumFitPolicy::HasPreviewEditableMesh(IDisplayClusterViewport* InViewport)
 {
-#if WITH_EDITOR
 	// This warp policy is based on IDisplayClusterWarpBlend only.
 	// Process only viewports with a projection policy based on the warpblend interface.
 	TSharedPtr<IDisplayClusterWarpBlend, ESPMode::ThreadSafe> WarpBlend;
@@ -247,13 +200,13 @@ bool FDisplayClusterWarpInFrustumFitPolicy::HasPreviewMovableMesh(IDisplayCluste
 		return false;
 	}
 
-	// If owner DCRA world is EditorPreview dont show movable mesh (Configurator, ICVFX Panel, etc)
+	// If owner DCRA world is EditorPreview dont show editable mesh (Configurator, ICVFX Panel, etc)
 	if (InViewport->GetConfiguration().IsRootActorWorldHasAnyType(EDisplayClusterRootActorType::Preview, EWorldType::EditorPreview))
 	{
 		return false;
 	}
 
-	// The movable mesh is an option for the UDisplayClusterInFrustumFitCameraComponent.
+	// The editable mesh is an option for the UDisplayClusterInFrustumFitCameraComponent.
 	if (UDisplayClusterInFrustumFitCameraComponent* ConfigurationCameraComponent = Cast<UDisplayClusterInFrustumFitCameraComponent>(InViewport->GetViewPointCameraComponent(EDisplayClusterRootActorType::Configuration)))
 	{
 		if (ConfigurationCameraComponent->bShowPreviewFrustumFit)
@@ -261,7 +214,6 @@ bool FDisplayClusterWarpInFrustumFitPolicy::HasPreviewMovableMesh(IDisplayCluste
 			return true;
 		}
 	}
-#endif
 
 	return false;
 }
@@ -344,6 +296,73 @@ void FDisplayClusterWarpInFrustumFitPolicy::EndCalcFrustum(IDisplayClusterViewpo
 	}
 }
 
+void FDisplayClusterWarpInFrustumFitPolicy::OnUpdatePreviewEditableMesh(IDisplayClusterViewportPreview& InViewportPreview, UMeshComponent* InEditableMeshComponent, const EDisplayClusterDisplayDeviceMaterialType InMaterialType, UMaterialInstanceDynamic* InMaterialInstance) const
+{
+	if (!InEditableMeshComponent || !InMaterialInstance)
+	{
+		return;
+	}
+
+	IDisplayClusterViewport* InViewport = InViewportPreview.GetViewport();
+	if (!InViewport)
+	{
+		return;
+	}
+
+	// Process only viewports with a projection policy based on the warpblend interface.
+	TSharedPtr<IDisplayClusterWarpBlend, ESPMode::ThreadSafe> WarpBlend;
+	if (!InViewport->GetProjectionPolicy().IsValid() || !InViewport->GetProjectionPolicy()->GetWarpBlendInterface(WarpBlend))
+	{
+		return;
+	}
+
+	FDisplayClusterWarpData& WarpData = WarpBlend->GetWarpData(0);
+
+	// Not all projection policies support a editable mesh.
+	if (WarpData.bValid && WarpData.bHasWarpPolicyChanges)
+	{
+		const FTransform CameraTransform(WarpData.WarpProjection.CameraRotation.Quaternion(), WarpData.WarpProjection.CameraLocation);
+
+		const float HScale = (WarpData.WarpProjection.Left - WarpData.WarpProjection.Right) / (WarpData.GeometryWarpProjection.Left - WarpData.GeometryWarpProjection.Right);
+		const float VScale = (WarpData.WarpProjection.Top - WarpData.WarpProjection.Bottom) / (WarpData.GeometryWarpProjection.Top - WarpData.GeometryWarpProjection.Bottom);
+
+		checkf(FMath::IsNearlyEqual(HScale, VScale), TEXT("Streching the stage geometry to fit a different aspect ratio is not supported!"));
+
+		const FVector Scale = FVector(1, HScale, HScale);
+
+		// Compute the relative transform from the view origin to the geometry
+		FTransform RelativeTransform(WarpData.WarpContext.MeshToStageMatrix * WarpData.Local2World.Inverse());
+		RelativeTransform.ScaleTranslation(Scale);
+
+		// Final transform is computed from the relative transform of the geometry to the view point, the frustum fit transform
+		// which will scale and position the geometry based on the fitted frustum, and the camera transform
+		const FTransform FinalTransform = RelativeTransform * CameraTransform;
+
+		InEditableMeshComponent->SetRelativeTransform(FinalTransform);
+
+		const FMatrix CameraBasis = FRotationMatrix::Make(WarpData.WarpProjection.CameraRotation).Inverse();
+
+		// Since the mesh needs to be skewed to scale appropriately, and since Unreal Engine does not support a skew transform
+		// through FTransform, the mesh needs to be skewed through the vertex shader using WorldPositionOffset,
+		// so pass in the "global" scale to the preview mesh's material instance
+		switch (InMaterialType)
+		{
+		case EDisplayClusterDisplayDeviceMaterialType::PreviewMeshMaterial:
+		case EDisplayClusterDisplayDeviceMaterialType::PreviewMeshTechvisMaterial:
+
+			InMaterialInstance->SetVectorParameterValue(TEXT("GlobalScale"), Scale);
+			InMaterialInstance->SetVectorParameterValue(TEXT("GlobalForward"), CameraBasis.GetUnitAxis(EAxis::X));
+			InMaterialInstance->SetVectorParameterValue(TEXT("GlobalRight"), CameraBasis.GetUnitAxis(EAxis::Y));
+			InMaterialInstance->SetVectorParameterValue(TEXT("GlobalUp"), CameraBasis.GetUnitAxis(EAxis::Z));
+			break;
+
+		default:
+			break;
+		}
+
+	}
+}
+
 #if WITH_EDITOR
 void FDisplayClusterWarpInFrustumFitPolicy::DrawDebugGroupBoundingBox(ADisplayClusterRootActor* RootActor, const FColor& Color)
 {
diff --git a/Engine/Plugins/Runtime/nDisplay/Source/DisplayClusterWarp/Private/Policy/InFrustumFit/DisplayClusterWarpInFrustumFitPolicy.h b/Engine/Plugins/Runtime/nDisplay/Source/DisplayClusterWarp/Private/Policy/InFrustumFit/DisplayClusterWarpInFrustumFitPolicy.h
--- a/Engine/Plugins/Runtime/nDisplay/Source/DisplayClusterWarp/Private/Policy/InFrustumFit/DisplayClusterWarpInFrustumFitPolicy.h
+++ b/Engine/Plugins/Runtime/nDisplay/Source/DisplayClusterWarp/Private/Policy/InFrustumFit/DisplayClusterWarpInFrustumFitPolicy.h
@@ -30,7 +30,9 @@ public:
 	virtual void BeginCalcFrustum(IDisplayClusterViewport* InViewport, const uint32 ContextNum) override;
 	virtual void EndCalcFrustum(IDisplayClusterViewport* InViewport, const uint32 ContextNum) override;
 
-	virtual bool HasPreviewMovableMesh(IDisplayClusterViewport* InViewport) override;
+	virtual bool HasPreviewEditableMesh(IDisplayClusterViewport* InViewport) override;
+
+	virtual void OnUpdatePreviewEditableMesh(IDisplayClusterViewportPreview& InViewportPreview, UMeshComponent* InEditableMeshComponent, const EDisplayClusterDisplayDeviceMaterialType InMaterialType, UMaterialInstanceDynamic* InMaterialInstance) const override;
 	//~~ End IDisplayClusterWarpPolicy
 
 private:
diff --git a/Engine/Plugins/Runtime/nDisplay/Source/DisplayClusterWarp/Private/WarpBlend/DisplayClusterWarpBlend.cpp b/Engine/Plugins/Runtime/nDisplay/Source/DisplayClusterWarp/Private/WarpBlend/DisplayClusterWarpBlend.cpp
--- a/Engine/Plugins/Runtime/nDisplay/Source/DisplayClusterWarp/Private/WarpBlend/DisplayClusterWarpBlend.cpp
+++ b/Engine/Plugins/Runtime/nDisplay/Source/DisplayClusterWarp/Private/WarpBlend/DisplayClusterWarpBlend.cpp
@@ -149,13 +149,13 @@ UMeshComponent* FDisplayClusterWarpBlend::GetOrCreatePreviewMeshComponent(IDispl
 	return GetOrCreatePreviewMeshComponentImpl(InViewport, false, bExistingComponent);
 }
 
-UMeshComponent* FDisplayClusterWarpBlend::GetOrCreatePreviewMovableMeshComponent(IDisplayClusterViewport* InViewport) const
+UMeshComponent* FDisplayClusterWarpBlend::GetOrCreatePreviewEditableMeshComponent(IDisplayClusterViewport* InViewport) const
 {
 	bool bExistingComponentDummy;
 	return GetOrCreatePreviewMeshComponentImpl(InViewport, true, bExistingComponentDummy);
 }
 
-UMeshComponent* FDisplayClusterWarpBlend::GetOrCreatePreviewMeshComponentImpl(IDisplayClusterViewport* InViewport, bool bMovableMesh, bool& bExistingComponent) const
+UMeshComponent* FDisplayClusterWarpBlend::GetOrCreatePreviewMeshComponentImpl(IDisplayClusterViewport* InViewport, bool bEditableMesh, bool& bExistingComponent) const
 {
 	const TSharedPtr<IDisplayClusterProjectionPolicy, ESPMode::ThreadSafe> ProjectionPolicy = InViewport ? InViewport->GetProjectionPolicy() : nullptr;
 
@@ -170,15 +170,15 @@ UMeshComponent* FDisplayClusterWarpBlend::GetOrCreatePreviewMeshComponentImpl(ID
 			{
 				UMeshComponent* ExistMeshComponent = static_cast<UMeshComponent*>(PreviewMeshComponent);
 
-				if (bMovableMesh)
+				if (bEditableMesh)
 				{
 					// create a mesh component copy
 					bExistingComponent = false;
 
-					// Get movable mesh root
-					USceneComponent* SceneOriginComp = ProjectionPolicy->GetPreviewMovableMeshOriginComponent(InViewport);
+					// Get editable mesh root
+					USceneComponent* SceneOriginComp = ProjectionPolicy->GetPreviewEditableMeshOriginComponent(InViewport);
 
-					const FString CompName = FString::Printf(TEXT("DCWarpBlend_MovableMesh_%s"), *ProjectionPolicy->GetId());
+					const FString CompName = FString::Printf(TEXT("DCWarpBlend_EditableMesh_%s"), *ProjectionPolicy->GetId());
 					const EObjectFlags ObjectFlags = EObjectFlags::RF_DuplicateTransient | RF_Transient | RF_TextExportTransient;
 					if (UMeshComponent* PreviewMeshComponentCopy = NewObject<UMeshComponent>(SceneOriginComp, ExistMeshComponent->GetClass(), *CompName, ObjectFlags, ExistMeshComponent))
 					{
@@ -214,7 +214,7 @@ UMeshComponent* FDisplayClusterWarpBlend::GetOrCreatePreviewMeshComponentImpl(ID
 			// Downscale preview mesh dimension to max limit
 			const uint32 PreviewGeometryDimLimit = 128;
 
-			USceneComponent* PreviewOriginComp = bMovableMesh ? ProjectionPolicy->GetPreviewMovableMeshOriginComponent(InViewport) : ProjectionPolicy->GetPreviewMeshOriginComponent(InViewport);
+			USceneComponent* PreviewOriginComp = bEditableMesh ? ProjectionPolicy->GetPreviewEditableMeshOriginComponent(InViewport) : ProjectionPolicy->GetPreviewMeshOriginComponent(InViewport);
 
 			// Create new WarpMesh component
 			FDisplayClusterWarpGeometryOBJ MeshData;
@@ -254,3 +254,11 @@ bool FDisplayClusterWarpBlend::ExportWarpMapGeometry(FDisplayClusterWarpGeometry
 {
 	return FDisplayClusterWarpBlendExporter_WarpMap::ExportWarpMap(GeometryContext, OutMeshData, InMaxDimension);
 }
+
+bool FDisplayClusterWarpBlend::HandleStartScene(IDisplayClusterViewport* InViewport)
+{
+	return true;
+}
+
+void FDisplayClusterWarpBlend::HandleEndScene(IDisplayClusterViewport* InViewport)
+{ }
diff --git a/Engine/Plugins/Runtime/nDisplay/Source/DisplayClusterWarp/Private/WarpBlend/DisplayClusterWarpBlend.h b/Engine/Plugins/Runtime/nDisplay/Source/DisplayClusterWarp/Private/WarpBlend/DisplayClusterWarpBlend.h
--- a/Engine/Plugins/Runtime/nDisplay/Source/DisplayClusterWarp/Private/WarpBlend/DisplayClusterWarpBlend.h
+++ b/Engine/Plugins/Runtime/nDisplay/Source/DisplayClusterWarp/Private/WarpBlend/DisplayClusterWarpBlend.h
@@ -30,6 +30,9 @@ public:
 		return AsShared();
 	}
 
+	virtual bool HandleStartScene(IDisplayClusterViewport* InViewport) override;
+	virtual void HandleEndScene(IDisplayClusterViewport* InViewport) override;
+
 	virtual bool MarkWarpGeometryComponentDirty(const FName& InComponentName) override;
 
 	virtual bool UpdateGeometryContext(const float InWorldScale) override;
@@ -83,7 +86,7 @@ public:
 	virtual const FDisplayClusterWarpData& GetWarpData(const uint32 ContextNum) const override;
 
 	virtual UMeshComponent* GetOrCreatePreviewMeshComponent(IDisplayClusterViewport* InViewport, bool& bExistingComponent) const override;
-	virtual UMeshComponent* GetOrCreatePreviewMovableMeshComponent(IDisplayClusterViewport* InViewport) const override;
+	virtual UMeshComponent* GetOrCreatePreviewEditableMeshComponent(IDisplayClusterViewport* InViewport) const override;
 
 	//~!IDisplayClusterWarpBlend
 
@@ -95,7 +98,7 @@ private:
 	* @param bCreateCopy - true to always create a copy of the mesh
 	* @param bExistingComponent - (out) true if the mesh component is not a copy
 	*/
-	UMeshComponent* GetOrCreatePreviewMeshComponentImpl(IDisplayClusterViewport* InViewport, bool bMovableMesh, bool& bExistingComponent) const;
+	UMeshComponent* GetOrCreatePreviewMeshComponentImpl(IDisplayClusterViewport* InViewport, bool bEditableMesh, bool& bExistingComponent) const;
 
 	/** Begin frustum calc. */
 	void BeginCalcFrustum(const TSharedPtr<FDisplayClusterWarpEye, ESPMode::ThreadSafe>& InWarpEye);
diff --git a/Engine/Plugins/Runtime/nDisplay/Source/DisplayClusterWarp/Public/Components/DisplayClusterInFrustumFitCameraComponent.h b/Engine/Plugins/Runtime/nDisplay/Source/DisplayClusterWarp/Public/Components/DisplayClusterInFrustumFitCameraComponent.h
--- a/Engine/Plugins/Runtime/nDisplay/Source/DisplayClusterWarp/Public/Components/DisplayClusterInFrustumFitCameraComponent.h
+++ b/Engine/Plugins/Runtime/nDisplay/Source/DisplayClusterWarp/Public/Components/DisplayClusterInFrustumFitCameraComponent.h
@@ -71,10 +71,9 @@ public:
 	UPROPERTY(EditAnywhere, Category = "Projection", meta = (EditCondition = "bEnableCameraProjection"))
 	bool bUseCameraPostprocess = false;
 
-#if WITH_EDITORONLY_DATA
+	/** Show additional warped preview meshes before the camera. */
 	UPROPERTY(EditAnywhere, Category = "Projection", meta = (EditCondition = "bEnableCameraProjection"))
 	bool bShowPreviewFrustumFit = false;
-#endif
 
 private:
 	// a unique type of warp policy for this component
diff --git a/Engine/Plugins/Runtime/nDisplay/Source/DisplayClusterWarp/Public/IDisplayClusterWarpBlend.h b/Engine/Plugins/Runtime/nDisplay/Source/DisplayClusterWarp/Public/IDisplayClusterWarpBlend.h
--- a/Engine/Plugins/Runtime/nDisplay/Source/DisplayClusterWarp/Public/IDisplayClusterWarpBlend.h
+++ b/Engine/Plugins/Runtime/nDisplay/Source/DisplayClusterWarp/Public/IDisplayClusterWarpBlend.h
@@ -21,6 +21,21 @@ public:
 	virtual TSharedPtr<IDisplayClusterWarpBlend, ESPMode::ThreadSafe> ToSharedPtr() = 0;
 	virtual TSharedPtr<const IDisplayClusterWarpBlend, ESPMode::ThreadSafe> ToSharedPtr() const = 0;
 
+	/**
+	* Called each time a new game level starts
+	*
+	* @param InViewport - a owner viewport
+	*/
+	virtual bool HandleStartScene(IDisplayClusterViewport* InViewport) = 0;
+
+	/**
+	* Called when current level is going to be closed (i.e. before loading a new map)
+	*
+	* @param InViewport - a owner viewport
+	*/
+	virtual void HandleEndScene(IDisplayClusterViewport* InViewport) = 0;
+
+
 	/** Update internal geometry cached data
 	 */
 	virtual bool UpdateGeometryContext(const float InWorldScale) = 0;
@@ -105,10 +120,10 @@ public:
 	virtual UMeshComponent* GetOrCreatePreviewMeshComponent(IDisplayClusterViewport* InViewport, bool& bExistingComponent) const = 0;
 
 	/**
-	* Build preview movable mesh
+	* Build preview editable mesh
 	* This MeshComponent is a copy of the preview mesh and can be moved freely with the UI visualization.
 	*
 	* @param InViewport - Projection specific parameters.
 	*/
-	virtual UMeshComponent* GetOrCreatePreviewMovableMeshComponent(IDisplayClusterViewport* InViewport) const = 0;
+	virtual UMeshComponent* GetOrCreatePreviewEditableMeshComponent(IDisplayClusterViewport* InViewport) const = 0;
 };
diff --git a/Engine/Plugins/VirtualProduction/ICVFXTesting/Source/ICVFXTesting/Public/ICVFXTestControllerAutoTest.h b/Engine/Plugins/VirtualProduction/ICVFXTesting/Source/ICVFXTesting/Public/ICVFXTestControllerAutoTest.h
--- a/Engine/Plugins/VirtualProduction/ICVFXTesting/Source/ICVFXTesting/Public/ICVFXTestControllerAutoTest.h
+++ b/Engine/Plugins/VirtualProduction/ICVFXTesting/Source/ICVFXTesting/Public/ICVFXTestControllerAutoTest.h
@@ -160,6 +160,8 @@ public:
 		constexpr bool bIncludeChildActors = false;
 		RootActor->GetComponents<UDisplayClusterICVFXCameraComponent>(CameraComponents, bIncludeChildActors);
 
+		UWorld* CurrentWorld = RootActor->GetWorld();
+
 		for (UDisplayClusterICVFXCameraComponent* CameraComponent : CameraComponents)
 		{
 			CameraComponent->CameraSettings.RenderSettings.AdvancedRenderSettings.GPUIndex = GetInnerGPUIndex();
@@ -170,7 +172,7 @@ public:
 			const FString NodeId = IDisplayCluster::Get().GetClusterMgr()->GetNodeId();
 			const EDisplayClusterRenderFrameMode RenderMode = EDisplayClusterRenderFrameMode::Mono;
 
-			ViewportManager->GetConfiguration().UpdateConfigurationForClusterNode(RenderMode, NodeId);
+			ViewportManager->GetConfiguration().UpdateConfigurationForClusterNode(RenderMode, CurrentWorld, NodeId);
 		}
 	}
 
