nDisplay: Fixed horde issue 379238
# jira none

diff --git a/Engine/Plugins/Runtime/nDisplay/Source/DisplayCluster/Private/Game/EngineClasses/Scene/DisplayClusterPreviewShareComponent.cpp b/Engine/Plugins/Runtime/nDisplay/Source/DisplayCluster/Private/Game/EngineClasses/Scene/DisplayClusterPreviewShareComponent.cpp
--- a/Engine/Plugins/Runtime/nDisplay/Source/DisplayCluster/Private/Game/EngineClasses/Scene/DisplayClusterPreviewShareComponent.cpp
+++ b/Engine/Plugins/Runtime/nDisplay/Source/DisplayCluster/Private/Game/EngineClasses/Scene/DisplayClusterPreviewShareComponent.cpp
@@ -15,6 +15,7 @@
 #include "Misc/DisplayClusterLog.h"
 #include "SharedMemoryMediaOutput.h"
 #include "SharedMemoryMediaSource.h"
+#include "UObject/Package.h"
 
 #if WITH_EDITOR
 #include "LevelEditor.h"
diff --git a/Engine/Plugins/Runtime/nDisplay/Source/DisplayCluster/Private/Game/EngineClasses/Scene/DisplayClusterRootActor.cpp b/Engine/Plugins/Runtime/nDisplay/Source/DisplayCluster/Private/Game/EngineClasses/Scene/DisplayClusterRootActor.cpp
--- a/Engine/Plugins/Runtime/nDisplay/Source/DisplayCluster/Private/Game/EngineClasses/Scene/DisplayClusterRootActor.cpp
+++ b/Engine/Plugins/Runtime/nDisplay/Source/DisplayCluster/Private/Game/EngineClasses/Scene/DisplayClusterRootActor.cpp
@@ -31,6 +31,7 @@
 #include "GameFramework/PlayerController.h"
 #include "Engine/Engine.h"
 #include "Engine/World.h"
+#include "Engine/GameInstance.h"
 
 #include "HAL/IConsoleManager.h"
 
@@ -51,7 +52,6 @@
 #include "Components/LineBatchComponent.h"
 #include "UObject/Package.h"
 
-
 #if WITH_EDITOR
 #include "IConcertSyncClientModule.h"
 #include "IConcertClientWorkspace.h"
diff --git a/Engine/Plugins/Runtime/nDisplay/Source/DisplayCluster/Private/Render/Viewport/Preview/DisplayClusterViewportPreviewMesh.cpp b/Engine/Plugins/Runtime/nDisplay/Source/DisplayCluster/Private/Render/Viewport/Preview/DisplayClusterViewportPreviewMesh.cpp
--- a/Engine/Plugins/Runtime/nDisplay/Source/DisplayCluster/Private/Render/Viewport/Preview/DisplayClusterViewportPreviewMesh.cpp
+++ b/Engine/Plugins/Runtime/nDisplay/Source/DisplayCluster/Private/Render/Viewport/Preview/DisplayClusterViewportPreviewMesh.cpp
@@ -76,7 +76,7 @@ void FDisplayClusterViewportPreviewMesh::Update(FDisplayClusterViewport& InViewp
 			CurrentMaterialPtr = InMeshMaterial;
 			CurrentMaterialType = NewCurrentMaterialType;
 
-			MaterialInstancePtr = UMaterialInstanceDynamic::Create(InMeshMaterial, GetTransientPackage());
+			MaterialInstancePtr = UMaterialInstanceDynamic::Create(InMeshMaterial, MeshComponent);
 			EnumAddFlags(RuntimeFlags, EDisplayClusterViewportPreviewMeshFlags::HasChangedMaterialInstance);
 
 			// Set preview material
