nDisplay: Preview resource destructors have been reworked.
When the RootActor is changed, the RerunConstructionScripts() function is called, which recreates all its components. In this case, the preview mesh must be updated accordingly.

diff --git a/Engine/Plugins/Runtime/nDisplay/Source/DisplayCluster/Private/Render/Viewport/Preview/DisplayClusterViewportPreview.cpp b/Engine/Plugins/Runtime/nDisplay/Source/DisplayCluster/Private/Render/Viewport/Preview/DisplayClusterViewportPreview.cpp
--- a/Engine/Plugins/Runtime/nDisplay/Source/DisplayCluster/Private/Render/Viewport/Preview/DisplayClusterViewportPreview.cpp
+++ b/Engine/Plugins/Runtime/nDisplay/Source/DisplayCluster/Private/Render/Viewport/Preview/DisplayClusterViewportPreview.cpp
@@ -21,6 +21,11 @@ FDisplayClusterViewportPreview::FDisplayClusterViewportPreview(const TSharedRef<
 	, PreviewEditableMesh(EDisplayClusterViewportPreviewMeshType::PreviewEditableMesh)
 { }
 
+FDisplayClusterViewportPreview::~FDisplayClusterViewportPreview()
+{
+	Release();
+}
+
 void FDisplayClusterViewportPreview::Initialize(FDisplayClusterViewport& InViewport)
 {
 	ViewportWeakPtr = InViewport.AsShared();
@@ -28,8 +33,12 @@ void FDisplayClusterViewportPreview::Initialize(FDisplayClusterViewport& InViewp
 
 void FDisplayClusterViewportPreview::Release()
 {
-	PreviewMesh.Release();
-	PreviewEditableMesh.Release();
+	if (FDisplayClusterViewport* InViewport = GetViewportImpl())
+	{
+		PreviewMesh.Release(*InViewport);
+		PreviewEditableMesh.Release(*InViewport);
+	}
+
 	PreviewRTT.Reset();
 	RuntimeFlags = EDisplayClusterViewportPreviewFlags::None;
 }
@@ -57,8 +66,8 @@ void FDisplayClusterViewportPreview::Update()
 	}
 	else
 	{
-		PreviewMesh.Release();
-		PreviewEditableMesh.Release();
+		PreviewMesh.Release(*InViewport);
+		PreviewEditableMesh.Release(*InViewport);
 	}
 
 	// Update Runtime Flags:
diff --git a/Engine/Plugins/Runtime/nDisplay/Source/DisplayCluster/Private/Render/Viewport/Preview/DisplayClusterViewportPreview.h b/Engine/Plugins/Runtime/nDisplay/Source/DisplayCluster/Private/Render/Viewport/Preview/DisplayClusterViewportPreview.h
--- a/Engine/Plugins/Runtime/nDisplay/Source/DisplayCluster/Private/Render/Viewport/Preview/DisplayClusterViewportPreview.h
+++ b/Engine/Plugins/Runtime/nDisplay/Source/DisplayCluster/Private/Render/Viewport/Preview/DisplayClusterViewportPreview.h
@@ -21,7 +21,7 @@ class FDisplayClusterViewportPreview
 {
 public:
 	FDisplayClusterViewportPreview(const TSharedRef<FDisplayClusterViewportConfiguration, ESPMode::ThreadSafe>& InConfiguration, const FString& ViewportId);
-	virtual ~FDisplayClusterViewportPreview() = default;
+	virtual ~FDisplayClusterViewportPreview();
 
 public:
 	//~ BEGIN IDisplayClusterViewportPreview
diff --git a/Engine/Plugins/Runtime/nDisplay/Source/DisplayCluster/Private/Render/Viewport/Preview/DisplayClusterViewportPreviewMesh.cpp b/Engine/Plugins/Runtime/nDisplay/Source/DisplayCluster/Private/Render/Viewport/Preview/DisplayClusterViewportPreviewMesh.cpp
--- a/Engine/Plugins/Runtime/nDisplay/Source/DisplayCluster/Private/Render/Viewport/Preview/DisplayClusterViewportPreviewMesh.cpp
+++ b/Engine/Plugins/Runtime/nDisplay/Source/DisplayCluster/Private/Render/Viewport/Preview/DisplayClusterViewportPreviewMesh.cpp
@@ -45,7 +45,7 @@ void FDisplayClusterViewportPreviewMesh::Update(FDisplayClusterViewport& InViewp
 	if (!ShouldUseMeshComponent(InViewport) || !InMeshMaterial)
 	{
 		// The mesh component and its resources are no longer used.
-		Release();
+		Release(InViewport);
 		return;
 	}
 
@@ -55,7 +55,7 @@ void FDisplayClusterViewportPreviewMesh::Update(FDisplayClusterViewport& InViewp
 	if (GetMeshComponent() != NewMeshComponent || bNewIsRootActorPreviewMesh != bIsRootActorMeshComponent)
 	{
 		// Release the reference to the old mesh component
-		ReleaseMeshComponent();
+		Release(InViewport);
 	}
 
 	if (NewMeshComponent != GetMeshComponent())
@@ -85,9 +85,16 @@ void FDisplayClusterViewportPreviewMesh::Update(FDisplayClusterViewport& InViewp
 	}
 }
 
-void FDisplayClusterViewportPreviewMesh::ReleaseMeshComponent()
+void FDisplayClusterViewportPreviewMesh::Release(FDisplayClusterViewport& InViewport)
 {
-	if (UMeshComponent* MeshComponent = GetMeshComponent())
+	UMeshComponent* MeshComponent = GetMeshComponent();
+	if (!MeshComponent && MeshComponentPtr)
+	{
+		// The mesh was destroyed earlier, (re-running build scripts inside RootActor), but we need to update the new mesh component to.
+		MeshComponent = GetOrCreatePreviewMeshComponent(&InViewport, bIsRootActorMeshComponent);
+	}
+
+	if (MeshComponent)
 	{
 		EnumAddFlags(RuntimeFlags, EDisplayClusterViewportPreviewMeshFlags::HasDeletedMeshComponent);
 
@@ -112,11 +119,6 @@ void FDisplayClusterViewportPreviewMesh::ReleaseMeshComponent()
 	CurrentMaterialType = EDisplayClusterDisplayDeviceMaterialType::DefaultPreviewMeshMaterial;
 
 	// The material instance references the mesh, so it must also be deleted
-	ReleaseMaterialInstance();
-}
-
-void FDisplayClusterViewportPreviewMesh::ReleaseMaterialInstance()
-{
 	if (UMaterialInstanceDynamic* MaterialInstance = GetMaterialInstance())
 	{
 		EnumAddFlags(RuntimeFlags, EDisplayClusterViewportPreviewMeshFlags::HasDeletedMaterialInstance);
diff --git a/Engine/Plugins/Runtime/nDisplay/Source/DisplayCluster/Private/Render/Viewport/Preview/DisplayClusterViewportPreviewMesh.h b/Engine/Plugins/Runtime/nDisplay/Source/DisplayCluster/Private/Render/Viewport/Preview/DisplayClusterViewportPreviewMesh.h
--- a/Engine/Plugins/Runtime/nDisplay/Source/DisplayCluster/Private/Render/Viewport/Preview/DisplayClusterViewportPreviewMesh.h
+++ b/Engine/Plugins/Runtime/nDisplay/Source/DisplayCluster/Private/Render/Viewport/Preview/DisplayClusterViewportPreviewMesh.h
@@ -56,10 +56,7 @@ public:
 		: MeshType(InMeshType)
 	{ }
 
-	~FDisplayClusterViewportPreviewMesh()
-	{
-		Release();
-	}
+	~FDisplayClusterViewportPreviewMesh() = default;
 
 public:
 	/** Get preview mesh component. */
@@ -79,16 +76,11 @@ public:
 	/** Get default material. */
 	UMaterial* GetDefaultMaterial() const;
 
-
 	/** Update mesh component and materials for viewport. */
 	void Update(FDisplayClusterViewport& InViewport, UDisplayClusterDisplayDeviceBaseComponent& InDisplayDeviceComponent);
 
 	/** Release mesh component and materials for viewport. */
-	void Release()
-	{
-		ReleaseMeshComponent();
-		ReleaseMaterialInstance();
-	}
+	void Release(FDisplayClusterViewport& InViewport);
 
 	/** Returns true if the runtime flags have any of the input flags. */
 	bool HasAnyFlag(const EDisplayClusterViewportPreviewMeshFlags InMeshFlags) const
@@ -97,12 +89,6 @@ public:
 	}
 
 private:
-	/** Remove mesh component. */
-	void ReleaseMeshComponent();
-
-	/** Remove material instance. */
-	void ReleaseMaterialInstance();
-
 	/** Returns true if this mesh type is supported by the viewport projection policy and DCRA. */
 	bool ShouldUseMeshComponent(FDisplayClusterViewport& InViewport) const;
 
