Fixed the issue with inner frustrum changing depth of field with overscan.

diff --git a/Engine/Plugins/Runtime/nDisplay/Source/DisplayCluster/Private/Game/EngineClasses/Scene/DisplayClusterICVFXCameraComponent.cpp b/Engine/Plugins/Runtime/nDisplay/Source/DisplayCluster/Private/Game/EngineClasses/Scene/DisplayClusterICVFXCameraComponent.cpp
--- a/Engine/Plugins/Runtime/nDisplay/Source/DisplayCluster/Private/Game/EngineClasses/Scene/DisplayClusterICVFXCameraComponent.cpp
+++ b/Engine/Plugins/Runtime/nDisplay/Source/DisplayCluster/Private/Game/EngineClasses/Scene/DisplayClusterICVFXCameraComponent.cpp
@@ -61,6 +61,14 @@ void UDisplayClusterICVFXCameraComponent::PostLoad()
 
 void UDisplayClusterICVFXCameraComponent::GetCameraView(float DeltaTime, FMinimalViewInfo& InOutViewInfo)
 {
+	const ADisplayClusterRootActor* RootActor = Cast<ADisplayClusterRootActor>(GetOwner());
+	if (RootActor == nullptr)
+	{
+		return;
+	}
+
+	const FDisplayClusterConfigurationICVFX_StageSettings& StageSettings = RootActor->GetStageSettings();
+
 	if (CameraSettings.ExternalCameraActor.IsValid())
 	{
 		// Get ViewInfo from external CineCamera
@@ -72,29 +80,7 @@ void UDisplayClusterICVFXCameraComponent::GetCameraView(float DeltaTime, FMinima
 		UCineCameraComponent::GetCameraView(DeltaTime, InOutViewInfo);
 	}
 
-	// CameraSettings can disable posprocess from this camera
-	if(!CameraSettings.RenderSettings.bUseCameraComponentPostprocess)
-	{
-		InOutViewInfo.PostProcessSettings = FPostProcessSettings();
-		InOutViewInfo.PostProcessBlendWeight = 0.0f;
-	}
-
-	// Add postprocess blur settings to viewinfo PP
-	const FDisplayClusterConfigurationICVFX_CameraMotionBlurOverridePPS& OverrideMotionBlurPPS = CameraSettings.CameraMotionBlur.MotionBlurPPS;
-	if (OverrideMotionBlurPPS.bReplaceEnable)
-	{
-		// Send camera postprocess to override
-		InOutViewInfo.PostProcessBlendWeight = 1.0f;
-
-		InOutViewInfo.PostProcessSettings.MotionBlurAmount = OverrideMotionBlurPPS.MotionBlurAmount;
-		InOutViewInfo.PostProcessSettings.bOverride_MotionBlurAmount = true;
-
-		InOutViewInfo.PostProcessSettings.MotionBlurMax = OverrideMotionBlurPPS.MotionBlurMax;
-		InOutViewInfo.PostProcessSettings.bOverride_MotionBlurMax = true;
-
-		InOutViewInfo.PostProcessSettings.MotionBlurPerObjectSize = OverrideMotionBlurPPS.MotionBlurPerObjectSize;
-		InOutViewInfo.PostProcessSettings.bOverride_MotionBlurPerObjectSize = true;
-	}
+	CameraSettings.SetupViewInfo(StageSettings, InOutViewInfo);
 }
 
 UCineCameraComponent* UDisplayClusterICVFXCameraComponent::GetActualCineCameraComponent()
diff --git a/Engine/Plugins/Runtime/nDisplay/Source/DisplayCluster/Private/Render/Viewport/Configuration/DisplayClusterViewportConfigurationHelpers_ICVFX.cpp b/Engine/Plugins/Runtime/nDisplay/Source/DisplayCluster/Private/Render/Viewport/Configuration/DisplayClusterViewportConfigurationHelpers_ICVFX.cpp
--- a/Engine/Plugins/Runtime/nDisplay/Source/DisplayCluster/Private/Render/Viewport/Configuration/DisplayClusterViewportConfigurationHelpers_ICVFX.cpp
+++ b/Engine/Plugins/Runtime/nDisplay/Source/DisplayCluster/Private/Render/Viewport/Configuration/DisplayClusterViewportConfigurationHelpers_ICVFX.cpp
@@ -95,7 +95,7 @@ bool FDisplayClusterViewportConfigurationHelpers_ICVFX::UpdateCameraProjectionSe
 
 	// Initialize camera policy with camera component and settings
 	FDisplayClusterProjectionCameraPolicySettings PolicyCameraSettings;
-	PolicyCameraSettings.FOVMultiplier = InCameraSettings.GetCameraFieldOfViewMultiplier(*StageSettings);
+	PolicyCameraSettings.FOVMultiplier = InCameraSettings.CustomFrustum.GetCameraFieldOfViewMultiplier(*StageSettings);
 
 	// Lens correction
 	PolicyCameraSettings.FrustumRotation = InCameraSettings.FrustumRotation;
@@ -533,7 +533,7 @@ void FDisplayClusterViewportConfigurationHelpers_ICVFX::UpdateCameraViewportSett
 	FDisplayClusterViewportConfigurationHelpers_ICVFX::UpdateCameraCustomFrustum(InCameraSettings.CustomFrustum, InOutRenderSettings.CustomFrustumSettings);
 
 	// Set RenderTargetAdaptRatio
-	InOutRenderSettings.RenderTargetAdaptRatio = InCameraSettings.GetCameraAdaptResolutionRatio(*StageSettings);
+	InOutRenderSettings.RenderTargetAdaptRatio = InCameraSettings.CustomFrustum.GetCameraAdaptResolutionRatio(*StageSettings);
 
 	// Set viewport buffer ratio
 	DstViewport.SetViewportBufferRatio(InCameraSettings.GetCameraBufferRatio(*StageSettings));
diff --git a/Engine/Plugins/Runtime/nDisplay/Source/DisplayClusterConfiguration/Private/DisplayClusterConfigurationTypes_ICVFX.cpp b/Engine/Plugins/Runtime/nDisplay/Source/DisplayClusterConfiguration/Private/DisplayClusterConfigurationTypes_ICVFX.cpp
--- a/Engine/Plugins/Runtime/nDisplay/Source/DisplayClusterConfiguration/Private/DisplayClusterConfigurationTypes_ICVFX.cpp
+++ b/Engine/Plugins/Runtime/nDisplay/Source/DisplayClusterConfiguration/Private/DisplayClusterConfigurationTypes_ICVFX.cpp
@@ -3,6 +3,7 @@
 #include "DisplayClusterConfigurationTypes_ICVFX.h"
 #include "DisplayClusterConfigurationTypes.h"
 #include "IDisplayCluster.h"
+#include "Camera/CameraTypes.h"
 
 namespace UE::DisplayClusterConfiguration::ICVFX
 {
@@ -34,6 +35,16 @@ FDisplayClusterConfigurationICVFX_CameraRenderSettings::FDisplayClusterConfigura
 	GenerateMips.bAutoGenerateMips = true;
 }
 
+void FDisplayClusterConfigurationICVFX_CameraRenderSettings::SetupViewInfo(const FDisplayClusterConfigurationICVFX_StageSettings& InStageSettings, FMinimalViewInfo& InOutViewInfo) const
+{
+	// CameraSettings can disable posprocess from this camera
+	if (!bUseCameraComponentPostprocess)
+	{
+		InOutViewInfo.PostProcessSettings = FPostProcessSettings();
+		InOutViewInfo.PostProcessBlendWeight = 0.0f;
+	}
+}
+
 ///////////////////////////////////////////////////////////////////////////////////////
 // FDisplayClusterConfigurationICVFX_LightcardCustomOCIO
 ///////////////////////////////////////////////////////////////////////////////////////
@@ -224,27 +235,6 @@ bool FDisplayClusterConfigurationICVFX_CameraSettings::IsChromakeyViewportSettin
 	return CameraOCIO.IsChromakeyViewportSettingsEqual(InClusterNodeId1, InClusterNodeId2);
 }
 
-float FDisplayClusterConfigurationICVFX_CameraSettings::GetCameraFieldOfViewMultiplier(const FDisplayClusterConfigurationICVFX_StageSettings& InStageSettings) const
-{
-	if (CustomFrustum.bEnable)
-	{
-		return CustomFrustum.FieldOfViewMultiplier;
-	}
-
-	return  1.f;
-}
-
-float FDisplayClusterConfigurationICVFX_CameraSettings::GetCameraAdaptResolutionRatio(const FDisplayClusterConfigurationICVFX_StageSettings& InStageSettings) const
-{
-	if (CustomFrustum.bAdaptResolution)
-	{
-		return GetCameraFieldOfViewMultiplier(InStageSettings);
-	}
-
-	// Don't use an adaptive resolution multiplier
-	return 1.f;
-}
-
 float FDisplayClusterConfigurationICVFX_CameraSettings::GetCameraBufferRatio(const FDisplayClusterConfigurationICVFX_StageSettings& InStageSettings) const
 {
 	return BufferRatio;
@@ -268,6 +258,13 @@ bool FDisplayClusterConfigurationICVFX_CameraSettings::GetCameraBorder(const FDi
 	return true;
 }
 
+void FDisplayClusterConfigurationICVFX_CameraSettings::SetupViewInfo(const FDisplayClusterConfigurationICVFX_StageSettings& InStageSettings, FMinimalViewInfo& InOutViewInfo)
+{
+	RenderSettings.SetupViewInfo(InStageSettings, InOutViewInfo);
+	CustomFrustum.SetupViewInfo(InStageSettings, InOutViewInfo);
+	CameraMotionBlur.SetupViewInfo(InStageSettings, InOutViewInfo);
+}
+
 FIntPoint FDisplayClusterConfigurationICVFX_CameraSettings::GetCameraFrameSize(const FDisplayClusterConfigurationICVFX_StageSettings& InStageSettings) const
 {
 	if (RenderSettings.CustomFrameSize.bUseCustomSize)
@@ -291,7 +288,7 @@ FVector4 FDisplayClusterConfigurationICVFX_CameraSettings::GetCameraSoftEdge(con
 {
 	FVector4 ResultSoftEdge(ForceInitToZero);
 
-	const float FieldOfViewMultiplier = GetCameraFieldOfViewMultiplier(InStageSettings);
+	const float FieldOfViewMultiplier = CustomFrustum.GetCameraFieldOfViewMultiplier(InStageSettings);
 
 	// softedge adjustments	
 	const float Overscan = (FieldOfViewMultiplier > 0) ? FieldOfViewMultiplier : 1;
@@ -552,3 +549,51 @@ bool FDisplayClusterConfigurationICVFX_ChromakeyRenderSettings::ShouldUseChromak
 	// ChromakeyRender requires actors for render.
 	return ShowOnlyList.IsVisibilityListValid();
 }
+
+void FDisplayClusterConfigurationICVFX_CameraCustomFrustum::SetupViewInfo(const FDisplayClusterConfigurationICVFX_StageSettings& InStageSettings, FMinimalViewInfo& InOutViewInfo) const
+{
+	// Since Circle of confusion is directly proportional to aperature, with wider FOV focal length needs to be shortened by the same amount as FOV.
+	const float FOVMultiplier = GetCameraFieldOfViewMultiplier(InStageSettings);
+	const float ClampedFieldOfViewMultiplier = (FOVMultiplier > 0.f) ? FOVMultiplier : 1.f;
+	InOutViewInfo.PostProcessSettings.DepthOfFieldFstop /= ClampedFieldOfViewMultiplier;
+}
+
+float FDisplayClusterConfigurationICVFX_CameraCustomFrustum::GetCameraFieldOfViewMultiplier(const FDisplayClusterConfigurationICVFX_StageSettings& InStageSettings) const
+{
+	if (bEnable)
+	{
+		return FieldOfViewMultiplier;
+	}
+
+	return  1.f;
+}
+
+float FDisplayClusterConfigurationICVFX_CameraCustomFrustum::GetCameraAdaptResolutionRatio(const FDisplayClusterConfigurationICVFX_StageSettings& InStageSettings) const
+{
+	if (bAdaptResolution)
+	{
+		return GetCameraFieldOfViewMultiplier(InStageSettings);
+	}
+
+	// Don't use an adaptive resolution multiplier
+	return 1.f;
+}
+
+void FDisplayClusterConfigurationICVFX_CameraMotionBlur::SetupViewInfo(const FDisplayClusterConfigurationICVFX_StageSettings& InStageSettings, FMinimalViewInfo& InOutViewInfo) const
+{
+	// Add postprocess blur settings to viewinfo PP
+	if (MotionBlurPPS.bReplaceEnable)
+	{
+		// Send camera postprocess to override
+		InOutViewInfo.PostProcessBlendWeight = 1.0f;
+
+		InOutViewInfo.PostProcessSettings.MotionBlurAmount = MotionBlurPPS.MotionBlurAmount;
+		InOutViewInfo.PostProcessSettings.bOverride_MotionBlurAmount = true;
+
+		InOutViewInfo.PostProcessSettings.MotionBlurMax = MotionBlurPPS.MotionBlurMax;
+		InOutViewInfo.PostProcessSettings.bOverride_MotionBlurMax = true;
+
+		InOutViewInfo.PostProcessSettings.MotionBlurPerObjectSize = MotionBlurPPS.MotionBlurPerObjectSize;
+		InOutViewInfo.PostProcessSettings.bOverride_MotionBlurPerObjectSize = true;
+	}
+}
diff --git a/Engine/Plugins/Runtime/nDisplay/Source/DisplayClusterConfiguration/Public/DisplayClusterConfigurationTypes_ICVFX.h b/Engine/Plugins/Runtime/nDisplay/Source/DisplayClusterConfiguration/Public/DisplayClusterConfigurationTypes_ICVFX.h
--- a/Engine/Plugins/Runtime/nDisplay/Source/DisplayClusterConfiguration/Public/DisplayClusterConfigurationTypes_ICVFX.h
+++ b/Engine/Plugins/Runtime/nDisplay/Source/DisplayClusterConfiguration/Public/DisplayClusterConfigurationTypes_ICVFX.h
@@ -513,6 +513,10 @@ public:
 	// Advanced render settings
 	UPROPERTY(EditAnywhere, BlueprintReadWrite, Category = NDisplay)
 	FDisplayClusterConfigurationICVFX_CameraAdvancedRenderSettings AdvancedRenderSettings;
+
+public:
+	/** Propagates general render related settings to the view info. */
+	void SetupViewInfo(const FDisplayClusterConfigurationICVFX_StageSettings& InStageSettings, FMinimalViewInfo& InOutViewInfo) const;
 };
 
 USTRUCT(BlueprintType)
@@ -555,6 +559,10 @@ public:
 	/** Motion Blur Settings Override */
 	UPROPERTY(EditAnywhere, BlueprintReadWrite, Category = NDisplay, meta = (DisplayName = "Motion Blur Settings Override"))
 	FDisplayClusterConfigurationICVFX_CameraMotionBlurOverridePPS MotionBlurPPS;
+
+public:
+	/** Propagates Motion Blur related settings to the view info. */
+	void SetupViewInfo(const FDisplayClusterConfigurationICVFX_StageSettings& InStageSettings, FMinimalViewInfo& InOutViewInfo) const;
 };
 
 USTRUCT(BlueprintType)
@@ -638,6 +646,17 @@ struct DISPLAYCLUSTERCONFIGURATION_API FDisplayClusterConfigurationICVFX_CameraC
 	/** Pixel/Percent value to alter the frustum to the bottom */
 	UPROPERTY(EditAnywhere, BlueprintReadWrite, Category = NDisplay, meta = (DisplayName = "Bottom", ClampMin = "-500.0", UIMin = "-500.0", ClampMax = "500.0", UIMax = "500.0"))
 	float Bottom = 0;
+
+public:
+	/** Propagates Custom Frustum related settings to the view info. */
+	void SetupViewInfo(const FDisplayClusterConfigurationICVFX_StageSettings& InStageSettings, FMinimalViewInfo& InOutViewInfo) const;
+
+	/** Get camera FOV multiplier. */
+	float GetCameraFieldOfViewMultiplier(const FDisplayClusterConfigurationICVFX_StageSettings& InStageSettings) const;
+
+	/** Get camera adapt resolution ratio. */
+	float GetCameraAdaptResolutionRatio(const FDisplayClusterConfigurationICVFX_StageSettings& InStageSettings) const;
+
 };
 
 USTRUCT(BlueprintType)
@@ -670,12 +689,6 @@ public:
 	/** Return calculated soft edges values. */
 	FVector4 GetCameraSoftEdge(const FDisplayClusterConfigurationICVFX_StageSettings& InStageSettings) const;
 
-	/** Get camera FOV multiplier. */
-	float GetCameraFieldOfViewMultiplier(const FDisplayClusterConfigurationICVFX_StageSettings& InStageSettings) const;
-
-	/** Get camera adapt resolution ratio. */
-	float GetCameraAdaptResolutionRatio(const FDisplayClusterConfigurationICVFX_StageSettings& InStageSettings) const;
-
 	/** Get camera buffer ratio. */
 	float GetCameraBufferRatio(const FDisplayClusterConfigurationICVFX_StageSettings& InStageSettings) const;
 
@@ -688,6 +701,9 @@ public:
 	/** Get camera border settings. */
 	bool GetCameraBorder(const FDisplayClusterConfigurationICVFX_StageSettings& InStageSettings, FLinearColor& OutBorderColor, float& OutBorderThickness) const;
 
+	/** Sets up view info for each relevant setting such as render, custom frustrum and motion blur settings. */
+	void SetupViewInfo(const FDisplayClusterConfigurationICVFX_StageSettings& InStageSettings, FMinimalViewInfo& InOutViewInfo);
+
 public:
 	/** Render the inner frustum for this ICVFX camera. */
 	UPROPERTY(EditAnywhere, BlueprintReadWrite, Category = "In Camera VFX", meta = (DisplayName = "Enable Inner Frustum"))
