nDisplay: Fixed horde compilation issue  379600

diff --git a/Engine/Plugins/Runtime/nDisplay/Source/DisplayCluster/Private/Render/Viewport/Preview/DisplayClusterViewportPreview.cpp b/Engine/Plugins/Runtime/nDisplay/Source/DisplayCluster/Private/Render/Viewport/Preview/DisplayClusterViewportPreview.cpp
--- a/Engine/Plugins/Runtime/nDisplay/Source/DisplayCluster/Private/Render/Viewport/Preview/DisplayClusterViewportPreview.cpp
+++ b/Engine/Plugins/Runtime/nDisplay/Source/DisplayCluster/Private/Render/Viewport/Preview/DisplayClusterViewportPreview.cpp
@@ -38,6 +38,12 @@ void FDisplayClusterViewportPreview::Release()
 		PreviewMesh.Release(*InViewport);
 		PreviewEditableMesh.Release(*InViewport);
 	}
+	else
+	{
+		// Just reset internal references
+		PreviewMesh.Reset();
+		PreviewEditableMesh.Reset();
+	}
 
 	PreviewRTT.Reset();
 	RuntimeFlags = EDisplayClusterViewportPreviewFlags::None;
@@ -45,6 +51,13 @@ void FDisplayClusterViewportPreview::Release()
 
 void FDisplayClusterViewportPreview::Update()
 {
+	FDisplayClusterViewport* InViewport = GetViewportImpl();
+	if (InViewport == nullptr)
+	{
+		Release();
+		return;
+	}
+
 	RuntimeFlags = EDisplayClusterViewportPreviewFlags::None;
 
 	// Update viewport output RTT
@@ -55,8 +68,7 @@ void FDisplayClusterViewportPreview::Update()
 		PreviewRTT = NewPreviewRTT;
 	}
 
-	FDisplayClusterViewport* InViewport = GetViewportImpl();
-	UDisplayClusterDisplayDeviceBaseComponent* InDisplayDeviceComponent = InViewport ? InViewport->GetDisplayDeviceComponent(EDisplayClusterRootActorType::Configuration) : nullptr;
+	UDisplayClusterDisplayDeviceBaseComponent* InDisplayDeviceComponent = InViewport->GetDisplayDeviceComponent(EDisplayClusterRootActorType::Configuration);
 
 	// Update preview meshes only if DisplayDevice is used
 	if (InDisplayDeviceComponent)
diff --git a/Engine/Plugins/Runtime/nDisplay/Source/DisplayCluster/Private/Render/Viewport/Preview/DisplayClusterViewportPreviewMesh.cpp b/Engine/Plugins/Runtime/nDisplay/Source/DisplayCluster/Private/Render/Viewport/Preview/DisplayClusterViewportPreviewMesh.cpp
--- a/Engine/Plugins/Runtime/nDisplay/Source/DisplayCluster/Private/Render/Viewport/Preview/DisplayClusterViewportPreviewMesh.cpp
+++ b/Engine/Plugins/Runtime/nDisplay/Source/DisplayCluster/Private/Render/Viewport/Preview/DisplayClusterViewportPreviewMesh.cpp
@@ -114,10 +114,6 @@ void FDisplayClusterViewportPreviewMesh::Release(FDisplayClusterViewport& InView
 		}
 	}
 
-	// Remove mesh
-	MeshComponentPtr = nullptr;
-	CurrentMaterialType = EDisplayClusterDisplayDeviceMaterialType::DefaultPreviewMeshMaterial;
-
 	// The material instance references the mesh, so it must also be deleted
 	if (UMaterialInstanceDynamic* MaterialInstance = GetMaterialInstance())
 	{
@@ -131,7 +127,14 @@ void FDisplayClusterViewportPreviewMesh::Release(FDisplayClusterViewport& InView
 		}
 	}
 
+	Reset();
+}
+
+void FDisplayClusterViewportPreviewMesh::Reset()
+{
+	MeshComponentPtr = nullptr;
 	MaterialInstancePtr = nullptr;
+	CurrentMaterialType = EDisplayClusterDisplayDeviceMaterialType::DefaultPreviewMeshMaterial;
 }
 
 bool FDisplayClusterViewportPreviewMesh::ShouldUseMeshComponent(FDisplayClusterViewport& InViewport) const
diff --git a/Engine/Plugins/Runtime/nDisplay/Source/DisplayCluster/Private/Render/Viewport/Preview/DisplayClusterViewportPreviewMesh.h b/Engine/Plugins/Runtime/nDisplay/Source/DisplayCluster/Private/Render/Viewport/Preview/DisplayClusterViewportPreviewMesh.h
--- a/Engine/Plugins/Runtime/nDisplay/Source/DisplayCluster/Private/Render/Viewport/Preview/DisplayClusterViewportPreviewMesh.h
+++ b/Engine/Plugins/Runtime/nDisplay/Source/DisplayCluster/Private/Render/Viewport/Preview/DisplayClusterViewportPreviewMesh.h
@@ -79,9 +79,12 @@ public:
 	/** Update mesh component and materials for viewport. */
 	void Update(FDisplayClusterViewport& InViewport, UDisplayClusterDisplayDeviceBaseComponent& InDisplayDeviceComponent);
 
-	/** Release mesh component and materials for viewport. */
+	/** Restore default material and release mesh component with materials for viewport. */
 	void Release(FDisplayClusterViewport& InViewport);
 
+	/** Release mesh component and materials for viewport. */
+	void Reset();
+
 	/** Returns true if the runtime flags have any of the input flags. */
 	bool HasAnyFlag(const EDisplayClusterViewportPreviewMeshFlags InMeshFlags) const
 	{
