nDisplay: Fixed preview rendering issues
- Fixed handling of scene start/end when changing RootActor references.
- Reworked preview settings in RootActor: Added a new property to RootActor that defines the source type of these settings - to support overriding.
- The error of opening shared memory in FSharedMemoryMediaPlayer::TickFetch() has been fixed. If opening fails, you should try to create and null-initialize shared memory.
- Code cleanup for RootActor: editor related logic moved to a separate file.

diff --git a/Engine/Plugins/Runtime/nDisplay/Source/DisplayCluster/Private/Game/EngineClasses/Scene/DisplayClusterPreviewShareComponent.cpp b/Engine/Plugins/Runtime/nDisplay/Source/DisplayCluster/Private/Game/EngineClasses/Scene/DisplayClusterPreviewShareComponent.cpp
--- a/Engine/Plugins/Runtime/nDisplay/Source/DisplayCluster/Private/Game/EngineClasses/Scene/DisplayClusterPreviewShareComponent.cpp
+++ b/Engine/Plugins/Runtime/nDisplay/Source/DisplayCluster/Private/Game/EngineClasses/Scene/DisplayClusterPreviewShareComponent.cpp
@@ -21,7 +21,7 @@
 #include "LevelEditor.h"
 #endif
 
-namespace UE::DisplayClusterPreviewShare
+namespace UE::DisplayCluster::PreviewShareComponent
 {
 	/** Convenience struct when iterating over config viewports */
 	struct FViewportData
@@ -43,7 +43,7 @@ namespace UE::DisplayClusterPreviewShare
 	};
 
 	/** Convenience function to get the data needed when iterating over config viewports of a given nDisplay actor */
-	TArray<FViewportData> GetConfigViewportsFromRootActor(const ADisplayClusterRootActor* RootActor)
+	static inline TArray<FViewportData> GetConfigViewportsFromRootActor(const ADisplayClusterRootActor* RootActor)
 	{
 		TArray<FViewportData> ViewportDatas;
 
@@ -86,7 +86,7 @@ namespace UE::DisplayClusterPreviewShare
 	}
 
 	/** Retrieves the viewport configuration from the DCRA's current config using the given NodeId and ViewportId */
-	static UDisplayClusterConfigurationViewport* GetViewportFromDCRA(ADisplayClusterRootActor* RootActor, const FString& NodeId, const FString& ViewportId)
+	static inline UDisplayClusterConfigurationViewport* GetViewportFromDCRA(ADisplayClusterRootActor* RootActor, const FString& NodeId, const FString& ViewportId)
 	{
 		if (!RootActor)
 		{
@@ -124,7 +124,7 @@ namespace UE::DisplayClusterPreviewShare
 		return *ViewportPtr;
 	}
 
-	void RemoveExternalMapReferences(ADisplayClusterRootActor* RootActor)
+	static inline void RemoveExternalMapReferences(ADisplayClusterRootActor* RootActor)
 	{
 		// In all cases, we should make sure we never restore in illegal external reference,
 		// because if we do, it will be very hard for the user to fix this and save the map, since these
@@ -161,7 +161,7 @@ namespace UE::DisplayClusterPreviewShare
 			}
 		}
 	}
-}
+};
 
 UDisplayClusterPreviewShareComponent::UDisplayClusterPreviewShareComponent(const FObjectInitializer& ObjectInitializer)
 	: Super(ObjectInitializer)
@@ -207,8 +207,6 @@ bool UDisplayClusterPreviewShareComponent::AllowedToShare() const
 	return true;
 }
 
-
-
 void UDisplayClusterPreviewShareComponent::ModeChanged()
 {
 	// We ignore the desired Mode if it is not allowed.
@@ -220,8 +218,10 @@ void UDisplayClusterPreviewShareComponent::ModeChanged()
 	// Close all media before restarting the sharing
 	CloseAllMedia();
 
+	// Restore default preview settings
+	RestoreDefaultPreviewSettings();
+
 	// Restore original root actor settings (e.g. texture replace settings that we overwrote)
-	PullPreviewFromSourceActor(nullptr);
 	RestoreRootActorOriginalSettings();
 
 	switch (Mode)
@@ -247,7 +247,7 @@ void UDisplayClusterPreviewShareComponent::SetTickEnable(const bool bEnable)
 
 void UDisplayClusterPreviewShareComponent::RestoreRootActorOriginalSettings()
 {
-	using namespace UE::DisplayClusterPreviewShare;
+	using namespace UE::DisplayCluster::PreviewShareComponent;
 
 	ADisplayClusterRootActor* RootActor = Cast<ADisplayClusterRootActor>(GetOwner());
 
@@ -262,7 +262,6 @@ void UDisplayClusterPreviewShareComponent::RestoreRootActorOriginalSettings()
 	}
 
 	// Iterate over viewports and restore their settings
-
 	for (const FViewportData& ViewportData : GetConfigViewportsFromRootActor(RootActor))
 	{
 		check(ViewportData.Viewport); // The iterator implementation guarantees that it is not nullptr
@@ -329,9 +328,62 @@ void UDisplayClusterPreviewShareComponent::CloseAllMedia()
 	MediaSources.Empty();
 	MediaPlayers.Empty();
 	MediaTextures.Empty();
+}
+
+void UDisplayClusterPreviewShareComponent::RestoreDefaultPreviewSettings()
+{
+	// Release custom viewport manager (used by 'Send')
+	CustomViewportManager.Reset();
+
+	// Restore default preview settings on owner RootActor (Used by 'PullActor')
+	if (ADisplayClusterRootActor* RootActor = Cast<ADisplayClusterRootActor>(GetOwner()))
+	{
+		if (IDisplayClusterViewportManager* ViewportManager = RootActor->GetViewportManager())
+		{
+			// Use preview settings from RootActor
+			RootActor->PreviewSetttingsSource = EDisplayClusterConfigurationRootActorPreviewSettingsSource::RootActor;
 
-	// Release custom viewport manager
-	ReleaseCustomViewportManager();
+			// reconfigure to rendering from owner root actor
+			ViewportManager->GetConfiguration().SetRootActor(EDisplayClusterRootActorType::Scene | EDisplayClusterRootActorType::Configuration, RootActor);
+		}
+	}
+}
+
+bool UDisplayClusterPreviewShareComponent::UpdateCustomViewportManager(const ADisplayClusterRootActor* InSrcRootActor)
+{
+	check(IsInGameThread());
+
+	if (!InSrcRootActor)
+	{
+		// The source RootActor is required
+		return false;
+	}
+
+	if (!CustomViewportManager.IsValid())
+	{
+		CustomViewportManager = IDisplayClusterViewportManager::CreateViewportManager();
+	}
+
+	IDisplayClusterViewportConfiguration& CustomConfiguration = CustomViewportManager->GetConfiguration();
+
+	// Assign DCRA to render, excluding PreviewRootActor since we are not using the preview mesh, only the viewport textures.
+	CustomConfiguration.SetRootActor(EDisplayClusterRootActorType::Scene | EDisplayClusterRootActorType::Configuration, InSrcRootActor);
+
+	// Get preview settings from InSrcRootActor properties
+	FDisplayClusterViewport_PreviewSettings NewPreviewSettings = InSrcRootActor->GetPreviewSettings(true);
+
+	// Override preview settings:
+	NewPreviewSettings.bPreviewEnable = true;
+	NewPreviewSettings.bPreviewEnablePostProcess = bPreviewEnablePostProcess;
+	NewPreviewSettings.bEnablePreviewMesh = false;
+	NewPreviewSettings.bEnablePreviewEditableMesh = false;
+	NewPreviewSettings.bPreviewICVFXFrustums = false; // disable frustum rendering
+	CustomConfiguration.SetPreviewSettings(NewPreviewSettings);
+
+	// Request rendering of the entire cluster for our custom viewport manager
+	CustomViewportManager->GetViewportManagerPreview().UpdateEntireClusterPreviewRender(true);
+
+	return true;
 }
 
 FString UDisplayClusterPreviewShareComponent::GenerateViewportKey(const FString& NodeName, const FString& ViewportName) const
@@ -352,13 +404,15 @@ FString UDisplayClusterPreviewShareComponent::GenerateMediaUniqueName(const FStr
 
 void UDisplayClusterPreviewShareComponent::TickSend()
 {
-	using namespace UE::DisplayClusterPreviewShare;
+	using namespace UE::DisplayCluster::PreviewShareComponent;
 
 	// Rendering RootActor on CustomViewportManager with custom settings and sending results
 	ADisplayClusterRootActor* RootActor = Cast<ADisplayClusterRootActor>(GetOwner());
 	if(!UpdateCustomViewportManager(RootActor))
 	{
 		CloseAllMedia();
+		RestoreDefaultPreviewSettings();
+
 		return;
 	}
 
@@ -486,21 +540,40 @@ void UDisplayClusterPreviewShareComponent::TickSend()
 
 void UDisplayClusterPreviewShareComponent::TickPullActor()
 {
-	ADisplayClusterRootActor* SourceRootActor = SourceNDisplayActor.Get();
-	if (!IsValid(SourceRootActor))
-	{
-		return;
-	}
+	using namespace UE::DisplayCluster::PreviewShareComponent;
 
+	ADisplayClusterRootActor* SourceRootActor = SourceNDisplayActor.Get();
 	ADisplayClusterRootActor* RootActor = Cast<ADisplayClusterRootActor>(GetOwner());
-	if (RootActor == SourceRootActor || !RootActor)
+	IDisplayClusterViewportManager* ViewportManager = RootActor ? RootActor->GetViewportManager() : nullptr;
+	if (!IsValid(SourceRootActor) || !IsValid(RootActor) || RootActor == SourceRootActor || !ViewportManager)
 	{
 		// Disallow self referencing
+		RestoreDefaultPreviewSettings();
+
 		return;
 	}
 
 	// Pull preview textures
-	PullPreviewFromSourceActor(SourceRootActor);
+	{
+		// reconfigure to rendering from external DCRA
+		ViewportManager->GetConfiguration().SetRootActor(EDisplayClusterRootActorType::Scene | EDisplayClusterRootActorType::Configuration, SourceRootActor);
+
+		// Part of setttings gets from owner (DCRA)
+		const FDisplayClusterViewport_PreviewSettings OwnerPreviewSettings = RootActor->GetPreviewSettings(true);
+
+		// Get preview settings from SourceRootActor properties
+		FDisplayClusterViewport_PreviewSettings NewPreviewSettings = SourceRootActor->GetPreviewSettings(true);
+
+		// Override some preview settings for the configurator and set them as the source of the preview settings
+		NewPreviewSettings.bPreviewEnable = OwnerPreviewSettings.bPreviewEnable;
+		NewPreviewSettings.bPreviewICVFXFrustums = OwnerPreviewSettings.bPreviewICVFXFrustums;
+
+		// Override preview PostProcess from component
+		NewPreviewSettings.bPreviewEnablePostProcess = bPreviewEnablePostProcess;
+
+		ViewportManager->GetConfiguration().SetPreviewSettings(NewPreviewSettings);
+		RootActor->PreviewSetttingsSource = EDisplayClusterConfigurationRootActorPreviewSettingsSource::Configuration;
+	}
 
 	// Note: Did we really need to synchronize the ICVFX camera function for PullActor mode?
 	// 
@@ -520,51 +593,9 @@ void UDisplayClusterPreviewShareComponent::TickPullActor()
 	}
 }
 
-void UDisplayClusterPreviewShareComponent::PullPreviewFromSourceActor(const ADisplayClusterRootActor* SourceRootActor)
-{
-	using namespace UE::DisplayClusterPreviewShare;
-
-	// Grab destination root actor.
-	ADisplayClusterRootActor* RootActor = Cast<ADisplayClusterRootActor>(GetOwner());
-	if (RootActor == SourceRootActor || !RootActor)
-	{
-		// Disallow self referencing
-		return;
-	}
-
-	if (IDisplayClusterViewportManager* ViewportManager = RootActor->GetViewportManager())
-	{
-		if (SourceRootActor)
-		{
-			// Disable local settings for owner root actor
-			RootActor->bUseLocalPreviewSetttings = false;
-
-			// reconfigure to rendering from external DCRA
-			ViewportManager->GetConfiguration().SetRootActor(EDisplayClusterRootActorType::Scene | EDisplayClusterRootActorType::Configuration, SourceRootActor);
-
-			// Override preview settings:
-			FDisplayClusterViewport_PreviewSettings CustomPreviewSettings = SourceRootActor->GetPreviewSettings();
-			CustomPreviewSettings.bPreviewEnable = RootActor->bPreviewEnable;
-			CustomPreviewSettings.bPreviewEnablePostProcess = bPreviewEnablePostProcess;
-			CustomPreviewSettings.bPreviewICVFXFrustums = false; // disable frustum rendering
-			ViewportManager->GetConfiguration().SetPreviewSettings(CustomPreviewSettings);
-		}
-		else
-		{
-			// Restore default preview:
-
-			// use local settings for owner root actor
-			RootActor->bUseLocalPreviewSetttings = true;
-
-			// reconfigure to rendering from owner root actor
-			ViewportManager->GetConfiguration().SetRootActor(EDisplayClusterRootActorType::Scene | EDisplayClusterRootActorType::Configuration, RootActor);
-		}
-	}
-}
-
 void UDisplayClusterPreviewShareComponent::SyncIcvxCamerasFromSourceActor(const ADisplayClusterRootActor* SourceRootActor, const ADisplayClusterRootActor* DstRootActor)
 {
-	using namespace UE::DisplayClusterPreviewShare;
+	using namespace UE::DisplayCluster::PreviewShareComponent;
 
 	if (!SourceRootActor || !DstRootActor)
 	{
@@ -671,7 +702,7 @@ void UDisplayClusterPreviewShareComponent::SyncIcvxCamerasFromSourceActor(const
 
 void UDisplayClusterPreviewShareComponent::TickReceive()
 {
-	using namespace UE::DisplayClusterPreviewShare;
+	using namespace UE::DisplayCluster::PreviewShareComponent;
 
 	ADisplayClusterRootActor* RootActor = Cast<ADisplayClusterRootActor>(GetOwner());
 	IDisplayClusterViewportManager* ViewportManager = RootActor ? RootActor->GetViewportManager() : nullptr;
@@ -831,7 +862,6 @@ void UDisplayClusterPreviewShareComponent::TickReceive()
 			MediaTextures.Remove(ViewportKey);
 		}
 	}
-
 }
 
 void UDisplayClusterPreviewShareComponent::TickComponent(float DeltaTime, ELevelTick TickType, FActorComponentTickFunction* ThisTickFunction)
@@ -860,7 +890,6 @@ void UDisplayClusterPreviewShareComponent::TickComponent(float DeltaTime, ELevel
 	}
 }
 
-
 void UDisplayClusterPreviewShareComponent::OnComponentDestroyed(bool bDestroyingHierarchy)
 {
 	// Close all media to avoid keeping resources alive even though the component was destroyed (and possibly kept in undo buffer).
@@ -881,7 +910,6 @@ void UDisplayClusterPreviewShareComponent::OnComponentDestroyed(bool bDestroying
 	Super::OnComponentDestroyed(bDestroyingHierarchy);
 }
 
-
 void UDisplayClusterPreviewShareComponent::OnRegister()
 {
 	Super::OnRegister();
@@ -890,21 +918,9 @@ void UDisplayClusterPreviewShareComponent::OnRegister()
 	ModeChanged();
 }
 
-#if WITH_EDITOR
-void UDisplayClusterPreviewShareComponent::HandleMapChanged_Editor(UWorld* InWorld, EMapChangeType InMapChangeType)
-{
-	// We remove the reference to the source actor preview texture to avoid this being interpreted as a 
-	// memory leak when unloadin the source actor map while the current map is kept open.
-	if (Mode == EDisplayClusterPreviewShareMode::PullActor)
-	{
-		RestoreRootActorOriginalSettings();
-	}
-}
-#endif
-
 void UDisplayClusterPreviewShareComponent::Serialize(FArchive& Ar)
 {
-	using namespace UE::DisplayClusterPreviewShare;
+	using namespace UE::DisplayCluster::PreviewShareComponent;
 
 	if (Ar.IsSaving())
 	{
@@ -962,56 +978,6 @@ void UDisplayClusterPreviewShareComponent::SetUniqueName(const FString& NewUniqu
 	CloseAllMedia();
 }
 
-bool UDisplayClusterPreviewShareComponent::UpdateCustomViewportManager(const ADisplayClusterRootActor* InSrcRootActor)
-{
-	check(IsInGameThread());
-
-	if (!InSrcRootActor)
-	{
-		// The source RootActor is required
-		return false;
-	}
-
-	if (!CustomViewportManager.IsValid())
-	{
-		CustomViewportManager = IDisplayClusterViewportManager::CreateViewportManager();
-	}
-
-	IDisplayClusterViewportConfiguration& CustomConfiguration = CustomViewportManager->GetConfiguration();
-
-	// Assign DCRA to render, excluding PreviewRootActor since we are not using the preview mesh, only the viewport textures.
-	CustomConfiguration.SetRootActor(EDisplayClusterRootActorType::Scene | EDisplayClusterRootActorType::Configuration, InSrcRootActor);
-
-	// Override preview settings:
-	FDisplayClusterViewport_PreviewSettings CustomPreviewSettings = InSrcRootActor->GetPreviewSettings();
-	CustomPreviewSettings.bPreviewEnable = true;
-	CustomPreviewSettings.bPreviewEnablePostProcess = bPreviewEnablePostProcess;
-	CustomPreviewSettings.bEnablePreviewMesh = false;
-	CustomPreviewSettings.bEnablePreviewEditableMesh = false;
-	CustomPreviewSettings.bPreviewICVFXFrustums = false; // disable frustum rendering
-	CustomConfiguration.SetPreviewSettings(CustomPreviewSettings);
-
-	// Request rendering of the entire cluster for our custom viewport manager
-	CustomViewportManager->GetViewportManagerPreview().UpdateEntireClusterPreviewRender(true);
-
-	return true;
-}
-
-void UDisplayClusterPreviewShareComponent::ReleaseCustomViewportManager()
-{
-	// Note: shared resources refer to internal preview textures inside the ViewportManager, and all these resources will be freed.
-	// we must handle this correctly on both sides of the IPC
-
-	if (CustomViewportManager.IsValid())
-	{
-		// Reset all DCRA references
-		CustomViewportManager->GetConfiguration().SetRootActor(EDisplayClusterRootActorType::Any, nullptr);
-
-		// Immediately release the viewport manager with resources
-		CustomViewportManager.Reset();
-	}
-}
-
 #if WITH_EDITOR
 
 void UDisplayClusterPreviewShareComponent::PostEditChangeProperty(FPropertyChangedEvent& PropertyChangedEvent)
@@ -1035,4 +1001,15 @@ void UDisplayClusterPreviewShareComponent::PostEditChangeProperty(FPropertyChang
 		ModeChanged();
 	}
 }
+
+void UDisplayClusterPreviewShareComponent::HandleMapChanged_Editor(UWorld* InWorld, EMapChangeType InMapChangeType)
+{
+	// We remove the reference to the source actor preview texture to avoid this being interpreted as a 
+	// memory leak when unloadin the source actor map while the current map is kept open.
+	if (Mode == EDisplayClusterPreviewShareMode::PullActor)
+	{
+		RestoreRootActorOriginalSettings();
+	}
+}
+
 #endif // WITH_EDITOR
\ No newline at end of file
diff --git a/Engine/Plugins/Runtime/nDisplay/Source/DisplayCluster/Private/Game/EngineClasses/Scene/DisplayClusterRootActor.cpp b/Engine/Plugins/Runtime/nDisplay/Source/DisplayCluster/Private/Game/EngineClasses/Scene/DisplayClusterRootActor.cpp
--- a/Engine/Plugins/Runtime/nDisplay/Source/DisplayCluster/Private/Game/EngineClasses/Scene/DisplayClusterRootActor.cpp
+++ b/Engine/Plugins/Runtime/nDisplay/Source/DisplayCluster/Private/Game/EngineClasses/Scene/DisplayClusterRootActor.cpp
@@ -52,14 +52,6 @@
 #include "Components/LineBatchComponent.h"
 #include "UObject/Package.h"
 
-#if WITH_EDITOR
-#include "IConcertSyncClientModule.h"
-#include "IConcertClientWorkspace.h"
-#include "IConcertSyncClient.h"
-
-#include "AssetToolsModule.h"
-#endif
-
 namespace UE::DisplayCluster::RootActor
 {
 	template <typename TComp>
@@ -169,7 +161,6 @@ void ADisplayClusterRootActor::ResetEntireClusterPreviewRendering()
 	{
 		// Update the preview settings as is from this DCRA
 		ViewportManager->GetViewportManagerPreview().ResetEntireClusterPreviewRendering();
-
 	}
 }
 
@@ -207,9 +198,6 @@ void ADisplayClusterRootActor::RemoveViewportManager()
 {
 	if (ViewportManagerPtr.IsValid())
 	{
-		// Reset all DCRA references
-		ViewportManagerPtr->GetConfiguration().SetRootActor(EDisplayClusterRootActorType::Any, nullptr);
-
 		// Immediately release the viewport manager with resources
 		ViewportManagerPtr.Reset();
 	}
@@ -242,22 +230,22 @@ bool ADisplayClusterRootActor::IsPrimaryRootActorForPIE() const
 }
 
 bool ADisplayClusterRootActor::IsRunningPIE() const
-	{
+{
 #if WITH_EDITOR
-		const UWorld* World = GetWorld();
-		return World && World->IsPlayInEditor();
+	const UWorld* World = GetWorld();
+	return World && World->IsPlayInEditor();
 #else
-		// Without editor return false
-		return false;
+	// Without editor return false
+	return false;
 #endif
-	}
+}
 
 bool ADisplayClusterRootActor::IsRunningDisplayCluster() const
 {
 	if (OperationMode == EDisplayClusterOperationMode::Cluster || OperationMode == EDisplayClusterOperationMode::Editor)
 	{
-	return true;
-}
+		return true;
+	}
 
 	return false;
 }
@@ -435,14 +423,14 @@ void ADisplayClusterRootActor::UpdateConfigDataInstance(UDisplayClusterConfigura
 			// properly with instanced values.
 
 			const EObjectFlags CommonFlags = RF_Public | RF_Transactional;
-			
+
 			CurrentConfigData = NewObject<UDisplayClusterConfigurationData>(
 				this,
 				UDisplayClusterConfigurationData::StaticClass(),
 				NAME_None,
 				IsTemplate() ? RF_ArchetypeObject | CommonFlags : CommonFlags,
 				ConfigDataTemplate);
-			
+
 			if (CurrentConfigData->Cluster == nullptr)
 			{
 				CurrentConfigData->Cluster = NewObject<UDisplayClusterConfigurationCluster>(
@@ -540,7 +528,7 @@ void ADisplayClusterRootActor::GetTypedPrimitives(TSet<FPrimitiveComponentId>& O
 		{
 			if (InCompNames != nullptr)
 			{
-				if(InCompNames->Find(CompIt->GetName()) != INDEX_NONE)
+				if (InCompNames->Find(CompIt->GetName()) != INDEX_NONE)
 				{
 					// add only comp from names list
 					CollectPrimitiveComponentsImpl(OutPrimitives, CompIt, bCollectChildrenVisualizationComponent);
@@ -733,7 +721,7 @@ void ADisplayClusterRootActor::InitializeRootActor()
 	{
 		return;
 	}
-	
+
 	bool bIsPIE = false;
 
 	if (!CurrentConfigData && !ConfigDataName.IsNone())
@@ -818,7 +806,7 @@ void ADisplayClusterRootActor::SetLightCardOwnership()
 	{
 		FDisplayClusterConfigurationICVFX_VisibilityList& LightCardVisibilityList = CurrentData->StageSettings.Lightcard.ShowOnlyList;
 		LightCardVisibilityList.AutoAddedActors.Reset();
-		
+
 		TArray<UDisplayClusterICVFXCameraComponent*> ICVFXComponents;
 		GetComponents(ICVFXComponents);
 
@@ -832,7 +820,7 @@ void ADisplayClusterRootActor::SetLightCardOwnership()
 					LightCardActor->SetWeakRootActorOwner(this);
 				}
 			}
-		
+
 			for (UDisplayClusterICVFXCameraComponent* Camera : ICVFXComponents)
 			{
 				FDisplayClusterConfigurationICVFX_ChromakeyRenderSettings* CameraChromakeyRenderSettings = Camera->CameraSettings.Chromakey.GetWritableChromakeyRenderSettings(GetStageSettings());
@@ -844,7 +832,7 @@ void ADisplayClusterRootActor::SetLightCardOwnership()
 
 				FDisplayClusterConfigurationICVFX_VisibilityList& ChromakeyCards = CameraChromakeyRenderSettings->ShowOnlyList;
 				ChromakeyCards.AutoAddedActors.Reset();
-				
+
 				for (const TSoftObjectPtr<AActor>& Actor : ChromakeyCards.Actors)
 				{
 					if (ADisplayClusterChromakeyCardActor* ChromakeyCardActor = Actor.IsValid() ? Cast<ADisplayClusterChromakeyCardActor>(Actor.Get()) : nullptr)
@@ -1004,7 +992,7 @@ void ADisplayClusterRootActor::Tick(float DeltaSeconds)
 	{
 		LineBatch->Flush();
 	}
-	
+
 	// Support for DCRA preview in the scene for standalone\package
 	// Use settings only from active DCRA
 	const bool bIsPrimaryRootActor = IsPrimaryRootActor();
@@ -1053,12 +1041,13 @@ void ADisplayClusterRootActor::Tick(float DeltaSeconds)
 		}
 	}
 
-	// Update Preview settings
+	// Get preview settings from the this actor current source
+	// as they may already be configured externally, so just get the preview settings from the current source
 	FDisplayClusterViewport_PreviewSettings NewPreviewSettings = GetPreviewSettings();
 
 	const bool bIsRunningPIE = IsRunningPIE();
 	const bool bIsRunningGame = IsRunningGame();
-	
+
 	// RootActor can have a preview in the scene
 	bool bEnablePreviewInScene = NewPreviewSettings.bPreviewEnable;
 	bool bPreviewInGame = false;
@@ -1120,57 +1109,50 @@ void ADisplayClusterRootActor::Tick(float DeltaSeconds)
 		NewPreviewSettings.bPreviewEnable = false;
 	}
 
-	// Preview in scene is a special render:
-	if (!bEnablePreviewInScene)
+	// Update entire cluster preview rendering
+	if (IDisplayClusterViewportManager* ViewportManager = bEnablePreviewInScene ? GetOrCreateViewportManager() : GetViewportManager())
 	{
-		// Hide invisible DCRA
-		SetActorHiddenInGame(true);
-
-		// Stop preview rendering.
-		if (IDisplayClusterViewportManager* ViewportManager = GetViewportManager())
-		{
-			// Update preview settings to new
-			ViewportManager->GetConfiguration().SetPreviewSettings(NewPreviewSettings);
+		// Update preview settings to new
+		ViewportManager->GetConfiguration().SetPreviewSettings(NewPreviewSettings);
 
-			// Stop preview rendering
-			ViewportManager->GetViewportManagerPreview().UpdateEntireClusterPreviewRender(false);
-		}
+		// Stop preview rendering
+		ViewportManager->GetViewportManagerPreview().UpdateEntireClusterPreviewRender(bEnablePreviewInScene);
 	}
-	// When the preview is used by this DCRA, we must create a new ViewportManager
-	else
-	{
-		// Show DCRA preview in the scene for standalone\package
-		SetActorHiddenInGame(false);
 
-		if (IDisplayClusterViewportManager* ViewportManager = GetOrCreateViewportManager())
-		{
-			// Update preview settings to new
-			ViewportManager->GetConfiguration().SetPreviewSettings(NewPreviewSettings);
-
-			// Request rendering
-			ViewportManager->GetViewportManagerPreview().UpdateEntireClusterPreviewRender(true);
-		}
-	}
+	// Update RootActor visibility for game
+	SetActorHiddenInGame(!bEnablePreviewInScene);
 
 	SetLightCardOwnership();
 
 	Super::Tick(DeltaSeconds);
 }
 
-FDisplayClusterViewport_PreviewSettings ADisplayClusterRootActor::GetPreviewSettings() const
+FDisplayClusterViewport_PreviewSettings ADisplayClusterRootActor::GetPreviewSettings(bool bIgnorePreviewSetttingsSource) const
 {
-	if (IDisplayClusterViewportManager* ViewportManager = GetViewportManager())
+	if (!bIgnorePreviewSetttingsSource)
 	{
-		if (!bUseLocalPreviewSetttings)
+		// Obtain preview settings from other sources:
+		switch (PreviewSetttingsSource)
 		{
-			return ViewportManager->GetConfiguration().GetPreviewSettings();
+		case EDisplayClusterConfigurationRootActorPreviewSettingsSource::Configuration:
+			if (IDisplayClusterViewportManager* ViewportManager = GetViewportManager())
+			{
+				return ViewportManager->GetConfiguration().GetPreviewSettings();
+			}
+			break;
+
+
+		case EDisplayClusterConfigurationRootActorPreviewSettingsSource::RootActor:
+		default:
+			break;
 		}
 	}
 
+	// Get preview settings from RootActor properties:
+
 	FDisplayClusterViewport_PreviewSettings OutPreviewSettings;
 
 	// By default RootActor renders in scene colors
-	// but
 	OutPreviewSettings.EntireClusterPreviewRenderMode = EDisplayClusterRenderFrameMode::PreviewInScene;
 
 	OutPreviewSettings.bPreviewEnable       = bPreviewEnable;
@@ -1179,10 +1161,10 @@ FDisplayClusterViewport_PreviewSettings ADisplayClusterRootActor::GetPreviewSett
 	OutPreviewSettings.bEnablePreviewTechvis     = bEnablePreviewTechvis;
 	OutPreviewSettings.bPreviewEnablePostProcess = bPreviewEnablePostProcess;
 
-	OutPreviewSettings.bPreviewICVFXFrustums = bPreviewICVFXFrustums;
+	OutPreviewSettings.bPreviewICVFXFrustums           = bPreviewICVFXFrustums;
 	OutPreviewSettings.PreviewICVFXFrustumsFarDistance = PreviewICVFXFrustumsFarDistance;
 
-	OutPreviewSettings.bEnablePreviewMesh        = bEnablePreviewMesh;
+	OutPreviewSettings.bEnablePreviewMesh         = bEnablePreviewMesh;
 	OutPreviewSettings.bEnablePreviewEditableMesh = bEnablePreviewEditableMesh;
 
 	OutPreviewSettings.PreviewRenderTargetRatioMult = PreviewRenderTargetRatioMult;
@@ -1252,226 +1234,6 @@ void ADisplayClusterRootActor::BeginDestroy()
 	Super::BeginDestroy();
 }
 
-/**
-* Propagate map changes from the DefaultsOwner to the InstanceOwner. The map must be <FString, UObject*>.
-* Each instanced object value will also be setup to use default propagation.
-*
-* If the DefaultsOwner has added a key it will be added to all instances with a new UObject value templated from the
-* default value.
-* 
-* If the DefaultsOwner has removed a key it will be removed from all instances.
-* 
-* All other elements will be left alone and element values will always work through the default propagation system.
-*
-* This is necessary in the event the instance has modified individual element value properties. Default container
-* propagation will treat the entire container as dirty and not add or remove elements to the instance.
-* 
-* For nDisplay the instance can only edit element properties, not the container size.
-* This allows us to safely propagate size changes from the default map through a custom propagation system.
-*
-* @param MapProperty The map property to sync.
-* @param DefaultsOwner The direct object owning the map property with the default values.
-* @param InstanceOwner The direct object owning the map property with the instance values.
-*/
-static void PropagateDefaultMapToInstancedMap(const FMapProperty* MapProperty, UObject* DefaultsOwner, UObject* InstanceOwner)
-{
-	check(DefaultsOwner);
-	check(InstanceOwner);
-	
-	const void* MapContainerDefaults = MapProperty->ContainerPtrToValuePtr<void*>(DefaultsOwner);
-	const void* MapContainerInstance = MapProperty->ContainerPtrToValuePtr<void*>(InstanceOwner);
-	FScriptMapHelper MapDefaultsHelper(MapProperty, MapContainerDefaults);
-	FScriptMapHelper MapInstanceHelper(MapProperty, MapContainerInstance);
-
-	bool bHasChanged = false;
-
-	auto AddKeyWithInstancedObject = [&](const FString* Key, UObject* ArchetypeToUse)
-	{
-		if (const uint8* ExistingPair = MapInstanceHelper.FindMapPairPtrFromHash(Key))
-		{
-			MapInstanceHelper.RemovePair(ExistingPair);
-		}
-		
-		// Existing objects should only occur in the case of delete and undo.
-		UObject* ObjectToAdd = FindObject<UObject>(InstanceOwner, *ArchetypeToUse->GetName());
-		if (ObjectToAdd == nullptr)
-		{
-			// Create the instance to assign to the map. Provide the archetype as a template with
-			// the same name so properties are propagated automatically.
-			ObjectToAdd = NewObject<UObject>(InstanceOwner, ArchetypeToUse->GetClass(),
-			ArchetypeToUse->GetFName(), RF_Transactional, ArchetypeToUse);
-		}
-#if WITH_EDITOR
-		else if (GIsTransacting)
-		{
-			// HACK: Projection policy parameters can become cleared when a VP is deleted, undone, redone, undone.
-			// The policy params are meant to only be set from the DCRA CDO, so resetting them back to the VP archetype is safe.
-			// @todo figure out how the parameters are being cleared between the redo/undo. No other properties seem to be impacted. 
-			if (UDisplayClusterConfigurationViewport* InstanceViewport = Cast<UDisplayClusterConfigurationViewport>(ObjectToAdd))
-			{
-				const UDisplayClusterConfigurationViewport* ArchetypeViewport = CastChecked<UDisplayClusterConfigurationViewport>(ArchetypeToUse);
-				if (InstanceViewport->ProjectionPolicy.Parameters.IsEmpty() && !ArchetypeViewport->ProjectionPolicy.Parameters.IsEmpty())
-				{
-					UE_LOG(LogDisplayClusterBlueprint, Warning, TEXT("Projection Policy mismatch from archetype on viewport, correcting. Instance: %s, Archetype: %s."),
-						*ObjectToAdd->GetName(), *ObjectToAdd->GetArchetype()->GetName());
-
-					InstanceViewport->ProjectionPolicy.Parameters = ArchetypeViewport->ProjectionPolicy.Parameters;
-				}
-			}
-		}
-#endif
-		
-		MapInstanceHelper.AddPair(Key, &ObjectToAdd);
-		
-		bHasChanged = true;
-	};
-
-	// Look for elements that should be added.
-	for (FScriptMapHelper::FIterator DefaultIt = MapDefaultsHelper.CreateIterator(); DefaultIt; ++DefaultIt)
-	{
-		uint8* DefaultPairPtr = MapDefaultsHelper.GetPairPtr(*DefaultIt);
-		check(DefaultPairPtr);
-		
-		FString* Key = MapProperty->KeyProp->ContainerPtrToValuePtr<FString>(DefaultPairPtr);
-		check(Key);
-		
-		UObject** DefaultObjectPtr =  MapProperty->ValueProp->ContainerPtrToValuePtr<UObject*>(DefaultPairPtr);
-		check(DefaultObjectPtr);
-		
-		UObject* DefaultObject = *DefaultObjectPtr;
-		check(DefaultObject);
-		
-		if (UObject** InstancedObjectPtr = (UObject**)MapInstanceHelper.FindValueFromHash(Key))
-		{
-			if (const UObject* InstancedObject = *InstancedObjectPtr)
-			{
-				const bool bArchetypeCorrect = DefaultObject == InstancedObject->GetArchetype();
-				
-				// The archetype should always match the new default. There are edge cases with MU
-				// where the instance may be updated prior to the CDO and should be corrected once
-				// the BP is saved. If this occurs outside of MU there could be a serious problem.
-				if (!bArchetypeCorrect)
-				{
-					const ADisplayClusterRootActor* RootActor =
-						Cast<ADisplayClusterRootActor>(InstancedObject->GetTypedOuter(ADisplayClusterRootActor::StaticClass()));
-					if (!RootActor)
-					{
-						// Undo transactions can potentially trigger this while an object was renamed to the transient package.
-						check(InstancedObject->GetPackage() == GetTransientPackage());
-						continue;
-					}
-#if WITH_EDITOR
-					if (GEditor)
-					{
-						bool bIsMultiUserSession = false;
-						TSharedPtr<IConcertSyncClient> ConcertSyncClient = IConcertSyncClientModule::Get().GetClient(TEXT("MultiUser"));
-						if (ConcertSyncClient.IsValid())
-						{
-							TSharedPtr<IConcertClientWorkspace> Workspace = ConcertSyncClient->GetWorkspace();
-							bIsMultiUserSession = Workspace.IsValid();
-						}
-
-						ensure(bArchetypeCorrect || bIsMultiUserSession);
-					}
-#endif
-					UE_LOG(LogDisplayClusterBlueprint, Warning, TEXT("Archetype mismatch on nDisplay config %s. Make sure the config is compiled and saved. Property: %s, Instance: %s, Archetype: %s, Default: %s."),
-						*RootActor->GetName(), *MapProperty->GetName(), *InstancedObject->GetName(), *InstancedObject->GetArchetype()->GetName(), *DefaultObject->GetName());
-				}
-				continue;
-			}
-		}
-		
-		AddKeyWithInstancedObject(Key, DefaultObject);
-	}
-
-	// Look for elements that should be removed.
-	for (FScriptMapHelper::FIterator InstanceIt = MapInstanceHelper.CreateIterator(); InstanceIt; ++InstanceIt)
-	{
-		uint8* InstancePairPtr = MapInstanceHelper.GetPairPtr(*InstanceIt);
-		check(InstancePairPtr);
-		
-		FString* Key = MapProperty->KeyProp->ContainerPtrToValuePtr<FString>(InstancePairPtr);
-		check(Key);
-		
-		if (!MapDefaultsHelper.FindValueFromHash(Key))
-		{
-			if (UObject** InstanceObjectPtr = MapProperty->ValueProp->ContainerPtrToValuePtr<UObject*>(InstancePairPtr))
-			{
-				if (UObject* InstanceObject = *InstanceObjectPtr)
-				{
-					// Trash the object -- default propagation won't handle this any more.
-					// RemoveAt below will remove the reference to it. This transaction can still be undone.
-					// Rename to transient package now so the same name is available immediately.
-					// It's possible a new object needs to be created with this outer using the same name.
-					InstanceObject->Rename(nullptr, GetTransientPackage(), REN_DontCreateRedirectors);
-					InstanceObject->SetFlags(RF_Transient);
-				}
-			}
-			
-			MapInstanceHelper.RemoveAt(*InstanceIt);
-			bHasChanged = true;
-		}
-	}
-	
-	if (bHasChanged)
-	{
-		MapInstanceHelper.Rehash();
-#if WITH_EDITOR
-		InstanceOwner->PostEditChange();
-#endif
-	}
-}
-
-/**
- * Syncs default config data changes to a config instance.
- *
- * @param InDefaultConfigData The class default config data object.
- * @param InInstanceConfigData An instance config data object.
- */
-static void PropagateDataFromDefaultConfig(UDisplayClusterConfigurationData* InDefaultConfigData, UDisplayClusterConfigurationData* InInstanceConfigData)
-{
-	check(InDefaultConfigData);
-	check(InInstanceConfigData);
-	
-	const FMapProperty* ClusterNodesMapProperty = FindFieldChecked<FMapProperty>(UDisplayClusterConfigurationCluster::StaticClass(),
-																				GET_MEMBER_NAME_CHECKED(UDisplayClusterConfigurationCluster, Nodes));
-	const FMapProperty* ViewportsMapProperty = FindFieldChecked<FMapProperty>(UDisplayClusterConfigurationClusterNode::StaticClass(),
-																			GET_MEMBER_NAME_CHECKED(UDisplayClusterConfigurationClusterNode, Viewports));
-	
-	PropagateDefaultMapToInstancedMap(ClusterNodesMapProperty, InDefaultConfigData->Cluster, InInstanceConfigData->Cluster);
-	
-	for (const TTuple<FString, TObjectPtr<UDisplayClusterConfigurationClusterNode>>& ClusterKeyVal : InDefaultConfigData->Cluster->Nodes)
-	{
-		UDisplayClusterConfigurationClusterNode* DestinationValue = InInstanceConfigData->Cluster->Nodes.FindChecked(
-			ClusterKeyVal.Key);
-		PropagateDefaultMapToInstancedMap(ViewportsMapProperty, ClusterKeyVal.Value, DestinationValue);
-	}
-}
-
-#if WITH_EDITOR
-void ADisplayClusterRootActor::RerunConstructionScripts()
-{
-	DECLARE_SCOPE_CYCLE_COUNTER(TEXT("ADisplayClusterRootActor::RerunConstructionScripts"), STAT_RerunConstructionScripts, STATGROUP_NDisplay);
-	
-	const IDisplayClusterConfiguration& Config = IDisplayClusterConfiguration::Get();
-	if (!Config.IsTransactingSnapshot())
-	{
-		Super::RerunConstructionScripts();
-		
-		if (!IsTemplate())
-		{
-			if (UDisplayClusterConfigurationData* CurrentData = GetConfigData())
-			{
-				const ADisplayClusterRootActor* CDO = CastChecked<ADisplayClusterRootActor>(GetClass()->GetDefaultObject());
-				UDisplayClusterConfigurationData* DefaultData = CDO->GetConfigData();
-				PropagateDataFromDefaultConfig(DefaultData, CurrentData);
-			}
-		}
-		RerunConstructionScripts_Editor();
-	}
-}
-#endif
-
 UDisplayClusterCameraComponent* ADisplayClusterRootActor::GetDefaultCamera() const
 {
 	return DefaultViewPoint;
@@ -1609,7 +1371,7 @@ bool ADisplayClusterRootActor::SetReplaceTextureFlagForAllViewports(bool bReplac
 	else
 	{
 		// No need to set this on a non operational nDisplay root actor.
-		if ((GDisplayCluster->GetOperationMode() == EDisplayClusterOperationMode::Cluster) 
+		if ((GDisplayCluster->GetOperationMode() == EDisplayClusterOperationMode::Cluster)
 			&& (this != Display.GetGameMgr()->GetRootActor()))
 		{
 			return false;
@@ -1626,7 +1388,7 @@ bool ADisplayClusterRootActor::SetReplaceTextureFlagForAllViewports(bool bReplac
 		Nodes.Add(Node);
 	}
 
-	for (const UDisplayClusterConfigurationClusterNode* Node: Nodes)
+	for (const UDisplayClusterConfigurationClusterNode* Node : Nodes)
 	{
 		check(Node != nullptr);
 
@@ -1651,7 +1413,7 @@ bool ADisplayClusterRootActor::SetFreezeOuterViewports(bool bEnable)
 		UE_LOG(LogDisplayClusterGame, Warning, TEXT("ADisplayClusterRootActor::SetFreezeOuterViewports failed because ConfigData was null"));
 		return false;
 	}
-	
+
 	if (ConfigData->StageSettings.bFreezeRenderOuterViewports != bEnable)
 	{
 #if WITH_EDITOR
diff --git a/Engine/Plugins/Runtime/nDisplay/Source/DisplayCluster/Private/Game/EngineClasses/Scene/DisplayClusterRootActorEditor.cpp b/Engine/Plugins/Runtime/nDisplay/Source/DisplayCluster/Private/Game/EngineClasses/Scene/DisplayClusterRootActorEditor.cpp
--- a/Engine/Plugins/Runtime/nDisplay/Source/DisplayCluster/Private/Game/EngineClasses/Scene/DisplayClusterRootActorEditor.cpp
+++ b/Engine/Plugins/Runtime/nDisplay/Source/DisplayCluster/Private/Game/EngineClasses/Scene/DisplayClusterRootActorEditor.cpp
@@ -51,14 +51,241 @@
 
 #if WITH_EDITOR
 
+#include "IConcertSyncClientModule.h"
+#include "IConcertClientWorkspace.h"
+#include "IConcertSyncClient.h"
+#include "AssetToolsModule.h"
+
 #include "Async/Async.h"
 #include "EditorSupportDelegates.h"
 #include "LevelEditor.h"
 
+namespace UE::DisplayCluster::RootActor_Editor
+{
+	/**
+	* Propagate map changes from the DefaultsOwner to the InstanceOwner. The map must be <FString, UObject*>.
+	* Each instanced object value will also be setup to use default propagation.
+	*
+	* If the DefaultsOwner has added a key it will be added to all instances with a new UObject value templated from the
+	* default value.
+	*
+	* If the DefaultsOwner has removed a key it will be removed from all instances.
+	*
+	* All other elements will be left alone and element values will always work through the default propagation system.
+	*
+	* This is necessary in the event the instance has modified individual element value properties. Default container
+	* propagation will treat the entire container as dirty and not add or remove elements to the instance.
+	*
+	* For nDisplay the instance can only edit element properties, not the container size.
+	* This allows us to safely propagate size changes from the default map through a custom propagation system.
+	*
+	* @param MapProperty The map property to sync.
+	* @param DefaultsOwner The direct object owning the map property with the default values.
+	* @param InstanceOwner The direct object owning the map property with the instance values.
+	*/
+	static void PropagateDefaultMapToInstancedMap_Editor(const FMapProperty* MapProperty, UObject* DefaultsOwner, UObject* InstanceOwner)
+	{
+		check(DefaultsOwner);
+		check(InstanceOwner);
+
+		const void* MapContainerDefaults = MapProperty->ContainerPtrToValuePtr<void*>(DefaultsOwner);
+		const void* MapContainerInstance = MapProperty->ContainerPtrToValuePtr<void*>(InstanceOwner);
+		FScriptMapHelper MapDefaultsHelper(MapProperty, MapContainerDefaults);
+		FScriptMapHelper MapInstanceHelper(MapProperty, MapContainerInstance);
+
+		bool bHasChanged = false;
+
+		auto AddKeyWithInstancedObject = [&](const FString* Key, UObject* ArchetypeToUse)
+			{
+				if (const uint8* ExistingPair = MapInstanceHelper.FindMapPairPtrFromHash(Key))
+				{
+					MapInstanceHelper.RemovePair(ExistingPair);
+				}
+
+				// Existing objects should only occur in the case of delete and undo.
+				UObject* ObjectToAdd = FindObject<UObject>(InstanceOwner, *ArchetypeToUse->GetName());
+				if (ObjectToAdd == nullptr)
+				{
+					// Create the instance to assign to the map. Provide the archetype as a template with
+					// the same name so properties are propagated automatically.
+					ObjectToAdd = NewObject<UObject>(InstanceOwner, ArchetypeToUse->GetClass(),
+						ArchetypeToUse->GetFName(), RF_Transactional, ArchetypeToUse);
+				}
+#if WITH_EDITOR
+				else if (GIsTransacting)
+				{
+					// HACK: Projection policy parameters can become cleared when a VP is deleted, undone, redone, undone.
+					// The policy params are meant to only be set from the DCRA CDO, so resetting them back to the VP archetype is safe.
+					// @todo figure out how the parameters are being cleared between the redo/undo. No other properties seem to be impacted. 
+					if (UDisplayClusterConfigurationViewport* InstanceViewport = Cast<UDisplayClusterConfigurationViewport>(ObjectToAdd))
+					{
+						const UDisplayClusterConfigurationViewport* ArchetypeViewport = CastChecked<UDisplayClusterConfigurationViewport>(ArchetypeToUse);
+						if (InstanceViewport->ProjectionPolicy.Parameters.IsEmpty() && !ArchetypeViewport->ProjectionPolicy.Parameters.IsEmpty())
+						{
+							UE_LOG(LogDisplayClusterBlueprint, Warning, TEXT("Projection Policy mismatch from archetype on viewport, correcting. Instance: %s, Archetype: %s."),
+								*ObjectToAdd->GetName(), *ObjectToAdd->GetArchetype()->GetName());
+
+							InstanceViewport->ProjectionPolicy.Parameters = ArchetypeViewport->ProjectionPolicy.Parameters;
+						}
+					}
+				}
+#endif
+
+				MapInstanceHelper.AddPair(Key, &ObjectToAdd);
+
+				bHasChanged = true;
+			};
+
+		// Look for elements that should be added.
+		for (FScriptMapHelper::FIterator DefaultIt = MapDefaultsHelper.CreateIterator(); DefaultIt; ++DefaultIt)
+		{
+			uint8* DefaultPairPtr = MapDefaultsHelper.GetPairPtr(*DefaultIt);
+			check(DefaultPairPtr);
+
+			FString* Key = MapProperty->KeyProp->ContainerPtrToValuePtr<FString>(DefaultPairPtr);
+			check(Key);
+
+			UObject** DefaultObjectPtr = MapProperty->ValueProp->ContainerPtrToValuePtr<UObject*>(DefaultPairPtr);
+			check(DefaultObjectPtr);
+
+			UObject* DefaultObject = *DefaultObjectPtr;
+			check(DefaultObject);
+
+			if (UObject** InstancedObjectPtr = (UObject**)MapInstanceHelper.FindValueFromHash(Key))
+			{
+				if (const UObject* InstancedObject = *InstancedObjectPtr)
+				{
+					const bool bArchetypeCorrect = DefaultObject == InstancedObject->GetArchetype();
+
+					// The archetype should always match the new default. There are edge cases with MU
+					// where the instance may be updated prior to the CDO and should be corrected once
+					// the BP is saved. If this occurs outside of MU there could be a serious problem.
+					if (!bArchetypeCorrect)
+					{
+						const ADisplayClusterRootActor* RootActor =
+							Cast<ADisplayClusterRootActor>(InstancedObject->GetTypedOuter(ADisplayClusterRootActor::StaticClass()));
+						if (!RootActor)
+						{
+							// Undo transactions can potentially trigger this while an object was renamed to the transient package.
+							check(InstancedObject->GetPackage() == GetTransientPackage());
+							continue;
+						}
+#if WITH_EDITOR
+						if (GEditor)
+						{
+							bool bIsMultiUserSession = false;
+							TSharedPtr<IConcertSyncClient> ConcertSyncClient = IConcertSyncClientModule::Get().GetClient(TEXT("MultiUser"));
+							if (ConcertSyncClient.IsValid())
+							{
+								TSharedPtr<IConcertClientWorkspace> Workspace = ConcertSyncClient->GetWorkspace();
+								bIsMultiUserSession = Workspace.IsValid();
+							}
+
+							ensure(bArchetypeCorrect || bIsMultiUserSession);
+						}
+#endif
+						UE_LOG(LogDisplayClusterBlueprint, Warning, TEXT("Archetype mismatch on nDisplay config %s. Make sure the config is compiled and saved. Property: %s, Instance: %s, Archetype: %s, Default: %s."),
+							*RootActor->GetName(), *MapProperty->GetName(), *InstancedObject->GetName(), *InstancedObject->GetArchetype()->GetName(), *DefaultObject->GetName());
+					}
+					continue;
+				}
+			}
+
+			AddKeyWithInstancedObject(Key, DefaultObject);
+		}
+
+		// Look for elements that should be removed.
+		for (FScriptMapHelper::FIterator InstanceIt = MapInstanceHelper.CreateIterator(); InstanceIt; ++InstanceIt)
+		{
+			uint8* InstancePairPtr = MapInstanceHelper.GetPairPtr(*InstanceIt);
+			check(InstancePairPtr);
+
+			FString* Key = MapProperty->KeyProp->ContainerPtrToValuePtr<FString>(InstancePairPtr);
+			check(Key);
+
+			if (!MapDefaultsHelper.FindValueFromHash(Key))
+			{
+				if (UObject** InstanceObjectPtr = MapProperty->ValueProp->ContainerPtrToValuePtr<UObject*>(InstancePairPtr))
+				{
+					if (UObject* InstanceObject = *InstanceObjectPtr)
+					{
+						// Trash the object -- default propagation won't handle this any more.
+						// RemoveAt below will remove the reference to it. This transaction can still be undone.
+						// Rename to transient package now so the same name is available immediately.
+						// It's possible a new object needs to be created with this outer using the same name.
+						InstanceObject->Rename(nullptr, GetTransientPackage(), REN_DontCreateRedirectors);
+						InstanceObject->SetFlags(RF_Transient);
+					}
+				}
+
+				MapInstanceHelper.RemoveAt(*InstanceIt);
+				bHasChanged = true;
+			}
+		}
+
+		if (bHasChanged)
+		{
+			MapInstanceHelper.Rehash();
+#if WITH_EDITOR
+			InstanceOwner->PostEditChange();
+#endif
+		}
+	}
+
+	/**
+	 * Syncs default config data changes to a config instance.
+	 *
+	 * @param InDefaultConfigData The class default config data object.
+	 * @param InInstanceConfigData An instance config data object.
+	 */
+	static void PropagateDataFromDefaultConfig_Editor(UDisplayClusterConfigurationData* InDefaultConfigData, UDisplayClusterConfigurationData* InInstanceConfigData)
+	{
+		check(InDefaultConfigData);
+		check(InInstanceConfigData);
+
+		const FMapProperty* ClusterNodesMapProperty = FindFieldChecked<FMapProperty>(UDisplayClusterConfigurationCluster::StaticClass(),
+			GET_MEMBER_NAME_CHECKED(UDisplayClusterConfigurationCluster, Nodes));
+		const FMapProperty* ViewportsMapProperty = FindFieldChecked<FMapProperty>(UDisplayClusterConfigurationClusterNode::StaticClass(),
+			GET_MEMBER_NAME_CHECKED(UDisplayClusterConfigurationClusterNode, Viewports));
+
+		PropagateDefaultMapToInstancedMap_Editor(ClusterNodesMapProperty, InDefaultConfigData->Cluster, InInstanceConfigData->Cluster);
+
+		for (const TTuple<FString, TObjectPtr<UDisplayClusterConfigurationClusterNode>>& ClusterKeyVal : InDefaultConfigData->Cluster->Nodes)
+		{
+			UDisplayClusterConfigurationClusterNode* DestinationValue = InInstanceConfigData->Cluster->Nodes.FindChecked(
+				ClusterKeyVal.Key);
+			PropagateDefaultMapToInstancedMap_Editor(ViewportsMapProperty, ClusterKeyVal.Value, DestinationValue);
+		}
+	}
+};
 //////////////////////////////////////////////////////////////////////////////////////////////
 // ADisplayClusterRootActor
 //////////////////////////////////////////////////////////////////////////////////////////////
 
+void ADisplayClusterRootActor::RerunConstructionScripts()
+{
+	using namespace UE::DisplayCluster::RootActor_Editor;
+
+	DECLARE_SCOPE_CYCLE_COUNTER(TEXT("ADisplayClusterRootActor::RerunConstructionScripts"), STAT_RerunConstructionScripts, STATGROUP_NDisplay);
+
+	const IDisplayClusterConfiguration& Config = IDisplayClusterConfiguration::Get();
+	if (!Config.IsTransactingSnapshot())
+	{
+		Super::RerunConstructionScripts();
+
+		if (!IsTemplate())
+		{
+			if (UDisplayClusterConfigurationData* CurrentData = GetConfigData())
+			{
+				const ADisplayClusterRootActor* CDO = CastChecked<ADisplayClusterRootActor>(GetClass()->GetDefaultObject());
+				UDisplayClusterConfigurationData* DefaultData = CDO->GetConfigData();
+				PropagateDataFromDefaultConfig_Editor(DefaultData, CurrentData);
+			}
+		}
+		RerunConstructionScripts_Editor();
+	}
+}
+
 void ADisplayClusterRootActor::Constructor_Editor()
 {
 	// Allow tick in editor for preview rendering
@@ -74,9 +301,6 @@ void ADisplayClusterRootActor::Constructor_Editor()
 
 void ADisplayClusterRootActor::Destructor_Editor()
 {
-	OnPreviewGenerated.Unbind();
-	OnPreviewDestroyed.Unbind();
-
 	FCoreUObjectDelegates::OnPackageReloaded.RemoveAll(this);
 
 	if (GEditor)
@@ -109,9 +333,6 @@ void ADisplayClusterRootActor::Destroyed_Editor()
 void ADisplayClusterRootActor::BeginDestroy_Editor()
 {
 	ResetEntireClusterPreviewRendering();
-
-	OnPreviewGenerated.Unbind();
-	OnPreviewDestroyed.Unbind();
 }
 
 void ADisplayClusterRootActor::RerunConstructionScripts_Editor()
@@ -133,7 +354,7 @@ void ADisplayClusterRootActor::UpdateInnerFrustumPriority()
 		ResetInnerFrustumPriority();
 		return;
 	}
-	
+
 	TArray<UDisplayClusterICVFXCameraComponent*> Components;
 	GetComponents(Components);
 
@@ -153,7 +374,7 @@ void ADisplayClusterRootActor::UpdateInnerFrustumPriority()
 			return CameraRef == CameraRefN2;
 		}).Num() > 1;
 	});
-	
+
 	for (int32 Idx = 0; Idx < InnerFrustumPriority.Num(); ++Idx)
 	{
 		if (UDisplayClusterICVFXCameraComponent* CameraComponent = FindObjectFast<UDisplayClusterICVFXCameraComponent>(this, *InnerFrustumPriority[Idx].Name))
@@ -173,7 +394,7 @@ void ADisplayClusterRootActor::ResetInnerFrustumPriority()
 	{
 		InnerFrustumPriority.Add(Camera->GetName());
 	}
-	
+
 	// Initialize based on current render priority.
 	InnerFrustumPriority.Sort([this](const FDisplayClusterComponentRef& CameraA, const FDisplayClusterComponentRef& CameraB)
 	{
@@ -199,10 +420,6 @@ void ADisplayClusterRootActor::SetIsSelectedInEditor(bool bValue)
 	bIsSelectedInEditor = bValue;
 }
 
-static FName Name_RelativeLocation = USceneComponent::GetRelativeLocationPropertyName();
-static FName Name_RelativeRotation = USceneComponent::GetRelativeRotationPropertyName();
-static FName Name_RelativeScale3D = USceneComponent::GetRelativeScale3DPropertyName();
-
 void ADisplayClusterRootActor::PostTransacted(const FTransactionObjectEvent& TransactionEvent)
 {
 	if (bRequiresComponentRefresh
@@ -222,7 +439,7 @@ void ADisplayClusterRootActor::PostEditChangeChainProperty(FPropertyChangedChain
 {
 	const FProperty* RootProperty = PropertyChainEvent.PropertyChain.GetActiveNode()->GetValue();
 	const FProperty* TailProperty = PropertyChainEvent.PropertyChain.GetTail()->GetValue();
-	if (PropertyChainEvent.ChangeType  == EPropertyChangeType::Interactive && RootProperty && TailProperty)
+	if (PropertyChainEvent.ChangeType == EPropertyChangeType::Interactive && RootProperty && TailProperty)
 	{
 		const FName RootPropertyName = RootProperty->GetFName();
 		const FName TailPropertyName = TailProperty->GetFName();
@@ -249,6 +466,10 @@ void ADisplayClusterRootActor::PostEditChangeProperty(FPropertyChangedEvent& Pro
 		return;
 	}
 
+	static FName Name_RelativeLocation = USceneComponent::GetRelativeLocationPropertyName();
+	static FName Name_RelativeRotation = USceneComponent::GetRelativeRotationPropertyName();
+	static FName Name_RelativeScale3D = USceneComponent::GetRelativeScale3DPropertyName();
+
 	// The AActor method, simplified and modified to skip construction scripts.
 	// Component registration still needs to occur or the actor will look like it disappeared.
 	auto SuperCallWithoutConstructionScripts = [&]
@@ -263,7 +484,7 @@ void ADisplayClusterRootActor::PostEditChangeProperty(FPropertyChangedEvent& Pro
 		if ((GEditor && GEditor->bIsSimulatingInEditor && GetWorld() != nullptr) || ReregisterComponentsWhenModified())
 		{
 			// If a transaction is occurring we rely on the true parent method instead.
-			check (!CurrentTransactionAnnotation.IsValid());
+			check(!CurrentTransactionAnnotation.IsValid());
 
 			UnregisterAllComponents();
 			ReregisterAllComponents();
@@ -272,7 +493,7 @@ void ADisplayClusterRootActor::PostEditChangeProperty(FPropertyChangedEvent& Pro
 		// Let other systems know that an actor was moved
 		if (bTransformationChanged)
 		{
-			GEngine->BroadcastOnActorMoved( this );
+			GEngine->BroadcastOnActorMoved(this);
 		}
 
 		FEditorSupportDelegates::UpdateUI.Broadcast();
@@ -286,7 +507,7 @@ void ADisplayClusterRootActor::PostEditChangeProperty(FPropertyChangedEvent& Pro
 	{
 		bCanSkipConstructionScripts = !Blueprint->bRunConstructionScriptOnInteractiveChange;
 	}
-	
+
 	if (bCanSkipConstructionScripts && PropertyChangedEvent.ChangeType == EPropertyChangeType::Interactive && !CurrentTransactionAnnotation.IsValid())
 	{
 		// Avoid calling construction scripts when the change occurs while the user is dragging a slider.
@@ -298,7 +519,7 @@ void ADisplayClusterRootActor::PostEditChangeProperty(FPropertyChangedEvent& Pro
 		bResetPreviewComponents = true;
 		Super::PostEditChangeProperty(PropertyChangedEvent);
 	}
-	
+
 	// Config file has been changed, we should rebuild the nDisplay hierarchy
 	// Cluster node ID has been changed
 	if (PropertyName == GET_MEMBER_NAME_CHECKED(ADisplayClusterRootActor, bFreezePreviewRender)
@@ -328,7 +549,7 @@ void ADisplayClusterRootActor::PostEditChangeProperty(FPropertyChangedEvent& Pro
 		ResetEntireClusterPreviewRendering();
 		bReinitializeActor = false;
 	}
-	
+
 	if (bReinitializeActor)
 	{
 		InitializeRootActor();
diff --git a/Engine/Plugins/Runtime/nDisplay/Source/DisplayCluster/Private/Render/Device/DisplayClusterDeviceBase.cpp b/Engine/Plugins/Runtime/nDisplay/Source/DisplayCluster/Private/Render/Device/DisplayClusterDeviceBase.cpp
--- a/Engine/Plugins/Runtime/nDisplay/Source/DisplayCluster/Private/Render/Device/DisplayClusterDeviceBase.cpp
+++ b/Engine/Plugins/Runtime/nDisplay/Source/DisplayCluster/Private/Render/Device/DisplayClusterDeviceBase.cpp
@@ -328,8 +328,12 @@ bool FDisplayClusterDeviceBase::BeginNewFrame(FViewport* InViewport, UWorld* InW
 		{
 			const FString LocalNodeId = DeviceBaseHelpers::GetDisplayClusterAPI().GetConfigMgr()->GetLocalNodeId();
 
+			// Get preview settings from RootActor properties
+			FDisplayClusterViewport_PreviewSettings NewPreviewSettings = RootActor->GetPreviewSettings(true);
+			NewPreviewSettings.bPreviewEnable = false;
+
 			// Dont use preview setting on primary RootActor in game
-			ViewportManagerPtr->GetConfiguration().SetPreviewSettings(FDisplayClusterViewport_PreviewSettings());
+			ViewportManagerPtr->GetConfiguration().SetPreviewSettings(NewPreviewSettings);
 
 			// Update local node viewports (update\create\delete) and build new render frame
 			if (ViewportManagerPtr->GetConfiguration().UpdateConfigurationForClusterNode(RenderFrameMode, InWorld, LocalNodeId))
diff --git a/Engine/Plugins/Runtime/nDisplay/Source/DisplayCluster/Private/Render/Viewport/Configuration/DisplayClusterViewportConfiguration.cpp b/Engine/Plugins/Runtime/nDisplay/Source/DisplayCluster/Private/Render/Viewport/Configuration/DisplayClusterViewportConfiguration.cpp
--- a/Engine/Plugins/Runtime/nDisplay/Source/DisplayCluster/Private/Render/Viewport/Configuration/DisplayClusterViewportConfiguration.cpp
+++ b/Engine/Plugins/Runtime/nDisplay/Source/DisplayCluster/Private/Render/Viewport/Configuration/DisplayClusterViewportConfiguration.cpp
@@ -30,7 +30,7 @@
 
 #include "Engine/RendererSettings.h"
 
-namespace UE::DisplayCluster::ConfigurationHelpers
+namespace UE::DisplayCluster::Configuration
 {
 	static inline ADisplayClusterRootActor* ImplGetRootActor(const FDisplayClusterActorRef& InConfigurationRootActorRef)
 	{
@@ -45,27 +45,34 @@ namespace UE::DisplayCluster::ConfigurationHelpers
 		return nullptr;
 	}
 
-	static inline bool ImplSetRootActor(const ADisplayClusterRootActor* InRootActor, FDisplayClusterActorRef& InOutConfigurationRootActorRef)
+	static inline bool ImplIsChangedRootActor(const ADisplayClusterRootActor* InRootActor, FDisplayClusterActorRef& InOutConfigurationRootActorRef)
 	{
 		const bool bIsDefined = InOutConfigurationRootActorRef.IsDefinedSceneActor();
 		if (InRootActor == nullptr)
 		{
-			InOutConfigurationRootActorRef.ResetSceneActor();
-
 			return bIsDefined;
 		}
 
 		if (!bIsDefined || ImplGetRootActor(InOutConfigurationRootActorRef) != InRootActor)
 		{
-			InOutConfigurationRootActorRef.SetSceneActor(InRootActor);
-
 			return true;
 		}
 
 		return false;
 	}
+
+	static inline void ImplSetRootActor(const ADisplayClusterRootActor* InRootActor, FDisplayClusterActorRef& InOutConfigurationRootActorRef)
+	{
+		if (InRootActor == nullptr)
+		{
+			InOutConfigurationRootActorRef.ResetSceneActor();
+		}
+		else
+		{
+			InOutConfigurationRootActorRef.SetSceneActor(InRootActor);
+		}
+	}
 };
-using namespace UE::DisplayCluster;
 
 ///////////////////////////////////////////////////////////////////
 // FDisplayClusterViewportConfiguration
@@ -88,27 +95,60 @@ void FDisplayClusterViewportConfiguration::SetRootActor(const EDisplayClusterRoo
 {
 	check(IsInGameThread());
 
-	if (EnumHasAnyFlags(InRootActorType, EDisplayClusterRootActorType::Preview))
+	using namespace UE::DisplayCluster::Configuration;
+
+	// COllect all required RootActor refs changes
+	EDisplayClusterRootActorType RootActorChanges = (EDisplayClusterRootActorType)0;
+	if (EnumHasAnyFlags(InRootActorType, EDisplayClusterRootActorType::Preview) && ImplIsChangedRootActor(InRootActor, PreviewRootActorRef))
 	{
-		ConfigurationHelpers::ImplSetRootActor(InRootActor, PreviewRootActorRef);
+		EnumAddFlags(RootActorChanges, EDisplayClusterRootActorType::Preview);
 	}
-
-	if (EnumHasAnyFlags(InRootActorType, EDisplayClusterRootActorType::Scene))
+	if (EnumHasAnyFlags(InRootActorType, EDisplayClusterRootActorType::Scene) && ImplIsChangedRootActor(InRootActor, SceneRootActorRef))
+	{
+		EnumAddFlags(RootActorChanges, EDisplayClusterRootActorType::Scene);
+	}
+	if (EnumHasAnyFlags(InRootActorType, EDisplayClusterRootActorType::Configuration) && ImplIsChangedRootActor(InRootActor, ConfigurationRootActorRef))
 	{
-		ConfigurationHelpers::ImplSetRootActor(InRootActor, SceneRootActorRef);
+		EnumAddFlags(RootActorChanges, EDisplayClusterRootActorType::Configuration);
 	}
 
-	if (EnumHasAnyFlags(InRootActorType, EDisplayClusterRootActorType::Configuration))
+	if (RootActorChanges != (EDisplayClusterRootActorType)0)
 	{
-		ConfigurationHelpers::ImplSetRootActor(InRootActor, ConfigurationRootActorRef);
+
+		if (IsSceneOpened())
+		{
+			if (FDisplayClusterViewportManager* ViewportManager = GetViewportManagerImpl())
+			{
+				ViewportManager->HandleEndScene();
+
+				// Always reset entore cluster preview rendering when RootActor chagned
+				ViewportManager->GetViewportManagerPreview().ResetEntireClusterPreviewRendering();
+			}
+		}
+
+		// Change RootActor
+		if (EnumHasAnyFlags(RootActorChanges, EDisplayClusterRootActorType::Preview))
+		{
+			ImplSetRootActor(InRootActor, PreviewRootActorRef);
+		}
+		if (EnumHasAnyFlags(RootActorChanges, EDisplayClusterRootActorType::Scene))
+		{
+			ImplSetRootActor(InRootActor, SceneRootActorRef);
+		}
+		if (EnumHasAnyFlags(RootActorChanges, EDisplayClusterRootActorType::Configuration))
+		{
+			ImplSetRootActor(InRootActor, ConfigurationRootActorRef);
+		}
 	}
 }
 
 ADisplayClusterRootActor* FDisplayClusterViewportConfiguration::GetRootActor(const EDisplayClusterRootActorType InRootActorType) const
 {
+	using namespace UE::DisplayCluster::Configuration;
+
 	if (EnumHasAnyFlags(InRootActorType, EDisplayClusterRootActorType::Preview))
 	{
-		if (ADisplayClusterRootActor* OutRootActor = ConfigurationHelpers::ImplGetRootActor(PreviewRootActorRef))
+		if (ADisplayClusterRootActor* OutRootActor = ImplGetRootActor(PreviewRootActorRef))
 		{
 			return OutRootActor;
 		}
@@ -116,7 +156,7 @@ ADisplayClusterRootActor* FDisplayClusterViewportConfiguration::GetRootActor(con
 
 	if (EnumHasAnyFlags(InRootActorType, EDisplayClusterRootActorType::Scene))
 	{
-		if (ADisplayClusterRootActor* OutRootActor = ConfigurationHelpers::ImplGetRootActor(SceneRootActorRef))
+		if (ADisplayClusterRootActor* OutRootActor = ImplGetRootActor(SceneRootActorRef))
 		{
 			return OutRootActor;
 		}
@@ -124,7 +164,7 @@ ADisplayClusterRootActor* FDisplayClusterViewportConfiguration::GetRootActor(con
 
 	if (EnumHasAnyFlags(InRootActorType, EDisplayClusterRootActorType::Configuration))
 	{
-		if (ADisplayClusterRootActor* OutRootActor = ConfigurationHelpers::ImplGetRootActor(ConfigurationRootActorRef))
+		if (ADisplayClusterRootActor* OutRootActor = ImplGetRootActor(ConfigurationRootActorRef))
 		{
 			return OutRootActor;
 		}
@@ -151,31 +191,16 @@ void FDisplayClusterViewportConfiguration::SetCurrentWorldImpl(const UWorld* InW
 {
 	if (GetCurrentWorld() != InWorld)
 	{
-		if (FDisplayClusterViewportManager* ViewportManager = GetViewportManagerImpl())
+		if (IsSceneOpened())
 		{
-			const bool bIsSceneOpened = IsSceneOpened();
-
-			if (bIsSceneOpened)
+			if (FDisplayClusterViewportManager* ViewportManager = GetViewportManagerImpl())
 			{
 				ViewportManager->HandleEndScene();
 			}
-
-			if (InWorld)
-			{
-				// Ignore const UWorld
-				CurrentWorldRef = TWeakObjectPtr<UWorld>((UWorld*)InWorld);
-
-				// Restore scene status
-				if (bIsSceneOpened)
-				{
-					ViewportManager->HandleStartScene();
-				}
-			}
-		}
-		else
-		{
-			ViewportManager->HandleEndScene();
 		}
+
+		// Ignore const UWorld
+		CurrentWorldRef = (UWorld*)InWorld;
 	}
 }
 
@@ -301,7 +326,7 @@ void FDisplayClusterViewportConfiguration::ReleaseConfiguration()
 bool FDisplayClusterViewportConfiguration::ImplUpdateConfiguration(EDisplayClusterRenderFrameMode InRenderMode, const UWorld* InWorld, const FString& InClusterNodeId, const TArray<FString>* InViewportNames)
 {
 	check(IsInGameThread());
-	
+
 	FDisplayClusterViewportManager* ViewportManager = GetViewportManagerImpl();
 	if (!ViewportManager || !FDisplayClusterViewportConfigurationHelpers_RenderFrameSettings::UpdateRenderFrameConfiguration(ViewportManager, InRenderMode, *this))
 	{
diff --git a/Engine/Plugins/Runtime/nDisplay/Source/DisplayCluster/Private/Render/Viewport/DisplayClusterViewportManager.cpp b/Engine/Plugins/Runtime/nDisplay/Source/DisplayCluster/Private/Render/Viewport/DisplayClusterViewportManager.cpp
--- a/Engine/Plugins/Runtime/nDisplay/Source/DisplayCluster/Private/Render/Viewport/DisplayClusterViewportManager.cpp
+++ b/Engine/Plugins/Runtime/nDisplay/Source/DisplayCluster/Private/Render/Viewport/DisplayClusterViewportManager.cpp
@@ -641,7 +641,7 @@ bool FDisplayClusterViewportManager::BeginNewFrame(FViewport* InViewport, FDispl
 	PostProcessManager->HandleBeginNewFrame(OutRenderFrame);
 
 	// Update viewport preview instances if preview is used and DCRA supports previews:
-	if (Configuration->IsPreviewRendering() && Configuration->GetRootActor(EDisplayClusterRootActorType::Preview) != nullptr)
+	if (Configuration->IsPreviewRendering())
 	{
 		ViewportManagerPreview->Update();
 	}
diff --git a/Engine/Plugins/Runtime/nDisplay/Source/DisplayCluster/Private/Render/Viewport/Preview/DisplayClusterViewportPreview.cpp b/Engine/Plugins/Runtime/nDisplay/Source/DisplayCluster/Private/Render/Viewport/Preview/DisplayClusterViewportPreview.cpp
--- a/Engine/Plugins/Runtime/nDisplay/Source/DisplayCluster/Private/Render/Viewport/Preview/DisplayClusterViewportPreview.cpp
+++ b/Engine/Plugins/Runtime/nDisplay/Source/DisplayCluster/Private/Render/Viewport/Preview/DisplayClusterViewportPreview.cpp
@@ -33,31 +33,16 @@ void FDisplayClusterViewportPreview::Initialize(FDisplayClusterViewport& InViewp
 
 void FDisplayClusterViewportPreview::Release()
 {
-	if (FDisplayClusterViewport* InViewport = GetViewportImpl())
-	{
-		PreviewMesh.Release(*InViewport);
-		PreviewEditableMesh.Release(*InViewport);
-	}
-	else
-	{
-		// Just reset internal references
-		PreviewMesh.Reset();
-		PreviewEditableMesh.Reset();
-	}
-
 	PreviewRTT.Reset();
 	RuntimeFlags = EDisplayClusterViewportPreviewFlags::None;
-}
 
-void FDisplayClusterViewportPreview::Update()
-{
 	FDisplayClusterViewport* InViewport = GetViewportImpl();
-	if (InViewport == nullptr)
-	{
-		Release();
-		return;
+	PreviewMesh.Release(InViewport);
+	PreviewEditableMesh.Release(InViewport);
 	}
 
+void FDisplayClusterViewportPreview::Update()
+{
 	RuntimeFlags = EDisplayClusterViewportPreviewFlags::None;
 
 	// Update viewport output RTT
@@ -68,19 +53,12 @@ void FDisplayClusterViewportPreview::Update()
 		PreviewRTT = NewPreviewRTT;
 	}
 
-	UDisplayClusterDisplayDeviceBaseComponent* InDisplayDeviceComponent = InViewport->GetDisplayDeviceComponent(EDisplayClusterRootActorType::Configuration);
+	FDisplayClusterViewport* InViewport = GetViewportImpl();
+	UDisplayClusterDisplayDeviceBaseComponent* InDisplayDeviceComponent = InViewport ? InViewport->GetDisplayDeviceComponent(EDisplayClusterRootActorType::Configuration) : nullptr;
 
 	// Update preview meshes only if DisplayDevice is used
-	if (InDisplayDeviceComponent)
-	{
-		PreviewMesh.Update(*InViewport, *InDisplayDeviceComponent);
-		PreviewEditableMesh.Update(*InViewport, *InDisplayDeviceComponent);
-	}
-	else
-	{
-		PreviewMesh.Release(*InViewport);
-		PreviewEditableMesh.Release(*InViewport);
-	}
+	PreviewMesh.Update(InViewport, InDisplayDeviceComponent);
+	PreviewEditableMesh.Update(InViewport, InDisplayDeviceComponent);
 
 	// Update Runtime Flags:
 	if (PreviewMesh.HasAnyFlag(EDisplayClusterViewportPreviewMeshFlags::HasDeletedMaterialInstance | EDisplayClusterViewportPreviewMeshFlags::HasChangedMaterialInstance))
diff --git a/Engine/Plugins/Runtime/nDisplay/Source/DisplayCluster/Private/Render/Viewport/Preview/DisplayClusterViewportPreviewMesh.cpp b/Engine/Plugins/Runtime/nDisplay/Source/DisplayCluster/Private/Render/Viewport/Preview/DisplayClusterViewportPreviewMesh.cpp
--- a/Engine/Plugins/Runtime/nDisplay/Source/DisplayCluster/Private/Render/Viewport/Preview/DisplayClusterViewportPreviewMesh.cpp
+++ b/Engine/Plugins/Runtime/nDisplay/Source/DisplayCluster/Private/Render/Viewport/Preview/DisplayClusterViewportPreviewMesh.cpp
@@ -27,22 +27,28 @@ namespace UE::DisplayCluster::Preview
 ////////////////////////////////////////////////////////////////////////////////////////
 // FDisplayClusterViewportPreviewMesh
 ////////////////////////////////////////////////////////////////////////////////////////
-void FDisplayClusterViewportPreviewMesh::Update(FDisplayClusterViewport& InViewport, UDisplayClusterDisplayDeviceBaseComponent& InDisplayDeviceComponent)
+void FDisplayClusterViewportPreviewMesh::Update(FDisplayClusterViewport* InViewport, UDisplayClusterDisplayDeviceBaseComponent* InDisplayDeviceComponent)
 {
+	// Reset runtime flags before each update
+	RuntimeFlags = EDisplayClusterViewportPreviewMeshFlags::None;
+
+	if (!InViewport || !ShouldUseMeshComponent(InViewport) || !InDisplayDeviceComponent)
+	{
+		// The mesh component and its resources are no longer used.
+		Release(InViewport);
+		return;
+	}
+
 	// Update default material
-	DefaultMaterialPtr = InDisplayDeviceComponent.GetDisplayDeviceMaterial(EDisplayClusterDisplayDeviceMaterialType::DefaultPreviewMeshMaterial);
+	DefaultMaterialPtr = InDisplayDeviceComponent->GetDisplayDeviceMaterial(EDisplayClusterDisplayDeviceMaterialType::DefaultPreviewMeshMaterial);
 
 	// Get current preview material
-	EDisplayClusterDisplayDeviceMaterialType NewCurrentMaterialType = InViewport.GetConfiguration().IsTechvisEnabled()
+	EDisplayClusterDisplayDeviceMaterialType NewCurrentMaterialType = InViewport->GetConfiguration().IsTechvisEnabled()
 		? EDisplayClusterDisplayDeviceMaterialType::PreviewMeshTechvisMaterial
 		: EDisplayClusterDisplayDeviceMaterialType::PreviewMeshMaterial;
 
-	UMaterial* InMeshMaterial = InDisplayDeviceComponent.GetDisplayDeviceMaterial(NewCurrentMaterialType);
-
-	// Reset runtime flags before each update
-	RuntimeFlags = EDisplayClusterViewportPreviewMeshFlags::None;
-
-	if (!ShouldUseMeshComponent(InViewport) || !InMeshMaterial)
+	UMaterial* InMeshMaterial = InDisplayDeviceComponent->GetDisplayDeviceMaterial(NewCurrentMaterialType);
+	if (!InMeshMaterial)
 	{
 		// The mesh component and its resources are no longer used.
 		Release(InViewport);
@@ -51,7 +57,7 @@ void FDisplayClusterViewportPreviewMesh::Update(FDisplayClusterViewport& InViewp
 
 	// Get or create warp mesh:
 	bool bNewIsRootActorPreviewMesh = false;
-	UMeshComponent* NewMeshComponent = GetOrCreatePreviewMeshComponent(&InViewport, bNewIsRootActorPreviewMesh);
+	UMeshComponent* NewMeshComponent = GetOrCreatePreviewMeshComponent(InViewport, bNewIsRootActorPreviewMesh);
 	if (GetMeshComponent() != NewMeshComponent || bNewIsRootActorPreviewMesh != bIsRootActorMeshComponent)
 	{
 		// Release the reference to the old mesh component
@@ -85,13 +91,13 @@ void FDisplayClusterViewportPreviewMesh::Update(FDisplayClusterViewport& InViewp
 	}
 }
 
-void FDisplayClusterViewportPreviewMesh::Release(FDisplayClusterViewport& InViewport)
+void FDisplayClusterViewportPreviewMesh::Release(FDisplayClusterViewport* InViewport)
 {
 	UMeshComponent* MeshComponent = GetMeshComponent();
 	if (!MeshComponent && MeshComponentPtr)
 	{
 		// The mesh was destroyed earlier, (re-running build scripts inside RootActor), but we need to update the new mesh component to.
-		MeshComponent = GetOrCreatePreviewMeshComponent(&InViewport, bIsRootActorMeshComponent);
+		MeshComponent = GetOrCreatePreviewMeshComponent(InViewport, bIsRootActorMeshComponent);
 	}
 
 	if (MeshComponent)
@@ -127,34 +133,30 @@ void FDisplayClusterViewportPreviewMesh::Release(FDisplayClusterViewport& InView
 		}
 	}
 
-	Reset();
-}
-
-void FDisplayClusterViewportPreviewMesh::Reset()
-{
 	MeshComponentPtr = nullptr;
 	MaterialInstancePtr = nullptr;
 	CurrentMaterialType = EDisplayClusterDisplayDeviceMaterialType::DefaultPreviewMeshMaterial;
 }
 
-bool FDisplayClusterViewportPreviewMesh::ShouldUseMeshComponent(FDisplayClusterViewport& InViewport) const
+bool FDisplayClusterViewportPreviewMesh::ShouldUseMeshComponent(FDisplayClusterViewport* InViewport) const
 {
-	if (!InViewport.GetRenderSettings().bEnable)
+	if (!InViewport || !InViewport->GetRenderSettings().bEnable)
 	{
 		// disable viewport
 		return false;
 	}
 
-	TSharedPtr<IDisplayClusterProjectionPolicy, ESPMode::ThreadSafe> ProjectionPolicy = InViewport.GetProjectionPolicy();
+	const FDisplayClusterViewport_PreviewSettings PreviewSettings = InViewport->Configuration->GetRenderFrameSettings().PreviewSettings;
+	TSharedPtr<IDisplayClusterProjectionPolicy, ESPMode::ThreadSafe> ProjectionPolicy = InViewport->GetProjectionPolicy();
 	if (ProjectionPolicy.IsValid())
 	{
 		switch (MeshType)
 		{
 		case EDisplayClusterViewportPreviewMeshType::PreviewMesh:
-			return ProjectionPolicy->HasPreviewMesh(&InViewport) && InViewport.Configuration->GetRenderFrameSettings().PreviewSettings.bEnablePreviewMesh;
+			return ProjectionPolicy->HasPreviewMesh(InViewport) && PreviewSettings.bEnablePreviewMesh;
 
 		case EDisplayClusterViewportPreviewMeshType::PreviewEditableMesh:
-			return ProjectionPolicy->HasPreviewEditableMesh(&InViewport) && InViewport.Configuration->GetRenderFrameSettings().PreviewSettings.bEnablePreviewEditableMesh;
+			return ProjectionPolicy->HasPreviewEditableMesh(InViewport) && PreviewSettings.bEnablePreviewEditableMesh;
 
 		default:
 			break;
@@ -166,9 +168,7 @@ bool FDisplayClusterViewportPreviewMesh::ShouldUseMeshComponent(FDisplayClusterV
 
 UMeshComponent* FDisplayClusterViewportPreviewMesh::GetOrCreatePreviewMeshComponent(FDisplayClusterViewport* InViewport, bool& bOutIsRootActorComponent) const
 {
-	check(InViewport);
-
-	TSharedPtr<IDisplayClusterProjectionPolicy, ESPMode::ThreadSafe> ProjectionPolicy = InViewport->GetProjectionPolicy();
+	TSharedPtr<IDisplayClusterProjectionPolicy, ESPMode::ThreadSafe> ProjectionPolicy = InViewport ? InViewport->GetProjectionPolicy() : nullptr;
 	if (ProjectionPolicy.IsValid())
 	{
 		switch (MeshType)
diff --git a/Engine/Plugins/Runtime/nDisplay/Source/DisplayCluster/Private/Render/Viewport/Preview/DisplayClusterViewportPreviewMesh.h b/Engine/Plugins/Runtime/nDisplay/Source/DisplayCluster/Private/Render/Viewport/Preview/DisplayClusterViewportPreviewMesh.h
--- a/Engine/Plugins/Runtime/nDisplay/Source/DisplayCluster/Private/Render/Viewport/Preview/DisplayClusterViewportPreviewMesh.h
+++ b/Engine/Plugins/Runtime/nDisplay/Source/DisplayCluster/Private/Render/Viewport/Preview/DisplayClusterViewportPreviewMesh.h
@@ -77,13 +77,10 @@ public:
 	UMaterial* GetDefaultMaterial() const;
 
 	/** Update mesh component and materials for viewport. */
-	void Update(FDisplayClusterViewport& InViewport, UDisplayClusterDisplayDeviceBaseComponent& InDisplayDeviceComponent);
+	void Update(FDisplayClusterViewport* InViewport, UDisplayClusterDisplayDeviceBaseComponent* InDisplayDeviceComponent);
 
 	/** Restore default material and release mesh component with materials for viewport. */
-	void Release(FDisplayClusterViewport& InViewport);
-
-	/** Release mesh component and materials for viewport. */
-	void Reset();
+	void Release(FDisplayClusterViewport* InViewport);
 
 	/** Returns true if the runtime flags have any of the input flags. */
 	bool HasAnyFlag(const EDisplayClusterViewportPreviewMeshFlags InMeshFlags) const
@@ -93,7 +90,7 @@ public:
 
 private:
 	/** Returns true if this mesh type is supported by the viewport projection policy and DCRA. */
-	bool ShouldUseMeshComponent(FDisplayClusterViewport& InViewport) const;
+	bool ShouldUseMeshComponent(FDisplayClusterViewport* InViewport) const;
 
 	/** Get mesh component. */
 	UMeshComponent* GetOrCreatePreviewMeshComponent(FDisplayClusterViewport* InViewport, bool& bOutIsRootActorComponent) const;
diff --git a/Engine/Plugins/Runtime/nDisplay/Source/DisplayCluster/Public/Components/DisplayClusterPreviewShareComponent.h b/Engine/Plugins/Runtime/nDisplay/Source/DisplayCluster/Public/Components/DisplayClusterPreviewShareComponent.h
--- a/Engine/Plugins/Runtime/nDisplay/Source/DisplayCluster/Public/Components/DisplayClusterPreviewShareComponent.h
+++ b/Engine/Plugins/Runtime/nDisplay/Source/DisplayCluster/Public/Components/DisplayClusterPreviewShareComponent.h
@@ -104,14 +104,9 @@ protected:
 	/**
 	* Creates or updates a custom rendering for the specified RootActor.
 	* Returns nullptr if rendering is not possible, or a reference to a custom viewport manager storing rendered resources.
-	* 
-	.*/
+	*/
 	bool UpdateCustomViewportManager(const ADisplayClusterRootActor* InSrcRootActor);
 
-	/** Release the custom viewport manager instance, if it exists. */
-	void ReleaseCustomViewportManager();
-
-
 	//~ UActorComponent interface begin
 	virtual void OnRegister() override;
 	//~ UActorComponent interface end
@@ -142,6 +137,9 @@ private:
 	// Used to perform rendering with custom settings from an external DCRA
 	TSharedPtr<IDisplayClusterViewportManager, ESPMode::ThreadSafe> CustomViewportManager;
 
+	/** Disabling all special preview renderers and reverting to the default view. */
+	void RestoreDefaultPreviewSettings();
+
 	/** Closes all media related objects (i.e. media captures and media players) */
 	void CloseAllMedia();
 
@@ -166,9 +164,6 @@ private:
 	/** Logic that should run every tick when in Pull Actor mode */
 	void TickPullActor();
 
-	/** Retrieves preview textures from given source nDisplay actor and uses them to drive the preview textures of the owning actor */
-	void PullPreviewFromSourceActor(const ADisplayClusterRootActor* SourceRootActor);
-
 	/** Syncs the Icvfx cameras from the given source nDisplay actor to the given destination nDisplay actor */
 	void SyncIcvxCamerasFromSourceActor(const ADisplayClusterRootActor* SrcRootActor, const ADisplayClusterRootActor* DstRootActor);
 
diff --git a/Engine/Plugins/Runtime/nDisplay/Source/DisplayCluster/Public/DisplayClusterRootActor.h b/Engine/Plugins/Runtime/nDisplay/Source/DisplayCluster/Public/DisplayClusterRootActor.h
--- a/Engine/Plugins/Runtime/nDisplay/Source/DisplayCluster/Public/DisplayClusterRootActor.h
+++ b/Engine/Plugins/Runtime/nDisplay/Source/DisplayCluster/Public/DisplayClusterRootActor.h
@@ -402,8 +402,11 @@ public:
 		return LastDeltaSecondsValue;
 	}
 
-	/** Get current settings for preview rendering. */
-	FDisplayClusterViewport_PreviewSettings GetPreviewSettings() const;
+	/** Get current settings for preview rendering.
+	* 
+	* @param bIgnorePreviewSetttingsSource - if true, the PreviewSettingsSource is ignored
+	*/
+	FDisplayClusterViewport_PreviewSettings GetPreviewSettings(bool bIgnorePreviewSetttingsSource = false) const;
 
 //////////////////////////////////////////////////////////////////////////////////////////////
 // EDITOR RELATED SETTINGS
@@ -433,16 +436,16 @@ public:
 	bool bEnablePreviewTechvis = false;
 
 	/** Enable the use of a preview mesh for the preview for this DCRA. */
-	UPROPERTY()
+	UPROPERTY(Transient, NonTransactional)
 	bool bEnablePreviewMesh = true;
 
 	/** Enable the use of a preview editable mesh for the preview for this DCRA. */
-	UPROPERTY()
+	UPROPERTY(Transient, NonTransactional)
 	bool bEnablePreviewEditableMesh = true;
 
-	/** Always use preview settings defined in this RootActor. */
-	UPROPERTY()
-	bool bUseLocalPreviewSetttings = true;
+	/** Determines where the preview settings will be retrieved from. */
+	UPROPERTY(Transient, NonTransactional)
+	EDisplayClusterConfigurationRootActorPreviewSettingsSource PreviewSetttingsSource = EDisplayClusterConfigurationRootActorPreviewSettingsSource::RootActor;
 
 	/** Freeze preview render.  This will impact editor performance. */
 	UPROPERTY(EditAnywhere, BlueprintReadWrite, Category = "Editor Preview", meta = (DisplayName = "Freeze Editor Preview"))
@@ -458,7 +461,7 @@ public:
 
 #if WITH_EDITORONLY_DATA
 	/** When the MRQ is rendered, this flag is raised. */
-	UPROPERTY()
+	UPROPERTY(Transient, NonTransactional)
 	bool bMoviePipelineRenderPass = false;
 
 	/** Selectively preview a specific viewport or show all/none. */
@@ -525,10 +528,10 @@ public:
 	virtual bool ShouldTickIfViewportsOnly() const override { return true; }
 
 	UE_DEPRECATED(5.4, "This function has been deprecated.")
-	FOnPreviewUpdated& GetOnPreviewGenerated() { return OnPreviewGenerated; }
+	FOnPreviewUpdated& GetOnPreviewGenerated() { return DeprecatedPreviewDelegate; }
 
 	UE_DEPRECATED(5.4, "This function has been deprecated.")
-	FOnPreviewUpdated& GetOnPreviewDestroyed() { return OnPreviewDestroyed; }
+	FOnPreviewUpdated& GetOnPreviewDestroyed() { return DeprecatedPreviewDelegate; }
 
 	// return true, if preview enabled for this actor
 	UE_DEPRECATED(5.4, "This function has been deprecated.")
@@ -657,7 +660,7 @@ private:
 
 	TWeakPtr<IDisplayClusterConfiguratorBlueprintEditor> ToolkitPtr;
 
-	FOnPreviewUpdated OnPreviewGenerated;
-	FOnPreviewUpdated OnPreviewDestroyed;
+	// UE_DEPRECATED 5.4
+	FOnPreviewUpdated DeprecatedPreviewDelegate;
 #endif
 };
diff --git a/Engine/Plugins/Runtime/nDisplay/Source/DisplayClusterConfiguration/Public/DisplayClusterConfigurationTypes_Enums.h b/Engine/Plugins/Runtime/nDisplay/Source/DisplayClusterConfiguration/Public/DisplayClusterConfigurationTypes_Enums.h
--- a/Engine/Plugins/Runtime/nDisplay/Source/DisplayClusterConfiguration/Public/DisplayClusterConfigurationTypes_Enums.h
+++ b/Engine/Plugins/Runtime/nDisplay/Source/DisplayClusterConfiguration/Public/DisplayClusterConfigurationTypes_Enums.h
@@ -199,3 +199,15 @@ enum class EDisplayClusterConfigurationViewportLightcardOCIOMode : uint8
 	/** None. */
 	None UMETA(DisplayName = "None"),
 };
+
+/** Determines where the data for RootActor GetPreviewSettings() will come from: */
+UENUM()
+enum class EDisplayClusterConfigurationRootActorPreviewSettingsSource : uint8
+{
+	// The preview settings will be obtained from the RootActor properties.
+	RootActor,
+
+	// The preview settings will be retrieved from DC ViewportManagerConfiguration.
+	// This option is used when the preview settings need to be customized in a different way, but not to change the RootActor properties.
+	Configuration
+};
diff --git a/Engine/Plugins/Runtime/nDisplay/Source/DisplayClusterMoviePipeline/Private/DisplayClusterMoviePipelineViewportPass.cpp b/Engine/Plugins/Runtime/nDisplay/Source/DisplayClusterMoviePipeline/Private/DisplayClusterMoviePipelineViewportPass.cpp
--- a/Engine/Plugins/Runtime/nDisplay/Source/DisplayClusterMoviePipeline/Private/DisplayClusterMoviePipelineViewportPass.cpp
+++ b/Engine/Plugins/Runtime/nDisplay/Source/DisplayClusterMoviePipeline/Private/DisplayClusterMoviePipelineViewportPass.cpp
@@ -485,8 +485,12 @@ bool UDisplayClusterMoviePipelineViewportPassBase::InitializeDisplayCluster()
 
 			if (IDisplayClusterViewportManager* ViewportManager = DCRootActor->GetOrCreateViewportManager())
 			{
+				// Get preview settings from DCRootActor properties
+				FDisplayClusterViewport_PreviewSettings NewPreviewSettings = DCRootActor->GetPreviewSettings(true);
+				NewPreviewSettings.bPreviewEnable = false;
+
 				// Don't use preview settings for MRQ rendering (same as for game)
-				ViewportManager->GetConfiguration().SetPreviewSettings(FDisplayClusterViewport_PreviewSettings());
+				ViewportManager->GetConfiguration().SetPreviewSettings(NewPreviewSettings);
 
 				// Update local node viewports (update\create\delete) and build new render frame
 				if (ViewportManager->GetConfiguration().UpdateConfigurationForViewportsList(EDisplayClusterRenderFrameMode::MRQ_Mono, GetWorld(), DisplayClusterViewports))
