nDisplay: Fixed crash when calling FDisplayClusterRender_MeshComponentProxy destructor in game thread.
Usually the destructor is called from the rendering thread, but it can also be called from the game thread. (This happens when an object is deleted immediately after creation).

diff --git a/Engine/Plugins/Runtime/nDisplay/Source/DisplayCluster/Private/Render/Containers/DisplayClusterRender_MeshComponentProxy.cpp b/Engine/Plugins/Runtime/nDisplay/Source/DisplayCluster/Private/Render/Containers/DisplayClusterRender_MeshComponentProxy.cpp
--- a/Engine/Plugins/Runtime/nDisplay/Source/DisplayCluster/Private/Render/Containers/DisplayClusterRender_MeshComponentProxy.cpp
+++ b/Engine/Plugins/Runtime/nDisplay/Source/DisplayCluster/Private/Render/Containers/DisplayClusterRender_MeshComponentProxy.cpp
@@ -16,19 +16,17 @@ FDisplayClusterRender_MeshComponentProxy::FDisplayClusterRender_MeshComponentPro
 
 FDisplayClusterRender_MeshComponentProxy::~FDisplayClusterRender_MeshComponentProxy()
 {
-	check(IsInRenderingThread());
-
-	ImplRelease_RenderThread();
+	ImplRelease();
 }
 
 void FDisplayClusterRender_MeshComponentProxy::Release_RenderThread()
 {
 	check(IsInRenderingThread());
 
-	ImplRelease_RenderThread();
+	ImplRelease();
 }
 
-void FDisplayClusterRender_MeshComponentProxy::ImplRelease_RenderThread()
+void FDisplayClusterRender_MeshComponentProxy::ImplRelease()
 {
 	VertexBufferRHI.SafeRelease();
 	IndexBufferRHI.SafeRelease();
@@ -72,7 +70,7 @@ void FDisplayClusterRender_MeshComponentProxy::UpdateRHI_RenderThread(FRHIComman
 {
 	check(IsInRenderingThread());
 
-	ImplRelease_RenderThread();
+	ImplRelease();
 
 	if (InMeshData && InMeshData->IsValid())
 	{
diff --git a/Engine/Plugins/Runtime/nDisplay/Source/DisplayCluster/Private/Render/Containers/DisplayClusterRender_MeshComponentProxy.h b/Engine/Plugins/Runtime/nDisplay/Source/DisplayCluster/Private/Render/Containers/DisplayClusterRender_MeshComponentProxy.h
--- a/Engine/Plugins/Runtime/nDisplay/Source/DisplayCluster/Private/Render/Containers/DisplayClusterRender_MeshComponentProxy.h
+++ b/Engine/Plugins/Runtime/nDisplay/Source/DisplayCluster/Private/Render/Containers/DisplayClusterRender_MeshComponentProxy.h
@@ -20,7 +20,7 @@ public:
 	void Release_RenderThread();
 
 private:
-	void ImplRelease_RenderThread();
+	void ImplRelease();
 
 private:
 	/* RenderThread resources */
