nDisplay: Fix chromakey color not matching in -game nodes when the owning ICVFX camera wasn't configured to use the chromakey card.

diff --git a/Engine/Plugins/Runtime/nDisplay/Source/DisplayCluster/Private/Game/EngineClasses/Scene/DisplayClusterChromakeyCardActor.cpp b/Engine/Plugins/Runtime/nDisplay/Source/DisplayCluster/Private/Game/EngineClasses/Scene/DisplayClusterChromakeyCardActor.cpp
--- a/Engine/Plugins/Runtime/nDisplay/Source/DisplayCluster/Private/Game/EngineClasses/Scene/DisplayClusterChromakeyCardActor.cpp
+++ b/Engine/Plugins/Runtime/nDisplay/Source/DisplayCluster/Private/Game/EngineClasses/Scene/DisplayClusterChromakeyCardActor.cpp
@@ -5,35 +5,74 @@
 #include "DisplayClusterRootActor.h"
 #include "Components/DisplayClusterChromakeyCardStageActorComponent.h"
 #include "Components/DisplayClusterICVFXCameraComponent.h"
+#include "Misc/DisplayClusterLog.h"
 
 
 ADisplayClusterChromakeyCardActor::ADisplayClusterChromakeyCardActor(const FObjectInitializer& ObjectInitializer)
 	: Super(ObjectInitializer.DoNotCreateDefaultSubobject(LightCardStageActorComponentName))
 {
 	StageActorComponent = CreateOptionalDefaultSubobject<UDisplayClusterChromakeyCardStageActorComponent>(TEXT("ChromakeyStageActorComponent"));
+
+#if WITH_EDITOR
+	if (GIsEditor && !IsTemplate())
+	{
+		PropertyChangedHandle = FCoreUObjectDelegates::OnObjectPropertyChanged.AddUObject(this, &ADisplayClusterChromakeyCardActor::OnObjectPropertyChanged);
+	}
+#endif
+}
+
+ADisplayClusterChromakeyCardActor::~ADisplayClusterChromakeyCardActor()
+{
+#if WITH_EDITOR
+	if (PropertyChangedHandle.IsValid())
+	{
+		FCoreUObjectDelegates::OnObjectPropertyChanged.Remove(PropertyChangedHandle);
+	}
+#endif
 }
 
 void ADisplayClusterChromakeyCardActor::Tick(float DeltaSeconds)
 {
 	Super::Tick(DeltaSeconds);
+}
 
 #if WITH_EDITOR
-	if (GIsEditor)
+void ADisplayClusterChromakeyCardActor::PostEditUndo()
+{
+	Super::PostEditUndo();
+	
+	if (!GIsEditor && !IsTemplate())
 	{
-		UpdateChromakeySettings();
+		if (const UWorld* World = GetWorld())
+		{
+			TWeakObjectPtr<ADisplayClusterChromakeyCardActor> WeakPtrThis(this);
+
+			// Wait until the next frame to update color settings in -game. MU can call PostEditUndo
+			// on this object before the object that made the change (ICVFX cameras) has processed their new values.
+			World->GetTimerManager().SetTimerForNextTick([WeakPtrThis]()
+			{
+				if (WeakPtrThis.IsValid())
+				{
+					WeakPtrThis->UpdateChromakeySettings();
+				}
+			});
+		}
 	}
-#endif
 }
 
+#endif
+
 void ADisplayClusterChromakeyCardActor::AddToRootActor(ADisplayClusterRootActor* InRootActor)
 {
 	Super::AddToRootActor(InRootActor);
+	UpdateChromakeySettings();
 	// @todo uv chromakey cards - determine translucent sort order
 }
 
 void ADisplayClusterChromakeyCardActor::RemoveFromRootActor()
 {
 	Super::RemoveFromRootActor();
+	UpdateChromakeySettings();
 }
 
 bool ADisplayClusterChromakeyCardActor::IsReferencedByICVFXCamera(const UDisplayClusterICVFXCameraComponent* InCamera) const
@@ -76,6 +115,8 @@ bool ADisplayClusterChromakeyCardActor::IsReferencedByICVFXCamera(const UDisplay
 
 void ADisplayClusterChromakeyCardActor::UpdateChromakeySettings()
 {
+	DECLARE_SCOPE_CYCLE_COUNTER(TEXT("ADisplayClusterChromakeyCardActor::UpdateChromakeySettings"), STAT_UpdateChromakeySettings, STATGROUP_NDisplay);
+	
 	if (const ADisplayClusterRootActor* RootActor = GetRootActorOwner())
 	{
 		int32 Count = 0;
@@ -92,6 +133,41 @@ void ADisplayClusterChromakeyCardActor::UpdateChromakeySettings()
 			}
 		}
 
-		Color = Count > 0 ? ChromaColor / Count : FLinearColor::Green;
+		const FLinearColor NewColor = Count > 0 ? ChromaColor / Count : FLinearColor::Green;
+		if (NewColor != Color)
+		{
+#if WITH_EDITOR
+			if (GIsEditor)
+			{
+				Modify(false);
+			}
+#endif
+			
+			Color = NewColor;
+
+#if WITH_EDITOR
+			if (GIsEditor)
+			{
+				// Update MU
+				
+				FProperty* Property = FindFieldChecked<FProperty>(StaticClass(), GET_MEMBER_NAME_CHECKED(ADisplayClusterChromakeyCardActor, Color));
+				TArray<const FProperty*> ChangedProperties { Property };
+				SnapshotTransactionBuffer(this, ChangedProperties);
+			}
+#endif
+		}
+	}
+}
+
+#if WITH_EDITOR
+void ADisplayClusterChromakeyCardActor::OnObjectPropertyChanged(UObject* InObject,
+	FPropertyChangedEvent& InPropertyChangedEvent)
+{
+	if (InObject && ((InObject->IsA<UDisplayClusterICVFXCameraComponent>()
+		&& IsReferencedByICVFXCamera(CastChecked<UDisplayClusterICVFXCameraComponent>(InObject)))
+		|| InObject == this))
+	{
+		UpdateChromakeySettings();
 	}
 }
+#endif
diff --git a/Engine/Plugins/Runtime/nDisplay/Source/DisplayCluster/Public/DisplayClusterChromakeyCardActor.h b/Engine/Plugins/Runtime/nDisplay/Source/DisplayCluster/Public/DisplayClusterChromakeyCardActor.h
--- a/Engine/Plugins/Runtime/nDisplay/Source/DisplayCluster/Public/DisplayClusterChromakeyCardActor.h
+++ b/Engine/Plugins/Runtime/nDisplay/Source/DisplayCluster/Public/DisplayClusterChromakeyCardActor.h
@@ -14,9 +14,16 @@ class DISPLAYCLUSTER_API ADisplayClusterChromakeyCardActor : public ADisplayClus
 	GENERATED_BODY()
 public:
 	ADisplayClusterChromakeyCardActor(const FObjectInitializer& ObjectInitializer);
+	virtual ~ADisplayClusterChromakeyCardActor() override;
 
 	virtual void Tick(float DeltaSeconds) override;
 
+#if WITH_EDITOR
+	// ~Begin UObject interface
+	virtual void PostEditUndo() override;
+	// ~End UObject interface
+#endif
+
 	// ~Begin ADisplayClusterLightCardActor interface
 	virtual void AddToRootActor(ADisplayClusterRootActor* InRootActor) override;
 	virtual void RemoveFromRootActor() override;
@@ -32,4 +39,14 @@ public:
 	bool IsReferencedByICVFXCamera(const UDisplayClusterICVFXCameraComponent* InCamera) const;
 protected:
 	void UpdateChromakeySettings();
+
+#if WITH_EDITOR
+private:
+	/** Called when any UObject has its property changed. */
+	void OnObjectPropertyChanged(UObject* InObject, FPropertyChangedEvent& InPropertyChangedEvent);
+
+private:
+	/** Handle for UObject property change. */
+	FDelegateHandle PropertyChangedHandle;
+#endif
 };
