nDisplay: Fix deleted ICVFX panel stage actors not being added back to the root actor after undo.

diff --git a/Engine/Plugins/Runtime/nDisplay/Source/DisplayCluster/Private/Game/EngineClasses/Scene/DisplayClusterStageActorComponent.cpp b/Engine/Plugins/Runtime/nDisplay/Source/DisplayCluster/Private/Game/EngineClasses/Scene/DisplayClusterStageActorComponent.cpp
--- a/Engine/Plugins/Runtime/nDisplay/Source/DisplayCluster/Private/Game/EngineClasses/Scene/DisplayClusterStageActorComponent.cpp
+++ b/Engine/Plugins/Runtime/nDisplay/Source/DisplayCluster/Private/Game/EngineClasses/Scene/DisplayClusterStageActorComponent.cpp
@@ -6,6 +6,10 @@
 
 void UDisplayClusterStageActorComponent::SetRootActor(ADisplayClusterRootActor* InRootActor)
 {
+#if WITH_EDITOR
+	Modify();
+#endif
+	
 	RootActor = InRootActor;
 }
 
