nDisplay: Fixed crash when accessing TickableGameObject with value nullptr.
This was happening when UpdateTickableGameObject() was called twice.

diff --git a/Engine/Plugins/Runtime/nDisplay/Source/DisplayCluster/Public/Misc/DisplayClusterDataCache.h b/Engine/Plugins/Runtime/nDisplay/Source/DisplayCluster/Public/Misc/DisplayClusterDataCache.h
--- a/Engine/Plugins/Runtime/nDisplay/Source/DisplayCluster/Public/Misc/DisplayClusterDataCache.h
+++ b/Engine/Plugins/Runtime/nDisplay/Source/DisplayCluster/Public/Misc/DisplayClusterDataCache.h
@@ -130,7 +130,7 @@ protected:
 		if (CachedObjects.IsEmpty())
 		{
 			// Unregister tick event
-			if (TickHandle.IsValid())
+			if (TickableGameObject.IsValid() && TickHandle.IsValid())
 			{
 				TickableGameObject->OnTick().Remove(TickHandle);
 			}
