nDisplay: reworked preview rendering mesh materials implementation
- The material must be assigned to the mesh at each tick, without any conditions in case it can be changed externally.
- Added a preview option to disable material overlay when using preview rendering.
- disable overlay material for ICVFX panel

diff --git a/Engine/Plugins/Runtime/nDisplay/Source/DisplayCluster/Private/Game/EngineClasses/Scene/DisplayClusterRootActor.cpp b/Engine/Plugins/Runtime/nDisplay/Source/DisplayCluster/Private/Game/EngineClasses/Scene/DisplayClusterRootActor.cpp
--- a/Engine/Plugins/Runtime/nDisplay/Source/DisplayCluster/Private/Game/EngineClasses/Scene/DisplayClusterRootActor.cpp
+++ b/Engine/Plugins/Runtime/nDisplay/Source/DisplayCluster/Private/Game/EngineClasses/Scene/DisplayClusterRootActor.cpp
@@ -1160,6 +1160,7 @@ FDisplayClusterViewport_PreviewSettings ADisplayClusterRootActor::GetPreviewSett
 
 	OutPreviewSettings.bEnablePreviewTechvis     = bEnablePreviewTechvis;
 	OutPreviewSettings.bPreviewEnablePostProcess = bPreviewEnablePostProcess;
+	OutPreviewSettings.bPreviewEnableOverlayMaterial = bPreviewEnableOverlayMaterial;
 
 	OutPreviewSettings.bPreviewICVFXFrustums           = bPreviewICVFXFrustums;
 	OutPreviewSettings.PreviewICVFXFrustumsFarDistance = PreviewICVFXFrustumsFarDistance;
diff --git a/Engine/Plugins/Runtime/nDisplay/Source/DisplayCluster/Private/Render/Viewport/Preview/DisplayClusterViewportPreview.cpp b/Engine/Plugins/Runtime/nDisplay/Source/DisplayCluster/Private/Render/Viewport/Preview/DisplayClusterViewportPreview.cpp
--- a/Engine/Plugins/Runtime/nDisplay/Source/DisplayCluster/Private/Render/Viewport/Preview/DisplayClusterViewportPreview.cpp
+++ b/Engine/Plugins/Runtime/nDisplay/Source/DisplayCluster/Private/Render/Viewport/Preview/DisplayClusterViewportPreview.cpp
@@ -17,8 +17,8 @@ FDisplayClusterViewportPreview::FDisplayClusterViewportPreview(const TSharedRef<
 	: Configuration(InConfiguration)
 	, ViewportId(InViewportId)
 	, ClusterNodeId(InConfiguration->GetClusterNodeId())
-	, PreviewMesh(EDisplayClusterViewportPreviewMeshType::PreviewMesh)
-	, PreviewEditableMesh(EDisplayClusterViewportPreviewMeshType::PreviewEditableMesh)
+	, PreviewMesh(EDisplayClusterViewportPreviewMeshType::PreviewMesh, InConfiguration)
+	, PreviewEditableMesh(EDisplayClusterViewportPreviewMeshType::PreviewEditableMesh, InConfiguration)
 { }
 
 FDisplayClusterViewportPreview::~FDisplayClusterViewportPreview()
diff --git a/Engine/Plugins/Runtime/nDisplay/Source/DisplayCluster/Private/Render/Viewport/Preview/DisplayClusterViewportPreviewMesh.cpp b/Engine/Plugins/Runtime/nDisplay/Source/DisplayCluster/Private/Render/Viewport/Preview/DisplayClusterViewportPreviewMesh.cpp
--- a/Engine/Plugins/Runtime/nDisplay/Source/DisplayCluster/Private/Render/Viewport/Preview/DisplayClusterViewportPreviewMesh.cpp
+++ b/Engine/Plugins/Runtime/nDisplay/Source/DisplayCluster/Private/Render/Viewport/Preview/DisplayClusterViewportPreviewMesh.cpp
@@ -9,8 +9,9 @@
 
 #include "Components/StaticMeshComponent.h"
 
-namespace UE::DisplayCluster::Preview
+namespace UE::DisplayCluster::ViewportPreviewMesh
 {
+	/** Return ptr on UObject if it is still valid. */
 	template<class T>
 	static inline T* GetObjectProperty(const TObjectPtr<T>& InProperty)
 	{
@@ -23,10 +24,35 @@ namespace UE::DisplayCluster::Preview
 
 		return InProperty;
 	}
-}
+
+	/** Reset the mesh material to default values from its archetype. */
+	static inline bool RestoreMeshMaterialsFromArchetype(UMeshComponent* InMeshComponent)
+	{
+		if (const UMeshComponent* MeshArchetype = Cast<UMeshComponent>(InMeshComponent->GetArchetype()))
+		{
+			// Retrieve material on the preview mesh from an archetype or OrigMaterial.
+			if (UMaterialInterface* OrigMaterial = MeshArchetype->OverrideMaterials.IsValidIndex(0) ? MeshArchetype->OverrideMaterials[0] : nullptr)
+			{
+				InMeshComponent->SetMaterial(0, OrigMaterial);
+
+				return true;
+			}
+		}
+
+		return false;
+	}
+};
+
 ////////////////////////////////////////////////////////////////////////////////////////
 // FDisplayClusterViewportPreviewMesh
 ////////////////////////////////////////////////////////////////////////////////////////
+EDisplayClusterDisplayDeviceMaterialType FDisplayClusterViewportPreviewMesh::GetCurrentMaterialType() const
+{
+	return Configuration->IsTechvisEnabled()
+		? EDisplayClusterDisplayDeviceMaterialType::PreviewMeshTechvisMaterial
+		: EDisplayClusterDisplayDeviceMaterialType::PreviewMeshMaterial;
+}
+
 void FDisplayClusterViewportPreviewMesh::Update(FDisplayClusterViewport* InViewport, UDisplayClusterDisplayDeviceBaseComponent* InDisplayDeviceComponent)
 {
 	// Reset runtime flags before each update
@@ -43,11 +69,7 @@ void FDisplayClusterViewportPreviewMesh::Update(FDisplayClusterViewport* InViewp
 	DefaultMaterialPtr = InDisplayDeviceComponent->GetDisplayDeviceMaterial(EDisplayClusterDisplayDeviceMaterialType::DefaultPreviewMeshMaterial);
 
 	// Get current preview material
-	EDisplayClusterDisplayDeviceMaterialType NewCurrentMaterialType = InViewport->GetConfiguration().IsTechvisEnabled()
-		? EDisplayClusterDisplayDeviceMaterialType::PreviewMeshTechvisMaterial
-		: EDisplayClusterDisplayDeviceMaterialType::PreviewMeshMaterial;
-
-	UMaterial* InMeshMaterial = InDisplayDeviceComponent->GetDisplayDeviceMaterial(NewCurrentMaterialType);
+	UMaterial* InMeshMaterial = InDisplayDeviceComponent->GetDisplayDeviceMaterial(GetCurrentMaterialType());
 	if (!InMeshMaterial)
 	{
 		// The mesh component and its resources are no longer used.
@@ -63,6 +85,11 @@ void FDisplayClusterViewportPreviewMesh::Update(FDisplayClusterViewport* InViewp
 		// Release the reference to the old mesh component
 		Release(InViewport);
 	}
+	else if (CurrentMaterialPtr != InMeshMaterial)
+	{
+		// Release a material instance when the material is changed
+		ReleaseMaterialInstance();
+	}
 
 	if (NewMeshComponent != GetMeshComponent())
 	{
@@ -77,25 +104,87 @@ void FDisplayClusterViewportPreviewMesh::Update(FDisplayClusterViewport* InViewp
 	// Update material instance and assign to the  mesh
 	if (UMeshComponent* MeshComponent = GetMeshComponent())
 	{
-		if (GetCurrentMaterial() != InMeshMaterial || GetMaterialInstance() == nullptr)
+		// Update material instance and assign to the  mesh
+		if (GetMaterialInstance() == nullptr)
 		{
 			CurrentMaterialPtr = InMeshMaterial;
-			CurrentMaterialType = NewCurrentMaterialType;
 
 			MaterialInstancePtr = UMaterialInstanceDynamic::Create(InMeshMaterial, MeshComponent);
 			EnumAddFlags(RuntimeFlags, EDisplayClusterViewportPreviewMeshFlags::HasChangedMaterialInstance);
+		}
+
+		// The material must be assigned to the mesh at each tick, without any conditions in case it can be changed externally.
+		MeshComponent->SetMaterial(0, GetMaterialInstance());
+	}
+
+	// Handling the material overlay logic for the preview mesh:
+	UpdateOverlayMaterial(InViewport);
+}
+
+void FDisplayClusterViewportPreviewMesh::SetCustomOverlayMaterial(UMeshComponent* InMeshComponent, UMaterialInterface* InOverlayMaterial)
+{
+	if (InMeshComponent)
+	{
+		if (OrigOverlayMaterial == nullptr)
+		{
+			OrigOverlayMaterial = InMeshComponent->GetOverlayMaterial();
+		}
+
+		InMeshComponent->SetOverlayMaterial(InOverlayMaterial);
+	}
+}
 
-			// Set preview material
-			MeshComponent->SetMaterial(0, GetMaterialInstance());
+void FDisplayClusterViewportPreviewMesh::RestoreOverlayMaterial(UMeshComponent* InMeshComponent)
+{
+	if (InMeshComponent)
+	{
+		if (OrigOverlayMaterial)
+		{
+			InMeshComponent->SetOverlayMaterial(OrigOverlayMaterial);
+		}
+		else
+		{
+			if (const UMeshComponent* MeshArchetype = Cast<UMeshComponent>(InMeshComponent->GetArchetype()))
+			{
+				// Retrieve material on the preview mesh from an archetype or OrigMaterial.
+				if (UMaterialInterface* ArchetypeOverlayMaterial = MeshArchetype->GetOverlayMaterial())
+				{
+					InMeshComponent->SetOverlayMaterial(ArchetypeOverlayMaterial);
+				}
+			}
 		}
 	}
 }
 
-void FDisplayClusterViewportPreviewMesh::Release(FDisplayClusterViewport* InViewport)
+void FDisplayClusterViewportPreviewMesh::UpdateOverlayMaterial(FDisplayClusterViewport* InViewport)
 {
+	// Update material instance and assign to the  mesh
+	if (UMeshComponent* MeshComponent = GetMeshComponent())
+	{
+		if (Configuration->GetPreviewSettings().bPreviewEnableOverlayMaterial)
+		{
+			// Restore overlay material
+			RestoreOverlayMaterial(MeshComponent);
+		}
+		else
+		{
+			// Disable overlay material when preview is used
+			SetCustomOverlayMaterial(MeshComponent, nullptr);
+		}
+	}
+}
+
+void FDisplayClusterViewportPreviewMesh::ReleaseMeshComponent(FDisplayClusterViewport* InViewport)
+{
+	using namespace UE::DisplayCluster::ViewportPreviewMesh;
+
 	UMeshComponent* MeshComponent = GetMeshComponent();
 	if (!MeshComponent && MeshComponentPtr)
 	{
+		// Restore materials on deleted mesh
+		RestoreOverlayMaterial(MeshComponentPtr);
+		RestoreMeshMaterialsFromArchetype(MeshComponentPtr);
+
 		// The mesh was destroyed earlier, (re-running build scripts inside RootActor), but we need to update the new mesh component to.
 		MeshComponent = GetOrCreatePreviewMeshComponent(InViewport, bIsRootActorMeshComponent);
 	}
@@ -104,22 +193,31 @@ void FDisplayClusterViewportPreviewMesh::Release(FDisplayClusterViewport* InView
 	{
 		EnumAddFlags(RuntimeFlags, EDisplayClusterViewportPreviewMeshFlags::HasDeletedMeshComponent);
 
+		// Restore materials on exists mesh
+		RestoreOverlayMaterial(MeshComponent);
+		if (!RestoreMeshMaterialsFromArchetype(MeshComponent) && DefaultMaterialPtr)
+		{
+			MeshComponent->SetMaterial(0, DefaultMaterialPtr);
+		}
+
 		if (!bIsRootActorMeshComponent)
 		{
 			// Release this mesh component from DCRA
 			MeshComponent->UnregisterComponent();
 			MeshComponent->DestroyComponent();
 		}
-		else if (UMaterial* DefaultMaterial = GetDefaultMaterial())
+		else
 		{
 			// Restore the default material for an existing mesh in DCRA
 			EnumAddFlags(RuntimeFlags, EDisplayClusterViewportPreviewMeshFlags::HasRestoredDefaultMaterial);
-
-			CurrentMaterialPtr = DefaultMaterial;
-			MeshComponent->SetMaterial(0, DefaultMaterial);
 		}
 	}
 
+	MeshComponentPtr = nullptr;
+}
+
+void FDisplayClusterViewportPreviewMesh::ReleaseMaterialInstance()
+{
 	// The material instance references the mesh, so it must also be deleted
 	if (UMaterialInstanceDynamic* MaterialInstance = GetMaterialInstance())
 	{
@@ -133,9 +231,8 @@ void FDisplayClusterViewportPreviewMesh::Release(FDisplayClusterViewport* InView
 		}
 	}
 
-	MeshComponentPtr = nullptr;
 	MaterialInstancePtr = nullptr;
-	CurrentMaterialType = EDisplayClusterDisplayDeviceMaterialType::DefaultPreviewMeshMaterial;
+	CurrentMaterialPtr = nullptr;
 }
 
 bool FDisplayClusterViewportPreviewMesh::ShouldUseMeshComponent(FDisplayClusterViewport* InViewport) const
@@ -146,7 +243,7 @@ bool FDisplayClusterViewportPreviewMesh::ShouldUseMeshComponent(FDisplayClusterV
 		return false;
 	}
 
-	const FDisplayClusterViewport_PreviewSettings PreviewSettings = InViewport->Configuration->GetRenderFrameSettings().PreviewSettings;
+	const FDisplayClusterViewport_PreviewSettings PreviewSettings = Configuration->GetRenderFrameSettings().PreviewSettings;
 	TSharedPtr<IDisplayClusterProjectionPolicy, ESPMode::ThreadSafe> ProjectionPolicy = InViewport->GetProjectionPolicy();
 	if (ProjectionPolicy.IsValid())
 	{
@@ -190,20 +287,14 @@ UMeshComponent* FDisplayClusterViewportPreviewMesh::GetOrCreatePreviewMeshCompon
 
 UMeshComponent* FDisplayClusterViewportPreviewMesh::GetMeshComponent() const
 {
-	return UE::DisplayCluster::Preview::GetObjectProperty(MeshComponentPtr);
-}
+	using namespace UE::DisplayCluster::ViewportPreviewMesh;
 
-UMaterialInstanceDynamic* FDisplayClusterViewportPreviewMesh::GetMaterialInstance() const
-{
-	return UE::DisplayCluster::Preview::GetObjectProperty(MaterialInstancePtr);
+	return GetObjectProperty(MeshComponentPtr);
 }
 
-UMaterial* FDisplayClusterViewportPreviewMesh::GetCurrentMaterial() const
+UMaterialInstanceDynamic* FDisplayClusterViewportPreviewMesh::GetMaterialInstance() const
 {
-	return UE::DisplayCluster::Preview::GetObjectProperty(CurrentMaterialPtr);
-}
+	using namespace UE::DisplayCluster::ViewportPreviewMesh;
 
-UMaterial* FDisplayClusterViewportPreviewMesh::GetDefaultMaterial() const
-{
-	return UE::DisplayCluster::Preview::GetObjectProperty(DefaultMaterialPtr);
+	return GetObjectProperty(MaterialInstancePtr);
 }
diff --git a/Engine/Plugins/Runtime/nDisplay/Source/DisplayCluster/Private/Render/Viewport/Preview/DisplayClusterViewportPreviewMesh.h b/Engine/Plugins/Runtime/nDisplay/Source/DisplayCluster/Private/Render/Viewport/Preview/DisplayClusterViewportPreviewMesh.h
--- a/Engine/Plugins/Runtime/nDisplay/Source/DisplayCluster/Private/Render/Viewport/Preview/DisplayClusterViewportPreviewMesh.h
+++ b/Engine/Plugins/Runtime/nDisplay/Source/DisplayCluster/Private/Render/Viewport/Preview/DisplayClusterViewportPreviewMesh.h
@@ -7,6 +7,7 @@
 #include "Materials/MaterialInstanceDynamic.h"
 
 #include "Render/DisplayDevice/Containers/DisplayClusterDisplayDevice_Enums.h"
+#include "Render/Viewport/Configuration/DisplayClusterViewportConfiguration.h"
 
 #include "Misc/DisplayClusterObjectRef.h"
 
@@ -52,8 +53,8 @@ ENUM_CLASS_FLAGS(EDisplayClusterViewportPreviewMeshFlags);
 class FDisplayClusterViewportPreviewMesh
 {
 public:
-	FDisplayClusterViewportPreviewMesh(const EDisplayClusterViewportPreviewMeshType InMeshType)
-		: MeshType(InMeshType)
+	FDisplayClusterViewportPreviewMesh(const EDisplayClusterViewportPreviewMeshType InMeshType, const TSharedRef<FDisplayClusterViewportConfiguration, ESPMode::ThreadSafe> InConfiguration)
+		: Configuration(InConfiguration), MeshType(InMeshType)
 	{ }
 
 	~FDisplayClusterViewportPreviewMesh() = default;
@@ -65,22 +66,18 @@ public:
 	/** Get MID used on preview mesh. */
 	UMaterialInstanceDynamic* GetMaterialInstance() const;
 
-	EDisplayClusterDisplayDeviceMaterialType GetCurrentMaterialType() const
-	{
-		return CurrentMaterialType;
-	}
-
-	/** Get current material. */
-	UMaterial* GetCurrentMaterial() const;
-
-	/** Get default material. */
-	UMaterial* GetDefaultMaterial() const;
+	/** Return current material type that used on this mesh. */
+	EDisplayClusterDisplayDeviceMaterialType GetCurrentMaterialType() const;
 
 	/** Update mesh component and materials for viewport. */
 	void Update(FDisplayClusterViewport* InViewport, UDisplayClusterDisplayDeviceBaseComponent* InDisplayDeviceComponent);
 
 	/** Restore default material and release mesh component with materials for viewport. */
-	void Release(FDisplayClusterViewport* InViewport);
+	void Release(FDisplayClusterViewport* InViewport)
+	{
+		ReleaseMeshComponent(InViewport);
+		ReleaseMaterialInstance();
+	}
 
 	/** Returns true if the runtime flags have any of the input flags. */
 	bool HasAnyFlag(const EDisplayClusterViewportPreviewMeshFlags InMeshFlags) const
@@ -92,13 +89,32 @@ private:
 	/** Returns true if this mesh type is supported by the viewport projection policy and DCRA. */
 	bool ShouldUseMeshComponent(FDisplayClusterViewport* InViewport) const;
 
+	/** Release material instance */
+	void ReleaseMaterialInstance();
+
+	/** Release mesh component. */
+	void ReleaseMeshComponent(FDisplayClusterViewport* InViewport);
+
+	/** Update overlay materials on mesh. */
+	void UpdateOverlayMaterial(FDisplayClusterViewport* InViewport);
+
+	/** Set custom overlay materials on mesh. */
+	void SetCustomOverlayMaterial(UMeshComponent* InMeshComponent, UMaterialInterface* InOverlayMaterial);
+
+	/** Restore overlay material on mesh.*/
+	void RestoreOverlayMaterial(UMeshComponent* InMeshComponent);
+
 	/** Get mesh component. */
 	UMeshComponent* GetOrCreatePreviewMeshComponent(FDisplayClusterViewport* InViewport, bool& bOutIsRootActorComponent) const;
 
-private:
+public:
+	// Configuration of the current cluster node
+	const TSharedRef<FDisplayClusterViewportConfiguration, ESPMode::ThreadSafe> Configuration;
+
 	// the type of mesh
 	const EDisplayClusterViewportPreviewMeshType MeshType;
 
+private:
 	// runtime flags
 	EDisplayClusterViewportPreviewMeshFlags RuntimeFlags;
 
@@ -108,6 +124,9 @@ private:
 	// This mesh component exists in DCRA and does not need to be deleted.
 	bool bIsRootActorMeshComponent = false;
 
+	// Overlay materials for the mesh can be customized when using preview rendering.
+	TObjectPtr<UMaterialInterface> OrigOverlayMaterial;
+
 	// Preview material used on the mesh
 	TObjectPtr<UMaterialInstanceDynamic> MaterialInstancePtr = nullptr;
 
@@ -116,7 +135,4 @@ private:
 
 	// The default material defined in the DisplayDevice
 	TObjectPtr<UMaterial> DefaultMaterialPtr = nullptr;
-
-	// The type of material that is currently used on the mesh
-	EDisplayClusterDisplayDeviceMaterialType CurrentMaterialType = EDisplayClusterDisplayDeviceMaterialType::DefaultPreviewMeshMaterial;
 };
diff --git a/Engine/Plugins/Runtime/nDisplay/Source/DisplayCluster/Public/DisplayClusterRootActor.h b/Engine/Plugins/Runtime/nDisplay/Source/DisplayCluster/Public/DisplayClusterRootActor.h
--- a/Engine/Plugins/Runtime/nDisplay/Source/DisplayCluster/Public/DisplayClusterRootActor.h
+++ b/Engine/Plugins/Runtime/nDisplay/Source/DisplayCluster/Public/DisplayClusterRootActor.h
@@ -431,6 +431,10 @@ public:
 	UPROPERTY(EditAnywhere, BlueprintReadWrite, Category = "Editor Preview", meta = (DisplayName = "Enable Post Process"), BlueprintSetter = SetPreviewEnablePostProcess)
 	bool bPreviewEnablePostProcess = false;
 
+	/** Show overlay material on the preview mesh when preview rendering is enabled (UMeshComponent::OverlayMaterial). */
+	UPROPERTY(EditAnywhere, BlueprintReadWrite, Category = "Editor Preview", meta = (DisplayName = "Enable Preview Overlay"))
+	bool bPreviewEnableOverlayMaterial = true;
+
 	/** Configure the root actor for Techvis rendering with preview components. */
 	UPROPERTY(EditAnywhere, BlueprintReadWrite, Category = "Editor Preview")
 	bool bEnablePreviewTechvis = false;
diff --git a/Engine/Plugins/Runtime/nDisplay/Source/DisplayCluster/Public/Render/Viewport/Containers/DisplayClusterViewport_PreviewSettings.h b/Engine/Plugins/Runtime/nDisplay/Source/DisplayCluster/Public/Render/Viewport/Containers/DisplayClusterViewport_PreviewSettings.h
--- a/Engine/Plugins/Runtime/nDisplay/Source/DisplayCluster/Public/Render/Viewport/Containers/DisplayClusterViewport_PreviewSettings.h
+++ b/Engine/Plugins/Runtime/nDisplay/Source/DisplayCluster/Public/Render/Viewport/Containers/DisplayClusterViewport_PreviewSettings.h
@@ -31,6 +31,9 @@ struct FDisplayClusterViewport_PreviewSettings
 	// But in it, PostProcess is also rendered with the flag turned off.
 	bool bPreviewEnablePostProcess = false;
 
+	// Show overlay material on the preview mesh when preview rendering is enabled (UMeshComponent::OverlayMaterial)
+	bool bPreviewEnableOverlayMaterial = true;
+
 	// Allows you to process preview meshes inside DCRA (get or create a mesh from a projection policy, update materials on the preview mesh, etc.).
 	bool bEnablePreviewMesh = false;
 
diff --git a/Engine/Plugins/Runtime/nDisplay/Source/DisplayClusterLightCardEditor/Private/Viewport/DisplayClusterLightCardEditorViewportClient.cpp b/Engine/Plugins/Runtime/nDisplay/Source/DisplayClusterLightCardEditor/Private/Viewport/DisplayClusterLightCardEditorViewportClient.cpp
--- a/Engine/Plugins/Runtime/nDisplay/Source/DisplayClusterLightCardEditor/Private/Viewport/DisplayClusterLightCardEditorViewportClient.cpp
+++ b/Engine/Plugins/Runtime/nDisplay/Source/DisplayClusterLightCardEditor/Private/Viewport/DisplayClusterLightCardEditorViewportClient.cpp
@@ -1643,6 +1643,7 @@ void FDisplayClusterLightCardEditorViewportClient::UpdatePreviewActor(ADisplayCl
 				RootActorProxy->bPreviewICVFXFrustums = false;
 				RootActorProxy->bEnablePreviewTechvis = false;
 				RootActorProxy->bPreviewEnablePostProcess = true;
+				RootActorProxy->bPreviewEnableOverlayMaterial = false;
 				RootActorProxy->bPreviewEnable = true;
 				RootActorProxy->bPreviewICVFXFrustums = false;
 
