nDisplay: Added definition and customization of DisplayDevice materials and meshes for Warp policies via ViewPointComponent.
- InFrustumFit's warp policy now uses and customizes its own material for preview meshes in front of the camera.
- Added changes to the IDisplayClusterWarpPolicy interface
- Added changes to the UDisplayClusterDisplayDeviceBaseComponent and UDisplayClusterDisplayDeviceComponent
- The preview materials are changed only after the DC viewport has finished rendering.
- Added new flag EDisplayClusterViewportResourceState::UpdatedOnRenderingThread
- Code clean up

Note: PDisplayClusterWarpStrings.h references a new material that does not exist "/nDisplay/Materials/Preview/M_InFrustumFitPreview".
This material will be added in the next CL, as its use assumes that the math for InFrustumFit will be removed from the DisplayDevice materials.

diff --git a/Engine/Plugins/Runtime/nDisplay/Source/DisplayCluster/Private/Game/EngineClasses/Scene/DisplayClusterRootActor.cpp b/Engine/Plugins/Runtime/nDisplay/Source/DisplayCluster/Private/Game/EngineClasses/Scene/DisplayClusterRootActor.cpp
--- a/Engine/Plugins/Runtime/nDisplay/Source/DisplayCluster/Private/Game/EngineClasses/Scene/DisplayClusterRootActor.cpp
+++ b/Engine/Plugins/Runtime/nDisplay/Source/DisplayCluster/Private/Game/EngineClasses/Scene/DisplayClusterRootActor.cpp
@@ -139,7 +139,9 @@ ADisplayClusterRootActor::ADisplayClusterRootActor(const FObjectInitializer& Obj
 
 	// Our internal display device which always exists
 	BasicDisplayDeviceComponent = CreateDefaultSubobject<UDisplayClusterDisplayDeviceComponent>(GetInternalDisplayDeviceName());
-	LineBatcherComponent = CreateDefaultSubobject<ULineBatchComponent>(TEXT("LineBatcher"));
+
+	// Create transient line batcher component
+	LineBatcherComponent = CreateDefaultSubobject<ULineBatchComponent>(TEXT("LineBatcher"), true);
 
 	ResetEntireClusterPreviewRendering();
 
diff --git a/Engine/Plugins/Runtime/nDisplay/Source/DisplayCluster/Private/Render/DisplayDevice/Components/DisplayClusterDisplayDeviceBaseComponent.cpp b/Engine/Plugins/Runtime/nDisplay/Source/DisplayCluster/Private/Render/DisplayDevice/Components/DisplayClusterDisplayDeviceBaseComponent.cpp
--- a/Engine/Plugins/Runtime/nDisplay/Source/DisplayCluster/Private/Render/DisplayDevice/Components/DisplayClusterDisplayDeviceBaseComponent.cpp
+++ b/Engine/Plugins/Runtime/nDisplay/Source/DisplayCluster/Private/Render/DisplayDevice/Components/DisplayClusterDisplayDeviceBaseComponent.cpp
@@ -9,6 +9,68 @@
 #include "UObject/ConstructorHelpers.h"
 #include "Engine/TextureRenderTarget2D.h"
 
+namespace UE::DisplayCluster::DisplayDeviceBaseComponent
+{
+	static inline void ImplUpdatePreviewMeshTechvis(IDisplayClusterViewportPreview& InViewportPreview, UMeshComponent& InMeshComponent)
+	{
+		if (InViewportPreview.GetConfiguration().IsTechvisEnabled())
+		{
+			InMeshComponent.bAffectDynamicIndirectLighting = true;
+			InMeshComponent.bAffectIndirectLightingWhileHidden = true;
+		}
+		else
+		{
+			// When disabling just revert to the archetype values. The original values at the time of enabling
+			// techvis aren't saved, and techvis is enabled by default so they'll always be set to true on new meshes.
+			if (const UMeshComponent* MeshArchetype = Cast<UMeshComponent>(InMeshComponent.GetArchetype()))
+			{
+				InMeshComponent.bAffectDynamicIndirectLighting = MeshArchetype->bAffectDynamicIndirectLighting;
+				InMeshComponent.bAffectIndirectLightingWhileHidden = MeshArchetype->bAffectIndirectLightingWhileHidden;
+			}
+		}
+
+		// For all preview meshes:
+		InMeshComponent.SetCastShadow(false);
+		InMeshComponent.SetHiddenInGame(false);
+		InMeshComponent.SetVisibility(true);
+
+		InMeshComponent.bVisibleInReflectionCaptures = false;
+		InMeshComponent.bVisibleInRayTracing = false;
+		InMeshComponent.bVisibleInRealTimeSkyCaptures = false;
+	};
+
+	static inline void ImplUpdatePreviewMeshMaterialInstanceParameters(IDisplayClusterViewportPreview& InViewportPreview, const EDisplayClusterDisplayDeviceMaterialType InMaterialType, UMaterialInstanceDynamic& InMaterialInstance)
+	{
+		// Update preview RTT parameter values:
+		switch (InMaterialType)
+		{
+		case EDisplayClusterDisplayDeviceMaterialType::PreviewMeshMaterial:
+		case EDisplayClusterDisplayDeviceMaterialType::PreviewMeshTechvisMaterial:
+			if (InViewportPreview.HasAnyFlags(EDisplayClusterViewportPreviewFlags::HasChangedPreviewMeshMaterialInstance | EDisplayClusterViewportPreviewFlags::HasChangedPreviewEditableMeshMaterialInstance | EDisplayClusterViewportPreviewFlags::HasChangedPreviewRTT))
+			{
+				// Updates the RTT parameter for the material instance when the PreviewTexture or mesh material is changed.
+				if (UTextureRenderTarget2D* PreviewTexture = InViewportPreview.GetPreviewTextureRenderTarget2D())
+				{
+					InMaterialInstance.SetTextureParameterValue(UE::DisplayClusterDisplayDeviceStrings::material::attr::Preview, PreviewTexture);
+				}
+				else
+				{
+					// Note: when preview texture not available, show default
+					static TObjectPtr<UTexture2D> DefaultGridTexture = LoadObject<UTexture2D>(nullptr, TEXT("/Engine/EditorMaterials/T_1x1_Grid.T_1x1_Grid"));
+					if (DefaultGridTexture)
+					{
+						InMaterialInstance.SetTextureParameterValue(UE::DisplayClusterDisplayDeviceStrings::material::attr::Preview, DefaultGridTexture);
+					}
+				}
+			}
+			break;
+
+		default:
+			break;
+		}
+	}
+};
+
 bool UDisplayClusterDisplayDeviceBaseComponent::ShouldUseDisplayDevice(IDisplayClusterViewportConfiguration& InConfiguration) const
 {
 	// Use this Display device only for preview
@@ -47,18 +109,27 @@ UDisplayClusterDisplayDeviceBaseComponent::UDisplayClusterDisplayDeviceBaseCompo
 	PreviewMeshTechvisMaterial = PreviewMeshTechvisMaterialObj.Object;
 }
 
-TObjectPtr<UMaterial> UDisplayClusterDisplayDeviceBaseComponent::GetDisplayDeviceMaterial(const EDisplayClusterDisplayDeviceMaterialType InMaterialType) const
+TObjectPtr<UMaterial> UDisplayClusterDisplayDeviceBaseComponent::GetDisplayDeviceMaterial(const EDisplayClusterDisplayDeviceMeshType InMeshType, const EDisplayClusterDisplayDeviceMaterialType InMaterialType) const
 {
-	switch (InMaterialType)
+	switch (InMeshType)
 	{
-	case EDisplayClusterDisplayDeviceMaterialType::DefaultPreviewMeshMaterial:
+	case EDisplayClusterDisplayDeviceMeshType::DefaultMesh:
 		return MeshMaterial;
 
-	case EDisplayClusterDisplayDeviceMaterialType::PreviewMeshMaterial:
-		return PreviewMeshMaterial;
+	case EDisplayClusterDisplayDeviceMeshType::PreviewMesh:
+	case EDisplayClusterDisplayDeviceMeshType::PreviewEditableMesh:
+		switch (InMaterialType)
+		{
+		case EDisplayClusterDisplayDeviceMaterialType::PreviewMeshMaterial:
+			return PreviewMeshMaterial;
 
-	case EDisplayClusterDisplayDeviceMaterialType::PreviewMeshTechvisMaterial:
-		return PreviewMeshTechvisMaterial;
+		case EDisplayClusterDisplayDeviceMaterialType::PreviewMeshTechvisMaterial:
+			return PreviewMeshTechvisMaterial;
+
+		default:
+			break;
+		}
+		break;
 
 	default:
 		break;
@@ -67,76 +138,35 @@ TObjectPtr<UMaterial> UDisplayClusterDisplayDeviceBaseComponent::GetDisplayDevic
 	return nullptr;
 }
 
-void UDisplayClusterDisplayDeviceBaseComponent::OnUpdateDisplayDeviceMaterialInstance(IDisplayClusterViewportPreview& InViewportPreview, const EDisplayClusterDisplayDeviceMeshType InMeshType, const EDisplayClusterDisplayDeviceMaterialType InMaterialType, UMaterialInstanceDynamic* InMaterialInstance) const
+void UDisplayClusterDisplayDeviceBaseComponent::OnUpdateDisplayDeviceMeshAndMaterialInstance(IDisplayClusterViewportPreview& InViewportPreview, const EDisplayClusterDisplayDeviceMeshType InMeshType, const EDisplayClusterDisplayDeviceMaterialType InMaterialType, UMeshComponent* InMeshComponent, UMaterialInstanceDynamic* InMeshMaterialInstance) const
 {
-	if (!InMaterialInstance || !ShouldUseDisplayDevice(InViewportPreview.GetConfiguration()))
-	{
-		return;
-	}
+	using namespace UE::DisplayCluster::DisplayDeviceBaseComponent;
 
-	// Update preview RTT parameter values:
-	switch (InMaterialType)
-	{
-	case EDisplayClusterDisplayDeviceMaterialType::PreviewMeshMaterial:
-	case EDisplayClusterDisplayDeviceMaterialType::PreviewMeshTechvisMaterial:
-		if (InViewportPreview.HasAnyFlags(EDisplayClusterViewportPreviewFlags::HasChangedPreviewMeshMaterialInstance | EDisplayClusterViewportPreviewFlags::HasChangedPreviewEditableMeshMaterialInstance | EDisplayClusterViewportPreviewFlags::HasChangedPreviewRTT))
-		{
-			// Updates the RTT parameter for the material instance when the PreviewTexture or mesh material is changed.
-			InMaterialInstance->SetTextureParameterValue(UE::DisplayClusterDisplayDeviceStrings::material::attr::Preview, InViewportPreview.GetPreviewTextureRenderTarget2D());
-		}
-		break;
-	default:
-		break;
-	}
-}
-
-void UDisplayClusterDisplayDeviceBaseComponent::OnUpdateDisplayDeviceMeshComponent(IDisplayClusterViewportPreview& InViewportPreview, const EDisplayClusterDisplayDeviceMeshType InMeshType, UMeshComponent* InMeshComponent) const
-{
-	if (!InMeshComponent || !ShouldUseDisplayDevice(InViewportPreview.GetConfiguration()))
+	if (!InMeshComponent || !InMeshMaterialInstance || !ShouldUseDisplayDevice(InViewportPreview.GetConfiguration()))
 	{
 		return;
 	}
 
-	// Only some of the preview meshes are supported by Techvis
 	switch (InMeshType)
 	{
-		case EDisplayClusterDisplayDeviceMeshType::PreviewMesh:
-		case EDisplayClusterDisplayDeviceMeshType::PreviewEditableMesh:
-			break;
+	case EDisplayClusterDisplayDeviceMeshType::PreviewMesh:
+	case EDisplayClusterDisplayDeviceMeshType::PreviewEditableMesh:
+		// Only some of the preview meshes are supported by Techvis
+		ImplUpdatePreviewMeshTechvis(InViewportPreview, *InMeshComponent);
 
-		default:
-			// Dont apply Techvis for other mesh types
-			return;
-	}
+		// Only some of the preview meshes are supported by preview
+		ImplUpdatePreviewMeshMaterialInstanceParameters(InViewportPreview, InMaterialType, *InMeshMaterialInstance);
+		break;
 
-	if (InViewportPreview.GetConfiguration().IsTechvisEnabled())
-	{
-		InMeshComponent->bAffectDynamicIndirectLighting = true;
-		InMeshComponent->bAffectIndirectLightingWhileHidden = true;
-	}
-	else
-	{
-		// When disabling just revert to the archetype values. The original values at the time of enabling
-		// techvis aren't saved, and techvis is enabled by default so they'll always be set to true on new meshes.
-		if (const UMeshComponent* MeshArchetype = Cast<UMeshComponent>(InMeshComponent->GetArchetype()))
-		{
-			InMeshComponent->bAffectDynamicIndirectLighting = MeshArchetype->bAffectDynamicIndirectLighting;
-			InMeshComponent->bAffectIndirectLightingWhileHidden = MeshArchetype->bAffectIndirectLightingWhileHidden;
-		}
+	default:
+		break;
 	}
-
-	// For all preview meshes:
-	InMeshComponent->SetCastShadow(false);
-	InMeshComponent->SetHiddenInGame(false);
-	InMeshComponent->SetVisibility(true);
-
-	InMeshComponent->bVisibleInReflectionCaptures = false;
-	InMeshComponent->bVisibleInRayTracing = false;
-	InMeshComponent->bVisibleInRealTimeSkyCaptures = false;
 }
 
 void UDisplayClusterDisplayDeviceBaseComponent::SetupSceneView(const IDisplayClusterViewportPreview& InViewportPreview, uint32 ContextNum, FSceneViewFamily& InOutViewFamily, FSceneView& InOutView) const
-{ }
+{
+	// This component does not change rendering settings
+}
 
 void UDisplayClusterDisplayDeviceBaseComponent::UpdateDisplayDeviceProxyImpl(IDisplayClusterViewportConfiguration& InConfiguration)
 {
diff --git a/Engine/Plugins/Runtime/nDisplay/Source/DisplayCluster/Private/Render/DisplayDevice/Components/DisplayClusterDisplayDeviceComponent.cpp b/Engine/Plugins/Runtime/nDisplay/Source/DisplayCluster/Private/Render/DisplayDevice/Components/DisplayClusterDisplayDeviceComponent.cpp
--- a/Engine/Plugins/Runtime/nDisplay/Source/DisplayCluster/Private/Render/DisplayDevice/Components/DisplayClusterDisplayDeviceComponent.cpp
+++ b/Engine/Plugins/Runtime/nDisplay/Source/DisplayCluster/Private/Render/DisplayDevice/Components/DisplayClusterDisplayDeviceComponent.cpp
@@ -15,28 +15,28 @@
 UDisplayClusterDisplayDeviceComponent::UDisplayClusterDisplayDeviceComponent()
 { }
 
-void UDisplayClusterDisplayDeviceComponent::OnUpdateDisplayDeviceMaterialInstance(IDisplayClusterViewportPreview& InViewportPreview, const EDisplayClusterDisplayDeviceMeshType InMeshType, const EDisplayClusterDisplayDeviceMaterialType InMaterialType, UMaterialInstanceDynamic* InMaterialInstance) const
+void UDisplayClusterDisplayDeviceComponent::OnUpdateDisplayDeviceMeshAndMaterialInstance(IDisplayClusterViewportPreview& InViewportPreview, const EDisplayClusterDisplayDeviceMeshType InMeshType, const EDisplayClusterDisplayDeviceMaterialType InMaterialType, UMeshComponent* InMeshComponent, UMaterialInstanceDynamic* InMeshMaterialInstance) const
 {
-	if (!InMaterialInstance || !ShouldUseDisplayDevice(InViewportPreview.GetConfiguration()))
-	{
-		return;
-	}
+	// Call a method of the base class
+	Super::OnUpdateDisplayDeviceMeshAndMaterialInstance(InViewportPreview, InMeshType, InMaterialType, InMeshComponent, InMeshMaterialInstance);
 
-	switch (InMaterialType)
+	// Customizes and overrides material parameters after the base class
+	if (InMeshMaterialInstance && ShouldUseDisplayDevice(InViewportPreview.GetConfiguration()))
 	{
-	case EDisplayClusterDisplayDeviceMaterialType::PreviewMeshMaterial:
-		InMaterialInstance->SetScalarParameterValue(UE::DisplayClusterDisplayDeviceStrings::material::attr::Exposure, 0.f);
-		break;
+		switch (InMaterialType)
+		{
+		case EDisplayClusterDisplayDeviceMaterialType::PreviewMeshMaterial:
+			InMeshMaterialInstance->SetScalarParameterValue(UE::DisplayClusterDisplayDeviceStrings::material::attr::Exposure, 0.f);
+			break;
 
-	case EDisplayClusterDisplayDeviceMaterialType::PreviewMeshTechvisMaterial:
-		InMaterialInstance->SetScalarParameterValue(UE::DisplayClusterDisplayDeviceStrings::material::attr::Exposure, Exposure);
-		break;
+		case EDisplayClusterDisplayDeviceMaterialType::PreviewMeshTechvisMaterial:
+			InMeshMaterialInstance->SetScalarParameterValue(UE::DisplayClusterDisplayDeviceStrings::material::attr::Exposure, Exposure);
+			break;
 
-	default:
-		break;
+		default:
+			break;
+		}
 	}
-
-	Super::OnUpdateDisplayDeviceMaterialInstance(InViewportPreview, InMeshType, InMaterialType, InMaterialInstance);
 }
 
 void UDisplayClusterDisplayDeviceComponent::UpdateDisplayDeviceProxyImpl(IDisplayClusterViewportConfiguration& InConfiguration)
diff --git a/Engine/Plugins/Runtime/nDisplay/Source/DisplayCluster/Private/Render/Viewport/DisplayClusterViewportResources.cpp b/Engine/Plugins/Runtime/nDisplay/Source/DisplayCluster/Private/Render/Viewport/DisplayClusterViewportResources.cpp
--- a/Engine/Plugins/Runtime/nDisplay/Source/DisplayCluster/Private/Render/Viewport/DisplayClusterViewportResources.cpp
+++ b/Engine/Plugins/Runtime/nDisplay/Source/DisplayCluster/Private/Render/Viewport/DisplayClusterViewportResources.cpp
@@ -27,10 +27,16 @@ bool FDisplayClusterViewportResources::GetRHIResources_RenderThread(const EDispl
 	{
 		for (const TSharedPtr<FDisplayClusterViewportResource, ESPMode::ThreadSafe>& ViewportResourceIt : *ExistResources)
 		{
-			if (FRHITexture2D* RHITexture2D = ViewportResourceIt.IsValid() ? ViewportResourceIt->GetViewportResourceRHI_RenderThread() : nullptr)
+			if (ViewportResourceIt.IsValid())
 			{
-				// Collects only valid resources.
-				OutResources.Add(RHITexture2D);
+				// When resource accessed on rendering thread, update this flag
+				EnumAddFlags(ViewportResourceIt->GetResourceState(), EDisplayClusterViewportResourceState::UpdatedOnRenderingThread);
+
+				if (FRHITexture2D* RHITexture2D = ViewportResourceIt->GetViewportResourceRHI_RenderThread())
+				{
+					// Collects only valid resources.
+					OutResources.Add(RHITexture2D);
+				}
 			}
 		}
 
diff --git a/Engine/Plugins/Runtime/nDisplay/Source/DisplayCluster/Private/Render/Viewport/Preview/DisplayClusterViewportPreview.cpp b/Engine/Plugins/Runtime/nDisplay/Source/DisplayCluster/Private/Render/Viewport/Preview/DisplayClusterViewportPreview.cpp
--- a/Engine/Plugins/Runtime/nDisplay/Source/DisplayCluster/Private/Render/Viewport/Preview/DisplayClusterViewportPreview.cpp
+++ b/Engine/Plugins/Runtime/nDisplay/Source/DisplayCluster/Private/Render/Viewport/Preview/DisplayClusterViewportPreview.cpp
@@ -7,6 +7,7 @@
 #include "Render/Warp/IDisplayClusterWarpPolicy.h"
 
 #include "Render/DisplayDevice/Components/DisplayClusterDisplayDeviceBaseComponent.h"
+#include "Components/DisplayClusterCameraComponent.h"
 
 #include "Components/StaticMeshComponent.h"
 
@@ -17,8 +18,8 @@ FDisplayClusterViewportPreview::FDisplayClusterViewportPreview(const TSharedRef<
 	: Configuration(InConfiguration)
 	, ViewportId(InViewportId)
 	, ClusterNodeId(InConfiguration->GetClusterNodeId())
-	, PreviewMesh(EDisplayClusterViewportPreviewMeshType::PreviewMesh, InConfiguration)
-	, PreviewEditableMesh(EDisplayClusterViewportPreviewMeshType::PreviewEditableMesh, InConfiguration)
+	, PreviewMesh(EDisplayClusterDisplayDeviceMeshType::PreviewMesh, InConfiguration)
+	, PreviewEditableMesh(EDisplayClusterDisplayDeviceMeshType::PreviewEditableMesh, InConfiguration)
 { }
 
 FDisplayClusterViewportPreview::~FDisplayClusterViewportPreview()
@@ -53,12 +54,19 @@ void FDisplayClusterViewportPreview::Update()
 		PreviewRTT = NewPreviewRTT;
 	}
 
+	if (PreviewRTT.IsValid() && EnumHasAnyFlags(PreviewRTT->GetResourceState(), EDisplayClusterViewportResourceState::UpdatedOnRenderingThread))
+	{
+		// Update flags when preview RTT is valid
+		EnumAddFlags(RuntimeFlags, EDisplayClusterViewportPreviewFlags::HasValidPreviewRTT);
+	}
+
 	FDisplayClusterViewport* InViewport = GetViewportImpl();
+	UDisplayClusterCameraComponent* ViewPointComponent = InViewport ? InViewport->GetViewPointCameraComponent(EDisplayClusterRootActorType::Configuration) : nullptr;
 	UDisplayClusterDisplayDeviceBaseComponent* InDisplayDeviceComponent = InViewport ? InViewport->GetDisplayDeviceComponent(EDisplayClusterRootActorType::Configuration) : nullptr;
 
 	// Update preview meshes only if DisplayDevice is used
-	PreviewMesh.Update(InViewport, InDisplayDeviceComponent);
-	PreviewEditableMesh.Update(InViewport, InDisplayDeviceComponent);
+	PreviewMesh.Update(InViewport, InDisplayDeviceComponent, ViewPointComponent);
+	PreviewEditableMesh.Update(InViewport, InDisplayDeviceComponent, ViewPointComponent);
 
 	// Update Runtime Flags:
 	if (PreviewMesh.HasAnyFlag(EDisplayClusterViewportPreviewMeshFlags::HasDeletedMaterialInstance | EDisplayClusterViewportPreviewMeshFlags::HasChangedMaterialInstance))
@@ -73,22 +81,16 @@ void FDisplayClusterViewportPreview::Update()
 	// Update the preview mesh and materials in the DisplayDevice component
 	if (InDisplayDeviceComponent)
 	{
-		// Update material instances
-		InDisplayDeviceComponent->OnUpdateDisplayDeviceMaterialInstance(*this, EDisplayClusterDisplayDeviceMeshType::PreviewMesh, PreviewMesh.GetCurrentMaterialType(), PreviewMesh.GetMaterialInstance());
-		InDisplayDeviceComponent->OnUpdateDisplayDeviceMaterialInstance(*this, EDisplayClusterDisplayDeviceMeshType::PreviewEditableMesh, PreviewMesh.GetCurrentMaterialType(), PreviewEditableMesh.GetMaterialInstance());
-
-		// Update mesh components
-		InDisplayDeviceComponent->OnUpdateDisplayDeviceMeshComponent(*this, EDisplayClusterDisplayDeviceMeshType::PreviewMesh, PreviewMesh.GetMeshComponent());
-		InDisplayDeviceComponent->OnUpdateDisplayDeviceMeshComponent(*this, EDisplayClusterDisplayDeviceMeshType::PreviewEditableMesh, PreviewEditableMesh.GetMeshComponent());
+		// Update mesh and material instances
+		InDisplayDeviceComponent->OnUpdateDisplayDeviceMeshAndMaterialInstance(*this, EDisplayClusterDisplayDeviceMeshType::PreviewMesh, PreviewMesh.GetCurrentMaterialType(), PreviewMesh.GetMeshComponent(), PreviewMesh.GetMaterialInstance());
+		InDisplayDeviceComponent->OnUpdateDisplayDeviceMeshAndMaterialInstance(*this, EDisplayClusterDisplayDeviceMeshType::PreviewEditableMesh, PreviewMesh.GetCurrentMaterialType(), PreviewEditableMesh.GetMeshComponent(), PreviewEditableMesh.GetMaterialInstance());
 	}
 
-	// The warp policy can update editable meshes and their material parameters
-	if (InViewport && InViewport->GetProjectionPolicy().IsValid())
+	// Update the preview mesh and materials in the ViewPointComponent component
+	if (ViewPointComponent)
 	{
-		if (IDisplayClusterWarpPolicy* WarpPolicy = InViewport->GetProjectionPolicy()->GetWarpPolicy())
-		{
-			WarpPolicy->OnUpdatePreviewEditableMesh(*this, PreviewEditableMesh.GetMeshComponent(), PreviewEditableMesh.GetCurrentMaterialType(), PreviewEditableMesh.GetMaterialInstance());
-		}
+		ViewPointComponent->OnUpdateDisplayDeviceMeshAndMaterialInstance(*this, EDisplayClusterDisplayDeviceMeshType::PreviewMesh, PreviewMesh.GetCurrentMaterialType(), PreviewMesh.GetMeshComponent(), PreviewMesh.GetMaterialInstance());
+		ViewPointComponent->OnUpdateDisplayDeviceMeshAndMaterialInstance(*this, EDisplayClusterDisplayDeviceMeshType::PreviewEditableMesh, PreviewMesh.GetCurrentMaterialType(), PreviewEditableMesh.GetMeshComponent(), PreviewEditableMesh.GetMaterialInstance());
 	}
 }
 
diff --git a/Engine/Plugins/Runtime/nDisplay/Source/DisplayCluster/Private/Render/Viewport/Preview/DisplayClusterViewportPreviewMesh.cpp b/Engine/Plugins/Runtime/nDisplay/Source/DisplayCluster/Private/Render/Viewport/Preview/DisplayClusterViewportPreviewMesh.cpp
--- a/Engine/Plugins/Runtime/nDisplay/Source/DisplayCluster/Private/Render/Viewport/Preview/DisplayClusterViewportPreviewMesh.cpp
+++ b/Engine/Plugins/Runtime/nDisplay/Source/DisplayCluster/Private/Render/Viewport/Preview/DisplayClusterViewportPreviewMesh.cpp
@@ -6,14 +6,14 @@
 
 #include "Render/DisplayDevice/Components/DisplayClusterDisplayDeviceBaseComponent.h"
 #include "Render/Projection/IDisplayClusterProjectionPolicy.h"
-
+#include "Components/DisplayClusterCameraComponent.h"
 #include "Components/StaticMeshComponent.h"
 
 namespace UE::DisplayCluster::ViewportPreviewMesh
 {
 	/** Return ptr on UObject if it is still valid. */
 	template<class T>
-	static inline T* GetObjectProperty(const TObjectPtr<T>& InProperty)
+	static inline T* GetExistingObject(const TObjectPtr<T>& InProperty)
 	{
 		if (InProperty == nullptr
 			|| InProperty->GetName().Find(TEXT("TRASH_")) != INDEX_NONE
@@ -41,6 +41,24 @@ namespace UE::DisplayCluster::ViewportPreviewMesh
 
 		return false;
 	}
+
+	/** Get material for preview mesh. */
+	TObjectPtr<UMaterial> FindPreviewMeshMaterial(const EDisplayClusterDisplayDeviceMeshType InMeshType, const EDisplayClusterDisplayDeviceMaterialType InMaterialType, UDisplayClusterDisplayDeviceBaseComponent* InDisplayDeviceComponent, UDisplayClusterCameraComponent* ViewPointComponent)
+	{
+		TObjectPtr<UMaterial> OutMaterial = nullptr;
+
+		// First get the material from the ViewPoint component (WarpPolicy)
+		OutMaterial = GetExistingObject(ViewPointComponent ? ViewPointComponent->GetDisplayDeviceMaterial(InMeshType, InMaterialType) : nullptr);
+
+		// Finally, get the material from the DisplayDevice
+		if (!OutMaterial)
+		{
+			OutMaterial = GetExistingObject(InDisplayDeviceComponent ? InDisplayDeviceComponent->GetDisplayDeviceMaterial(InMeshType, InMaterialType) : nullptr);
+		}
+
+		// Ignore deleted materials
+		return OutMaterial;
+	}
 };
 
 ////////////////////////////////////////////////////////////////////////////////////////
@@ -53,27 +71,31 @@ EDisplayClusterDisplayDeviceMaterialType FDisplayClusterViewportPreviewMesh::Get
 		: EDisplayClusterDisplayDeviceMaterialType::PreviewMeshMaterial;
 }
 
-void FDisplayClusterViewportPreviewMesh::Update(FDisplayClusterViewport* InViewport, UDisplayClusterDisplayDeviceBaseComponent* InDisplayDeviceComponent)
+void FDisplayClusterViewportPreviewMesh::Update(FDisplayClusterViewport* InViewport, UDisplayClusterDisplayDeviceBaseComponent* InDisplayDeviceComponent, UDisplayClusterCameraComponent* ViewPointComponent)
 {
+	using namespace UE::DisplayCluster::ViewportPreviewMesh;
+
 	// Reset runtime flags before each update
 	RuntimeFlags = EDisplayClusterViewportPreviewMeshFlags::None;
 
-	if (!InViewport || !ShouldUseMeshComponent(InViewport) || !InDisplayDeviceComponent)
+	if (!InViewport || !InViewport->ViewportPreview->HasAnyFlags(EDisplayClusterViewportPreviewFlags::HasValidPreviewRTT) || !ShouldUseMeshComponent(InViewport) || !InDisplayDeviceComponent)
 	{
 		// The mesh component and its resources are no longer used.
 		Release(InViewport);
+
 		return;
 	}
 
 	// Update default material
-	DefaultMaterialPtr = InDisplayDeviceComponent->GetDisplayDeviceMaterial(EDisplayClusterDisplayDeviceMaterialType::DefaultPreviewMeshMaterial);
+	DefaultMaterialPtr = FindPreviewMeshMaterial(EDisplayClusterDisplayDeviceMeshType::DefaultMesh, GetCurrentMaterialType(), InDisplayDeviceComponent, ViewPointComponent);
 
 	// Get current preview material
-	UMaterial* InMeshMaterial = InDisplayDeviceComponent->GetDisplayDeviceMaterial(GetCurrentMaterialType());
+	UMaterial* InMeshMaterial = FindPreviewMeshMaterial(MeshType, GetCurrentMaterialType(), InDisplayDeviceComponent, ViewPointComponent);
 	if (!InMeshMaterial)
 	{
-		// The mesh component and its resources are no longer used.
+		// Do not use preview if material is not defined.
 		Release(InViewport);
+
 		return;
 	}
 
@@ -249,10 +271,10 @@ bool FDisplayClusterViewportPreviewMesh::ShouldUseMeshComponent(FDisplayClusterV
 	{
 		switch (MeshType)
 		{
-		case EDisplayClusterViewportPreviewMeshType::PreviewMesh:
+		case EDisplayClusterDisplayDeviceMeshType::PreviewMesh:
 			return ProjectionPolicy->HasPreviewMesh(InViewport) && PreviewSettings.bEnablePreviewMesh;
 
-		case EDisplayClusterViewportPreviewMeshType::PreviewEditableMesh:
+		case EDisplayClusterDisplayDeviceMeshType::PreviewEditableMesh:
 			return ProjectionPolicy->HasPreviewEditableMesh(InViewport) && PreviewSettings.bEnablePreviewEditableMesh;
 
 		default:
@@ -270,10 +292,10 @@ UMeshComponent* FDisplayClusterViewportPreviewMesh::GetOrCreatePreviewMeshCompon
 	{
 		switch (MeshType)
 		{
-		case EDisplayClusterViewportPreviewMeshType::PreviewMesh:
+		case EDisplayClusterDisplayDeviceMeshType::PreviewMesh:
 			return ProjectionPolicy->GetOrCreatePreviewMeshComponent(InViewport, bOutIsRootActorComponent);
 
-		case EDisplayClusterViewportPreviewMeshType::PreviewEditableMesh:
+		case EDisplayClusterDisplayDeviceMeshType::PreviewEditableMesh:
 			bOutIsRootActorComponent = false;
 			return ProjectionPolicy->GetOrCreatePreviewEditableMeshComponent(InViewport);
 
@@ -289,12 +311,12 @@ UMeshComponent* FDisplayClusterViewportPreviewMesh::GetMeshComponent() const
 {
 	using namespace UE::DisplayCluster::ViewportPreviewMesh;
 
-	return GetObjectProperty(MeshComponentPtr);
+	return GetExistingObject(MeshComponentPtr);
 }
 
 UMaterialInstanceDynamic* FDisplayClusterViewportPreviewMesh::GetMaterialInstance() const
 {
 	using namespace UE::DisplayCluster::ViewportPreviewMesh;
 
-	return GetObjectProperty(MaterialInstancePtr);
+	return GetExistingObject(MaterialInstancePtr);
 }
diff --git a/Engine/Plugins/Runtime/nDisplay/Source/DisplayCluster/Private/Render/Viewport/Preview/DisplayClusterViewportPreviewMesh.h b/Engine/Plugins/Runtime/nDisplay/Source/DisplayCluster/Private/Render/Viewport/Preview/DisplayClusterViewportPreviewMesh.h
--- a/Engine/Plugins/Runtime/nDisplay/Source/DisplayCluster/Private/Render/Viewport/Preview/DisplayClusterViewportPreviewMesh.h
+++ b/Engine/Plugins/Runtime/nDisplay/Source/DisplayCluster/Private/Render/Viewport/Preview/DisplayClusterViewportPreviewMesh.h
@@ -12,17 +12,9 @@
 #include "Misc/DisplayClusterObjectRef.h"
 
 class UDisplayClusterDisplayDeviceBaseComponent;
+class UDisplayClusterCameraComponent;
 class FDisplayClusterViewport;
 
-/**
-* DCRA has several types of preview meshes.
-*/
-enum class EDisplayClusterViewportPreviewMeshType : uint8
-{
-	PreviewMesh = 0,
-	PreviewEditableMesh
-};
-
 /**
  * Runtime configuration of preview mesh.
  */
@@ -53,7 +45,7 @@ ENUM_CLASS_FLAGS(EDisplayClusterViewportPreviewMeshFlags);
 class FDisplayClusterViewportPreviewMesh
 {
 public:
-	FDisplayClusterViewportPreviewMesh(const EDisplayClusterViewportPreviewMeshType InMeshType, const TSharedRef<FDisplayClusterViewportConfiguration, ESPMode::ThreadSafe> InConfiguration)
+	FDisplayClusterViewportPreviewMesh(const EDisplayClusterDisplayDeviceMeshType InMeshType, const TSharedRef<FDisplayClusterViewportConfiguration, ESPMode::ThreadSafe> InConfiguration)
 		: Configuration(InConfiguration), MeshType(InMeshType)
 	{ }
 
@@ -70,7 +62,7 @@ public:
 	EDisplayClusterDisplayDeviceMaterialType GetCurrentMaterialType() const;
 
 	/** Update mesh component and materials for viewport. */
-	void Update(FDisplayClusterViewport* InViewport, UDisplayClusterDisplayDeviceBaseComponent* InDisplayDeviceComponent);
+	void Update(FDisplayClusterViewport* InViewport, UDisplayClusterDisplayDeviceBaseComponent* InDisplayDeviceComponent, UDisplayClusterCameraComponent* ViewPointComponent);
 
 	/** Restore default material and release mesh component with materials for viewport. */
 	void Release(FDisplayClusterViewport* InViewport)
@@ -112,7 +104,7 @@ public:
 	const TSharedRef<FDisplayClusterViewportConfiguration, ESPMode::ThreadSafe> Configuration;
 
 	// the type of mesh
-	const EDisplayClusterViewportPreviewMeshType MeshType;
+	const EDisplayClusterDisplayDeviceMeshType MeshType;
 
 private:
 	// runtime flags
diff --git a/Engine/Plugins/Runtime/nDisplay/Source/DisplayCluster/Private/Render/Viewport/Resource/DisplayClusterViewportResource.h b/Engine/Plugins/Runtime/nDisplay/Source/DisplayCluster/Private/Render/Viewport/Resource/DisplayClusterViewportResource.h
--- a/Engine/Plugins/Runtime/nDisplay/Source/DisplayCluster/Private/Render/Viewport/Resource/DisplayClusterViewportResource.h
+++ b/Engine/Plugins/Runtime/nDisplay/Source/DisplayCluster/Private/Render/Viewport/Resource/DisplayClusterViewportResource.h
@@ -28,6 +28,9 @@ enum class EDisplayClusterViewportResourceState : uint8
 
 	// Disable reallocation of this resource for the current frame
 	DisableReallocate = 1 << 3,
+
+	// This resource was updated in the rendering thread
+	UpdatedOnRenderingThread = 1 << 4,
 };
 ENUM_CLASS_FLAGS(EDisplayClusterViewportResourceState);
 
diff --git a/Engine/Plugins/Runtime/nDisplay/Source/DisplayCluster/Public/Components/DisplayClusterCameraComponent.h b/Engine/Plugins/Runtime/nDisplay/Source/DisplayCluster/Public/Components/DisplayClusterCameraComponent.h
--- a/Engine/Plugins/Runtime/nDisplay/Source/DisplayCluster/Public/Components/DisplayClusterCameraComponent.h
+++ b/Engine/Plugins/Runtime/nDisplay/Source/DisplayCluster/Public/Components/DisplayClusterCameraComponent.h
@@ -4,15 +4,21 @@
 
 #include "CoreMinimal.h"
 #include "Components/SceneComponent.h"
+#include "Materials/MaterialInstanceDynamic.h"
 #include "Components/IDisplayClusterComponent.h"
+#include "Render/DisplayDevice/Containers/DisplayClusterDisplayDevice_Enums.h"
 
 #include "DisplayClusterCameraComponent.generated.h"
 
+class UMaterial;
+class UMaterialInstanceDynamic;
+class UMeshComponent;
 class UBillboardComponent;
 class UTexture2D;
 class IDisplayClusterViewportManager;
 class IDisplayClusterWarpPolicy;
 class IDisplayClusterViewportConfiguration;
+class IDisplayClusterViewportPreview;
 struct FMinimalViewInfo;
 
 UENUM()
@@ -27,7 +33,7 @@ enum class EDisplayClusterEyeStereoOffset : uint8
 /**
  * 3D point in space used to render nDisplay viewports from
  */
-UCLASS(ClassGroup = (DisplayCluster), meta = (BlueprintSpawnableComponent, DisplayName = "NDisplay View Origin"))
+UCLASS(ClassGroup = (DisplayCluster), HideCategories = (Navigation, AssetUserData, LOD, Physics, Cooking, Activation, Tags, Gizmo, Collision, ComponentReplication, Events, Sockets, ComponentTick), meta = (BlueprintSpawnableComponent, DisplayName = "NDisplay View Origin"))
 class DISPLAYCLUSTER_API UDisplayClusterCameraComponent
 	: public USceneComponent
 	, public IDisplayClusterComponent
@@ -72,6 +78,31 @@ public:
 		return nullptr;
 	}
 
+	/** Override DisplayDevice material by type for 
+	* The UDisplayClusterInFrustumFitCameraComponent uses its own material to display additional deformed preview meshes in front of the camera.
+	*
+	* @param InMeshType     - mesh type
+	* @param InMaterialType - the type of material being requested
+	* 
+	* @return nullptr if DisplayDevice material is used.
+	*/
+	virtual TObjectPtr<UMaterial> GetDisplayDeviceMaterial(const EDisplayClusterDisplayDeviceMeshType InMeshType, const EDisplayClusterDisplayDeviceMaterialType InMaterialType) const
+	{
+		return nullptr;
+	}
+
+	/** Perform any operations on the mesh and material instance, such as setting parameter values.
+	*
+	* @param InViewport             - current viewport
+	* @param InMeshType             - mesh type
+	* @param InMaterialType         - type of material being requested
+	* @param InMeshComponent        - mesh component to be updated
+	* @param InMeshMaterialInstance - material instance that used on this mesh
+	*/
+	virtual void OnUpdateDisplayDeviceMeshAndMaterialInstance(IDisplayClusterViewportPreview& InViewportPreview, const EDisplayClusterDisplayDeviceMeshType InMeshType, const EDisplayClusterDisplayDeviceMaterialType InMaterialType, UMeshComponent* InMeshComponent, UMaterialInstanceDynamic* InMeshMaterialInstance) const
+	{ }
+
+
 public:
 	/**
 	* Get interpupillary distance
diff --git a/Engine/Plugins/Runtime/nDisplay/Source/DisplayCluster/Public/Render/DisplayDevice/Components/DisplayClusterDisplayDeviceBaseComponent.h b/Engine/Plugins/Runtime/nDisplay/Source/DisplayCluster/Public/Render/DisplayDevice/Components/DisplayClusterDisplayDeviceBaseComponent.h
--- a/Engine/Plugins/Runtime/nDisplay/Source/DisplayCluster/Public/Render/DisplayDevice/Components/DisplayClusterDisplayDeviceBaseComponent.h
+++ b/Engine/Plugins/Runtime/nDisplay/Source/DisplayCluster/Public/Render/DisplayDevice/Components/DisplayClusterDisplayDeviceBaseComponent.h
@@ -12,6 +12,7 @@ class UMaterial;
 class FSceneView;
 class FSceneViewFamily;
 class UMeshComponent;
+class UDisplayClusterCameraComponent;
 class UWorld;
 
 class IDisplayClusterDisplayDeviceProxy;
@@ -37,27 +38,20 @@ public:
 public:
 	/** Receive material by type
 	* 
+	* @param InMeshType     - mesh type
 	* @param InMaterialType - the type of material being requested
 	*/
-	virtual TObjectPtr<UMaterial> GetDisplayDeviceMaterial(const EDisplayClusterDisplayDeviceMaterialType InMaterialType) const;
+	virtual TObjectPtr<UMaterial> GetDisplayDeviceMaterial(const EDisplayClusterDisplayDeviceMeshType InMeshType, const EDisplayClusterDisplayDeviceMaterialType InMaterialType) const;
 	
 	/** Perform any operations on the  material instance, such as setting parameter values.
 	* 
-	* @param InViewport         - current viewport
+	* @param InViewportPreview  - current viewport
 	* @param InMeshType         - mesh type
 	* @param InMaterialType     - type of material being requested
-	* @param InMaterialInstance - material instance
-	*/
-	virtual void OnUpdateDisplayDeviceMaterialInstance(IDisplayClusterViewportPreview& InViewportPreview, const EDisplayClusterDisplayDeviceMeshType InMeshType, const EDisplayClusterDisplayDeviceMaterialType InMaterialType, UMaterialInstanceDynamic* InMaterialInstance) const;
-
-	/** Update preview mesh component, projection policy, preview meshes etc.
-	* Techvis should be implemented inside this function.
-	*
-	* @param InViewport      - the current viewport
-	* @param InMeshType      - mesh type
 	* @param InMeshComponent - mesh component to be updated
+	* @param InMeshMaterialInstance - material instance that used on this mesh
 	*/
-	virtual void OnUpdateDisplayDeviceMeshComponent(IDisplayClusterViewportPreview& InViewportPreview, const EDisplayClusterDisplayDeviceMeshType InMeshType,UMeshComponent* InMeshComponent) const;
+	virtual void OnUpdateDisplayDeviceMeshAndMaterialInstance(IDisplayClusterViewportPreview& InViewportPreview, const EDisplayClusterDisplayDeviceMeshType InMeshType, const EDisplayClusterDisplayDeviceMaterialType InMaterialType, UMeshComponent* InMeshComponent, UMaterialInstanceDynamic* InMeshMaterialInstance) const;
 
 	/** Update and Get proxy object for rendering thread.
 	* 
diff --git a/Engine/Plugins/Runtime/nDisplay/Source/DisplayCluster/Public/Render/DisplayDevice/Components/DisplayClusterDisplayDeviceComponent.h b/Engine/Plugins/Runtime/nDisplay/Source/DisplayCluster/Public/Render/DisplayDevice/Components/DisplayClusterDisplayDeviceComponent.h
--- a/Engine/Plugins/Runtime/nDisplay/Source/DisplayCluster/Public/Render/DisplayDevice/Components/DisplayClusterDisplayDeviceComponent.h
+++ b/Engine/Plugins/Runtime/nDisplay/Source/DisplayCluster/Public/Render/DisplayDevice/Components/DisplayClusterDisplayDeviceComponent.h
@@ -25,7 +25,7 @@ public:
 	UDisplayClusterDisplayDeviceComponent();
 
 	//~BEGIN UDisplayClusterDisplayDeviceBaseComponent
-	virtual void OnUpdateDisplayDeviceMaterialInstance(IDisplayClusterViewportPreview& InViewportPreview, const EDisplayClusterDisplayDeviceMeshType InMeshType, const EDisplayClusterDisplayDeviceMaterialType InMaterialType, UMaterialInstanceDynamic* InMaterialInstance) const override;
+	virtual void OnUpdateDisplayDeviceMeshAndMaterialInstance(IDisplayClusterViewportPreview& InViewportPreview, const EDisplayClusterDisplayDeviceMeshType InMeshType, const EDisplayClusterDisplayDeviceMaterialType InMaterialType, UMeshComponent* InMeshComponent, UMaterialInstanceDynamic* InMeshMaterialInstance) const override;
 	virtual void UpdateDisplayDeviceProxyImpl(IDisplayClusterViewportConfiguration& InConfiguration) override;
 	//~~END UDisplayClusterDisplayDeviceBaseComponent
 
diff --git a/Engine/Plugins/Runtime/nDisplay/Source/DisplayCluster/Public/Render/DisplayDevice/Containers/DisplayClusterDisplayDevice_Enums.h b/Engine/Plugins/Runtime/nDisplay/Source/DisplayCluster/Public/Render/DisplayDevice/Containers/DisplayClusterDisplayDevice_Enums.h
--- a/Engine/Plugins/Runtime/nDisplay/Source/DisplayCluster/Public/Render/DisplayDevice/Containers/DisplayClusterDisplayDevice_Enums.h
+++ b/Engine/Plugins/Runtime/nDisplay/Source/DisplayCluster/Public/Render/DisplayDevice/Containers/DisplayClusterDisplayDevice_Enums.h
@@ -9,9 +9,6 @@
 */
 enum class EDisplayClusterDisplayDeviceMaterialType : uint8
 {
-	// Default material used on preview mesh when preview is disabled
-	DefaultPreviewMeshMaterial = 0,
-
 	// Preview on mesh
 	PreviewMeshMaterial,
 
@@ -24,9 +21,13 @@ enum class EDisplayClusterDisplayDeviceMaterialType : uint8
 */
 enum class EDisplayClusterDisplayDeviceMeshType : uint8
 {
+	// Default mesh without preview
+	DefaultMesh = 0,
+
 	// Preview mesh
-	PreviewMesh = 0,
+	PreviewMesh,
 
 	// Preview editable mesh
 	PreviewEditableMesh,
+
 };
diff --git a/Engine/Plugins/Runtime/nDisplay/Source/DisplayCluster/Public/Render/Viewport/Containers/DisplayClusterViewportPreview_Enums.h b/Engine/Plugins/Runtime/nDisplay/Source/DisplayCluster/Public/Render/Viewport/Containers/DisplayClusterViewportPreview_Enums.h
--- a/Engine/Plugins/Runtime/nDisplay/Source/DisplayCluster/Public/Render/Viewport/Containers/DisplayClusterViewportPreview_Enums.h
+++ b/Engine/Plugins/Runtime/nDisplay/Source/DisplayCluster/Public/Render/Viewport/Containers/DisplayClusterViewportPreview_Enums.h
@@ -14,10 +14,13 @@ enum class EDisplayClusterViewportPreviewFlags : uint8
 	// the RTT has been changed
 	HasChangedPreviewRTT = 1 << 0,
 
+	// This viewport has a valid preview texture.
+	HasValidPreviewRTT = 1 << 1,
+
 	// Preview mesh material instance has been changed.
-	HasChangedPreviewMeshMaterialInstance = 1 << 1,
+	HasChangedPreviewMeshMaterialInstance = 1 << 2,
 
 	// Preview editable mesh material instance has been changed.
-	HasChangedPreviewEditableMeshMaterialInstance = 1 << 2,
+	HasChangedPreviewEditableMeshMaterialInstance = 1 << 3,
 };
 ENUM_CLASS_FLAGS(EDisplayClusterViewportPreviewFlags);
diff --git a/Engine/Plugins/Runtime/nDisplay/Source/DisplayCluster/Public/Render/Warp/IDisplayClusterWarpPolicy.h b/Engine/Plugins/Runtime/nDisplay/Source/DisplayCluster/Public/Render/Warp/IDisplayClusterWarpPolicy.h
--- a/Engine/Plugins/Runtime/nDisplay/Source/DisplayCluster/Public/Render/Warp/IDisplayClusterWarpPolicy.h
+++ b/Engine/Plugins/Runtime/nDisplay/Source/DisplayCluster/Public/Render/Warp/IDisplayClusterWarpPolicy.h
@@ -119,14 +119,14 @@ public:
 		return false;
 	}
 
-	/** Update Editable mesh material instance parameters.
-	* (this function is called every tick)
+	/** Perform any operations on the  mesh and material instance, such as setting parameter values.
 	* 
-	* @param InViewport              - the current viewport
-	* @param InEditableMeshComponent - editable mesh component
-	* @param InMaterialType          - the type of material being requested
-	* @param InMaterialInstance      - material instance
+	* @param InViewport - current viewport
+	* @param InMeshType - mesh type
+	* @param InMaterialType - type of material being requested
+	* @param InMeshComponent - mesh component to be updated
+	* @param InMeshMaterialInstance - material instance that used on this mesh
 	*/
-	virtual void OnUpdatePreviewEditableMesh(IDisplayClusterViewportPreview& InViewportPreview, UMeshComponent* InEditableMeshComponent, const EDisplayClusterDisplayDeviceMaterialType InMaterialType, UMaterialInstanceDynamic* InMaterialInstance) const
+	virtual void OnUpdateDisplayDeviceMeshAndMaterialInstance(IDisplayClusterViewportPreview& InViewportPreview, const EDisplayClusterDisplayDeviceMeshType InMeshType, const EDisplayClusterDisplayDeviceMaterialType InMaterialType, UMeshComponent* InMeshComponent, UMaterialInstanceDynamic* InMeshMaterialInstance) const
 	{ }
 };
diff --git a/Engine/Plugins/Runtime/nDisplay/Source/DisplayClusterWarp/Private/Components/DisplayClusterInFrustumFitCameraComponent.cpp b/Engine/Plugins/Runtime/nDisplay/Source/DisplayClusterWarp/Private/Components/DisplayClusterInFrustumFitCameraComponent.cpp
--- a/Engine/Plugins/Runtime/nDisplay/Source/DisplayClusterWarp/Private/Components/DisplayClusterInFrustumFitCameraComponent.cpp
+++ b/Engine/Plugins/Runtime/nDisplay/Source/DisplayClusterWarp/Private/Components/DisplayClusterInFrustumFitCameraComponent.cpp
@@ -8,7 +8,7 @@
 #include "CineCameraComponent.h"
 
 #include "DisplayClusterWarpLog.h"
-#include "DisplayClusterWarpStrings.h"
+#include "PDisplayClusterWarpStrings.h"
 
 #include "IDisplayCluster.h"
 #include "IDisplayClusterWarpBlend.h"
@@ -21,6 +21,9 @@
 #include "Render/Warp/IDisplayClusterWarpPolicyFactory.h"
 #include "Render/Warp/IDisplayClusterWarpPolicy.h"
 
+#include "Materials/Material.h"
+#include "UObject/ConstructorHelpers.h"
+
 #include "DisplayClusterRootActor.h"
 
 namespace UE::DisplayClusterWarp::ViewPointComponent
@@ -40,7 +43,6 @@ namespace UE::DisplayClusterWarp::ViewPointComponent
 		return nullptr;
 	}
 };
-using namespace UE::DisplayClusterWarp::ViewPointComponent;
 
 //--------------------------------------------------------------------------------
 // UDisplayClusterInFrustumFitCameraComponent
@@ -54,6 +56,55 @@ UDisplayClusterInFrustumFitCameraComponent::UDisplayClusterInFrustumFitCameraCom
 	bAutoActivate = true;
 }
 
+TObjectPtr<UMaterial> UDisplayClusterInFrustumFitCameraComponent::GetDisplayDeviceMaterial(const EDisplayClusterDisplayDeviceMeshType InMeshType, const EDisplayClusterDisplayDeviceMaterialType InMaterialType) const
+{
+	// Uses its own material to display additional deformed preview meshes in front of the camera.
+	if (IsEnabled() && WarpPolicy.IsValid())
+	{
+		// Special preview material is used for editable meshes: they should fly in front of the camera and deform according to its frustum.
+		if (InMeshType == EDisplayClusterDisplayDeviceMeshType::PreviewEditableMesh)
+		{
+			static TObjectPtr<UMaterial> InFrustumFitMaterial = LoadObject<UMaterial>(nullptr, UE::DisplayClusterWarpStrings::InFrustumFit::material::Name, nullptr, LOAD_None, nullptr);
+			if (InFrustumFitMaterial && InMeshType == EDisplayClusterDisplayDeviceMeshType::PreviewEditableMesh)
+			{
+				switch (InMaterialType)
+				{
+				case EDisplayClusterDisplayDeviceMaterialType::PreviewMeshMaterial:
+				case EDisplayClusterDisplayDeviceMaterialType::PreviewMeshTechvisMaterial:
+					// Note: Add additional techvis material for 'InFrustumFitCamera' if needed.
+					return InFrustumFitMaterial.Get();
+
+				default:
+					break;
+				}
+			}
+		}
+	}
+
+	return nullptr;
+}
+
+void UDisplayClusterInFrustumFitCameraComponent::OnUpdateDisplayDeviceMeshAndMaterialInstance(IDisplayClusterViewportPreview& InViewportPreview, const EDisplayClusterDisplayDeviceMeshType InMeshType, const EDisplayClusterDisplayDeviceMaterialType InMaterialType, UMeshComponent* InMeshComponent, UMaterialInstanceDynamic* InMeshMaterialInstance) const
+{
+	if (IsEnabled() && WarpPolicy.IsValid())
+	{
+		// The preview material used for editable meshes requires a set of unique parameters that are set from the warp policy.
+		if (InMeshComponent && InMeshMaterialInstance && InMeshType == EDisplayClusterDisplayDeviceMeshType::PreviewEditableMesh)
+		{
+			switch (InMaterialType)
+			{
+			case EDisplayClusterDisplayDeviceMaterialType::PreviewMeshMaterial:
+			case EDisplayClusterDisplayDeviceMaterialType::PreviewMeshTechvisMaterial:
+				WarpPolicy->OnUpdateDisplayDeviceMeshAndMaterialInstance(InViewportPreview, InMeshType, InMaterialType, InMeshComponent, InMeshMaterialInstance);
+				break;
+
+			default:
+				break;
+			}
+		}
+	}
+}
+
 const UDisplayClusterInFrustumFitCameraComponent& UDisplayClusterInFrustumFitCameraComponent::GetConfigurationInFrustumFitCameraComponent(IDisplayClusterViewportConfiguration& InViewportConfiguration) const
 {
 	if (ADisplayClusterRootActor* ConfigurationRootActor = InViewportConfiguration.GetRootActor(EDisplayClusterRootActorType::Configuration))
@@ -135,8 +186,10 @@ bool UDisplayClusterInFrustumFitCameraComponent::ShouldUseEntireClusterViewports
 
 IDisplayClusterWarpPolicy* UDisplayClusterInFrustumFitCameraComponent::GetWarpPolicy(IDisplayClusterViewportManager* InViewportManager)
 {
+	using namespace UE::DisplayClusterWarp::ViewPointComponent;
+
 	// We can ask for different types of warp policies, depending on the rules of the user settings
-	const FString NewWaprPolicyType = DisplayClusterWarpStrings::warp::InFrustumFit;
+	const FString NewWaprPolicyType = UE::DisplayClusterWarpStrings::warp::InFrustumFit;
 
 	// when returns different type, this will recreate warp policy instance
 	if (WarpPolicy.IsValid() && WarpPolicy->GetType() != NewWaprPolicyType)
diff --git a/Engine/Plugins/Runtime/nDisplay/Source/DisplayClusterWarp/Private/DisplayClusterWarpModule.cpp b/Engine/Plugins/Runtime/nDisplay/Source/DisplayClusterWarp/Private/DisplayClusterWarpModule.cpp
--- a/Engine/Plugins/Runtime/nDisplay/Source/DisplayClusterWarp/Private/DisplayClusterWarpModule.cpp
+++ b/Engine/Plugins/Runtime/nDisplay/Source/DisplayClusterWarp/Private/DisplayClusterWarpModule.cpp
@@ -18,7 +18,7 @@ FDisplayClusterWarpModule::FDisplayClusterWarpModule()
 
 	// InFrustumFit warp policy
 	Factory = MakeShared<FDisplayClusterWarpInFrustumFitPolicyFactory>();
-	WarpPolicyFactories.Emplace(DisplayClusterWarpStrings::warp::InFrustumFit, Factory);
+	WarpPolicyFactories.Emplace(UE::DisplayClusterWarpStrings::warp::InFrustumFit, Factory);
 
 	UE_LOG(LogDisplayClusterWarp, Log, TEXT("Warp module has been instantiated"));
 }
diff --git a/Engine/Plugins/Runtime/nDisplay/Source/DisplayClusterWarp/Private/PDisplayClusterWarpStrings.h b/Engine/Plugins/Runtime/nDisplay/Source/DisplayClusterWarp/Private/PDisplayClusterWarpStrings.h
new file mode 100644
--- /dev/null
+++ b/Engine/Plugins/Runtime/nDisplay/Source/DisplayClusterWarp/Private/PDisplayClusterWarpStrings.h
@@ -0,0 +1,24 @@
+// Copyright Epic Games, Inc. All Rights Reserved.
+
+#pragma once
+
+#include "DisplayClusterWarpStrings.h"
+
+namespace UE::DisplayClusterWarpStrings
+{
+	namespace InFrustumFit
+	{
+		namespace material
+		{
+			static constexpr const TCHAR* Name = TEXT("/nDisplay/Materials/Preview/M_InFrustumFitPreview");
+
+			namespace attr
+			{
+				static constexpr const TCHAR* GlobalScale = TEXT("GlobalScale");
+				static constexpr const TCHAR* GlobalForward = TEXT("GlobalForward");
+				static constexpr const TCHAR* GlobalRight = TEXT("GlobalRight");
+				static constexpr const TCHAR* GlobalUp = TEXT("GlobalUp");
+			}
+		}
+	}
+};
diff --git a/Engine/Plugins/Runtime/nDisplay/Source/DisplayClusterWarp/Private/Policy/InFrustumFit/DisplayClusterWarpInFrustumFitPolicy.cpp b/Engine/Plugins/Runtime/nDisplay/Source/DisplayClusterWarp/Private/Policy/InFrustumFit/DisplayClusterWarpInFrustumFitPolicy.cpp
--- a/Engine/Plugins/Runtime/nDisplay/Source/DisplayClusterWarp/Private/Policy/InFrustumFit/DisplayClusterWarpInFrustumFitPolicy.cpp
+++ b/Engine/Plugins/Runtime/nDisplay/Source/DisplayClusterWarp/Private/Policy/InFrustumFit/DisplayClusterWarpInFrustumFitPolicy.cpp
@@ -10,7 +10,7 @@
 #include "Render/Projection/IDisplayClusterProjectionPolicy.h"
 
 #include "IDisplayClusterWarpBlend.h"
-#include "DisplayClusterWarpStrings.h"
+#include "PDisplayClusterWarpStrings.h"
 
 #include "Components/DisplayClusterCameraComponent.h"
 #include "Components/DisplayClusterInFrustumFitCameraComponent.h"
@@ -52,7 +52,7 @@ FDisplayClusterWarpInFrustumFitPolicy::FDisplayClusterWarpInFrustumFitPolicy(con
 
 const FString& FDisplayClusterWarpInFrustumFitPolicy::GetType() const
 {
-	static const FString Type(DisplayClusterWarpStrings::warp::InFrustumFit);
+	static const FString Type(UE::DisplayClusterWarpStrings::warp::InFrustumFit);
 
 	return Type;
 }
@@ -296,29 +296,27 @@ void FDisplayClusterWarpInFrustumFitPolicy::EndCalcFrustum(IDisplayClusterViewpo
 	}
 }
 
-void FDisplayClusterWarpInFrustumFitPolicy::OnUpdatePreviewEditableMesh(IDisplayClusterViewportPreview& InViewportPreview, UMeshComponent* InEditableMeshComponent, const EDisplayClusterDisplayDeviceMaterialType InMaterialType, UMaterialInstanceDynamic* InMaterialInstance) const
+void FDisplayClusterWarpInFrustumFitPolicy::OnUpdateDisplayDeviceMeshAndMaterialInstance(IDisplayClusterViewportPreview& InViewportPreview, const EDisplayClusterDisplayDeviceMeshType InMeshType, const EDisplayClusterDisplayDeviceMaterialType InMaterialType, UMeshComponent* InMeshComponent, UMaterialInstanceDynamic* InMeshMaterialInstance) const
 {
-	if (!InEditableMeshComponent || !InMaterialInstance)
-	{
-		return;
-	}
+	// The preview material used for editable meshes requires a set of unique parameters that are set from the warp policy.
+	check(InMeshComponent && InMeshMaterialInstance);
 
-	IDisplayClusterViewport* InViewport = InViewportPreview.GetViewport();
-	if (!InViewport)
+	if (InMeshType != EDisplayClusterDisplayDeviceMeshType::PreviewEditableMesh)
 	{
+		// Only for editable mesh
 		return;
 	}
 
 	// Process only viewports with a projection policy based on the warpblend interface.
 	TSharedPtr<IDisplayClusterWarpBlend, ESPMode::ThreadSafe> WarpBlend;
-	if (!InViewport->GetProjectionPolicy().IsValid() || !InViewport->GetProjectionPolicy()->GetWarpBlendInterface(WarpBlend))
+	IDisplayClusterViewport* InViewport = InViewportPreview.GetViewport();
+	if (!InViewport || !InViewport->GetProjectionPolicy().IsValid() || !InViewport->GetProjectionPolicy()->GetWarpBlendInterface(WarpBlend))
 	{
 		return;
 	}
 
-	FDisplayClusterWarpData& WarpData = WarpBlend->GetWarpData(0);
-
 	// Not all projection policies support a editable mesh.
+	const FDisplayClusterWarpData& WarpData = WarpBlend->GetWarpData(0);
 	if (WarpData.bValid && WarpData.bHasWarpPolicyChanges)
 	{
 		const FTransform CameraTransform(WarpData.WarpProjection.CameraRotation.Quaternion(), WarpData.WarpProjection.CameraLocation);
@@ -327,64 +325,60 @@ void FDisplayClusterWarpInFrustumFitPolicy::OnUpdatePreviewEditableMesh(IDisplay
 		const float VScale = (WarpData.WarpProjection.Top - WarpData.WarpProjection.Bottom) / (WarpData.GeometryWarpProjection.Top - WarpData.GeometryWarpProjection.Bottom);
 
 		checkf(FMath::IsNearlyEqual(HScale, VScale), TEXT("Streching the stage geometry to fit a different aspect ratio is not supported!"));
+		const FVector Scale(1, HScale, HScale);
 
-		const FVector Scale = FVector(1, HScale, HScale);
+		const FMatrix CameraBasis = FRotationMatrix::Make(WarpData.WarpProjection.CameraRotation).Inverse();
 
 		// Compute the relative transform from the view origin to the geometry
-		FTransform RelativeTransform(WarpData.WarpContext.MeshToStageMatrix * WarpData.Local2World.Inverse());
+		FTransform RelativeTransform = FTransform(WarpData.WarpContext.MeshToStageMatrix * WarpData.Local2World.Inverse());
 		RelativeTransform.ScaleTranslation(Scale);
 
 		// Final transform is computed from the relative transform of the geometry to the view point, the frustum fit transform
 		// which will scale and position the geometry based on the fitted frustum, and the camera transform
 		const FTransform FinalTransform = RelativeTransform * CameraTransform;
-
-		InEditableMeshComponent->SetRelativeTransform(FinalTransform);
-
-		const FMatrix CameraBasis = FRotationMatrix::Make(WarpData.WarpProjection.CameraRotation).Inverse();
+		InMeshComponent->SetRelativeTransform(FinalTransform);
 
 		// Since the mesh needs to be skewed to scale appropriately, and since Unreal Engine does not support a skew transform
 		// through FTransform, the mesh needs to be skewed through the vertex shader using WorldPositionOffset,
 		// so pass in the "global" scale to the preview mesh's material instance
-		switch (InMaterialType)
-		{
-		case EDisplayClusterDisplayDeviceMaterialType::PreviewMeshMaterial:
-		case EDisplayClusterDisplayDeviceMaterialType::PreviewMeshTechvisMaterial:
-
-			InMaterialInstance->SetVectorParameterValue(TEXT("GlobalScale"), Scale);
-			InMaterialInstance->SetVectorParameterValue(TEXT("GlobalForward"), CameraBasis.GetUnitAxis(EAxis::X));
-			InMaterialInstance->SetVectorParameterValue(TEXT("GlobalRight"), CameraBasis.GetUnitAxis(EAxis::Y));
-			InMaterialInstance->SetVectorParameterValue(TEXT("GlobalUp"), CameraBasis.GetUnitAxis(EAxis::Z));
-			break;
-
-		default:
-			break;
-		}
-
+		InMeshMaterialInstance->SetVectorParameterValue(UE::DisplayClusterWarpStrings::InFrustumFit::material::attr::GlobalScale, Scale);
+		InMeshMaterialInstance->SetVectorParameterValue(UE::DisplayClusterWarpStrings::InFrustumFit::material::attr::GlobalForward, CameraBasis.GetUnitAxis(EAxis::X));
+		InMeshMaterialInstance->SetVectorParameterValue(UE::DisplayClusterWarpStrings::InFrustumFit::material::attr::GlobalRight, CameraBasis.GetUnitAxis(EAxis::Y));
+		InMeshMaterialInstance->SetVectorParameterValue(UE::DisplayClusterWarpStrings::InFrustumFit::material::attr::GlobalUp, CameraBasis.GetUnitAxis(EAxis::Z));
 	}
 }
 
 #if WITH_EDITOR
-void FDisplayClusterWarpInFrustumFitPolicy::DrawDebugGroupBoundingBox(ADisplayClusterRootActor* RootActor, const FColor& Color)
+#include "Components/LineBatchComponent.h"
+
+void FDisplayClusterWarpInFrustumFitPolicy::DrawDebugGroupBoundingBox(ADisplayClusterRootActor* SceneRootActor, const FLinearColor& Color)
 {
-	if (RootActor)
+	// DCRA uses its own LineBatcher
+	ULineBatchComponent* LineBatcher = SceneRootActor ? SceneRootActor->GetLineBatchComponent() : nullptr;
+	UWorld* World = SceneRootActor ? SceneRootActor->GetWorld() : nullptr;
+	if (LineBatcher && World)
 	{
-		if (UWorld* World = RootActor->GetWorld())
-		{
-			FBox WorldBox = GroupAABBox.TransformBy(RootActor->GetActorTransform());
-			DrawDebugBox(World, WorldBox.GetCenter(), WorldBox.GetExtent(), Color);
-			DrawDebugPoint(World, WorldBox.GetCenter(), 5, Color);
-		}
+		const float Thickness = 1.f;
+		const float PointSize = 5.f;
+		const FBox WorldBox = GroupAABBox.TransformBy(SceneRootActor->GetActorTransform());
+
+		LineBatcher->DrawBox(WorldBox.GetCenter(), WorldBox.GetExtent(), Color, 0, SDPG_World, Thickness);
+		LineBatcher->DrawPoint(WorldBox.GetCenter(), Color, PointSize, SDPG_World);
 	}
 }
 
-void FDisplayClusterWarpInFrustumFitPolicy::DrawDebugGroupFrustum(ADisplayClusterRootActor* RootActor, UDisplayClusterInFrustumFitCameraComponent* CameraComponent, const FColor& Color)
+void FDisplayClusterWarpInFrustumFitPolicy::DrawDebugGroupFrustum(ADisplayClusterRootActor* SceneRootActor, UDisplayClusterInFrustumFitCameraComponent* CameraComponent, const FLinearColor& Color)
 {
-	if (RootActor && CameraComponent)
+	// DCRA uses its own LineBatcher
+	ULineBatchComponent* LineBatcher = SceneRootActor ? SceneRootActor->GetLineBatchComponent() : nullptr;
+	if (LineBatcher && CameraComponent)
 	{
-		UWorld* World = RootActor->GetWorld();
-		IDisplayClusterViewportConfiguration* ViewportConfiguration = RootActor->GetViewportConfiguration();
+		UWorld* World = SceneRootActor->GetWorld();
+		IDisplayClusterViewportConfiguration* ViewportConfiguration = SceneRootActor->GetViewportConfiguration();
 		if (ViewportConfiguration && World)
 		{
+			const float Thickness = 1.0f;
+
 			// Get the configuration in use
 			const UDisplayClusterInFrustumFitCameraComponent& ConfigurationCameraComponent = CameraComponent->GetConfigurationInFrustumFitCameraComponent(*ViewportConfiguration);
 
@@ -400,7 +394,7 @@ void FDisplayClusterWarpInFrustumFitPolicy::DrawDebugGroupFrustum(ADisplayCluste
 			}
 			else
 			{
-				const FBox WorldBox = GroupAABBox.TransformBy(RootActor->GetActorTransform());
+				const FBox WorldBox = GroupAABBox.TransformBy(SceneRootActor->GetActorTransform());
 				ViewDirection = (WorldBox.GetCenter() - CameraComponent->GetComponentLocation()).GetSafeNormal();
 
 				if (SymmetricForwardCorrection.IsSet())
@@ -409,7 +403,6 @@ void FDisplayClusterWarpInFrustumFitPolicy::DrawDebugGroupFrustum(ADisplayCluste
 				}
 			}
 
-			DrawDebugLine(World, CameraLoc, CameraLoc + ViewDirection * 50, Color);
 
 			const FRotator ViewRotator = ViewDirection.ToOrientationRotator();
 			const FVector FrustumTopLeft = ViewRotator.RotateVector(FVector(GroupGeometryWarpProjection.ZNear, GroupGeometryWarpProjection.Left, GroupGeometryWarpProjection.Top) / GroupGeometryWarpProjection.ZNear);
@@ -430,23 +423,25 @@ void FDisplayClusterWarpInFrustumFitPolicy::DrawDebugGroupFrustum(ADisplayCluste
 				CameraLoc + FrustumBottomLeft * FarPlane,
 			};
 
+			LineBatcher->DrawLine(CameraLoc, CameraLoc + ViewDirection * 50, Color, SDPG_World, Thickness, 0.f);
+
 			// Near plane rectangle
-			DrawDebugLine(World, FrustumVertices[0], FrustumVertices[1], Color);
-			DrawDebugLine(World, FrustumVertices[1], FrustumVertices[2], Color);
-			DrawDebugLine(World, FrustumVertices[2], FrustumVertices[3], Color);
-			DrawDebugLine(World, FrustumVertices[3], FrustumVertices[0], Color);
+			LineBatcher->DrawLine(FrustumVertices[0], FrustumVertices[1], Color, SDPG_World, Thickness, 0.f);
+			LineBatcher->DrawLine(FrustumVertices[1], FrustumVertices[2], Color, SDPG_World, Thickness, 0.f);
+			LineBatcher->DrawLine(FrustumVertices[2], FrustumVertices[3], Color, SDPG_World, Thickness, 0.f);
+			LineBatcher->DrawLine(FrustumVertices[3], FrustumVertices[0], Color, SDPG_World, Thickness, 0.f);
 
 			// Frustum
-			DrawDebugLine(World, FrustumVertices[0], FrustumVertices[4], Color);
-			DrawDebugLine(World, FrustumVertices[1], FrustumVertices[5], Color);
-			DrawDebugLine(World, FrustumVertices[2], FrustumVertices[6], Color);
-			DrawDebugLine(World, FrustumVertices[3], FrustumVertices[7], Color);
+			LineBatcher->DrawLine(FrustumVertices[0], FrustumVertices[4], Color, SDPG_World, Thickness, 0.f);
+			LineBatcher->DrawLine(FrustumVertices[1], FrustumVertices[5], Color, SDPG_World, Thickness, 0.f);
+			LineBatcher->DrawLine(FrustumVertices[2], FrustumVertices[6], Color, SDPG_World, Thickness, 0.f);
+			LineBatcher->DrawLine(FrustumVertices[3], FrustumVertices[7], Color, SDPG_World, Thickness, 0.f);
 
 			// Far plane rectangle
-			DrawDebugLine(World, FrustumVertices[4], FrustumVertices[5], Color);
-			DrawDebugLine(World, FrustumVertices[5], FrustumVertices[6], Color);
-			DrawDebugLine(World, FrustumVertices[6], FrustumVertices[7], Color);
-			DrawDebugLine(World, FrustumVertices[7], FrustumVertices[4], Color);
+			LineBatcher->DrawLine(FrustumVertices[4], FrustumVertices[5], Color, SDPG_World, Thickness, 0.f);
+			LineBatcher->DrawLine(FrustumVertices[5], FrustumVertices[6], Color, SDPG_World, Thickness, 0.f);
+			LineBatcher->DrawLine(FrustumVertices[6], FrustumVertices[7], Color, SDPG_World, Thickness, 0.f);
+			LineBatcher->DrawLine(FrustumVertices[7], FrustumVertices[4], Color, SDPG_World, Thickness, 0.f);
 		}
 	}
 }
diff --git a/Engine/Plugins/Runtime/nDisplay/Source/DisplayClusterWarp/Private/Policy/InFrustumFit/DisplayClusterWarpInFrustumFitPolicy.h b/Engine/Plugins/Runtime/nDisplay/Source/DisplayClusterWarp/Private/Policy/InFrustumFit/DisplayClusterWarpInFrustumFitPolicy.h
--- a/Engine/Plugins/Runtime/nDisplay/Source/DisplayClusterWarp/Private/Policy/InFrustumFit/DisplayClusterWarpInFrustumFitPolicy.h
+++ b/Engine/Plugins/Runtime/nDisplay/Source/DisplayClusterWarp/Private/Policy/InFrustumFit/DisplayClusterWarpInFrustumFitPolicy.h
@@ -32,7 +32,7 @@ public:
 
 	virtual bool HasPreviewEditableMesh(IDisplayClusterViewport* InViewport) override;
 
-	virtual void OnUpdatePreviewEditableMesh(IDisplayClusterViewportPreview& InViewportPreview, UMeshComponent* InEditableMeshComponent, const EDisplayClusterDisplayDeviceMaterialType InMaterialType, UMaterialInstanceDynamic* InMaterialInstance) const override;
+	virtual void OnUpdateDisplayDeviceMeshAndMaterialInstance(IDisplayClusterViewportPreview& InViewportPreview, const EDisplayClusterDisplayDeviceMeshType InMeshType, const EDisplayClusterDisplayDeviceMaterialType InMaterialType, UMeshComponent* InMeshComponent, UMaterialInstanceDynamic* InMeshMaterialInstance) const override;
 	//~~ End IDisplayClusterWarpPolicy
 
 private:
@@ -50,10 +50,10 @@ private:
 
 #if WITH_EDITOR
 	/** Renders a debug visualization for the group bounding box */
-	void DrawDebugGroupBoundingBox(ADisplayClusterRootActor* RootActor, const FColor& Color);
+	void DrawDebugGroupBoundingBox(ADisplayClusterRootActor* RootActor, const FLinearColor& Color);
 
 	/** Renders a debug visualization for the group frustum */
-	void DrawDebugGroupFrustum(ADisplayClusterRootActor* RootActor, UDisplayClusterInFrustumFitCameraComponent* CameraComponent, const FColor& Color);
+	void DrawDebugGroupFrustum(ADisplayClusterRootActor* RootActor, UDisplayClusterInFrustumFitCameraComponent* CameraComponent, const FLinearColor& Color);
 #endif
 
 private:
diff --git a/Engine/Plugins/Runtime/nDisplay/Source/DisplayClusterWarp/Private/Policy/InFrustumFit/DisplayClusterWarpInFrustumFitPolicyFactory.cpp b/Engine/Plugins/Runtime/nDisplay/Source/DisplayClusterWarp/Private/Policy/InFrustumFit/DisplayClusterWarpInFrustumFitPolicyFactory.cpp
--- a/Engine/Plugins/Runtime/nDisplay/Source/DisplayClusterWarp/Private/Policy/InFrustumFit/DisplayClusterWarpInFrustumFitPolicyFactory.cpp
+++ b/Engine/Plugins/Runtime/nDisplay/Source/DisplayClusterWarp/Private/Policy/InFrustumFit/DisplayClusterWarpInFrustumFitPolicyFactory.cpp
@@ -13,7 +13,7 @@ TSharedPtr<IDisplayClusterWarpPolicy, ESPMode::ThreadSafe> FDisplayClusterWarpIn
 {
 	TSharedPtr<IDisplayClusterWarpPolicy, ESPMode::ThreadSafe> Result;
 
-	if (InWarpPolicyType == DisplayClusterWarpStrings::warp::InFrustumFit)
+	if (InWarpPolicyType == UE::DisplayClusterWarpStrings::warp::InFrustumFit)
 	{
 		Result = MakeShared<FDisplayClusterWarpInFrustumFitPolicy, ESPMode::ThreadSafe>(InWarpPolicyName);
 	}
diff --git a/Engine/Plugins/Runtime/nDisplay/Source/DisplayClusterWarp/Public/Components/DisplayClusterInFrustumFitCameraComponent.h b/Engine/Plugins/Runtime/nDisplay/Source/DisplayClusterWarp/Public/Components/DisplayClusterInFrustumFitCameraComponent.h
--- a/Engine/Plugins/Runtime/nDisplay/Source/DisplayClusterWarp/Public/Components/DisplayClusterInFrustumFitCameraComponent.h
+++ b/Engine/Plugins/Runtime/nDisplay/Source/DisplayClusterWarp/Public/Components/DisplayClusterInFrustumFitCameraComponent.h
@@ -33,6 +33,9 @@ public:
 	virtual void GetDesiredView(IDisplayClusterViewportConfiguration& InViewportConfiguration, FMinimalViewInfo& InOutViewInfo, float* OutCustomNearClippingPlane = nullptr) override;
 	virtual bool ShouldUseEntireClusterViewports(IDisplayClusterViewportManager* InViewportManager) const override;
 	virtual IDisplayClusterWarpPolicy* GetWarpPolicy(IDisplayClusterViewportManager* InViewportManager) override;
+
+	virtual TObjectPtr<UMaterial> GetDisplayDeviceMaterial(const EDisplayClusterDisplayDeviceMeshType InMeshType, const EDisplayClusterDisplayDeviceMaterialType InMaterialType) const override;
+	virtual void OnUpdateDisplayDeviceMeshAndMaterialInstance(IDisplayClusterViewportPreview& InViewportPreview, const EDisplayClusterDisplayDeviceMeshType InMeshType, const EDisplayClusterDisplayDeviceMaterialType InMaterialType, UMeshComponent* InMeshComponent, UMaterialInstanceDynamic* InMeshMaterialInstance) const override;
 	//~~End UDisplayClusterCameraComponent
 
 #if WITH_EDITOR
diff --git a/Engine/Plugins/Runtime/nDisplay/Source/DisplayClusterWarp/Public/DisplayClusterWarpStrings.h b/Engine/Plugins/Runtime/nDisplay/Source/DisplayClusterWarp/Public/DisplayClusterWarpStrings.h
--- a/Engine/Plugins/Runtime/nDisplay/Source/DisplayClusterWarp/Public/DisplayClusterWarpStrings.h
+++ b/Engine/Plugins/Runtime/nDisplay/Source/DisplayClusterWarp/Public/DisplayClusterWarpStrings.h
@@ -4,7 +4,7 @@
 
 #include "CoreMinimal.h"
 
-namespace DisplayClusterWarpStrings
+namespace UE::DisplayClusterWarpStrings
 {
 	namespace warp
 	{
