nDisplay: Fix PVS warning due to call of virtual Disconnect() in destructor of cluster client.

It now uses a custom deleter that calls virtual Disconnect() of the client before deleting it.

diff --git a/Engine/Plugins/Runtime/nDisplay/Source/DisplayCluster/Private/Cluster/Controller/DisplayClusterClusterNodeCtrlBase.cpp b/Engine/Plugins/Runtime/nDisplay/Source/DisplayCluster/Private/Cluster/Controller/DisplayClusterClusterNodeCtrlBase.cpp
--- a/Engine/Plugins/Runtime/nDisplay/Source/DisplayCluster/Private/Cluster/Controller/DisplayClusterClusterNodeCtrlBase.cpp
+++ b/Engine/Plugins/Runtime/nDisplay/Source/DisplayCluster/Private/Cluster/Controller/DisplayClusterClusterNodeCtrlBase.cpp
@@ -17,7 +17,8 @@
 FDisplayClusterClusterNodeCtrlBase::FDisplayClusterClusterNodeCtrlBase(const FString& CtrlName, const FString& NodeName)
 	: NodeName(NodeName)
 	, ControllerName(CtrlName)
-	, ExternalEventsClientJson(MakeUnique<FDisplayClusterClusterEventsJsonClient>())
+	, ExternalEventsClientJson(new FDisplayClusterClusterEventsJsonClient(), FDisplayClusterClientDeleter())
+	, ExternalEventsClientBinary(new FDisplayClusterClusterEventsBinaryClient(), FDisplayClusterClientDeleter())
 {
 }
 
diff --git a/Engine/Plugins/Runtime/nDisplay/Source/DisplayCluster/Private/Cluster/Controller/DisplayClusterClusterNodeCtrlBase.h b/Engine/Plugins/Runtime/nDisplay/Source/DisplayCluster/Private/Cluster/Controller/DisplayClusterClusterNodeCtrlBase.h
--- a/Engine/Plugins/Runtime/nDisplay/Source/DisplayCluster/Private/Cluster/Controller/DisplayClusterClusterNodeCtrlBase.h
+++ b/Engine/Plugins/Runtime/nDisplay/Source/DisplayCluster/Private/Cluster/Controller/DisplayClusterClusterNodeCtrlBase.h
@@ -3,8 +3,8 @@
 #pragma once
 
 #include "Cluster/Controller/IDisplayClusterClusterNodeController.h"
+#include "Network/IDisplayClusterClient.h"
 
-class IDisplayClusterClient;
 class IDisplayClusterServer;
 class FDisplayClusterTcpListener;
 class FDisplayClusterClusterEventsBinaryClient;
@@ -147,9 +147,9 @@ private:
 
 	// JSON client for sending events outside of the cluster
 	FCriticalSection ExternEventsClientJsonGuard;
-	TUniquePtr<FDisplayClusterClusterEventsJsonClient>   ExternalEventsClientJson;
+	TUniquePtr<FDisplayClusterClusterEventsJsonClient, FDisplayClusterClientDeleter> ExternalEventsClientJson;
 
 	// Binary client for sending events outside of the cluster
 	FCriticalSection ExternEventsClientBinaryGuard;
-	TUniquePtr<FDisplayClusterClusterEventsBinaryClient> ExternalEventsClientBinary;
+	TUniquePtr<FDisplayClusterClusterEventsBinaryClient, FDisplayClusterClientDeleter> ExternalEventsClientBinary;
 };
diff --git a/Engine/Plugins/Runtime/nDisplay/Source/DisplayCluster/Private/Cluster/Controller/DisplayClusterClusterNodeCtrlSecondary.cpp b/Engine/Plugins/Runtime/nDisplay/Source/DisplayCluster/Private/Cluster/Controller/DisplayClusterClusterNodeCtrlSecondary.cpp
--- a/Engine/Plugins/Runtime/nDisplay/Source/DisplayCluster/Private/Cluster/Controller/DisplayClusterClusterNodeCtrlSecondary.cpp
+++ b/Engine/Plugins/Runtime/nDisplay/Source/DisplayCluster/Private/Cluster/Controller/DisplayClusterClusterNodeCtrlSecondary.cpp
@@ -153,10 +153,10 @@ bool FDisplayClusterClusterNodeCtrlSecondary::InitializeClients()
 	UE_LOG(LogDisplayClusterCluster, Log, TEXT("%s - initializing secondary node clients..."), *GetControllerName());
 
 	// Instantiate local clients
-	ClusterSyncClient         = MakeUnique<FDisplayClusterClusterSyncClient>();
-	RenderSyncClient          = MakeUnique<FDisplayClusterRenderSyncClient>();
-	ClusterEventsJsonClient   = MakeUnique<FDisplayClusterClusterEventsJsonClient>();
-	ClusterEventsBinaryClient = MakeUnique<FDisplayClusterClusterEventsBinaryClient>();
+	ClusterSyncClient         = TUniquePtr<FDisplayClusterClusterSyncClient        , FDisplayClusterClientDeleter>(new FDisplayClusterClusterSyncClient()        , FDisplayClusterClientDeleter());
+	RenderSyncClient          = TUniquePtr<FDisplayClusterRenderSyncClient         , FDisplayClusterClientDeleter>(new FDisplayClusterRenderSyncClient()         , FDisplayClusterClientDeleter());
+	ClusterEventsJsonClient   = TUniquePtr<FDisplayClusterClusterEventsJsonClient  , FDisplayClusterClientDeleter>(new FDisplayClusterClusterEventsJsonClient()  , FDisplayClusterClientDeleter());
+	ClusterEventsBinaryClient = TUniquePtr<FDisplayClusterClusterEventsBinaryClient, FDisplayClusterClientDeleter>(new FDisplayClusterClusterEventsBinaryClient(), FDisplayClusterClientDeleter());
 
 	return ClusterSyncClient && RenderSyncClient && ClusterEventsJsonClient && ClusterEventsBinaryClient;
 }
diff --git a/Engine/Plugins/Runtime/nDisplay/Source/DisplayCluster/Private/Cluster/Controller/DisplayClusterClusterNodeCtrlSecondary.h b/Engine/Plugins/Runtime/nDisplay/Source/DisplayCluster/Private/Cluster/Controller/DisplayClusterClusterNodeCtrlSecondary.h
--- a/Engine/Plugins/Runtime/nDisplay/Source/DisplayCluster/Private/Cluster/Controller/DisplayClusterClusterNodeCtrlSecondary.h
+++ b/Engine/Plugins/Runtime/nDisplay/Source/DisplayCluster/Private/Cluster/Controller/DisplayClusterClusterNodeCtrlSecondary.h
@@ -73,8 +73,8 @@ protected:
 
 private:
 	// Cluster node clients
-	TUniquePtr<FDisplayClusterClusterSyncClient>         ClusterSyncClient;
-	TUniquePtr<FDisplayClusterRenderSyncClient>          RenderSyncClient;
-	TUniquePtr<FDisplayClusterClusterEventsJsonClient>   ClusterEventsJsonClient;
-	TUniquePtr<FDisplayClusterClusterEventsBinaryClient> ClusterEventsBinaryClient;
+	TUniquePtr<FDisplayClusterClusterSyncClient, FDisplayClusterClientDeleter>         ClusterSyncClient;
+	TUniquePtr<FDisplayClusterRenderSyncClient, FDisplayClusterClientDeleter>          RenderSyncClient;
+	TUniquePtr<FDisplayClusterClusterEventsJsonClient, FDisplayClusterClientDeleter>   ClusterEventsJsonClient;
+	TUniquePtr<FDisplayClusterClusterEventsBinaryClient, FDisplayClusterClientDeleter> ClusterEventsBinaryClient;
 };
diff --git a/Engine/Plugins/Runtime/nDisplay/Source/DisplayCluster/Private/Cluster/DisplayClusterClusterManager.cpp b/Engine/Plugins/Runtime/nDisplay/Source/DisplayCluster/Private/Cluster/DisplayClusterClusterManager.cpp
--- a/Engine/Plugins/Runtime/nDisplay/Source/DisplayCluster/Private/Cluster/DisplayClusterClusterManager.cpp
+++ b/Engine/Plugins/Runtime/nDisplay/Source/DisplayCluster/Private/Cluster/DisplayClusterClusterManager.cpp
@@ -308,9 +308,10 @@ void FDisplayClusterClusterManager::UnregisterSyncObject(IDisplayClusterClusterS
 	}
 }
 
-TUniquePtr<IDisplayClusterGenericBarriersClient> FDisplayClusterClusterManager::CreateGenericBarriersClient(const FString& ClientName)
+TUniquePtr<IDisplayClusterGenericBarriersClient, FDisplayClusterGenericBarriersClientDeleter> FDisplayClusterClusterManager::CreateGenericBarriersClient(const FString& ClientName)
 {
-	return MakeUnique<FDisplayClusterGenericBarrierClient>(ClientName);
+	return TUniquePtr<FDisplayClusterGenericBarrierClient, FDisplayClusterGenericBarriersClientDeleter>(
+		new FDisplayClusterGenericBarrierClient(ClientName), FDisplayClusterGenericBarriersClientDeleter());
 }
 
 void FDisplayClusterClusterManager::AddClusterEventListener(TScriptInterface<IDisplayClusterClusterEventListener> Listener)
diff --git a/Engine/Plugins/Runtime/nDisplay/Source/DisplayCluster/Private/Cluster/DisplayClusterClusterManager.h b/Engine/Plugins/Runtime/nDisplay/Source/DisplayCluster/Private/Cluster/DisplayClusterClusterManager.h
--- a/Engine/Plugins/Runtime/nDisplay/Source/DisplayCluster/Private/Cluster/DisplayClusterClusterManager.h
+++ b/Engine/Plugins/Runtime/nDisplay/Source/DisplayCluster/Private/Cluster/DisplayClusterClusterManager.h
@@ -62,7 +62,7 @@ public:
 	virtual void RegisterSyncObject(IDisplayClusterClusterSyncObject* SyncObj, EDisplayClusterSyncGroup SyncGroup) override;
 	virtual void UnregisterSyncObject(IDisplayClusterClusterSyncObject* SyncObj) override;
 
-	virtual TUniquePtr<IDisplayClusterGenericBarriersClient> CreateGenericBarriersClient(const FString& ClientName) override;
+	virtual TUniquePtr<IDisplayClusterGenericBarriersClient, FDisplayClusterGenericBarriersClientDeleter> CreateGenericBarriersClient(const FString& ClientName) override;
 
 	virtual void AddClusterEventListener(TScriptInterface<IDisplayClusterClusterEventListener>) override;
 	virtual void RemoveClusterEventListener(TScriptInterface<IDisplayClusterClusterEventListener>) override;
diff --git a/Engine/Plugins/Runtime/nDisplay/Source/DisplayCluster/Private/Cluster/DisplayClusterNetDriverHelper.cpp b/Engine/Plugins/Runtime/nDisplay/Source/DisplayCluster/Private/Cluster/DisplayClusterNetDriverHelper.cpp
--- a/Engine/Plugins/Runtime/nDisplay/Source/DisplayCluster/Private/Cluster/DisplayClusterNetDriverHelper.cpp
+++ b/Engine/Plugins/Runtime/nDisplay/Source/DisplayCluster/Private/Cluster/DisplayClusterNetDriverHelper.cpp
@@ -79,7 +79,7 @@ bool FDisplayClusterNetDriverHelper::RegisterClusterEventsBinaryClient(uint32 Cl
 		return false;
 	}
 	
-	TSharedPtr<FDisplayClusterClusterEventsBinaryClient> Client = MakeShared<FDisplayClusterClusterEventsBinaryClient>();
+	TSharedPtr<FDisplayClusterClusterEventsBinaryClient> Client(new FDisplayClusterClusterEventsBinaryClient(), FDisplayClusterClientDeleter());
 	
 	if (!Client->Connect(ClientAddress, ClientPort, 1, 0.f))
 	{
diff --git a/Engine/Plugins/Runtime/nDisplay/Source/DisplayCluster/Private/Network/DisplayClusterClient.h b/Engine/Plugins/Runtime/nDisplay/Source/DisplayCluster/Private/Network/DisplayClusterClient.h
--- a/Engine/Plugins/Runtime/nDisplay/Source/DisplayCluster/Private/Network/DisplayClusterClient.h
+++ b/Engine/Plugins/Runtime/nDisplay/Source/DisplayCluster/Private/Network/DisplayClusterClient.h
@@ -23,14 +23,10 @@ public:
 		: FDisplayClusterSocketOperations(CreateSocket(InName), DisplayClusterConstants::net::PacketBufferSize, InName)
 	{ }
 
-	virtual ~FDisplayClusterClientBase()
-	{
-		Disconnect();
-	}
-
 public:
 	// Connects to a server
 	virtual bool Connect(const FString& Address, const uint16 Port, const uint32 ConnectRetriesAmount, const uint32 ConnectRetryDelay) override;
+
 	// Terminates current connection
 	virtual void Disconnect() override;
 
@@ -48,6 +44,7 @@ public:
 protected:
 	// Creates client socket
 	FSocket* CreateSocket(const FString& InName);
+
 };
 
 
diff --git a/Engine/Plugins/Runtime/nDisplay/Source/DisplayCluster/Private/Network/IDisplayClusterClient.h b/Engine/Plugins/Runtime/nDisplay/Source/DisplayCluster/Private/Network/IDisplayClusterClient.h
--- a/Engine/Plugins/Runtime/nDisplay/Source/DisplayCluster/Private/Network/IDisplayClusterClient.h
+++ b/Engine/Plugins/Runtime/nDisplay/Source/DisplayCluster/Private/Network/IDisplayClusterClient.h
@@ -4,6 +4,7 @@
 
 #include "CoreMinimal.h"
 
+#include "Cluster/IDisplayClusterGenericBarriersClient.h"
 
 /**
  * DisplayCluster TCP client interface
@@ -23,3 +24,6 @@ public:
 	// Returns client name
 	virtual FString GetName() const = 0;
 };
+
+// Custom deleter for IDisplayClusterClient. See FDisplayClusterGenericClientDeleter for more info.
+using FDisplayClusterClientDeleter = FDisplayClusterGenericClientDeleter<IDisplayClusterClient>;
diff --git a/Engine/Plugins/Runtime/nDisplay/Source/DisplayCluster/Public/Cluster/IDisplayClusterClusterManager.h b/Engine/Plugins/Runtime/nDisplay/Source/DisplayCluster/Public/Cluster/IDisplayClusterClusterManager.h
--- a/Engine/Plugins/Runtime/nDisplay/Source/DisplayCluster/Public/Cluster/IDisplayClusterClusterManager.h
+++ b/Engine/Plugins/Runtime/nDisplay/Source/DisplayCluster/Public/Cluster/IDisplayClusterClusterManager.h
@@ -3,6 +3,8 @@
 #pragma once
 
 #include "CoreMinimal.h"
+
+#include "Cluster/IDisplayClusterGenericBarriersClient.h"
 #include "Delegates/Delegate.h"
 #include "UObject/ScriptInterface.h"
 #include "DisplayClusterEnums.h"
@@ -13,7 +15,7 @@ class IDisplayClusterClusterEventBinaryListener;
 class IDisplayClusterClusterEventJsonListener;
 class IDisplayClusterClusterEventListener;
 class IDisplayClusterClusterSyncObject;
-class IDisplayClusterGenericBarriersClient;
+
 
 DECLARE_MULTICAST_DELEGATE_OneParam(FOnClusterEventJson, const FDisplayClusterClusterEventJson& /* Event */);
 typedef FOnClusterEventJson::FDelegate FOnClusterEventJsonListener;
@@ -56,7 +58,7 @@ public:
 	virtual void UnregisterSyncObject(IDisplayClusterClusterSyncObject* SyncObj) = 0;
 
 	/** Spawns a client for generic barriers service. */
-	virtual TUniquePtr<IDisplayClusterGenericBarriersClient> CreateGenericBarriersClient(const FString& ClientName) = 0;
+	virtual TUniquePtr<IDisplayClusterGenericBarriersClient, FDisplayClusterGenericBarriersClientDeleter> CreateGenericBarriersClient(const FString& ClientName) = 0;
 
 	//////////////////////////////////////////////////////////////////////////////////////////////
 	// Cluster events
diff --git a/Engine/Plugins/Runtime/nDisplay/Source/DisplayCluster/Public/Cluster/IDisplayClusterGenericBarriersClient.h b/Engine/Plugins/Runtime/nDisplay/Source/DisplayCluster/Public/Cluster/IDisplayClusterGenericBarriersClient.h
--- a/Engine/Plugins/Runtime/nDisplay/Source/DisplayCluster/Public/Cluster/IDisplayClusterGenericBarriersClient.h
+++ b/Engine/Plugins/Runtime/nDisplay/Source/DisplayCluster/Public/Cluster/IDisplayClusterGenericBarriersClient.h
@@ -67,3 +67,26 @@ public:
 	/** Synchronize calling thread on a specific barrier with custom data */
 	virtual bool Synchronize(const FString& BarrierId, const FString& UniqueThreadMarker, const TArray<uint8>& RequestData, TArray<uint8>& OutResponseData) = 0;
 };
+
+/**
+ * Custom deleter to avoid calling the virtual Disconnect function in the destructor.
+ * This is necessary since the destructor of the base class is calling it but expecting
+ * the derived class implementation to be called. This also avoids the PVS warning:
+ * V1053: Calling the 'Disconnect' virtual function in the destructor may lead to unexpected result at runtime.
+ */
+template<typename ClientType>
+struct FDisplayClusterGenericClientDeleter
+{
+	void operator()(ClientType* Client)
+	{
+		if (Client && Client->IsConnected())
+		{
+			Client->Disconnect();
+		}
+
+		delete Client;
+	}
+};
+
+// Custom deleter for IDisplayClusterGenericBarriersClient.
+using FDisplayClusterGenericBarriersClientDeleter = FDisplayClusterGenericClientDeleter<IDisplayClusterGenericBarriersClient>;
diff --git a/Engine/Plugins/Runtime/nDisplay/Source/DisplayClusterMedia/Public/Synchronization/DisplayClusterMediaOutputSynchronizationPolicyEthernetBarrierBase.h b/Engine/Plugins/Runtime/nDisplay/Source/DisplayClusterMedia/Public/Synchronization/DisplayClusterMediaOutputSynchronizationPolicyEthernetBarrierBase.h
--- a/Engine/Plugins/Runtime/nDisplay/Source/DisplayClusterMedia/Public/Synchronization/DisplayClusterMediaOutputSynchronizationPolicyEthernetBarrierBase.h
+++ b/Engine/Plugins/Runtime/nDisplay/Source/DisplayClusterMedia/Public/Synchronization/DisplayClusterMediaOutputSynchronizationPolicyEthernetBarrierBase.h
@@ -73,7 +73,7 @@ private:
 	FString ThreadMarker;
 
 	/** Barrier sync client. */
-	TUniquePtr<IDisplayClusterGenericBarriersClient> EthernetBarrierClient;
+	TUniquePtr<IDisplayClusterGenericBarriersClient, FDisplayClusterGenericBarriersClientDeleter> EthernetBarrierClient;
 };
 
 /*
