Fix warning V523: The 'then' statement is equivalent to the 'else' statement.
#rnx

diff --git a/Engine/Plugins/Runtime/nDisplay/Source/DisplayCluster/Private/Game/EngineClasses/Scene/DisplayClusterRootActor.cpp b/Engine/Plugins/Runtime/nDisplay/Source/DisplayCluster/Private/Game/EngineClasses/Scene/DisplayClusterRootActor.cpp
--- a/Engine/Plugins/Runtime/nDisplay/Source/DisplayCluster/Private/Game/EngineClasses/Scene/DisplayClusterRootActor.cpp
+++ b/Engine/Plugins/Runtime/nDisplay/Source/DisplayCluster/Private/Game/EngineClasses/Scene/DisplayClusterRootActor.cpp
@@ -734,27 +734,10 @@ void ADisplayClusterRootActor::InitializeRootActor()
 
 	StageGeometryComponent->Invalidate();
 
-	// Packaged, PIE and -game runtime
-	if (IsRunningGame() || IsRunningPIE())
-	{
-		if (CurrentConfigData)
-		{
-			BuildHierarchy();
-
-			return;
-		}
-	}
-#if WITH_EDITOR
-	// Initialize from file property by default in Editor
-	else
+	if (CurrentConfigData)
 	{
-		if (CurrentConfigData)
-		{
-			BuildHierarchy();
-			return;
-		}
+		BuildHierarchy();
 	}
-#endif
 }
 
 void ADisplayClusterRootActor::UpdateProceduralMeshComponentData(const UProceduralMeshComponent* InProceduralMeshComponent)
