nDisplay: Not simulating click when rendering offscreen.

Also, removing an annoying warning when shared memory media tries to open mapped memory before creating it.

#rnx

diff --git a/Engine/Plugins/Runtime/nDisplay/Source/DisplayCluster/Private/Render/DisplayClusterRenderManager.cpp b/Engine/Plugins/Runtime/nDisplay/Source/DisplayCluster/Private/Render/DisplayClusterRenderManager.cpp
--- a/Engine/Plugins/Runtime/nDisplay/Source/DisplayCluster/Private/Render/DisplayClusterRenderManager.cpp
+++ b/Engine/Plugins/Runtime/nDisplay/Source/DisplayCluster/Private/Render/DisplayClusterRenderManager.cpp
@@ -251,7 +251,9 @@ static bool SimulateMouseClick(HWND WindowHandle)
 
 void FDisplayClusterRenderManager::PreTick(float DeltaSeconds)
 {
-	if(!bWasWindowFocused)
+	static const bool bIsRenderingOffscreen = FParse::Param(FCommandLine::Get(), TEXT("RenderOffscreen"));
+
+	if(!bIsRenderingOffscreen && !bWasWindowFocused)
 	{
 		if (UGameViewportClient* GameViewportClient = GEngine->GameViewport)
 		{
diff --git a/Engine/Plugins/Runtime/nDisplay/Source/SharedMemoryMedia/Private/SharedMemoryMediaCapture.cpp b/Engine/Plugins/Runtime/nDisplay/Source/SharedMemoryMedia/Private/SharedMemoryMediaCapture.cpp
--- a/Engine/Plugins/Runtime/nDisplay/Source/SharedMemoryMedia/Private/SharedMemoryMediaCapture.cpp
+++ b/Engine/Plugins/Runtime/nDisplay/Source/SharedMemoryMedia/Private/SharedMemoryMediaCapture.cpp
@@ -51,14 +51,22 @@ bool USharedMemoryMediaCapture::InitializeCapture()
 
 		const FString SharedMemoryRegionName = Guid.ToString(EGuidFormats::DigitsWithHyphensInBraces);
 
-		// Open existing shared memory region, in case it exists:
+		// Open existing shared memory region, in case it exists. 
 
 		const uint32 AccessMode = FPlatformMemory::ESharedMemoryAccess::Read | FPlatformMemory::ESharedMemoryAccess::Write;
 
+		// Disable LogHAL warnings caused by failing to open the shared memory.
+
+		const ELogVerbosity::Type LogHALVerbosity = LogHAL.GetVerbosity();
+		LogHAL.SetVerbosity(ELogVerbosity::Error);
+
 		FPlatformMemory::FSharedMemoryRegion* SharedMemoryRegion = FPlatformMemory::MapNamedSharedMemoryRegion(
 			*SharedMemoryRegionName, false /* bCreate */, AccessMode, SharedMemorySize
 		);
 
+		// Restore the verbosity level of LogHAL to the previous value
+		LogHAL.SetVerbosity(LogHALVerbosity);
+
 		// If it doesn't exist, then we allocate and zero-initialize it.
 		if (!SharedMemoryRegion)
 		{
