nDisplay: Partial backout of 29676031

#rnx

diff --git a/Engine/Plugins/Runtime/nDisplay/Source/SharedMemoryMedia/Private/SharedMemoryMediaCapture.cpp b/Engine/Plugins/Runtime/nDisplay/Source/SharedMemoryMedia/Private/SharedMemoryMediaCapture.cpp
--- a/Engine/Plugins/Runtime/nDisplay/Source/SharedMemoryMedia/Private/SharedMemoryMediaCapture.cpp
+++ b/Engine/Plugins/Runtime/nDisplay/Source/SharedMemoryMedia/Private/SharedMemoryMediaCapture.cpp
@@ -51,22 +51,14 @@ bool USharedMemoryMediaCapture::InitializeCapture()
 
 		const FString SharedMemoryRegionName = Guid.ToString(EGuidFormats::DigitsWithHyphensInBraces);
 
-		// Open existing shared memory region, in case it exists. 
+		// Open existing shared memory region, in case it exists:
 
 		const uint32 AccessMode = FPlatformMemory::ESharedMemoryAccess::Read | FPlatformMemory::ESharedMemoryAccess::Write;
 
-		// Disable LogHAL warnings caused by failing to open the shared memory.
-
-		const ELogVerbosity::Type LogHALVerbosity = LogHAL.GetVerbosity();
-		LogHAL.SetVerbosity(ELogVerbosity::Error);
-
 		FPlatformMemory::FSharedMemoryRegion* SharedMemoryRegion = FPlatformMemory::MapNamedSharedMemoryRegion(
 			*SharedMemoryRegionName, false /* bCreate */, AccessMode, SharedMemorySize
 		);
 
-		// Restore the verbosity level of LogHAL to the previous value
-		LogHAL.SetVerbosity(LogHALVerbosity);
-
 		// If it doesn't exist, then we allocate and zero-initialize it.
 		if (!SharedMemoryRegion)
 		{
