nDisplay: Work around annoying LogHAL warning when checking if a named mapped memory already exists.

diff --git a/Engine/Plugins/Runtime/nDisplay/Source/SharedMemoryMedia/Private/SharedMemoryMediaCapture.cpp b/Engine/Plugins/Runtime/nDisplay/Source/SharedMemoryMedia/Private/SharedMemoryMediaCapture.cpp
--- a/Engine/Plugins/Runtime/nDisplay/Source/SharedMemoryMedia/Private/SharedMemoryMediaCapture.cpp
+++ b/Engine/Plugins/Runtime/nDisplay/Source/SharedMemoryMedia/Private/SharedMemoryMediaCapture.cpp
@@ -55,10 +55,21 @@ bool USharedMemoryMediaCapture::InitializeCapture()
 
 		const uint32 AccessMode = FPlatformMemory::ESharedMemoryAccess::Read | FPlatformMemory::ESharedMemoryAccess::Write;
 
+#if !NO_LOGGING
+		// Disable LogHAL warnings caused by failing to open the shared memory.
+		const ELogVerbosity::Type LogHALVerbosity = LogHAL.GetVerbosity();
+		LogHAL.SetVerbosity(ELogVerbosity::Error);
+#endif // !NO_LOGGING
+
 		FPlatformMemory::FSharedMemoryRegion* SharedMemoryRegion = FPlatformMemory::MapNamedSharedMemoryRegion(
 			*SharedMemoryRegionName, false /* bCreate */, AccessMode, SharedMemorySize
 		);
 
+#if !NO_LOGGING
+		// Restore the verbosity level of LogHAL to the previous value
+		LogHAL.SetVerbosity(LogHALVerbosity);
+#endif // !NO_LOGGING
+
 		// If it doesn't exist, then we allocate and zero-initialize it.
 		if (!SharedMemoryRegion)
 		{
