nDisplay: Added option to disable overriding transforms (xform, camera, screens) from external config (.ndisplay).

Users are encountering two problems:
* Unexpected behavior when they launch the cluster and the projection is different than expected, and seemingly gets "fixed" when wiggling the NDC via multi-user.
* Difficulty in sharing the same NDC archetype in multiple levels that have NDC instances with transform overrides in them.

Default to disable overrides, because it is the common desired case in VP. Users that dynamically create or modify the .ndisplay and wish those changes to override, should set this new option to true.

diff --git a/Engine/Plugins/Runtime/nDisplay/Source/DisplayCluster/Private/Game/EngineClasses/Scene/DisplayClusterRootActor.cpp b/Engine/Plugins/Runtime/nDisplay/Source/DisplayCluster/Private/Game/EngineClasses/Scene/DisplayClusterRootActor.cpp
--- a/Engine/Plugins/Runtime/nDisplay/Source/DisplayCluster/Private/Game/EngineClasses/Scene/DisplayClusterRootActor.cpp
+++ b/Engine/Plugins/Runtime/nDisplay/Source/DisplayCluster/Private/Game/EngineClasses/Scene/DisplayClusterRootActor.cpp
@@ -401,6 +401,7 @@ void ADisplayClusterRootActor::OverrideFromConfig(UDisplayClusterConfigurationDa
 	}
 
 	// Update component transforms if any of them have been changed in the file
+	if (ConfigData->bOverrideTransformsFromExternalConfig)
 	{
 		TUniquePtr<FDisplayClusterRootActorInitializer> Initializer = MakeUnique<FDisplayClusterRootActorInitializer>();
 		Initializer->UpdateComponentTransformsOnly(this, ConfigData);
diff --git a/Engine/Plugins/Runtime/nDisplay/Source/DisplayClusterConfiguration/Private/Formats/JSON500/DisplayClusterConfigurationJsonParser_500.cpp b/Engine/Plugins/Runtime/nDisplay/Source/DisplayClusterConfiguration/Private/Formats/JSON500/DisplayClusterConfigurationJsonParser_500.cpp
--- a/Engine/Plugins/Runtime/nDisplay/Source/DisplayClusterConfiguration/Private/Formats/JSON500/DisplayClusterConfigurationJsonParser_500.cpp
+++ b/Engine/Plugins/Runtime/nDisplay/Source/DisplayClusterConfiguration/Private/Formats/JSON500/DisplayClusterConfigurationJsonParser_500.cpp
@@ -106,6 +106,7 @@ namespace JSON500
 		Config->bFollowLocalPlayerCamera = CfgJson.Misc.bFollowLocalPlayerCamera;
 		Config->bExitOnEsc = CfgJson.Misc.bExitOnEsc;
 		Config->bOverrideViewportsFromExternalConfig = CfgJson.Misc.bOverrideViewportsFromExternalConfig;
+		Config->bOverrideTransformsFromExternalConfig = CfgJson.Misc.bOverrideTransformsFromExternalConfig;
 
 		// Scene
 		{
@@ -310,6 +311,7 @@ namespace JSON500
 		Json.Misc.bFollowLocalPlayerCamera = Config->bFollowLocalPlayerCamera;
 		Json.Misc.bExitOnEsc = Config->bExitOnEsc;
 		Json.Misc.bOverrideViewportsFromExternalConfig = Config->bOverrideViewportsFromExternalConfig;
+		Json.Misc.bOverrideTransformsFromExternalConfig = Config->bOverrideTransformsFromExternalConfig;
 
 		// Scene
 		{
diff --git a/Engine/Plugins/Runtime/nDisplay/Source/DisplayClusterConfiguration/Private/Formats/JSON500/DisplayClusterConfigurationJsonTypes_500.h b/Engine/Plugins/Runtime/nDisplay/Source/DisplayClusterConfiguration/Private/Formats/JSON500/DisplayClusterConfigurationJsonTypes_500.h
--- a/Engine/Plugins/Runtime/nDisplay/Source/DisplayClusterConfiguration/Private/Formats/JSON500/DisplayClusterConfigurationJsonTypes_500.h
+++ b/Engine/Plugins/Runtime/nDisplay/Source/DisplayClusterConfiguration/Private/Formats/JSON500/DisplayClusterConfigurationJsonTypes_500.h
@@ -215,6 +215,7 @@ public:
 		: bFollowLocalPlayerCamera(false)
 		, bExitOnEsc(true)
 		, bOverrideViewportsFromExternalConfig(false)
+		, bOverrideTransformsFromExternalConfig(false)
 	{ }
 
 public:
@@ -226,6 +227,9 @@ public:
 
 	UPROPERTY()
 	bool bOverrideViewportsFromExternalConfig;
+
+	UPROPERTY()
+	bool bOverrideTransformsFromExternalConfig;
 };
 
 
diff --git a/Engine/Plugins/Runtime/nDisplay/Source/DisplayClusterConfiguration/Public/DisplayClusterConfigurationTypes.h b/Engine/Plugins/Runtime/nDisplay/Source/DisplayClusterConfiguration/Public/DisplayClusterConfigurationTypes.h
--- a/Engine/Plugins/Runtime/nDisplay/Source/DisplayClusterConfiguration/Public/DisplayClusterConfigurationTypes.h
+++ b/Engine/Plugins/Runtime/nDisplay/Source/DisplayClusterConfiguration/Public/DisplayClusterConfigurationTypes.h
@@ -583,6 +583,9 @@ public:
 	UPROPERTY(EditDefaultsOnly, BlueprintReadWrite, Category = Configuration, meta = (DisplayName = "Override Viewports From .ndisplay", DisplayAfter = "bExitOnEsc"))
 	bool bOverrideViewportsFromExternalConfig = false;
 
+	UPROPERTY(EditDefaultsOnly, BlueprintReadWrite, Category = Configuration, meta = (DisplayName = "Override Transforms From .ndisplay", DisplayAfter = "bOverrideViewportsFromExternalConfig"))
+	bool bOverrideTransformsFromExternalConfig = false;
+
 	UPROPERTY(EditAnywhere, BlueprintReadWrite, Category = Configuration, meta = (DisplayName = "Global Media Settings"))
 	FDisplayClusterConfigurationGlobalMediaSettings MediaSettings;
 
