nDisplay: Display device exposure property now affects all materials and not just the techvis material.

The property can also be useful to brighten or darken the in-editor preview, especially when preview post process is enabled.

Also added a Gamma property, which can be useful when receiving a texture that needs such conversion.

diff --git a/Engine/Plugins/Runtime/nDisplay/Source/DisplayCluster/Private/Render/DisplayDevice/Components/DisplayClusterDisplayDeviceComponent.cpp b/Engine/Plugins/Runtime/nDisplay/Source/DisplayCluster/Private/Render/DisplayDevice/Components/DisplayClusterDisplayDeviceComponent.cpp
--- a/Engine/Plugins/Runtime/nDisplay/Source/DisplayCluster/Private/Render/DisplayDevice/Components/DisplayClusterDisplayDeviceComponent.cpp
+++ b/Engine/Plugins/Runtime/nDisplay/Source/DisplayCluster/Private/Render/DisplayDevice/Components/DisplayClusterDisplayDeviceComponent.cpp
@@ -23,19 +23,8 @@ void UDisplayClusterDisplayDeviceComponent::OnUpdateDisplayDeviceMeshAndMaterial
 	// Customizes and overrides material parameters after the base class
 	if (InMeshMaterialInstance && ShouldUseDisplayDevice(InViewportPreview.GetConfiguration()))
 	{
-		switch (InMaterialType)
-		{
-		case EDisplayClusterDisplayDeviceMaterialType::PreviewMeshMaterial:
-			InMeshMaterialInstance->SetScalarParameterValue(UE::DisplayClusterDisplayDeviceStrings::material::attr::Exposure, 0.f);
-			break;
-
-		case EDisplayClusterDisplayDeviceMaterialType::PreviewMeshTechvisMaterial:
-			InMeshMaterialInstance->SetScalarParameterValue(UE::DisplayClusterDisplayDeviceStrings::material::attr::Exposure, Exposure);
-			break;
-
-		default:
-			break;
-		}
+		InMeshMaterialInstance->SetScalarParameterValue(UE::DisplayClusterDisplayDeviceStrings::material::attr::Exposure, Exposure);
+		InMeshMaterialInstance->SetScalarParameterValue(UE::DisplayClusterDisplayDeviceStrings::material::attr::Gamma, Gamma);
 	}
 }
 
diff --git a/Engine/Plugins/Runtime/nDisplay/Source/DisplayCluster/Private/Render/DisplayDevice/DisplayClusterDisplayDeviceStrings.h b/Engine/Plugins/Runtime/nDisplay/Source/DisplayCluster/Private/Render/DisplayDevice/DisplayClusterDisplayDeviceStrings.h
--- a/Engine/Plugins/Runtime/nDisplay/Source/DisplayCluster/Private/Render/DisplayDevice/DisplayClusterDisplayDeviceStrings.h
+++ b/Engine/Plugins/Runtime/nDisplay/Source/DisplayCluster/Private/Render/DisplayDevice/DisplayClusterDisplayDeviceStrings.h
@@ -21,6 +21,7 @@ namespace UE::DisplayClusterDisplayDeviceStrings
 			static constexpr const TCHAR* Preview = TEXT("Preview");
 			static constexpr const TCHAR* Opacity = TEXT("Opacity");
 			static constexpr const TCHAR* Exposure = TEXT("Exposure");
+			static constexpr const TCHAR* Gamma = TEXT("Gamma");
 		}
 	}
 };
diff --git a/Engine/Plugins/Runtime/nDisplay/Source/DisplayCluster/Public/Render/DisplayDevice/Components/DisplayClusterDisplayDeviceComponent.h b/Engine/Plugins/Runtime/nDisplay/Source/DisplayCluster/Public/Render/DisplayDevice/Components/DisplayClusterDisplayDeviceComponent.h
--- a/Engine/Plugins/Runtime/nDisplay/Source/DisplayCluster/Public/Render/DisplayDevice/Components/DisplayClusterDisplayDeviceComponent.h
+++ b/Engine/Plugins/Runtime/nDisplay/Source/DisplayCluster/Public/Render/DisplayDevice/Components/DisplayClusterDisplayDeviceComponent.h
@@ -33,7 +33,11 @@ protected:
 	/** Adjust the exposure for the emissive input. */
 	UPROPERTY(EditAnywhere, Category=Material)
 	float Exposure = 0.f;
-	
+
+	/** Adjust the gamma for the emissive input. */
+	UPROPERTY(EditAnywhere, Category = Material)
+	float Gamma = 1.f;
+
 	/** Color grading settings. */
 	UPROPERTY(EditAnywhere, Category = RenderPass, meta = (DisplayAfter="bEnableRenderPass"))
 	FOpenColorIOColorConversionSettings ColorConversionSettings;
