nDisplay: Avoid HasFocus related stalls when running the cluster.

diff --git a/Engine/Plugins/Runtime/nDisplay/Source/DisplayCluster/Private/Game/EngineClasses/Basics/DisplayClusterViewportClient.cpp b/Engine/Plugins/Runtime/nDisplay/Source/DisplayCluster/Private/Game/EngineClasses/Basics/DisplayClusterViewportClient.cpp
--- a/Engine/Plugins/Runtime/nDisplay/Source/DisplayCluster/Private/Game/EngineClasses/Basics/DisplayClusterViewportClient.cpp
+++ b/Engine/Plugins/Runtime/nDisplay/Source/DisplayCluster/Private/Game/EngineClasses/Basics/DisplayClusterViewportClient.cpp
@@ -104,6 +104,16 @@ static FAutoConsoleVariableRef CVarDisplayClusterDebugDraw(
 	ECVF_RenderThreadSafe
 );
 
+// Replaces FApp::HasFocus
+bool GDisplayClusterReplaceHasFocusFunction = true;
+static FAutoConsoleVariableRef CVarDisplayClusterReplaceHasFocusFunction(
+	TEXT("DC.ReplaceHasFocusFunction"),
+	GDisplayClusterReplaceHasFocusFunction,
+	TEXT("Replaces the function that FApp::HasFocus() uses, to mitigate OS stalls that happen in some systems."),
+	ECVF_ReadOnly
+);
+
+
 struct FCompareViewFamilyBySizeAndGPU
 {
 	FORCEINLINE bool operator()(const FSceneViewFamilyContext& A, const FSceneViewFamilyContext& B) const
@@ -374,6 +384,13 @@ void UDisplayClusterViewportClient::Init(struct FWorldContext& WorldContext, UGa
 		{
 			AllowMotionBlurInVR->Set(int32(1));
 		}
+
+		// Replace FApp::HasFocus to avoid stalls observed in some render nodes. 
+		// It always return true, so all code behaves as if the application were in focus, even when rendering offscreen.
+		if (GDisplayClusterReplaceHasFocusFunction)
+		{
+			FApp::SetHasFocusFunction([]() { return true; });
+		}
 	}
 
 	Super::Init(WorldContext, OwningGameInstance, bCreateNewAudioDevice);
