nDisplay: Reworked UDisplayClusterPreviewShareComponent
- PullActor option: the preview setting from the RootActor owner is used to render previews.
- The bPreviewEnablePostProcess parameter has been removed as it is no longer used.
- A new DisplayDeviceRootActorType parameter has been added to FDisplayClusterViewport_PreviewSettings to determine which type of DCRA is used as the source to retrieve the DisplayDeviceComponent.

diff --git a/Engine/Plugins/Runtime/nDisplay/Source/DisplayCluster/Private/Game/EngineClasses/Scene/DisplayClusterPreviewShareComponent.cpp b/Engine/Plugins/Runtime/nDisplay/Source/DisplayCluster/Private/Game/EngineClasses/Scene/DisplayClusterPreviewShareComponent.cpp
--- a/Engine/Plugins/Runtime/nDisplay/Source/DisplayCluster/Private/Game/EngineClasses/Scene/DisplayClusterPreviewShareComponent.cpp
+++ b/Engine/Plugins/Runtime/nDisplay/Source/DisplayCluster/Private/Game/EngineClasses/Scene/DisplayClusterPreviewShareComponent.cpp
@@ -374,7 +374,6 @@ bool UDisplayClusterPreviewShareComponent::UpdateCustomViewportManager(const ADi
 
 	// Override preview settings:
 	NewPreviewSettings.bPreviewEnable = true;
-	NewPreviewSettings.bPreviewEnablePostProcess = bPreviewEnablePostProcess;
 	NewPreviewSettings.bEnablePreviewMesh = false;
 	NewPreviewSettings.bEnablePreviewEditableMesh = false;
 	NewPreviewSettings.bPreviewICVFXFrustums = false; // disable frustum rendering
@@ -538,7 +537,7 @@ void UDisplayClusterPreviewShareComponent::TickSend()
 
 }
 
-void UDisplayClusterPreviewShareComponent::TickPullActor()
+void UDisplayClusterPreviewShareComponent::TickPullActor(const bool bUseSourceActorSettings)
 {
 	using namespace UE::DisplayCluster::PreviewShareComponent;
 
@@ -555,23 +554,21 @@ void UDisplayClusterPreviewShareComponent::TickPullActor()
 
 	// Pull preview textures
 	{
-		// reconfigure to rendering from external DCRA
+		// Reconfigure to rendering from external DCRA
 		ViewportManager->GetConfiguration().SetRootActor(EDisplayClusterRootActorType::Scene | EDisplayClusterRootActorType::Configuration, SourceRootActor);
 
-		// Part of setttings gets from owner (DCRA)
-		const FDisplayClusterViewport_PreviewSettings OwnerPreviewSettings = RootActor->GetPreviewSettings(true);
+		// Gets preview settings from various RootActors
+		ADisplayClusterRootActor* PreviewSettingsRootActor = bUseSourceActorSettings ? SourceRootActor : RootActor;
 
-		// Get preview settings from SourceRootActor properties
-		FDisplayClusterViewport_PreviewSettings NewPreviewSettings = SourceRootActor->GetPreviewSettings(true);
+		// Override preview settings:
+		FDisplayClusterViewport_PreviewSettings NewPreviewSettings = PreviewSettingsRootActor->GetPreviewSettings(true);
 
-		// Override some preview settings for the configurator and set them as the source of the preview settings
-		NewPreviewSettings.bPreviewEnable = OwnerPreviewSettings.bPreviewEnable;
-		NewPreviewSettings.bPreviewICVFXFrustums = OwnerPreviewSettings.bPreviewICVFXFrustums;
-
-		// Override preview PostProcess from component
-		NewPreviewSettings.bPreviewEnablePostProcess = bPreviewEnablePostProcess;
+		// When using the preview settings from the local RootActor, we also use the DisplayDevice from the same.
+		NewPreviewSettings.DisplayDeviceRootActorType = bUseSourceActorSettings ? EDisplayClusterRootActorType::Configuration : EDisplayClusterRootActorType::Preview;
 
 		ViewportManager->GetConfiguration().SetPreviewSettings(NewPreviewSettings);
+
+		// Sets configuration mode to use overridden values instead of values from RootActor.
 		RootActor->PreviewSetttingsSource = EDisplayClusterConfigurationRootActorPreviewSettingsSource::Configuration;
 	}
 
@@ -874,7 +871,7 @@ void UDisplayClusterPreviewShareComponent::TickComponent(float DeltaTime, ELevel
 	}
 	case EDisplayClusterPreviewShareMode::PullActor:
 	{
-		TickPullActor();
+		TickPullActor(false);
 		break;
 	}
 	case EDisplayClusterPreviewShareMode::Send:
diff --git a/Engine/Plugins/Runtime/nDisplay/Source/DisplayCluster/Private/Render/Viewport/DisplayClusterViewport.cpp b/Engine/Plugins/Runtime/nDisplay/Source/DisplayCluster/Private/Render/Viewport/DisplayClusterViewport.cpp
--- a/Engine/Plugins/Runtime/nDisplay/Source/DisplayCluster/Private/Render/Viewport/DisplayClusterViewport.cpp
+++ b/Engine/Plugins/Runtime/nDisplay/Source/DisplayCluster/Private/Render/Viewport/DisplayClusterViewport.cpp
@@ -117,7 +117,7 @@ FDisplayClusterViewportProxyData* FDisplayClusterViewport::CreateViewportProxyDa
 	OutViewportProxyData->OpenColorIO = OpenColorIO;
 
 	// Get Display Device proxy object
-	if (UDisplayClusterDisplayDeviceBaseComponent* DisplayDevice = GetDisplayDeviceComponent(EDisplayClusterRootActorType::Configuration))
+	if (UDisplayClusterDisplayDeviceBaseComponent* DisplayDevice = GetDisplayDeviceComponent(Configuration->GetPreviewSettings().DisplayDeviceRootActorType))
 	{
 		OutViewportProxyData->DisplayDeviceProxy = DisplayDevice->GetDisplayDeviceProxy(GetConfiguration());
 	}
@@ -422,7 +422,7 @@ void FDisplayClusterViewport::SetupSceneView(uint32 ContextNum, class UWorld* Wo
 	CameraMotionBlur.SetupSceneView(Contexts[ContextNum], InOutView);
 
 	// Handle DisplayDevice
-	if (UDisplayClusterDisplayDeviceBaseComponent* InDisplayDeviceComponent = GetDisplayDeviceComponent(EDisplayClusterRootActorType::Configuration))
+	if (UDisplayClusterDisplayDeviceBaseComponent* InDisplayDeviceComponent = GetDisplayDeviceComponent(Configuration->GetPreviewSettings().DisplayDeviceRootActorType))
 	{
 		InDisplayDeviceComponent->SetupSceneView(*ViewportPreview, ContextNum, InOutViewFamily, InOutView);
 	}
diff --git a/Engine/Plugins/Runtime/nDisplay/Source/DisplayCluster/Private/Render/Viewport/Preview/DisplayClusterViewportPreview.cpp b/Engine/Plugins/Runtime/nDisplay/Source/DisplayCluster/Private/Render/Viewport/Preview/DisplayClusterViewportPreview.cpp
--- a/Engine/Plugins/Runtime/nDisplay/Source/DisplayCluster/Private/Render/Viewport/Preview/DisplayClusterViewportPreview.cpp
+++ b/Engine/Plugins/Runtime/nDisplay/Source/DisplayCluster/Private/Render/Viewport/Preview/DisplayClusterViewportPreview.cpp
@@ -62,7 +62,7 @@ void FDisplayClusterViewportPreview::Update()
 
 	FDisplayClusterViewport* InViewport = GetViewportImpl();
 	UDisplayClusterCameraComponent* ViewPointComponent = InViewport ? InViewport->GetViewPointCameraComponent(EDisplayClusterRootActorType::Configuration) : nullptr;
-	UDisplayClusterDisplayDeviceBaseComponent* InDisplayDeviceComponent = InViewport ? InViewport->GetDisplayDeviceComponent(EDisplayClusterRootActorType::Configuration) : nullptr;
+	UDisplayClusterDisplayDeviceBaseComponent* InDisplayDeviceComponent = InViewport ? InViewport->GetDisplayDeviceComponent(Configuration->GetPreviewSettings().DisplayDeviceRootActorType) : nullptr;
 
 	// Update preview meshes only if DisplayDevice is used
 	PreviewMesh.Update(InViewport, InDisplayDeviceComponent, ViewPointComponent);
diff --git a/Engine/Plugins/Runtime/nDisplay/Source/DisplayCluster/Public/Components/DisplayClusterPreviewShareComponent.h b/Engine/Plugins/Runtime/nDisplay/Source/DisplayCluster/Public/Components/DisplayClusterPreviewShareComponent.h
--- a/Engine/Plugins/Runtime/nDisplay/Source/DisplayCluster/Public/Components/DisplayClusterPreviewShareComponent.h
+++ b/Engine/Plugins/Runtime/nDisplay/Source/DisplayCluster/Public/Components/DisplayClusterPreviewShareComponent.h
@@ -128,10 +128,6 @@ public:
 	UPROPERTY(EditAnywhere, Category = Sharing, meta = (EditCondition = "Mode == EDisplayClusterPreviewShareMode::PullActor"))
 	EDisplayClusterPreviewShareIcvfxSync IcvfxCamerasSyncType = EDisplayClusterPreviewShareIcvfxSync::PullActor;
 
-	/** Override the Enable PostProcess value. */
-	UPROPERTY(EditAnywhere, BlueprintReadWrite, Category = Sharing, meta = (DisplayName = "Enable Post Process"))
-	bool bPreviewEnablePostProcess = true;
-
 private:
 	// DC ViewportManager instance for this component
 	// Used to perform rendering with custom settings from an external DCRA
@@ -161,8 +157,11 @@ private:
 	/** Logic that should run every tick when in Receive mode */
 	void TickReceive();
 
-	/** Logic that should run every tick when in Pull Actor mode */
-	void TickPullActor();
+	/** Logic that should run every tick when in Pull Actor mode
+	* 
+	* @param bUseSourceActorSettings - Get preview settings from Source nDisplay Actor
+	*/
+	void TickPullActor(const bool bUseSourceActorSettings);
 
 	/** Syncs the Icvfx cameras from the given source nDisplay actor to the given destination nDisplay actor */
 	void SyncIcvxCamerasFromSourceActor(const ADisplayClusterRootActor* SrcRootActor, const ADisplayClusterRootActor* DstRootActor);
diff --git a/Engine/Plugins/Runtime/nDisplay/Source/DisplayCluster/Public/Render/Viewport/Containers/DisplayClusterViewport_Enums.h b/Engine/Plugins/Runtime/nDisplay/Source/DisplayCluster/Public/Render/Viewport/Containers/DisplayClusterViewport_Enums.h
--- a/Engine/Plugins/Runtime/nDisplay/Source/DisplayCluster/Public/Render/Viewport/Containers/DisplayClusterViewport_Enums.h
+++ b/Engine/Plugins/Runtime/nDisplay/Source/DisplayCluster/Public/Render/Viewport/Containers/DisplayClusterViewport_Enums.h
@@ -158,3 +158,26 @@ enum class EDisplayClusterViewportContextState : uint8
 	InvalidProjectionMatrix = 1 << 4,
 };
 ENUM_CLASS_FLAGS(EDisplayClusterViewportContextState);
+
+/**
+* Type of DCRA used by configuration
+*/
+enum class EDisplayClusterRootActorType : uint8
+{
+	// This DCRA will be used to render previews. The meshes and preview materials are created at runtime.
+	Preview = 1 << 0,
+
+	// A reference to DCRA in the scene, used as a source for math calculations and references.
+	// Locations in the scene and math data are taken from this DCRA.
+	Scene = 1 << 1,
+
+	// Reference to DCRA, used as a source of configuration data from DCRA and its components.
+	Configuration = 1 << 2,
+
+	// This value can only be used in very specific cases:
+	// For function GetRootActor() : Return any of the DCRAs that are not nullptr, in ascending order of type: Preview, Scene, Configuration.
+	// For function SetRootActor() : Sets all references to DRCA to the specified value.
+	Any = Preview | Scene | Configuration,
+};
+ENUM_CLASS_FLAGS(EDisplayClusterRootActorType);
+
diff --git a/Engine/Plugins/Runtime/nDisplay/Source/DisplayCluster/Public/Render/Viewport/Containers/DisplayClusterViewport_PreviewSettings.h b/Engine/Plugins/Runtime/nDisplay/Source/DisplayCluster/Public/Render/Viewport/Containers/DisplayClusterViewport_PreviewSettings.h
--- a/Engine/Plugins/Runtime/nDisplay/Source/DisplayCluster/Public/Render/Viewport/Containers/DisplayClusterViewport_PreviewSettings.h
+++ b/Engine/Plugins/Runtime/nDisplay/Source/DisplayCluster/Public/Render/Viewport/Containers/DisplayClusterViewport_PreviewSettings.h
@@ -3,6 +3,7 @@
 #pragma once
 
 #include "CoreMinimal.h"
+#include "Render/Viewport/Containers/DisplayClusterViewport_Enums.h"
 
 /**
 * Viewport preview-in-scene rendering settings.
@@ -56,4 +57,7 @@ struct FDisplayClusterViewport_PreviewSettings
 
 	// Max amount of Viewports Per Frame
 	int ViewportsPerFrame = 1;
+
+	// The DisplayDevice component will be obtained from the RootActor with the specified type
+	EDisplayClusterRootActorType DisplayDeviceRootActorType = EDisplayClusterRootActorType::Configuration;
 };
diff --git a/Engine/Plugins/Runtime/nDisplay/Source/DisplayCluster/Public/Render/Viewport/IDisplayClusterViewportConfiguration.h b/Engine/Plugins/Runtime/nDisplay/Source/DisplayCluster/Public/Render/Viewport/IDisplayClusterViewportConfiguration.h
--- a/Engine/Plugins/Runtime/nDisplay/Source/DisplayCluster/Public/Render/Viewport/IDisplayClusterViewportConfiguration.h
+++ b/Engine/Plugins/Runtime/nDisplay/Source/DisplayCluster/Public/Render/Viewport/IDisplayClusterViewportConfiguration.h
@@ -14,28 +14,6 @@ class IDisplayClusterViewportManager;
 struct FDisplayClusterConfigurationICVFX_StageSettings;
 struct FDisplayClusterConfigurationRenderFrame;
 
-/**
-* Type of DCRA used by configuration
-*/
-enum class EDisplayClusterRootActorType : uint8
-{
-	// This DCRA will be used to render previews. The meshes and preview materials are created at runtime.
-	Preview = 1 << 0,
-
-	// A reference to DCRA in the scene, used as a source for math calculations and references.
-	// Locations in the scene and math data are taken from this DCRA.
-	Scene = 1 << 1,
-
-	// Reference to DCRA, used as a source of configuration data from DCRA and its components.
-	Configuration = 1 << 2,
-
-	// This value can only be used in very specific cases:
-	// For function GetRootActor() : Return any of the DCRAs that are not nullptr, in ascending order of type: Preview, Scene, Configuration.
-	// For function SetRootActor() : Sets all references to DRCA to the specified value.
-	Any = Preview | Scene | Configuration,
-};
-ENUM_CLASS_FLAGS(EDisplayClusterRootActorType);
-
 /**
  * Viewport manager configuration.
  */
