SharedMemoryMedia: Fixed inverted alpha when viewing in player.

It will invert alpha by default but there is a new option in the Media Output to not invert.

Note: GPU time increased up to 100us when alpha inversion is selected on Quadro RTX A6000, for a 4K texture.

diff --git a/Engine/Plugins/Runtime/nDisplay/Source/SharedMemoryMedia/Private/SharedMemoryMediaCapture.cpp b/Engine/Plugins/Runtime/nDisplay/Source/SharedMemoryMedia/Private/SharedMemoryMediaCapture.cpp
--- a/Engine/Plugins/Runtime/nDisplay/Source/SharedMemoryMedia/Private/SharedMemoryMediaCapture.cpp
+++ b/Engine/Plugins/Runtime/nDisplay/Source/SharedMemoryMedia/Private/SharedMemoryMediaCapture.cpp
@@ -5,14 +5,24 @@
 
 #include "GenericPlatform/GenericPlatformMemory.h"
 #include "Internationalization/TextLocalizationResource.h"
+#include "MediaShaders.h"
 #include "RenderGraphUtils.h"
+#include "SceneView.h"
+#include "ScreenPass.h"
 
 #include "SharedMemoryMediaModule.h"
 #include "SharedMemoryMediaPlatform.h"
 #include "SharedMemoryMediaTypes.h"
 
+
+
 DECLARE_GPU_STAT(SharedMemory_Capture);
 
+BEGIN_SHADER_PARAMETER_STRUCT(FCopyToSharedGpuTexturePass, )
+	RDG_TEXTURE_ACCESS(SrcTexture, ERHIAccess::CopySrc)
+	RDG_TEXTURE_ACCESS(DstTexture, ERHIAccess::CopyDest)
+END_SHADER_PARAMETER_STRUCT()
+
 
 bool USharedMemoryMediaCapture::InitializeCapture()
 {
@@ -232,14 +242,89 @@ void USharedMemoryMediaCapture::OnCustomCapture_RenderingThread(
 		}
 	}
 
+	FRDGTextureRef SourceTexture = InSourceTexture;
+
+	// When enabled, add pass to invert alpha
+
+	const USharedMemoryMediaOutput* SharedMemoryMediaOutput = Cast<USharedMemoryMediaOutput>(MediaOutput);
+	check(SharedMemoryMediaOutput);
+
+	if (SharedMemoryMediaOutput->bInvertAlpha)
+	{
+		ETextureCreateFlags InvertedAlphaTextureFlags = ETextureCreateFlags::ResolveTargetable;
+
+		if (EnumHasAnyFlags(SourceTexture->Desc.Flags, TexCreate_SRGB))
+		{
+			InvertedAlphaTextureFlags |= TexCreate_SRGB;
+		}
+
+		const FRDGTextureDesc Desc = FRDGTextureDesc::Create2D(
+			FIntPoint(CopyInfo.Size.X, CopyInfo.Size.Y),
+			SourceTexture->Desc.Format,
+			FClearValueBinding::Black,
+			InvertedAlphaTextureFlags
+		);
+
+		FRDGTextureRef InvertedAlphaTexture = GraphBuilder.CreateTexture(Desc, TEXT("SharedMemoryMediaInvertedAlphaTexture"));
+		check(InvertedAlphaTexture);
+
+		AddInvertAlphaConversionPass(GraphBuilder, SourceTexture, InvertedAlphaTexture);
+
+		SourceTexture = InvertedAlphaTexture;
+	}
+
 	// Add the copy texture pass
-	AddCopyToSharedGpuTexturePass(GraphBuilder, InSourceTexture, GFrameCounterRenderThread % NUMBUFFERS);
+	AddCopyToSharedGpuTexturePass(GraphBuilder, SourceTexture, GFrameCounterRenderThread % NUMBUFFERS);
+}
+
+void USharedMemoryMediaCapture::AddInvertAlphaConversionPass(FRDGBuilder& GraphBuilder, const FRDGTextureRef& SourceTexture, FRDGTextureRef DestTexture)
+{
+	// Rectangle area to use from source
+	const FIntRect ViewRect(0, 0, SourceTexture->Desc.GetSize().X, SourceTexture->Desc.GetSize().Y);
+
+	FGlobalShaderMap* GlobalShaderMap = GetGlobalShaderMap(GMaxRHIFeatureLevel);
+	TShaderMapRef<FScreenPassVS> VertexShader(GlobalShaderMap);
+
+	// Configure source/output viewport to get the right UV scaling from source texture to output texture
+	FScreenPassTextureViewport InputViewport(SourceTexture, ViewRect);
+	FScreenPassTextureViewport OutputViewport(DestTexture);
+
+	// In cases where texture is converted from a format that doesn't have A channel, we want to force set it to 1.
+	EMediaCaptureConversionOperation MediaConversionOperation = EMediaCaptureConversionOperation::INVERT_ALPHA;
+	FModifyAlphaSwizzleRgbaPS::FPermutationDomain PermutationVector;
+	PermutationVector.Set<FModifyAlphaSwizzleRgbaPS::FConversionOp>(static_cast<int32>(MediaConversionOperation));
+
+	TShaderMapRef<FModifyAlphaSwizzleRgbaPS> PixelShader(GlobalShaderMap, PermutationVector);
+	FModifyAlphaSwizzleRgbaPS::FParameters* Parameters = PixelShader->AllocateAndSetParameters(GraphBuilder, SourceTexture, DestTexture);
+
+	// Dummy SceneView created to use built in Draw Screen/Texture Pass
+
+	FSceneViewFamily ViewFamily(FSceneViewFamily::ConstructionValues(nullptr, nullptr, FEngineShowFlags(ESFIM_Game))
+		.SetTime(FGameTime())
+	);
+
+	FSceneViewInitOptions ViewInitOptions;
+
+	ViewInitOptions.ViewFamily = &ViewFamily;
+	ViewInitOptions.SetViewRectangle(ViewRect);
+	ViewInitOptions.ViewOrigin = FVector::ZeroVector;
+	ViewInitOptions.ViewRotationMatrix = FMatrix::Identity;
+	ViewInitOptions.ProjectionMatrix = FMatrix::Identity;
+
+	FSceneView View = FSceneView(ViewInitOptions);
+
+	AddDrawScreenPass(
+		GraphBuilder, 
+		RDG_EVENT_NAME("SharedMemoryMediaOutputInvertAlpha"), 
+		View, 
+		OutputViewport, 
+		InputViewport, 
+		VertexShader, 
+		PixelShader, 
+		Parameters
+	);
 }
 
-BEGIN_SHADER_PARAMETER_STRUCT(FCopyToSharedGpuTexturePass, )
-	RDG_TEXTURE_ACCESS(SrcTexture, ERHIAccess::CopySrc)
-	RDG_TEXTURE_ACCESS(DstTexture, ERHIAccess::CopyDest)
-END_SHADER_PARAMETER_STRUCT()
 
 void USharedMemoryMediaCapture::AddCopyToSharedGpuTexturePass(FRDGBuilder& GraphBuilder, FRDGTextureRef InSourceTexture, uint32 SharedTextureIdx)
 {
diff --git a/Engine/Plugins/Runtime/nDisplay/Source/SharedMemoryMedia/Public/SharedMemoryMediaCapture.h b/Engine/Plugins/Runtime/nDisplay/Source/SharedMemoryMedia/Public/SharedMemoryMediaCapture.h
--- a/Engine/Plugins/Runtime/nDisplay/Source/SharedMemoryMedia/Public/SharedMemoryMediaCapture.h
+++ b/Engine/Plugins/Runtime/nDisplay/Source/SharedMemoryMedia/Public/SharedMemoryMediaCapture.h
@@ -67,6 +67,9 @@ protected:
 	/** Adds to the graph builder the copy from the input texture to the shared gpu texture that corresponds to the current frame number */
 	void AddCopyToSharedGpuTexturePass(FRDGBuilder& GraphBuilder, FRDGTextureRef InSourceTexture, uint32 SharedTextureIdx);
 
+	/** Adds a pass to invert the alpha of the texture */
+	void AddInvertAlphaConversionPass(FRDGBuilder& GraphBuilder, const FRDGTextureRef& SourceTexture, FRDGTextureRef DestTexture);
+
 private:
 
 	/** Number of shared buffers */
diff --git a/Engine/Plugins/Runtime/nDisplay/Source/SharedMemoryMedia/Public/SharedMemoryMediaOutput.h b/Engine/Plugins/Runtime/nDisplay/Source/SharedMemoryMedia/Public/SharedMemoryMediaOutput.h
--- a/Engine/Plugins/Runtime/nDisplay/Source/SharedMemoryMedia/Public/SharedMemoryMediaOutput.h
+++ b/Engine/Plugins/Runtime/nDisplay/Source/SharedMemoryMedia/Public/SharedMemoryMediaOutput.h
@@ -24,6 +24,10 @@ public:
 	UPROPERTY(EditAnywhere, BlueprintReadWrite, Category = "Output")
 	FString UniqueName = TEXT("UniqueName");
 
+	/** If checked, the alpha channel of the texture will be inverted */
+	UPROPERTY(EditAnywhere, AdvancedDisplay, BlueprintReadWrite, Category = "Output")
+	bool bInvertAlpha = true;
+
 public:
 	
 	//~ Begin UMediaOutput interface
diff --git a/Engine/Plugins/Runtime/nDisplay/Source/SharedMemoryMedia/SharedMemoryMedia.Build.cs b/Engine/Plugins/Runtime/nDisplay/Source/SharedMemoryMedia/SharedMemoryMedia.Build.cs
--- a/Engine/Plugins/Runtime/nDisplay/Source/SharedMemoryMedia/SharedMemoryMedia.Build.cs
+++ b/Engine/Plugins/Runtime/nDisplay/Source/SharedMemoryMedia/SharedMemoryMedia.Build.cs
@@ -18,6 +18,7 @@ namespace UnrealBuildTool.Rules
 					"MediaAssets",
 					"MediaIOCore",
 					"RenderCore",
+					"Renderer",
 					"RHI",
 				});
 
