Fix CIS Issue 400137

diff --git a/Engine/Plugins/Runtime/nDisplay/Source/DisplayCluster/Private/Game/EngineClasses/Scene/DisplayClusterRootActorEditor.cpp b/Engine/Plugins/Runtime/nDisplay/Source/DisplayCluster/Private/Game/EngineClasses/Scene/DisplayClusterRootActorEditor.cpp
--- a/Engine/Plugins/Runtime/nDisplay/Source/DisplayCluster/Private/Game/EngineClasses/Scene/DisplayClusterRootActorEditor.cpp
+++ b/Engine/Plugins/Runtime/nDisplay/Source/DisplayCluster/Private/Game/EngineClasses/Scene/DisplayClusterRootActorEditor.cpp
@@ -139,7 +139,7 @@ namespace UE::DisplayCluster::RootActor_Editor
 		// Look for elements that should be added.
 		for (FScriptMapHelper::FIterator DefaultIt = MapDefaultsHelper.CreateIterator(); DefaultIt; ++DefaultIt)
 		{
-			uint8* DefaultPairPtr = MapDefaultsHelper.GetPairPtr(*DefaultIt);
+			uint8* DefaultPairPtr = MapDefaultsHelper.GetPairPtr(DefaultIt);
 			check(DefaultPairPtr);
 
 			FString* Key = MapProperty->KeyProp->ContainerPtrToValuePtr<FString>(DefaultPairPtr);
@@ -197,7 +197,7 @@ namespace UE::DisplayCluster::RootActor_Editor
 		// Look for elements that should be removed.
 		for (FScriptMapHelper::FIterator InstanceIt = MapInstanceHelper.CreateIterator(); InstanceIt; ++InstanceIt)
 		{
-			uint8* InstancePairPtr = MapInstanceHelper.GetPairPtr(*InstanceIt);
+			uint8* InstancePairPtr = MapInstanceHelper.GetPairPtr(InstanceIt);
 			check(InstancePairPtr);
 
 			FString* Key = MapProperty->KeyProp->ContainerPtrToValuePtr<FString>(InstancePairPtr);
@@ -218,7 +218,7 @@ namespace UE::DisplayCluster::RootActor_Editor
 					}
 				}
 
-				MapInstanceHelper.RemoveAt(*InstanceIt);
+				MapInstanceHelper.RemoveAt(InstanceIt.GetInternalIndex());
 				bHasChanged = true;
 			}
 		}
