nDisplay: Fixed display name of the in frustum fit view origin component.

#rnx

diff --git a/Engine/Plugins/Runtime/nDisplay/Source/DisplayClusterWarp/Public/Components/DisplayClusterInFrustumFitCameraComponent.h b/Engine/Plugins/Runtime/nDisplay/Source/DisplayClusterWarp/Public/Components/DisplayClusterInFrustumFitCameraComponent.h
--- a/Engine/Plugins/Runtime/nDisplay/Source/DisplayClusterWarp/Public/Components/DisplayClusterInFrustumFitCameraComponent.h
+++ b/Engine/Plugins/Runtime/nDisplay/Source/DisplayClusterWarp/Public/Components/DisplayClusterInFrustumFitCameraComponent.h
@@ -16,7 +16,7 @@ class IDisplayClusterWarpPolicy;
  * 3D point in space used to project the camera view onto a group of nDisplay viewports.
  * Support projection policies: mpcdi/pfm 2d/a3d, mesh.
  */
-UCLASS(ClassGroup = (DisplayCluster), meta = (BlueprintSpawnableComponent, DisplayName = "nDisplay InFrustumFit View Origin"))
+UCLASS(ClassGroup = (DisplayCluster), meta = (BlueprintSpawnableComponent, DisplayName = "NDisplay In Frustum Fit View Origin"))
 class DISPLAYCLUSTERWARP_API UDisplayClusterInFrustumFitCameraComponent : public UDisplayClusterCameraComponent
 {
 	GENERATED_BODY()
