nDisplay: Added log with chosen DCRA from a given asset path.

In this code path only the found actors were listed but not the chosen one. This log should help when debugging.

#rnx

diff --git a/Engine/Plugins/Runtime/nDisplay/Source/DisplayCluster/Private/Game/DisplayClusterGameManager.cpp b/Engine/Plugins/Runtime/nDisplay/Source/DisplayCluster/Private/Game/DisplayClusterGameManager.cpp
--- a/Engine/Plugins/Runtime/nDisplay/Source/DisplayCluster/Private/Game/DisplayClusterGameManager.cpp
+++ b/Engine/Plugins/Runtime/nDisplay/Source/DisplayCluster/Private/Game/DisplayClusterGameManager.cpp
@@ -89,6 +89,11 @@ bool FDisplayClusterGameManager::StartScene(UWorld* InWorld)
 		// If a corresponding DCRA instance was found, overwrite its settings
 		if (RootActor)
 		{
+			UE_LOG(LogDisplayClusterGame, Log, TEXT("Chose DCRA '%s' corresponding to asset '%s'"), 
+				*RootActor->GetName(),
+				*ConfigData->Info.AssetPath
+			);
+
 			RootActor->OverrideFromConfig(ConfigData);
 		}
 		// If no proper DCRA found,
