nDisplay: Added node name to barrier trace names.

Should make it easier to debug. W/o this, one has to look up which node has the ip address that is put in the session name, reducing debugging efficiency, particularly with large clusters.

diff --git a/Engine/Plugins/Runtime/nDisplay/Source/DisplayCluster/Private/Network/Service/ClusterEventsBinary/DisplayClusterClusterEventsBinaryService.cpp b/Engine/Plugins/Runtime/nDisplay/Source/DisplayCluster/Private/Network/Service/ClusterEventsBinary/DisplayClusterClusterEventsBinaryService.cpp
--- a/Engine/Plugins/Runtime/nDisplay/Source/DisplayCluster/Private/Network/Service/ClusterEventsBinary/DisplayClusterClusterEventsBinaryService.cpp
+++ b/Engine/Plugins/Runtime/nDisplay/Source/DisplayCluster/Private/Network/Service/ClusterEventsBinary/DisplayClusterClusterEventsBinaryService.cpp
@@ -37,7 +37,13 @@ FString FDisplayClusterClusterEventsBinaryService::GetProtocolName() const
 
 TSharedPtr<IDisplayClusterSession> FDisplayClusterClusterEventsBinaryService::CreateSession(FDisplayClusterSessionInfo& SessionInfo)
 {
-	SessionInfo.SessionName = FString::Printf(TEXT("%s_session_%lu_%s"), *GetName(), SessionInfo.SessionId, *SessionInfo.Endpoint.ToString());
+	SessionInfo.SessionName = FString::Printf(TEXT("%s_%lu_%s_%s"),
+		*GetName(),
+		SessionInfo.SessionId,
+		*SessionInfo.Endpoint.ToString(),
+		*SessionInfo.NodeId.Get(TEXT("(na)"))
+	);
+
 	return MakeShared<FDisplayClusterSession<FDisplayClusterPacketBinary, false>>(SessionInfo, *this, *this, FDisplayClusterService::GetThreadPriority());
 }
 
diff --git a/Engine/Plugins/Runtime/nDisplay/Source/DisplayCluster/Private/Network/Service/ClusterEventsJson/DisplayClusterClusterEventsJsonService.cpp b/Engine/Plugins/Runtime/nDisplay/Source/DisplayCluster/Private/Network/Service/ClusterEventsJson/DisplayClusterClusterEventsJsonService.cpp
--- a/Engine/Plugins/Runtime/nDisplay/Source/DisplayCluster/Private/Network/Service/ClusterEventsJson/DisplayClusterClusterEventsJsonService.cpp
+++ b/Engine/Plugins/Runtime/nDisplay/Source/DisplayCluster/Private/Network/Service/ClusterEventsJson/DisplayClusterClusterEventsJsonService.cpp
@@ -40,7 +40,13 @@ FString FDisplayClusterClusterEventsJsonService::GetProtocolName() const
 
 TSharedPtr<IDisplayClusterSession> FDisplayClusterClusterEventsJsonService::CreateSession(FDisplayClusterSessionInfo& SessionInfo)
 {
-	SessionInfo.SessionName = FString::Printf(TEXT("%s_session_%lu_%s"), *GetName(), SessionInfo.SessionId, *SessionInfo.Endpoint.ToString());
+	SessionInfo.SessionName = FString::Printf(TEXT("%s_%lu_%s_%s"),
+		*GetName(),
+		SessionInfo.SessionId,
+		*SessionInfo.Endpoint.ToString(),
+		*SessionInfo.NodeId.Get(TEXT("(na)"))
+	);
+
 	return MakeShared<FDisplayClusterSession<FDisplayClusterPacketJson, false>>(SessionInfo, *this, *this, FDisplayClusterService::GetThreadPriority());
 }
 
diff --git a/Engine/Plugins/Runtime/nDisplay/Source/DisplayCluster/Private/Network/Service/ClusterSync/DisplayClusterClusterSyncService.cpp b/Engine/Plugins/Runtime/nDisplay/Source/DisplayCluster/Private/Network/Service/ClusterSync/DisplayClusterClusterSyncService.cpp
--- a/Engine/Plugins/Runtime/nDisplay/Source/DisplayCluster/Private/Network/Service/ClusterSync/DisplayClusterClusterSyncService.cpp
+++ b/Engine/Plugins/Runtime/nDisplay/Source/DisplayCluster/Private/Network/Service/ClusterSync/DisplayClusterClusterSyncService.cpp
@@ -125,7 +125,13 @@ void FDisplayClusterClusterSyncService::KillSession(const FString& NodeId)
 
 TSharedPtr<IDisplayClusterSession> FDisplayClusterClusterSyncService::CreateSession(FDisplayClusterSessionInfo& SessionInfo)
 {
-	SessionInfo.SessionName = FString::Printf(TEXT("%s_session_%lu_%s"), *GetName(), SessionInfo.SessionId, *SessionInfo.Endpoint.ToString());
+	SessionInfo.SessionName = FString::Printf(TEXT("%s_%lu_%s_%s"),
+		*GetName(),
+		SessionInfo.SessionId,
+		*SessionInfo.Endpoint.ToString(),
+		*SessionInfo.NodeId.Get(TEXT("(na)"))
+	);
+
 	return MakeShared<FDisplayClusterSession<FDisplayClusterPacketInternal, true>>(SessionInfo, *this, *this, FDisplayClusterService::GetThreadPriority());
 }
 
diff --git a/Engine/Plugins/Runtime/nDisplay/Source/DisplayCluster/Private/Network/Service/GenericBarrier/DisplayClusterGenericBarrierService.cpp b/Engine/Plugins/Runtime/nDisplay/Source/DisplayCluster/Private/Network/Service/GenericBarrier/DisplayClusterGenericBarrierService.cpp
--- a/Engine/Plugins/Runtime/nDisplay/Source/DisplayCluster/Private/Network/Service/GenericBarrier/DisplayClusterGenericBarrierService.cpp
+++ b/Engine/Plugins/Runtime/nDisplay/Source/DisplayCluster/Private/Network/Service/GenericBarrier/DisplayClusterGenericBarrierService.cpp
@@ -65,7 +65,13 @@ TSharedPtr<IDisplayClusterBarrier, ESPMode::ThreadSafe> FDisplayClusterGenericBa
 
 TSharedPtr<IDisplayClusterSession> FDisplayClusterGenericBarrierService::CreateSession(FDisplayClusterSessionInfo& SessionInfo)
 {
-	SessionInfo.SessionName = FString::Printf(TEXT("%s_session_%lu_%s"), *GetName(), SessionInfo.SessionId, *SessionInfo.Endpoint.ToString());
+	SessionInfo.SessionName = FString::Printf(TEXT("%s_%lu_%s_%s"),
+		*GetName(),
+		SessionInfo.SessionId,
+		*SessionInfo.Endpoint.ToString(),
+		*SessionInfo.NodeId.Get(TEXT("(na)"))
+	);
+
 	return MakeShared<FDisplayClusterSession<FDisplayClusterPacketInternal, true>>(SessionInfo, *this, *this, FDisplayClusterService::GetThreadPriority());
 }
 
diff --git a/Engine/Plugins/Runtime/nDisplay/Source/DisplayCluster/Private/Network/Service/RenderSync/DisplayClusterRenderSyncService.cpp b/Engine/Plugins/Runtime/nDisplay/Source/DisplayCluster/Private/Network/Service/RenderSync/DisplayClusterRenderSyncService.cpp
--- a/Engine/Plugins/Runtime/nDisplay/Source/DisplayCluster/Private/Network/Service/RenderSync/DisplayClusterRenderSyncService.cpp
+++ b/Engine/Plugins/Runtime/nDisplay/Source/DisplayCluster/Private/Network/Service/RenderSync/DisplayClusterRenderSyncService.cpp
@@ -80,7 +80,13 @@ void FDisplayClusterRenderSyncService::KillSession(const FString& NodeId)
 
 TSharedPtr<IDisplayClusterSession> FDisplayClusterRenderSyncService::CreateSession(FDisplayClusterSessionInfo& SessionInfo)
 {
-	SessionInfo.SessionName = FString::Printf(TEXT("%s_session_%lu_%s"), *GetName(), SessionInfo.SessionId, *SessionInfo.Endpoint.ToString());
+	SessionInfo.SessionName = FString::Printf(TEXT("%s_%lu_%s_%s"),
+		*GetName(),
+		SessionInfo.SessionId,
+		*SessionInfo.Endpoint.ToString(),
+		*SessionInfo.NodeId.Get(TEXT("(na)"))
+	);
+
 	return MakeShared<FDisplayClusterSession<FDisplayClusterPacketInternal, true>>(SessionInfo, *this, *this, FDisplayClusterService::GetThreadPriority());
 }
 
