- Fixing rivermax PTP media sync margin not working. Still missing self repair but at least margin should be detected now.

diff --git a/Engine/Plugins/Runtime/nDisplay/Source/DisplayClusterMedia/Private/Synchronization/SyncPolicies/DisplayClusterMediaOutputSynchronizationPolicyThresholdBase.cpp b/Engine/Plugins/Runtime/nDisplay/Source/DisplayClusterMedia/Private/Synchronization/SyncPolicies/DisplayClusterMediaOutputSynchronizationPolicyThresholdBase.cpp
--- a/Engine/Plugins/Runtime/nDisplay/Source/DisplayClusterMedia/Private/Synchronization/SyncPolicies/DisplayClusterMediaOutputSynchronizationPolicyThresholdBase.cpp
+++ b/Engine/Plugins/Runtime/nDisplay/Source/DisplayClusterMedia/Private/Synchronization/SyncPolicies/DisplayClusterMediaOutputSynchronizationPolicyThresholdBase.cpp
@@ -23,7 +23,7 @@ void FDisplayClusterMediaOutputSynchronizationPolicyThresholdBaseHandler::Synchr
 	if (TimeLeftSeconds < MarginSeconds)
 	{
 		// Sleep for a bit longer to skip the alignment timepoint
-		const float SleepTime = (MarginSeconds - TimeLeftSeconds) * 1.01f;
+		const float SleepTime = TimeLeftSeconds * 1.01f;
 
 		UE_LOG(LogDisplayClusterMediaSync, VeryVerbose, TEXT("'%s': TimeLeft(%lf) < Margin(%lf) --> Sleeping for %lf..."),
 			*GetMediaDeviceId(), TimeLeftSeconds, MarginSeconds, SleepTime);
diff --git a/Engine/Plugins/VirtualProduction/Rivermax/RivermaxSync/Source/RivermaxSync/Private/MediaOutputSynchronizationPolicyRivermax.cpp b/Engine/Plugins/VirtualProduction/Rivermax/RivermaxSync/Source/RivermaxSync/Private/MediaOutputSynchronizationPolicyRivermax.cpp
--- a/Engine/Plugins/VirtualProduction/Rivermax/RivermaxSync/Source/RivermaxSync/Private/MediaOutputSynchronizationPolicyRivermax.cpp
+++ b/Engine/Plugins/VirtualProduction/Rivermax/RivermaxSync/Source/RivermaxSync/Private/MediaOutputSynchronizationPolicyRivermax.cpp
@@ -70,7 +70,7 @@ double FMediaOutputSynchronizationPolicyRivermaxHandler::GetTimeBeforeNextSyncPo
 				const uint64 TimeLeftNanosec = NextAlignmentTimeNanosec - CurrentTimeNanosec;
 
 				// Return remaining time in seconds
-				return double(TimeLeftNanosec * 1E9);
+				return double(TimeLeftNanosec * 1E-9);
 			}
 		}
 	}
