nDisplay: Detect when the exported config in the nDisplay asset is out of date.

This can happen in child blueprints when their parent changes. It can also happen when there are updates to the export logic.

diff --git a/Engine/Plugins/Runtime/nDisplay/Source/DisplayCluster/Private/Blueprints/DisplayClusterBlueprint.cpp b/Engine/Plugins/Runtime/nDisplay/Source/DisplayCluster/Private/Blueprints/DisplayClusterBlueprint.cpp
--- a/Engine/Plugins/Runtime/nDisplay/Source/DisplayCluster/Private/Blueprints/DisplayClusterBlueprint.cpp
+++ b/Engine/Plugins/Runtime/nDisplay/Source/DisplayCluster/Private/Blueprints/DisplayClusterBlueprint.cpp
@@ -123,6 +123,77 @@ void UDisplayClusterBlueprint::PreSave(FObjectPreSaveContext SaveContext)
 
 	UpdateConfigExportProperty();
 	DisplayClusterBlueprint::SendAnalytics(TEXT("Usage.nDisplay.ConfigSaved"), ConfigData);
+
+#if WITH_EDITOR
+
+	// Child blueprints need to re-generate their config export property as well
+	// Note: Using GetDerivedClasses will only get loaded classes, which is the normal case,
+	// and the rest will be caught when they get loaded as an out of date exported config
+	// will be detected.
+	if (GIsEditor)
+	{
+		TArray<UClass*> ChildClasses;
+		GetDerivedClasses(GeneratedClass, ChildClasses);
+
+		for (UClass* ChildClass : ChildClasses)
+		{
+			// CLASS_NewerVersionExists suggests there is a newer class that will update the asset, so we skip it.
+			if (ChildClass->HasAnyClassFlags(CLASS_Abstract | CLASS_Deprecated | CLASS_NewerVersionExists))
+			{
+				continue;
+			}
+
+			UDisplayClusterBlueprint* ChildDCBP = Cast<UDisplayClusterBlueprint>(ChildClass->ClassGeneratedBy);
+			check(ChildDCBP); // We already know it is a derived class
+
+			// Only mark as dirty if the config needs updating, to avoid unnecessary re-saves.
+
+			const FString OriginalChildConfigExport = ChildDCBP->ConfigExport;
+			ChildDCBP->UpdateConfigExportProperty();
+
+			if (!ChildDCBP->ConfigExport.Equals(OriginalChildConfigExport))
+			{
+				if (ChildDCBP->MarkPackageDirty())
+				{
+					UE_LOG(LogDisplayClusterBlueprint, Warning,
+						TEXT("ConfigExport of the child nDisplay blueprint actor '%s' is not up to date in the asset, so the package was marked as dirty and should be re-saved."),
+						*ChildDCBP->GetOutermost()->GetName() // If MarkPackageDirty succeeded then GetOutermost() must exist.
+					);
+				}
+			}
+		}
+	}
+
+#endif // WITH_EDITOR
+}
+
+void UDisplayClusterBlueprint::PostLoad()
+{
+	Super::PostLoad();
+
+#if WITH_EDITOR
+
+	// If the exported config is out of date, mark the package as dirty for the user
+	// to re-save. This may happen, for example, when the parent blueprint is updated, 
+	// or when the config export logic has changed.
+
+	if (GIsEditor)
+	{
+		const FString LoadedConfigExport = ConfigExport;
+		UpdateConfigExportProperty();
+
+		if (!ConfigExport.Equals(LoadedConfigExport))
+		{
+			if (MarkPackageDirty())
+			{			
+				UE_LOG(LogDisplayClusterBlueprint, Warning,
+					TEXT("ConfigExport of the nDisplay actor '%s' was not up to date in the asset, so the package was marked as dirty and should be re-saved."),
+					*GetOutermost()->GetName() // If MarkPackageDirty succeeded then GetOutermost() must exist.
+				);
+			}
+		}
+	}
+#endif // WITH_EDITOR
 }
 
 UDisplayClusterBlueprintGeneratedClass* UDisplayClusterBlueprint::GetGeneratedClass() const
diff --git a/Engine/Plugins/Runtime/nDisplay/Source/DisplayCluster/Public/Blueprints/DisplayClusterBlueprint.h b/Engine/Plugins/Runtime/nDisplay/Source/DisplayCluster/Public/Blueprints/DisplayClusterBlueprint.h
--- a/Engine/Plugins/Runtime/nDisplay/Source/DisplayCluster/Public/Blueprints/DisplayClusterBlueprint.h
+++ b/Engine/Plugins/Runtime/nDisplay/Source/DisplayCluster/Public/Blueprints/DisplayClusterBlueprint.h
@@ -32,6 +32,7 @@ public:
 	
 	//~ Begin UObject Interface
 	virtual void PreSave(FObjectPreSaveContext SaveContext) override;
+	virtual void PostLoad() override;
 	//~ End UObject Interface
 
 	class UDisplayClusterBlueprintGeneratedClass* GetGeneratedClass() const;
@@ -54,14 +55,15 @@ public:
 	/** Configure the config with needed blueprint data for an export. */
 	void PrepareConfigForExport();
 	
+	//** Updates the ConfigExport property. Called when saving the asset.
+	void UpdateConfigExportProperty();
+
 public:
 	// Holds the last saved config export. In the AssetRegistry to allow parsing without loading.
 	UPROPERTY(AssetRegistrySearchable)
 	FString ConfigExport;
 
 private:
-	//** Updates the ConfigExport property. Called when saving the asset.
-	void UpdateConfigExportProperty();
 
 	/** Returns normal component name. We get component name from a SCS component template which has _GEN_VARIABLE postfix. */
 	FString GetObjectNameFromSCSNode(const UObject* const Object) const;
