SharedMemoryMedia: Added single GPU option so that it can be faster in those cases.

Also improved some of the logging information.

diff --git a/Engine/Plugins/Runtime/nDisplay/Source/SharedMemoryMedia/Private/Platforms/Windows/SharedMemoryMediaPlatformWindowsD3D12.cpp b/Engine/Plugins/Runtime/nDisplay/Source/SharedMemoryMedia/Private/Platforms/Windows/SharedMemoryMediaPlatformWindowsD3D12.cpp
--- a/Engine/Plugins/Runtime/nDisplay/Source/SharedMemoryMedia/Private/Platforms/Windows/SharedMemoryMediaPlatformWindowsD3D12.cpp
+++ b/Engine/Plugins/Runtime/nDisplay/Source/SharedMemoryMedia/Private/Platforms/Windows/SharedMemoryMediaPlatformWindowsD3D12.cpp
@@ -87,11 +87,11 @@ FSharedMemoryMediaPlatformWindowsD3D12::~FSharedMemoryMediaPlatformWindowsD3D12(
 	// Close shared handles
 	for (int32 BufferIdx = 0; BufferIdx < UE::SharedMemoryMedia::SenderNumBuffers; ++BufferIdx)
 	{
-		ReleaseSharedCrossGpuTexture(BufferIdx);
+		ReleaseSharedTexture(BufferIdx);
 	}
 }
 
-FTextureRHIRef FSharedMemoryMediaPlatformWindowsD3D12::CreateSharedCrossGpuTexture(EPixelFormat Format, bool bSrgb, int32 Width, int32 Height, const FGuid& Guid, uint32 BufferIdx)
+FTextureRHIRef FSharedMemoryMediaPlatformWindowsD3D12::CreateSharedTexture(EPixelFormat Format, bool bSrgb, int32 Width, int32 Height, const FGuid& Guid, uint32 BufferIdx, bool bCrossGpu)
 {
 	using namespace UE::FSharedMemoryMediaPlatformWindowsD3D12;
 
@@ -104,6 +104,15 @@ FTextureRHIRef FSharedMemoryMediaPlatformWindowsD3D12::CreateSharedCrossGpuTextu
 
 	const DXGI_FORMAT DxgiFormat = FindSharedResourceDXGIFormat(DXGI_FORMAT(GPixelFormats[Format].PlatformFormat), bSrgb);
 
+	D3D12_RESOURCE_FLAGS Flags = D3D12_RESOURCE_FLAG_ALLOW_SIMULTANEOUS_ACCESS;
+
+	if (bCrossGpu)
+	{
+		Flags |= D3D12_RESOURCE_FLAG_ALLOW_CROSS_ADAPTER;
+	}
+
+	const D3D12_TEXTURE_LAYOUT Layout = bCrossGpu ? D3D12_TEXTURE_LAYOUT_ROW_MAJOR : D3D12_TEXTURE_LAYOUT_UNKNOWN;
+
 	D3D12_RESOURCE_DESC SharedCrossGpuTextureDesc = CD3DX12_RESOURCE_DESC::Tex2D(
 		DxgiFormat,
 		Width,
@@ -112,15 +121,21 @@ FTextureRHIRef FSharedMemoryMediaPlatformWindowsD3D12::CreateSharedCrossGpuTextu
 		1, // mipLevels
 		1, // sampleCount
 		0, // sampleQuality
-		D3D12_RESOURCE_FLAG_ALLOW_CROSS_ADAPTER | D3D12_RESOURCE_FLAG_ALLOW_SIMULTANEOUS_ACCESS, // flags
-		D3D12_TEXTURE_LAYOUT_ROW_MAJOR // layout
+		Flags, // flags
+		Layout // layout
 	);
 
 	CD3DX12_HEAP_PROPERTIES HeapProperties(D3D12_HEAP_TYPE_DEFAULT);
+	D3D12_HEAP_FLAGS HeapFlags = D3D12_HEAP_FLAG_SHARED;
+
+	if (bCrossGpu)
+	{
+		HeapFlags |= D3D12_HEAP_FLAG_SHARED_CROSS_ADAPTER;
+	}
 
 	HRESULT HResult = D3D12Device->CreateCommittedResource(
 		&HeapProperties,
-		D3D12_HEAP_FLAG_SHARED | D3D12_HEAP_FLAG_SHARED_CROSS_ADAPTER,
+		HeapFlags,
 		&SharedCrossGpuTextureDesc,
 		D3D12_RESOURCE_STATE_COPY_DEST,
 		nullptr, // pOptimizedClearValue
@@ -131,7 +146,7 @@ FTextureRHIRef FSharedMemoryMediaPlatformWindowsD3D12::CreateSharedCrossGpuTextu
 	{
 		UE_LOG(LogSharedMemoryMedia, Error, TEXT(
 			"D3D12Device->CreateCommittedResource failed when creating a cross GPU texture:\n"
-			"0x%lX - %s. Texture was %dx%d, EPixelFormat %d and DXGI type %d"),
+			"0x%lX - %s. Texture was %dx%d, EPixelFormat %d and DXGI type %d. Use -d3ddebug for more information. "),
 			HResult, *GetD3D12ComErrorDescription(HResult), Width, Height, Format, DxgiFormat);
 
 		return nullptr;
@@ -162,7 +177,7 @@ FTextureRHIRef FSharedMemoryMediaPlatformWindowsD3D12::CreateSharedCrossGpuTextu
 	);
 }
 
-void FSharedMemoryMediaPlatformWindowsD3D12::ReleaseSharedCrossGpuTexture(uint32 BufferIdx)
+void FSharedMemoryMediaPlatformWindowsD3D12::ReleaseSharedTexture(uint32 BufferIdx)
 {
 	check(BufferIdx < UE::SharedMemoryMedia::SenderNumBuffers);
 	if (SharedHandle[BufferIdx] != INVALID_HANDLE_VALUE)
@@ -180,7 +195,7 @@ void FSharedMemoryMediaPlatformWindowsD3D12::ReleaseSharedCrossGpuTexture(uint32
 	}
 }
 
-FTextureRHIRef FSharedMemoryMediaPlatformWindowsD3D12::OpenSharedCrossGpuTextureByGuid(const FGuid& Guid, FSharedMemoryMediaTextureDescription& OutTextureDescription)
+FTextureRHIRef FSharedMemoryMediaPlatformWindowsD3D12::OpenSharedTextureByGuid(const FGuid& Guid, FSharedMemoryMediaTextureDescription& OutTextureDescription)
 {
 	using namespace UE::FSharedMemoryMediaPlatformWindowsD3D12;
 
diff --git a/Engine/Plugins/Runtime/nDisplay/Source/SharedMemoryMedia/Private/Platforms/Windows/SharedMemoryMediaPlatformWindowsD3D12.h b/Engine/Plugins/Runtime/nDisplay/Source/SharedMemoryMedia/Private/Platforms/Windows/SharedMemoryMediaPlatformWindowsD3D12.h
--- a/Engine/Plugins/Runtime/nDisplay/Source/SharedMemoryMedia/Private/Platforms/Windows/SharedMemoryMediaPlatformWindowsD3D12.h
+++ b/Engine/Plugins/Runtime/nDisplay/Source/SharedMemoryMedia/Private/Platforms/Windows/SharedMemoryMediaPlatformWindowsD3D12.h
@@ -36,9 +36,9 @@ private:
 public:
 
 	//~ Begin FSharedMemoryMediaPlatform interface
-	virtual FTextureRHIRef CreateSharedCrossGpuTexture(EPixelFormat Format, bool bSrgb, int32 Width, int32 Height, const FGuid& Guid, uint32 BufferIdx) override;
-	virtual void ReleaseSharedCrossGpuTexture(uint32 BufferIdx) override;
-	virtual FTextureRHIRef OpenSharedCrossGpuTextureByGuid(const FGuid& Guid, FSharedMemoryMediaTextureDescription& OutTextureDescription) override;
+	virtual FTextureRHIRef CreateSharedTexture(EPixelFormat Format, bool bSrgb, int32 Width, int32 Height, const FGuid& Guid, uint32 BufferIdx, bool bCrossGpu) override;
+	virtual void ReleaseSharedTexture(uint32 BufferIdx) override;
+	virtual FTextureRHIRef OpenSharedTextureByGuid(const FGuid& Guid, FSharedMemoryMediaTextureDescription& OutTextureDescription) override;
 	//~ End FSharedMemoryMediaPlatform interface
 
 	FSharedMemoryMediaPlatformWindowsD3D12();
diff --git a/Engine/Plugins/Runtime/nDisplay/Source/SharedMemoryMedia/Private/SharedMemoryMediaCapture.cpp b/Engine/Plugins/Runtime/nDisplay/Source/SharedMemoryMedia/Private/SharedMemoryMediaCapture.cpp
--- a/Engine/Plugins/Runtime/nDisplay/Source/SharedMemoryMedia/Private/SharedMemoryMediaCapture.cpp
+++ b/Engine/Plugins/Runtime/nDisplay/Source/SharedMemoryMedia/Private/SharedMemoryMediaCapture.cpp
@@ -32,7 +32,7 @@ bool USharedMemoryMediaCapture::InitializeCapture()
 
 	if (!SharedMemoryMediaOutput)
 	{
-		UE_LOG(LogSharedMemoryMedia, Error, TEXT("Invalid MediaOutput, cannot InitializeCapture"));
+		UE_LOG(LogSharedMemoryMedia, Error, TEXT("Invalid MediaOutput for '%s', cannot InitializeCapture"), *GetName());
 		return false;
 	}
 
@@ -101,7 +101,9 @@ bool USharedMemoryMediaCapture::InitializeCapture()
 				FSharedMemoryMediaFrameMetadata* Data = static_cast<FSharedMemoryMediaFrameMetadata*>(SharedMemoryRegion->GetAddress());
 				Data->Receiver.FrameNumberAcked = ~0;
 
-				UE_LOG(LogSharedMemoryMedia, Verbose, TEXT("Created SharedMemoryRegion[%d] = %s"), BufferIdx, *SharedMemoryRegionName);
+				UE_LOG(LogSharedMemoryMedia, Verbose, TEXT("Created SharedMemoryRegion[%d] = %s for UniqueName '%s'"), 
+					BufferIdx, *SharedMemoryRegionName, *SharedMemoryMediaOutput->UniqueName
+				);
 			}
 		}
 
@@ -171,7 +173,7 @@ void USharedMemoryMediaCapture::StopCaptureImpl(bool bAllowPendingFrameToBeProce
 
 		if (PlatformData.IsValid())
 		{
-			PlatformData->ReleaseSharedCrossGpuTexture(BufferIdx);
+			PlatformData->ReleaseSharedTexture(BufferIdx);
 		}
 
 		SharedCrossGpuTextureGuids[BufferIdx] = FGuid();
@@ -210,6 +212,10 @@ void USharedMemoryMediaCapture::OnCustomCapture_RenderingThread(
 	RDG_GPU_STAT_SCOPE(GraphBuilder, SharedMemory_Capture)
 	TRACE_CPUPROFILER_EVENT_SCOPE(USharedMemoryMediaCapture::OnCustomCapture_RenderingThread);
 
+	// We'll be needing the output options.
+	const USharedMemoryMediaOutput* SharedMemoryMediaOutput = Cast<USharedMemoryMediaOutput>(MediaOutput);
+	check(SharedMemoryMediaOutput);
+
 	// Initialize shared gpu textures if needed.
 
 	check(PlatformData.IsValid());
@@ -220,25 +226,31 @@ void USharedMemoryMediaCapture::OnCustomCapture_RenderingThread(
 		{
 			const FGuid Guid = FGuid::NewGuid();
 						
-			SharedCrossGpuTextures[Idx] = PlatformData->CreateSharedCrossGpuTexture(
+			SharedCrossGpuTextures[Idx] = PlatformData->CreateSharedTexture(
 				InSourceTexture->Desc.Format,
 				EnumHasAnyFlags(InSourceTexture->Desc.Flags, TexCreate_SRGB),
 				CopyInfo.Size.X, 
 				CopyInfo.Size.Y, 
 				Guid, 
-				Idx
+				Idx,
+				SharedMemoryMediaOutput->bCrossGpu
 			);
 
 			if (!SharedCrossGpuTextures[Idx].IsValid())
 			{
-				UE_LOG(LogSharedMemoryMedia, Error, TEXT("Unable to create cross GPU texture of the requested type."));
+				UE_LOG(LogSharedMemoryMedia, Error, TEXT("Unable to create cross GPU texture of the requested type for Unique Name '%s'"), 
+					*SharedMemoryMediaOutput->UniqueName
+				);
 
 				SetState(EMediaCaptureState::Error);
 				return;
 			}
 
 			SharedCrossGpuTextureGuids[Idx] = Guid;
-			UE_LOG(LogSharedMemoryMedia, Verbose, TEXT("Created SharedGpuTextureGuid[%d] = %s"), Idx, *SharedCrossGpuTextureGuids[Idx].ToString());
+
+			UE_LOG(LogSharedMemoryMedia, Verbose, TEXT("Created SharedGpuTextureGuid[%d] = %s for UniqueName '%s'"), 
+				Idx, *SharedCrossGpuTextureGuids[Idx].ToString(), *SharedMemoryMediaOutput->UniqueName
+			);
 		}
 	}
 
@@ -246,9 +258,6 @@ void USharedMemoryMediaCapture::OnCustomCapture_RenderingThread(
 
 	// When enabled, add pass to invert alpha
 
-	const USharedMemoryMediaOutput* SharedMemoryMediaOutput = Cast<USharedMemoryMediaOutput>(MediaOutput);
-	check(SharedMemoryMediaOutput);
-
 	if (SharedMemoryMediaOutput->bInvertAlpha)
 	{
 		ETextureCreateFlags InvertedAlphaTextureFlags = ETextureCreateFlags::ResolveTargetable;
@@ -348,7 +357,9 @@ void USharedMemoryMediaCapture::AddCopyToSharedGpuTexturePass(FRDGBuilder& Graph
 			{
 				TRACE_CPUPROFILER_EVENT_SCOPE(SharedMemMediaOutputFenceBusy);
 
-				UE_LOG(LogSharedMemoryMedia, Verbose, TEXT("bTextureReadyFenceBusy[%d] for frame %d was busy, so we wait"), SharedTextureIdx, GFrameCounterRenderThread);
+				UE_LOG(LogSharedMemoryMedia, Verbose, TEXT("bTextureReadyFenceBusy[%d] for frame %d was busy, so we wait"),
+					SharedTextureIdx, GFrameCounterRenderThread
+				);
 
 				while (bTextureReadyFenceBusy[SharedTextureIdx])
 				{
@@ -377,7 +388,7 @@ void USharedMemoryMediaCapture::AddCopyToSharedGpuTexturePass(FRDGBuilder& Graph
 					RunningTasksCount--;
 				};
 
-				const FString CopyThreadName = FString::Printf(TEXT("SharedMemMediaOutputGpuTextureInTransitForFrame_%d"), FrameNumber);
+				const FString CopyThreadName = FString::Printf(TEXT("SharedMemMediaOutputGpuTextureInTransitForFrame_%d"), FrameNumber % 100);
 				TRACE_CPUPROFILER_EVENT_SCOPE_TEXT(*CopyThreadName);
 
 				// Wait for fence that indicates that the gpu texture has the data
diff --git a/Engine/Plugins/Runtime/nDisplay/Source/SharedMemoryMedia/Private/SharedMemoryMediaPlatform.h b/Engine/Plugins/Runtime/nDisplay/Source/SharedMemoryMedia/Private/SharedMemoryMediaPlatform.h
--- a/Engine/Plugins/Runtime/nDisplay/Source/SharedMemoryMedia/Private/SharedMemoryMediaPlatform.h
+++ b/Engine/Plugins/Runtime/nDisplay/Source/SharedMemoryMedia/Private/SharedMemoryMediaPlatform.h
@@ -16,19 +16,19 @@ class FSharedMemoryMediaPlatform
 public:
 
 	/** Creates a cross gpu texture */
-	virtual FTextureRHIRef CreateSharedCrossGpuTexture(EPixelFormat Format, bool bSrgb, int32 Width, int32 Height, const FGuid& Guid, uint32 BufferIdx)
+	virtual FTextureRHIRef CreateSharedTexture(EPixelFormat Format, bool bSrgb, int32 Width, int32 Height, const FGuid& Guid, uint32 BufferIdx, bool bCrossGpu)
 	{
 		return nullptr;
 	}
 
 	/** Opens a cross gpu texture specified by a Guid */
-	virtual FTextureRHIRef OpenSharedCrossGpuTextureByGuid(const FGuid& Guid, FSharedMemoryMediaTextureDescription& OutTextureDescription)
+	virtual FTextureRHIRef OpenSharedTextureByGuid(const FGuid& Guid, FSharedMemoryMediaTextureDescription& OutTextureDescription)
 	{
 		return nullptr;
 	}
 
 	/** Release any platform specific resources related to the indexed texture */
-	virtual void ReleaseSharedCrossGpuTexture(uint32 BufferIdx) {};
+	virtual void ReleaseSharedTexture(uint32 BufferIdx) {};
 };
 
 /** Factory of registered FSharedMemoryMediaPlatform RHI implementations. Use RegisterPlatformForRhi to register them. */
diff --git a/Engine/Plugins/Runtime/nDisplay/Source/SharedMemoryMedia/Private/SharedMemoryMediaPlayer.cpp b/Engine/Plugins/Runtime/nDisplay/Source/SharedMemoryMedia/Private/SharedMemoryMediaPlayer.cpp
--- a/Engine/Plugins/Runtime/nDisplay/Source/SharedMemoryMedia/Private/SharedMemoryMediaPlayer.cpp
+++ b/Engine/Plugins/Runtime/nDisplay/Source/SharedMemoryMedia/Private/SharedMemoryMediaPlayer.cpp
@@ -366,7 +366,7 @@ void FSharedMemoryMediaPlayer::TickFetch(FTimespan DeltaTime, FTimespan Timecode
 
 		// Now that we have the shared texture guid from the sender, use it to open the associated texture
 
-		SharedCrossGpuTextures[MemIdx] = PlatformData->OpenSharedCrossGpuTextureByGuid(SharedGpuTextureGuid, SharedCrossGpuTextureDescriptions[MemIdx]);
+		SharedCrossGpuTextures[MemIdx] = PlatformData->OpenSharedTextureByGuid(SharedGpuTextureGuid, SharedCrossGpuTextureDescriptions[MemIdx]);
 
 		if (!SharedCrossGpuTextures[MemIdx].IsValid())
 		{
diff --git a/Engine/Plugins/Runtime/nDisplay/Source/SharedMemoryMedia/Public/SharedMemoryMediaOutput.h b/Engine/Plugins/Runtime/nDisplay/Source/SharedMemoryMedia/Public/SharedMemoryMediaOutput.h
--- a/Engine/Plugins/Runtime/nDisplay/Source/SharedMemoryMedia/Public/SharedMemoryMediaOutput.h
+++ b/Engine/Plugins/Runtime/nDisplay/Source/SharedMemoryMedia/Public/SharedMemoryMediaOutput.h
@@ -28,6 +28,10 @@ public:
 	UPROPERTY(EditAnywhere, AdvancedDisplay, BlueprintReadWrite, Category = "Output")
 	bool bInvertAlpha = true;
 
+	/** If checked, the texture will be shared across different GPUs. Disable if not needed for faster performance */
+	UPROPERTY(EditAnywhere, AdvancedDisplay, BlueprintReadWrite, Category = "Output")
+	bool bCrossGpu = true;
+
 public:
 	
 	//~ Begin UMediaOutput interface
