nDislplay: Fix inner resolution math when buffer ratio < 1.

Issue is that FIntPoint * float decays to FIntPoint * IntType. So the resolution ended up as zero and media capture start would fail since zero resolution makes it fail.

diff --git a/Engine/Plugins/Runtime/nDisplay/Source/DisplayCluster/Private/Game/EngineClasses/Scene/DisplayClusterICVFXCameraComponent.cpp b/Engine/Plugins/Runtime/nDisplay/Source/DisplayCluster/Private/Game/EngineClasses/Scene/DisplayClusterICVFXCameraComponent.cpp
--- a/Engine/Plugins/Runtime/nDisplay/Source/DisplayCluster/Private/Game/EngineClasses/Scene/DisplayClusterICVFXCameraComponent.cpp
+++ b/Engine/Plugins/Runtime/nDisplay/Source/DisplayCluster/Private/Game/EngineClasses/Scene/DisplayClusterICVFXCameraComponent.cpp
@@ -132,8 +132,17 @@ void UDisplayClusterICVFXCameraComponent::UpdateOverscanEstimatedFrameSize()
 	const FDisplayClusterConfigurationICVFX_StageSettings& StageSettings = RootActor->GetStageSettings();
 
 	const float CameraBufferRatio = CameraSettings.GetCameraBufferRatio(StageSettings);
-	const FIntPoint InnerFrustumResolution = CameraSettings.GetCameraFrameSize(StageSettings) * CameraBufferRatio;
-	const FIntPoint UpscaledInnerFrustumResolution = InnerFrustumResolution * CameraSettings.CustomFrustum.FieldOfViewMultiplier;
+	const FIntPoint CameraFrameSize = CameraSettings.GetCameraFrameSize(StageSettings);
+
+	const FIntPoint InnerFrustumResolution(
+		CameraFrameSize.X * CameraBufferRatio,
+		CameraFrameSize.Y * CameraBufferRatio
+	);
+
+	const FIntPoint UpscaledInnerFrustumResolution(
+		InnerFrustumResolution.X * CameraSettings.CustomFrustum.FieldOfViewMultiplier,
+		InnerFrustumResolution.Y * CameraSettings.CustomFrustum.FieldOfViewMultiplier
+	);
 
 	// Read configuration data
 	FDisplayClusterViewport_CustomFrustumSettings CustomFrustumSettings;
