nDisplay: Demoted log of out of date ConfgExport from Warning to Display.

This should avoid packaging warnings in Horde. In normal usage, the asset will still be marked as dirty and the console will show the explanation to the user.

#rnx

diff --git a/Engine/Plugins/Runtime/nDisplay/Source/DisplayCluster/Private/Blueprints/DisplayClusterBlueprint.cpp b/Engine/Plugins/Runtime/nDisplay/Source/DisplayCluster/Private/Blueprints/DisplayClusterBlueprint.cpp
--- a/Engine/Plugins/Runtime/nDisplay/Source/DisplayCluster/Private/Blueprints/DisplayClusterBlueprint.cpp
+++ b/Engine/Plugins/Runtime/nDisplay/Source/DisplayCluster/Private/Blueprints/DisplayClusterBlueprint.cpp
@@ -155,7 +155,7 @@ void UDisplayClusterBlueprint::PreSave(FObjectPreSaveContext SaveContext)
 			{
 				if (ChildDCBP->MarkPackageDirty())
 				{
-					UE_LOG(LogDisplayClusterBlueprint, Warning,
+					UE_LOG(LogDisplayClusterBlueprint, Display,
 						TEXT("ConfigExport of the child nDisplay blueprint actor '%s' is not up to date in the asset, so the package was marked as dirty and should be re-saved."),
 						*ChildDCBP->GetOutermost()->GetName() // If MarkPackageDirty succeeded then GetOutermost() must exist.
 					);
@@ -186,7 +186,7 @@ void UDisplayClusterBlueprint::PostLoad()
 		{
 			if (MarkPackageDirty())
 			{			
-				UE_LOG(LogDisplayClusterBlueprint, Warning,
+				UE_LOG(LogDisplayClusterBlueprint, Display,
 					TEXT("ConfigExport of the nDisplay actor '%s' was not up to date in the asset, so the package was marked as dirty and should be re-saved."),
 					*GetOutermost()->GetName() // If MarkPackageDirty succeeded then GetOutermost() must exist.
 				);
