nDisplay: Fixed crashes if shaders do not exist on the current platform.
- Modified permutation rules for shaders that used to MeshProjection
- Always check if shaders are available on the current platform and hardware

diff --git a/Engine/Plugins/Runtime/nDisplay/Source/DisplayCluster/Private/Render/Viewport/DisplayClusterViewportFrameStatsViewExtension.cpp b/Engine/Plugins/Runtime/nDisplay/Source/DisplayCluster/Private/Render/Viewport/DisplayClusterViewportFrameStatsViewExtension.cpp
--- a/Engine/Plugins/Runtime/nDisplay/Source/DisplayCluster/Private/Render/Viewport/DisplayClusterViewportFrameStatsViewExtension.cpp
+++ b/Engine/Plugins/Runtime/nDisplay/Source/DisplayCluster/Private/Render/Viewport/DisplayClusterViewportFrameStatsViewExtension.cpp
@@ -113,7 +113,10 @@ FScreenPassTexture FDisplayClusterViewportFrameStatsViewExtension::PostProcessPa
 
 	// Render frame stats to a small render target to minimize shader costs
 	// Note: As a further optimization, this could be renderered once per view family.
+	if(VertexShader.IsValid())
 	{
+		// Always check if shaders are available on the current platform and hardware
+
 		const FScreenPassTextureViewport Viewport(TextureFrameStats);
 		FRenderTargetBinding FrameStatsRenderTargetBinding(TextureFrameStats, ERenderTargetLoadAction::EClear);
 
@@ -133,12 +136,18 @@ FScreenPassTexture FDisplayClusterViewportFrameStatsViewExtension::PostProcessPa
 		Output = FScreenPassRenderTarget::CreateFromInput(GraphBuilder, SceneColor, View.GetOverwriteLoadAction(), TEXT("DisplayCluster.FrameStatsOutput"));
 	}
 
+	TShaderMapRef<FDisplayClusterFrameStatsOutputShader> OutputPixelShader(ShaderMap);
+	if (!VertexShader.IsValid() || !OutputPixelShader.IsValid())
+	{
+		// Always check if shaders are available on the current platform and hardware
+		return Output;
+	}
+
 	// Render output composited with frame stats texture
 	{
 		const FScreenPassTextureViewport InputViewport(SceneColor);
 		const FScreenPassTextureViewport OutputViewport(Output);
 
-		TShaderMapRef<FDisplayClusterFrameStatsOutputShader> OutputPixelShader(ShaderMap);
 		FDisplayClusterFrameStatsOutputShaderParameters* Parameters = GraphBuilder.AllocParameters<FDisplayClusterFrameStatsOutputShaderParameters>();
 		Parameters->InputTexture = SceneColor.Texture;
 		Parameters->InputTextureSampler = TStaticSamplerState<>::GetRHI();
diff --git a/Engine/Plugins/Runtime/nDisplay/Source/DisplayCluster/Private/Render/Viewport/DisplayClusterViewportProxy.cpp b/Engine/Plugins/Runtime/nDisplay/Source/DisplayCluster/Private/Render/Viewport/DisplayClusterViewportProxy.cpp
--- a/Engine/Plugins/Runtime/nDisplay/Source/DisplayCluster/Private/Render/Viewport/DisplayClusterViewportProxy.cpp
+++ b/Engine/Plugins/Runtime/nDisplay/Source/DisplayCluster/Private/Render/Viewport/DisplayClusterViewportProxy.cpp
@@ -171,6 +171,11 @@ namespace UE::DisplayCluster::ViewportProxy
 			FGlobalShaderMap* ShaderMap = GetGlobalShaderMap(GMaxRHIFeatureLevel);
 			TShaderMapRef<FScreenVS> VertexShader(ShaderMap);
 			TShaderMapRef<TScreenPixelShader> PixelShader(ShaderMap);
+			if (!VertexShader.IsValid() || !PixelShader.IsValid())
+			{
+				// Always check if shaders are available on the current platform and hardware
+				return;
+			}
 
 			GraphicsPSOInit.BoundShaderState.VertexDeclarationRHI = GFilterVertexDeclaration.VertexDeclarationRHI;
 			GraphicsPSOInit.BoundShaderState.VertexShaderRHI = VertexShader.GetVertexShader();
diff --git a/Engine/Plugins/Runtime/nDisplay/Source/DisplayCluster/Private/Render/Viewport/Preview/DisplayClusterViewportManagerPreview_Render.cpp b/Engine/Plugins/Runtime/nDisplay/Source/DisplayCluster/Private/Render/Viewport/Preview/DisplayClusterViewportManagerPreview_Render.cpp
--- a/Engine/Plugins/Runtime/nDisplay/Source/DisplayCluster/Private/Render/Viewport/Preview/DisplayClusterViewportManagerPreview_Render.cpp
+++ b/Engine/Plugins/Runtime/nDisplay/Source/DisplayCluster/Private/Render/Viewport/Preview/DisplayClusterViewportManagerPreview_Render.cpp
@@ -58,23 +58,24 @@ int32 FDisplayClusterViewportManagerPreview::RenderClusterNodePreview(const int3
 			FSceneViewFamily& ViewFamily = *ViewportsViewFamily[0];
 
 			if (InSceneCanvas)
-	{
+			{
 				GetRendererModule().BeginRenderingViewFamily(InSceneCanvas, &ViewFamily);
-		}
+			}
 			else
 			{
-				FCanvas Canvas((FRenderTarget*)ViewportsViewFamily[0]->RenderTarget, nullptr, CurrentWorld, ERHIFeatureLevel::SM5, FCanvas::CDM_DeferDrawing /*FCanvas::CDM_ImmediateDrawing*/, 1.0f);
+				const ERHIFeatureLevel::Type FeatureLevel = CurrentWorld ? CurrentWorld->GetFeatureLevel() : GMaxRHIFeatureLevel;
+				FCanvas Canvas((FRenderTarget*)ViewportsViewFamily[0]->RenderTarget, nullptr, CurrentWorld, FeatureLevel, FCanvas::CDM_DeferDrawing /*FCanvas::CDM_ImmediateDrawing*/, 1.0f);
 				Canvas.Clear(FLinearColor::Black);
 
 				GetRendererModule().BeginRenderingViewFamily(&Canvas, &ViewFamily);
 			}
 
 			if (GNumExplicitGPUsForRendering > 1)
-		{
+			{
 				const FRHIGPUMask SubmitGPUMask = ViewFamily.Views.Num() == 1 ? ViewFamily.Views[0]->GPUMask : FRHIGPUMask::All();
 				ENQUEUE_RENDER_COMMAND(UDisplayClusterViewportClient_SubmitCommandList)(
 					[SubmitGPUMask](FRHICommandListImmediate& RHICmdList)
-			{
+					{
 						SCOPED_GPU_MASK(RHICmdList, SubmitGPUMask);
 						RHICmdList.SubmitCommandsHint();
 					});
@@ -83,11 +84,11 @@ int32 FDisplayClusterViewportManagerPreview::RenderClusterNodePreview(const int3
 			ViewportsViewFamily.RemoveAt(0);
 			OutViewportsAmount--;
 
-			}
+		}
 
 		// After all viewports we render compose
 		if (OutViewportsAmount > 0)
-			{
+		{
 			OutViewportsAmount--;
 
 			// Handle special viewports game-thread logic at frame end
@@ -102,7 +103,7 @@ int32 FDisplayClusterViewportManagerPreview::RenderClusterNodePreview(const int3
 
 			// Send event about cluster node rendering is finished
 			OnClusterNodePreviewGenerated.ExecuteIfBound(Configuration->GetClusterNodeId());
-			}
+		}
 	}
 
 	return OutViewportsAmount;
diff --git a/Engine/Plugins/Runtime/nDisplay/Source/DisplayClusterLightCardEditorShaders/Private/DisplayClusterMeshProjectionRenderer.cpp b/Engine/Plugins/Runtime/nDisplay/Source/DisplayClusterLightCardEditorShaders/Private/DisplayClusterMeshProjectionRenderer.cpp
--- a/Engine/Plugins/Runtime/nDisplay/Source/DisplayClusterLightCardEditorShaders/Private/DisplayClusterMeshProjectionRenderer.cpp
+++ b/Engine/Plugins/Runtime/nDisplay/Source/DisplayClusterLightCardEditorShaders/Private/DisplayClusterMeshProjectionRenderer.cpp
@@ -90,8 +90,11 @@ public:
 			return;
 		}
 
-		Shaders.TryGetVertexShader(PassShaders.VertexShader);
-		Shaders.TryGetPixelShader(PassShaders.PixelShader);
+		if (!Shaders.TryGetVertexShader(PassShaders.VertexShader) || !Shaders.TryGetPixelShader(PassShaders.PixelShader))
+		{
+			// Always check if shaders are available on the current platform and hardware
+			return;
+		}
 
 		FMeshDrawingPolicyOverrideSettings OverrideSettings = GetMeshOverrideSettings(MeshBatch);
 		ERasterizerFillMode MeshFillMode = ComputeMeshFillMode(*Material, OverrideSettings);
@@ -189,8 +192,7 @@ public:
 
 	static bool ShouldCompilePermutation(const FMeshMaterialShaderPermutationParameters& Parameters)
 	{
-		return IsFeatureLevelSupported(Parameters.Platform, ERHIFeatureLevel::SM5)
-			&& Parameters.VertexFactoryType == FindVertexFactoryType(TEXT("FLocalVertexFactory"));
+		return Parameters.VertexFactoryType == FindVertexFactoryType(TEXT("FLocalVertexFactory"));
 	}
 
 	void GetShaderBindings(
@@ -242,8 +244,7 @@ public:
 
 	static bool ShouldCompilePermutation(const FMeshMaterialShaderPermutationParameters& Parameters)
 	{
-		return IsFeatureLevelSupported(Parameters.Platform, ERHIFeatureLevel::SM5)
-			&& Parameters.VertexFactoryType == FindVertexFactoryType(TEXT("FLocalVertexFactory"));
+		return Parameters.VertexFactoryType == FindVertexFactoryType(TEXT("FLocalVertexFactory"));
 	}
 
 	void GetShaderBindings(
@@ -410,13 +411,13 @@ class FMeshProjectionHitProxyPS : public FMeshMaterialShader
 public:
 	static bool ShouldCompilePermutation(const FMeshMaterialShaderPermutationParameters& Parameters)
 	{
-		// Only compile the hit proxy shader on desktop editor platforms
-		return IsPCPlatform(Parameters.Platform) && EnumHasAllFlags(Parameters.Flags, EShaderPermutationFlags::HasEditorOnlyData)
+		// Only compile the hit proxy shader on editor platforms
+		return EnumHasAllFlags(Parameters.Flags, EShaderPermutationFlags::HasEditorOnlyData)
 			// and only compile for default materials or materials that are masked.
-			&& (Parameters.MaterialParameters.bIsSpecialEngineMaterial ||
-				!Parameters.MaterialParameters.bWritesEveryPixel ||
-				Parameters.MaterialParameters.bMaterialMayModifyMeshPosition ||
-				Parameters.MaterialParameters.bIsTwoSided);
+			&&
+			( !Parameters.MaterialParameters.bWritesEveryPixel
+			|| Parameters.MaterialParameters.bMaterialMayModifyMeshPosition
+			|| Parameters.MaterialParameters.bIsTwoSided);
 	}
 
 	FMeshProjectionHitProxyPS(const ShaderMetaType::CompiledShaderInitializerType& Initializer):
@@ -527,7 +528,7 @@ public:
 
 	static bool ShouldCompilePermutation(const FGlobalShaderPermutationParameters& Parameters)
 	{
-		return IsPCPlatform(Parameters.Platform) && IsFeatureLevelSupported(Parameters.Platform, ERHIFeatureLevel::SM5);
+		return true;
 	}
 
 	static void ModifyCompilationEnvironment(const FGlobalShaderPermutationParameters& Parameters, FShaderCompilerEnvironment& OutEnvironment)
@@ -580,7 +581,7 @@ public:
 
 	static bool ShouldCompilePermutation(const FGlobalShaderPermutationParameters& Parameters)
 	{
-		return IsPCPlatform(Parameters.Platform) && IsFeatureLevelSupported(Parameters.Platform, ERHIFeatureLevel::SM5);
+		return true;
 	}
 
 	static void ModifyCompilationEnvironment(const FGlobalShaderPermutationParameters& Parameters, FShaderCompilerEnvironment& OutEnvironment)
@@ -613,7 +614,7 @@ public:
 
 	static bool ShouldCompilePermutation(const FGlobalShaderPermutationParameters& Parameters)
 	{
-		return IsPCPlatform(Parameters.Platform) && IsFeatureLevelSupported(Parameters.Platform, ERHIFeatureLevel::SM5);
+		return true;
 	}
 
 	static void ModifyCompilationEnvironment(const FGlobalShaderPermutationParameters& Parameters, FShaderCompilerEnvironment& OutEnvironment)
@@ -651,8 +652,8 @@ public:
 
 	static bool ShouldCompilePermutation(const FGlobalShaderPermutationParameters& Parameters)
 	{
-		// Only PC platforms render editor primitives.
-		return IsPCPlatform(Parameters.Platform) && EnumHasAllFlags(Parameters.Flags, EShaderPermutationFlags::HasEditorOnlyData);
+		// Compile the selection outline shader only on editor platforms
+		return EnumHasAllFlags(Parameters.Flags, EShaderPermutationFlags::HasEditorOnlyData);
 	}
 };
 
@@ -1013,6 +1014,11 @@ void FDisplayClusterMeshProjectionRenderer::RenderColorOutput(FRDGBuilder& Graph
 		FGlobalShaderMap* GlobalShaderMap = GetGlobalShaderMap(View->FeatureLevel);
 		TShaderMapRef<FScreenPassVS> ScreenPassVS(GlobalShaderMap);
 		TShaderMapRef<FCopyRectPS> CopyPixelShader(GlobalShaderMap);
+		if (!ScreenPassVS.IsValid() || !CopyPixelShader.IsValid())
+		{
+			// Always check if shaders are available on the current platform and hardware
+			return;
+		}
 
 		FRHIBlendState* DefaultBlendState = FScreenPassPipelineState::FDefaultBlendState::GetRHI();
 		const FScreenPassTextureViewport RegionViewport(OutputRenderTargetBinding.GetTexture());
@@ -1044,8 +1050,6 @@ void FDisplayClusterMeshProjectionRenderer::RenderHitProxyOutput(FRDGBuilder& Gr
 	FRenderTargetBinding& OutputRenderTargetBinding,
 	FHitProxyConsumer* HitProxyConsumer)
 {
-
-
 	FRDGTextureDesc Desc(FRDGTextureDesc::Create2D(OutputRenderTargetBinding.GetTexture()->Desc.Extent, PF_B8G8R8A8, FClearValueBinding::Black, TexCreate_RenderTargetable | TexCreate_ShaderResource));
 	FRDGTextureRef HitProxyTexture = GraphBuilder.CreateTexture(Desc, TEXT("DisplayClusterMeshProjection.HitProxyTexture"));
 
@@ -1067,6 +1071,11 @@ void FDisplayClusterMeshProjectionRenderer::RenderHitProxyOutput(FRDGBuilder& Gr
 		FGlobalShaderMap* GlobalShaderMap = GetGlobalShaderMap(View->FeatureLevel);
 		TShaderMapRef<FScreenPassVS> ScreenPassVS(GlobalShaderMap);
 		TShaderMapRef<FCopyRectPS> CopyPixelShader(GlobalShaderMap);
+		if (!ScreenPassVS.IsValid() || !CopyPixelShader.IsValid())
+		{
+			// Always check if shaders are available on the current platform and hardware
+			return;
+		}
 
 		FRHIBlendState* DefaultBlendState = FScreenPassPipelineState::FDefaultBlendState::GetRHI();
 		const FScreenPassTextureViewport RegionViewport(OutputRenderTargetBinding.GetTexture());
@@ -1269,6 +1278,11 @@ void AddSeparableFilterPass(FRDGBuilder& GraphBuilder,
 	}
 
 	TShaderMapRef<FMeshProjectionNormalsFilterCS<FilterType>> ComputeShader(GlobalShaderMap);
+	if (!ComputeShader.IsValid())
+	{
+		// Always check if shaders are available on the current platform and hardware
+		return;
+	}
 
 	FComputeShaderUtils::AddPass(
 		GraphBuilder,
@@ -1345,6 +1359,11 @@ void FDisplayClusterMeshProjectionRenderer::AddNormalsFilterPass(FRDGBuilder& Gr
 		PassParameters->NormalCorrectionMatrix = NormalCorrectionMatrix;
 
 		TShaderMapRef<FMeshProjectionNormalsCreateRWTexturesCS> ComputeShader(GlobalShaderMap);
+		if (!ComputeShader.IsValid())
+		{
+			// Always check if shaders are available on the current platform and hardware
+			return;
+		}
 
 		FComputeShaderUtils::AddPass(
 			GraphBuilder,
@@ -1387,6 +1406,11 @@ void FDisplayClusterMeshProjectionRenderer::AddNormalsFilterPass(FRDGBuilder& Gr
 
 		TShaderMapRef<FScreenPassVS> ScreenPassVS(GlobalShaderMap);
 		TShaderMapRef<FMeshProjectionNormalsOutputPS> OutputNormalsPS(GlobalShaderMap);
+		if (!ScreenPassVS.IsValid() || !OutputNormalsPS.IsValid())
+		{
+			// Always check if shaders are available on the current platform and hardware
+			return;
+		}
 
 		FRHIBlendState* DefaultBlendState = FScreenPassPipelineState::FDefaultBlendState::GetRHI();
 
@@ -1463,6 +1487,11 @@ void FDisplayClusterMeshProjectionRenderer::AddSelectionOutlineScreenPass(FRDGBu
 
 	TShaderMapRef<FScreenPassVS> ScreenPassVS(GlobalShaderMap);
 	TShaderMapRef<FMeshProjectionSelectionOutlinePS> SelectionOutlinePS(GlobalShaderMap);
+	if (!ScreenPassVS.IsValid() || !SelectionOutlinePS.IsValid())
+	{
+		// Always check if shaders are available on the current platform and hardware
+		return;
+	}
 
 	FRHIBlendState* DefaultBlendState = FScreenPassPipelineState::FDefaultBlendState::GetRHI();
 
diff --git a/Engine/Plugins/Runtime/nDisplay/Source/DisplayClusterMedia/Private/DisplayClusterMediaHelpers.cpp b/Engine/Plugins/Runtime/nDisplay/Source/DisplayClusterMedia/Private/DisplayClusterMediaHelpers.cpp
--- a/Engine/Plugins/Runtime/nDisplay/Source/DisplayClusterMedia/Private/DisplayClusterMediaHelpers.cpp
+++ b/Engine/Plugins/Runtime/nDisplay/Source/DisplayClusterMedia/Private/DisplayClusterMediaHelpers.cpp
@@ -88,6 +88,11 @@ namespace DisplayClusterMediaHelpers
 			FGlobalShaderMap* ShaderMap = GetGlobalShaderMap(GMaxRHIFeatureLevel);
 			TShaderMapRef<FScreenVS> VertexShader(ShaderMap);
 			TShaderMapRef<TScreenPixelShader> PixelShader(ShaderMap);
+			if (!VertexShader.IsValid() || !PixelShader.IsValid())
+			{
+				// Always check if shaders are available on the current platform and hardware
+				return;
+			}
 
 			GraphicsPSOInit.BoundShaderState.VertexDeclarationRHI = GFilterVertexDeclaration.VertexDeclarationRHI;
 			GraphicsPSOInit.BoundShaderState.VertexShaderRHI = VertexShader.GetVertexShader();
diff --git a/Engine/Plugins/Runtime/nDisplay/Source/DisplayClusterShaders/Private/Shaders/DisplayClusterShadersPostprocess_Blur.cpp b/Engine/Plugins/Runtime/nDisplay/Source/DisplayClusterShaders/Private/Shaders/DisplayClusterShadersPostprocess_Blur.cpp
--- a/Engine/Plugins/Runtime/nDisplay/Source/DisplayClusterShaders/Private/Shaders/DisplayClusterShadersPostprocess_Blur.cpp
+++ b/Engine/Plugins/Runtime/nDisplay/Source/DisplayClusterShaders/Private/Shaders/DisplayClusterShadersPostprocess_Blur.cpp
@@ -134,6 +134,11 @@ static void PicpBlurPostProcess_RenderThread(
 
 	TShaderMapRef<FDirectProjectionVS>                VertexShader(GlobalShaderMap);
 	TShaderMapRef<TPicpBlurPostProcessPS<ShaderType>> PixelShader(GlobalShaderMap);
+	if (!VertexShader.IsValid() || !PixelShader.IsValid())
+	{
+		// Always check if shaders are available on the current platform and hardware
+		return;
+	}
 
 	FPicpBlurPostProcessPS::FParameters PsParameters{};
 	PsParameters.SrcTexture = InShaderTexture;
diff --git a/Engine/Plugins/Runtime/nDisplay/Source/DisplayClusterShaders/Private/Shaders/DisplayClusterShadersPostprocess_OutputRemap.cpp b/Engine/Plugins/Runtime/nDisplay/Source/DisplayClusterShaders/Private/Shaders/DisplayClusterShadersPostprocess_OutputRemap.cpp
--- a/Engine/Plugins/Runtime/nDisplay/Source/DisplayClusterShaders/Private/Shaders/DisplayClusterShadersPostprocess_OutputRemap.cpp
+++ b/Engine/Plugins/Runtime/nDisplay/Source/DisplayClusterShaders/Private/Shaders/DisplayClusterShadersPostprocess_OutputRemap.cpp
@@ -172,6 +172,11 @@ bool FDisplayClusterShadersPostprocess_OutputRemap::RenderPostprocess_OutputRema
 
 		TShaderMapRef<FOutputRemapVS> VertexShader(GlobalShaderMap);
 		TShaderMapRef<FOutputRemapPS> PixelShader(GlobalShaderMap);
+		if (!VertexShader.IsValid() || !PixelShader.IsValid())
+		{
+			// Always check if shaders are available on the current platform and hardware
+			return false;
+		}
 
 		FGraphicsPipelineStateInitializer GraphicsPSOInit;
 		// Set the graphic pipeline state.
diff --git a/Engine/Plugins/Runtime/nDisplay/Source/DisplayClusterShaders/Private/Shaders/DisplayClusterShadersPreprocess_UVLightCards.cpp b/Engine/Plugins/Runtime/nDisplay/Source/DisplayClusterShaders/Private/Shaders/DisplayClusterShadersPreprocess_UVLightCards.cpp
--- a/Engine/Plugins/Runtime/nDisplay/Source/DisplayClusterShaders/Private/Shaders/DisplayClusterShadersPreprocess_UVLightCards.cpp
+++ b/Engine/Plugins/Runtime/nDisplay/Source/DisplayClusterShaders/Private/Shaders/DisplayClusterShadersPreprocess_UVLightCards.cpp
@@ -32,12 +32,10 @@ public:
 	{
 		PassUniformBuffer.Bind(Initializer.ParameterMap, FSceneTextureUniformParameters::FTypeInfo::GetStructMetadata()->GetShaderVariableName());
 	}
-	              
+
 	static bool ShouldCompilePermutation(const FMeshMaterialShaderPermutationParameters& Parameters)
 	{
-		return EnumHasAllFlags(Parameters.Flags, EShaderPermutationFlags::HasEditorOnlyData)
-			&& IsFeatureLevelSupported(Parameters.Platform, ERHIFeatureLevel::SM5)
-			&& Parameters.VertexFactoryType == FindVertexFactoryType(TEXT("FLocalVertexFactory"));
+		return Parameters.VertexFactoryType == FindVertexFactoryType(TEXT("FLocalVertexFactory"));
 	}
 };
 
@@ -57,9 +55,7 @@ public:
 
 	static bool ShouldCompilePermutation(const FMeshMaterialShaderPermutationParameters& Parameters)
 	{
-		return EnumHasAllFlags(Parameters.Flags, EShaderPermutationFlags::HasEditorOnlyData)
-			&& IsFeatureLevelSupported(Parameters.Platform, ERHIFeatureLevel::SM5)
-			&& Parameters.VertexFactoryType == FindVertexFactoryType(TEXT("FLocalVertexFactory"));
+		return Parameters.VertexFactoryType == FindVertexFactoryType(TEXT("FLocalVertexFactory"));
 	}
 };
 
@@ -95,8 +91,11 @@ public:
 			return;
 		}
 
-		Shaders.TryGetVertexShader(PassShaders.VertexShader);
-		Shaders.TryGetPixelShader(PassShaders.PixelShader);
+		if (!Shaders.TryGetVertexShader(PassShaders.VertexShader) || !Shaders.TryGetPixelShader(PassShaders.PixelShader))
+		{
+			// Always check if shaders are available on the current platform and hardware
+			return;
+		}
 
 		FMeshDrawingPolicyOverrideSettings OverrideSettings = ComputeMeshOverrideSettings(MeshBatch);
 		ERasterizerFillMode MeshFillMode = ComputeMeshFillMode(Material, OverrideSettings);
@@ -337,6 +336,11 @@ bool FDisplayClusterShadersPreprocess_UVLightCards::RenderPreprocess_UVLightCard
 	FGlobalShaderMap* GlobalShaderMap = GetGlobalShaderMap(View->FeatureLevel);
 	TShaderMapRef<FScreenPassVS> ScreenPassVS(GlobalShaderMap);
 	TShaderMapRef<FGammaCorrectionPS> GammaCorrectionPS(GlobalShaderMap);
+	if (!ScreenPassVS.IsValid() || !GammaCorrectionPS.IsValid())
+	{
+		// Always check if shaders are available on the current platform and hardware
+		return false;
+	}
 
 	AddDrawScreenPass(
 		GraphBuilder,
diff --git a/Engine/Plugins/Runtime/nDisplay/Source/DisplayClusterShaders/Private/Shaders/DisplayClusterShadersWarpblend_ICVFX.cpp b/Engine/Plugins/Runtime/nDisplay/Source/DisplayClusterShaders/Private/Shaders/DisplayClusterShadersWarpblend_ICVFX.cpp
--- a/Engine/Plugins/Runtime/nDisplay/Source/DisplayClusterShaders/Private/Shaders/DisplayClusterShadersWarpblend_ICVFX.cpp
+++ b/Engine/Plugins/Runtime/nDisplay/Source/DisplayClusterShaders/Private/Shaders/DisplayClusterShadersWarpblend_ICVFX.cpp
@@ -980,7 +980,11 @@ public:
 		FGlobalShaderMap* ShaderMap = GetGlobalShaderMap(GMaxRHIFeatureLevel);
 		TShaderMapRef<FIcvfxWarpVS> VertexShader(ShaderMap, RenderPassData.VSPermutationVector);
 		TShaderMapRef<FIcvfxPassthroughPS> PixelShader(ShaderMap);
-
+		if (!VertexShader.IsValid() || !PixelShader.IsValid())
+		{
+			// Always check if shaders are available on the current platform and hardware
+			return false;
+		}
 
 		FGraphicsPipelineStateInitializer GraphicsPSOInit;
 		// Set the graphic pipeline state.
diff --git a/Engine/Plugins/Runtime/nDisplay/Source/DisplayClusterShaders/Private/Shaders/DisplayClusterShadersWarpblend_MPCDI.cpp b/Engine/Plugins/Runtime/nDisplay/Source/DisplayClusterShaders/Private/Shaders/DisplayClusterShadersWarpblend_MPCDI.cpp
--- a/Engine/Plugins/Runtime/nDisplay/Source/DisplayClusterShaders/Private/Shaders/DisplayClusterShadersWarpblend_MPCDI.cpp
+++ b/Engine/Plugins/Runtime/nDisplay/Source/DisplayClusterShaders/Private/Shaders/DisplayClusterShadersWarpblend_MPCDI.cpp
@@ -417,6 +417,11 @@ public:
 		FGlobalShaderMap* ShaderMap = GetGlobalShaderMap(GMaxRHIFeatureLevel);
 		TShaderMapRef<FMpcdiWarpVS> VertexShader(ShaderMap, RenderPassData.VSPermutationVector);
 		TShaderMapRef<FMpcdiPassthroughPS> PixelShader(ShaderMap);
+		if (!VertexShader.IsValid() || !PixelShader.IsValid())
+		{
+			// Always check if shaders are available on the current platform and hardware
+			return false;
+		}
 
 		FGraphicsPipelineStateInitializer GraphicsPSOInit;
 		// Set the graphic pipeline state.
@@ -459,6 +464,11 @@ public:
 		FGlobalShaderMap* ShaderMap = GetGlobalShaderMap(GMaxRHIFeatureLevel);
 		TShaderMapRef<FMpcdiWarpVS> VertexShader(ShaderMap, RenderPassData.VSPermutationVector);
 		TShaderMapRef<FMpcdiWarpPS> PixelShader(ShaderMap, RenderPassData.PSPermutationVector);
+		if (!VertexShader.IsValid() || !PixelShader.IsValid())
+		{
+			// Always check if shaders are available on the current platform and hardware
+			return false;
+		}
 
 		FGraphicsPipelineStateInitializer GraphicsPSOInit;
 		// Set the graphic pipeline state.
