[nDisplay] Fixed editor crash caused by accessing an UObject that have already been garbage collected

diff --git a/Engine/Plugins/Runtime/nDisplay/Source/DisplayCluster/Private/Render/Viewport/Preview/DisplayClusterViewportPreviewMesh.cpp b/Engine/Plugins/Runtime/nDisplay/Source/DisplayCluster/Private/Render/Viewport/Preview/DisplayClusterViewportPreviewMesh.cpp
--- a/Engine/Plugins/Runtime/nDisplay/Source/DisplayCluster/Private/Render/Viewport/Preview/DisplayClusterViewportPreviewMesh.cpp
+++ b/Engine/Plugins/Runtime/nDisplay/Source/DisplayCluster/Private/Render/Viewport/Preview/DisplayClusterViewportPreviewMesh.cpp
@@ -11,20 +11,6 @@
 
 namespace UE::DisplayCluster::ViewportPreviewMesh
 {
-	/** Return ptr on UObject if it is still valid. */
-	template<class T>
-	static inline T* GetExistingObject(const TObjectPtr<T>& InProperty)
-	{
-		if (InProperty == nullptr
-			|| InProperty->GetName().Find(TEXT("TRASH_")) != INDEX_NONE
-			|| InProperty->HasAnyFlags(RF_BeginDestroyed | RF_FinishDestroyed))
-		{
-			return nullptr;
-		}
-
-		return InProperty;
-	}
-
 	/** Reset the mesh material to default values from its archetype. */
 	static inline bool RestoreMeshMaterialsFromArchetype(UMeshComponent* InMeshComponent)
 	{
@@ -48,12 +34,12 @@ namespace UE::DisplayCluster::ViewportPreviewMesh
 		TObjectPtr<UMaterial> OutMaterial = nullptr;
 
 		// First get the material from the ViewPoint component (WarpPolicy)
-		OutMaterial = GetExistingObject(ViewPointComponent ? ViewPointComponent->GetDisplayDeviceMaterial(InMeshType, InMaterialType) : nullptr);
+		OutMaterial = ViewPointComponent ? ViewPointComponent->GetDisplayDeviceMaterial(InMeshType, InMaterialType) : nullptr;
 
 		// Finally, get the material from the DisplayDevice
 		if (!OutMaterial)
 		{
-			OutMaterial = GetExistingObject(InDisplayDeviceComponent ? InDisplayDeviceComponent->GetDisplayDeviceMaterial(InMeshType, InMaterialType) : nullptr);
+			OutMaterial = InDisplayDeviceComponent ? InDisplayDeviceComponent->GetDisplayDeviceMaterial(InMeshType, InMaterialType) : nullptr;
 		}
 
 		// Ignore deleted materials
@@ -143,6 +129,19 @@ void FDisplayClusterViewportPreviewMesh::Update(FDisplayClusterViewport* InViewp
 	UpdateOverlayMaterial(InViewport);
 }
 
+void FDisplayClusterViewportPreviewMesh::AddReferencedObjects(FReferenceCollector& Collector)
+{
+	TObjectPtr<UObject> OwnedObjects[] = { MeshComponentPtr, OrigOverlayMaterial, MaterialInstancePtr, CurrentMaterialPtr, DefaultMaterialPtr };
+
+	for (TObjectPtr<UObject>& OwnedObject : OwnedObjects)
+	{
+		if (OwnedObject)
+		{
+			Collector.AddReferencedObject(OwnedObject);
+		}
+	}
+}
+
 void FDisplayClusterViewportPreviewMesh::SetCustomOverlayMaterial(UMeshComponent* InMeshComponent, UMaterialInterface* InOverlayMaterial)
 {
 	if (InMeshComponent)
@@ -309,14 +308,10 @@ UMeshComponent* FDisplayClusterViewportPreviewMesh::GetOrCreatePreviewMeshCompon
 
 UMeshComponent* FDisplayClusterViewportPreviewMesh::GetMeshComponent() const
 {
-	using namespace UE::DisplayCluster::ViewportPreviewMesh;
-
-	return GetExistingObject(MeshComponentPtr);
+	return MeshComponentPtr;
 }
 
 UMaterialInstanceDynamic* FDisplayClusterViewportPreviewMesh::GetMaterialInstance() const
 {
-	using namespace UE::DisplayCluster::ViewportPreviewMesh;
-
-	return GetExistingObject(MaterialInstancePtr);
+	return MaterialInstancePtr;
 }
diff --git a/Engine/Plugins/Runtime/nDisplay/Source/DisplayCluster/Private/Render/Viewport/Preview/DisplayClusterViewportPreviewMesh.h b/Engine/Plugins/Runtime/nDisplay/Source/DisplayCluster/Private/Render/Viewport/Preview/DisplayClusterViewportPreviewMesh.h
--- a/Engine/Plugins/Runtime/nDisplay/Source/DisplayCluster/Private/Render/Viewport/Preview/DisplayClusterViewportPreviewMesh.h
+++ b/Engine/Plugins/Runtime/nDisplay/Source/DisplayCluster/Private/Render/Viewport/Preview/DisplayClusterViewportPreviewMesh.h
@@ -11,6 +11,9 @@
 
 #include "Misc/DisplayClusterObjectRef.h"
 
+#include "UObject/GCObject.h"
+
+
 class UDisplayClusterDisplayDeviceBaseComponent;
 class UDisplayClusterCameraComponent;
 class FDisplayClusterViewport;
@@ -43,6 +46,7 @@ ENUM_CLASS_FLAGS(EDisplayClusterViewportPreviewMeshFlags);
 * Manage preview mesh of the viewport
 */
 class FDisplayClusterViewportPreviewMesh
+	: public FGCObject
 {
 public:
 	FDisplayClusterViewportPreviewMesh(const EDisplayClusterDisplayDeviceMeshType InMeshType, const TSharedRef<FDisplayClusterViewportConfiguration, ESPMode::ThreadSafe> InConfiguration)
@@ -77,6 +81,15 @@ public:
 		return EnumHasAnyFlags(RuntimeFlags, InMeshFlags);
 	}
 
+public:
+	//~ Begin FGCObject interface
+	virtual void AddReferencedObjects(FReferenceCollector& Collector) override;
+	virtual FString GetReferencerName() const override
+	{
+		return TEXT("FDisplayClusterViewportPreviewMesh");
+	}
+	//~ End FGCObject interface
+
 private:
 	/** Returns true if this mesh type is supported by the viewport projection policy and DCRA. */
 	bool ShouldUseMeshComponent(FDisplayClusterViewport* InViewport) const;
@@ -111,7 +124,7 @@ private:
 	EDisplayClusterViewportPreviewMeshFlags RuntimeFlags;
 
 	// Mesh component used for preview
-	TObjectPtr<UMeshComponent> MeshComponentPtr = nullptr;
+	TObjectPtr<UMeshComponent> MeshComponentPtr;
 
 	// This mesh component exists in DCRA and does not need to be deleted.
 	bool bIsRootActorMeshComponent = false;
@@ -120,11 +133,11 @@ private:
 	TObjectPtr<UMaterialInterface> OrigOverlayMaterial;
 
 	// Preview material used on the mesh
-	TObjectPtr<UMaterialInstanceDynamic> MaterialInstancePtr = nullptr;
+	TObjectPtr<UMaterialInstanceDynamic> MaterialInstancePtr;
 
 	// The current material assigned to the preview mesh
-	TObjectPtr<UMaterial> CurrentMaterialPtr = nullptr;
+	TObjectPtr<UMaterial> CurrentMaterialPtr;
 
 	// The default material defined in the DisplayDevice
-	TObjectPtr<UMaterial> DefaultMaterialPtr = nullptr;
+	TObjectPtr<UMaterial> DefaultMaterialPtr;
 };
