SharedMemoryMedia: Added ability to have multiple simultaneous receivers. Currently limited to 4.

diff --git a/Engine/Plugins/Runtime/nDisplay/Source/SharedMemoryMedia/Private/SharedMemoryMediaCapture.cpp b/Engine/Plugins/Runtime/nDisplay/Source/SharedMemoryMedia/Private/SharedMemoryMediaCapture.cpp
--- a/Engine/Plugins/Runtime/nDisplay/Source/SharedMemoryMedia/Private/SharedMemoryMediaCapture.cpp
+++ b/Engine/Plugins/Runtime/nDisplay/Source/SharedMemoryMedia/Private/SharedMemoryMediaCapture.cpp
@@ -99,7 +99,7 @@ bool USharedMemoryMediaCapture::InitializeCapture()
 
 				// Except some special data
 				FSharedMemoryMediaFrameMetadata* Data = static_cast<FSharedMemoryMediaFrameMetadata*>(SharedMemoryRegion->GetAddress());
-				Data->Receiver.FrameNumberAcked = ~0;
+				Data->Initialize();
 
 				UE_LOG(LogSharedMemoryMedia, Verbose, TEXT("Created SharedMemoryRegion[%d] = %s for UniqueName '%s'"), 
 					BufferIdx, *SharedMemoryRegionName, *SharedMemoryMediaOutput->UniqueName
@@ -413,31 +413,34 @@ void USharedMemoryMediaCapture::AddCopyToSharedGpuTexturePass(FRDGBuilder& Graph
 					FMemory::Memcpy(&SharedMemoryData->Sender, &SenderMetadata, sizeof(FSharedMemoryMediaFrameMetadata::FSender));
 				}
 
-				// Wait for FrameNumber ack
-
-				if (SharedMemoryData->Receiver.KeepAliveShiftRegister)
+				// Wait for FrameNumber ack. Since we may have more than one receiver, we must wait until 
+				// all active receivers have acked a frame number equal or greater than FrameNumber.
 				{
 					TRACE_CPUPROFILER_EVENT_SCOPE(WaitForGpuTextureAck)
 
 					const double StartTimeSeconds = FPlatformTime::Seconds();
 					constexpr double TimeoutSeconds = 0.5;
 
-					while (SharedMemoryData->Receiver.FrameNumberAcked < FrameNumber && SharedMemoryData->Receiver.KeepAliveShiftRegister)
+					while (!SharedMemoryData->AllReceiversAckedFrameNumber(FrameNumber))
 					{
 						FPlatformProcess::SleepNoStats(SpinWaitTimeSeconds);
 
 						if ((FPlatformTime::Seconds() - StartTimeSeconds) > TimeoutSeconds)
 						{
-							// @todo use proper log category
-							UE_LOG(LogSharedMemoryMedia, Warning, TEXT("FSharedMemoryMediaCapture timed out waiting for its receiver to ack frame %d"), FrameNumber);
+							UE_LOG(LogSharedMemoryMedia, Warning, 
+								TEXT("FSharedMemoryMediaCapture timed out waiting for its receiver to ack frame %d"), FrameNumber
+							);
 
+							// We break this while loop because we will not be waiting any longer, 
+							// even if receivers haven't acked the frame.
+							// This includes any newly connected receivers that may have joined during this time.
 							break;
 						}
-					};
+					}
 				}
 
-				// Shift the keep alive. The bit depth of the keep alive is the number of frames the receiver has to re-set it before it expires.
-				SharedMemoryData->Receiver.KeepAliveShiftRegister >>= 1;
+				// Decrement the keep alive. The receiver must keep resetting it to avoid letting it expire.
+				SharedMemoryData->DecrementKeepAlives();
 
 				// Clear fence and flag that we're ready for a new frame
 
diff --git a/Engine/Plugins/Runtime/nDisplay/Source/SharedMemoryMedia/Private/SharedMemoryMediaPlayer.cpp b/Engine/Plugins/Runtime/nDisplay/Source/SharedMemoryMedia/Private/SharedMemoryMediaPlayer.cpp
--- a/Engine/Plugins/Runtime/nDisplay/Source/SharedMemoryMedia/Private/SharedMemoryMediaPlayer.cpp
+++ b/Engine/Plugins/Runtime/nDisplay/Source/SharedMemoryMedia/Private/SharedMemoryMediaPlayer.cpp
@@ -37,12 +37,18 @@ static FAutoConsoleVariableRef CVarDisplayClusterSharedMemoryLatency(
 );
 
 
-TSet<FString> FSharedMemoryMediaPlayer::UniqueNamesRegistered;
 
 FSharedMemoryMediaPlayer::FSharedMemoryMediaPlayer()
 	: Super()
 {
 	Samples = new FSharedMemoryMediaSamples();
+
+	// Invalidate the receiver indices, indicating that we need to find a spot in the shared metadata
+	// to interact with the sender.
+	for (int32 Idx = 0; Idx < NUMSHAREDMEM; ++Idx)
+	{
+		ReceiverIndex[Idx] = INDEX_NONE;
+	}
 }
 
 FSharedMemoryMediaPlayer::~FSharedMemoryMediaPlayer()
@@ -51,25 +57,6 @@ FSharedMemoryMediaPlayer::~FSharedMemoryMediaPlayer()
 	delete Samples;
 }
 
-bool FSharedMemoryMediaPlayer::RegisterUniqueName(FString& InUniqueName)
-{
-	// We don't allow empty unique names
-	if (!InUniqueName.Len())
-	{
-		return false;
-	}
-
-	bool bIsAlreadyInSet = false;
-
-	UniqueNamesRegistered.Add(InUniqueName, &bIsAlreadyInSet);
-
-	return !bIsAlreadyInSet;
-}
-
-void FSharedMemoryMediaPlayer::UnregisterUniqueName(FString& InUniqueName)
-{
-	UniqueNamesRegistered.Remove(InUniqueName);
-}
 
 bool FSharedMemoryMediaPlayer::Open(const FString& Url, const IMediaOptions* Options)
 {
@@ -125,14 +112,6 @@ bool FSharedMemoryMediaPlayer::Open(const FString& Url, const IMediaOptions* Opt
 		}
 	}
 
-	if (!RegisterUniqueName(UniqueName))
-	{
-		UE_LOG(LogSharedMemoryMedia, Error, TEXT("Only one SharedMemoryMediaPlayer the UniqueName '%s' can play at a time"), *UniqueName);
-		UniqueName.Reset(); // Reset the unique name to avoid unregistering when not open
-		Close();
-		return true;
-	}
-
 	// Set a playing state.
 	PlayerState = EMediaState::Playing;
 
@@ -168,13 +147,12 @@ void FSharedMemoryMediaPlayer::Close()
 		if (SharedMemory[BufferIdx])
 		{
 			// Let the sender know that we're not monitoring this and that it should not wait for acks
-			{
-				FSharedMemoryMediaFrameMetadata* Data = static_cast<FSharedMemoryMediaFrameMetadata*>(SharedMemory[BufferIdx]->GetAddress());
 
-				if (Data)
-				{
-					Data->Receiver.KeepAliveShiftRegister = 0;
-				}
+			FSharedMemoryMediaFrameMetadata* Data = static_cast<FSharedMemoryMediaFrameMetadata*>(SharedMemory[BufferIdx]->GetAddress());
+
+			if (Data)
+			{
+				Data->DisconnectReceiverFromSlot(ReceiverId, ReceiverIndex[BufferIdx]); // Internally checks for INDEX_NONE
 			}
 
 			// Close the shared memory
@@ -182,6 +160,9 @@ void FSharedMemoryMediaPlayer::Close()
 			SharedMemory[BufferIdx] = nullptr;
 		}
 
+		// Unconditionally invalidate the receiver indices
+		ReceiverIndex[BufferIdx] = INDEX_NONE;
+
 		// Free gpu fences
 		{
 			// Fences should not be busy since we flush rendering commands.
@@ -200,12 +181,8 @@ void FSharedMemoryMediaPlayer::Close()
 	// Reset our state variables
 	ModeState.Reset();
 
-	// Allow this or other players with the given UniqueName to play
-	UnregisterUniqueName(UniqueName);
-
 	// Update our player state
 	PlayerState = EMediaState::Closed;
-
 }
 
 FGuid FSharedMemoryMediaPlayer::GetPlayerPluginGUID() const
@@ -345,7 +322,8 @@ void FSharedMemoryMediaPlayer::TickFetch(FTimespan DeltaTime, FTimespan Timecode
 			SampleCommonTexture = nullptr;
 
 			// No need for the sender to wait for us while we figure out the new textures.
-			SharedMemoryMetadataPtr->Receiver.KeepAliveShiftRegister = 0;
+			SharedMemoryMetadataPtr->DisconnectReceiverFromSlot(ReceiverId, ReceiverIndex[MemIdx]); // Already checks for INDEX_NONE
+			ReceiverIndex[MemIdx] = INDEX_NONE;
 
 			// We now exit this block, which will treat the invalid gpu texture case.
 		}
@@ -418,14 +396,38 @@ void FSharedMemoryMediaPlayer::TickFetch(FTimespan DeltaTime, FTimespan Timecode
 		return;
 	}
 
-	// Let sender know that we're monitoring the shared memory by setting the Keepalive shift register
+	// Make sure that we have a receiver slot, and keep it alive.
 	for (int32 MemIdx = 0; MemIdx < NUMSHAREDMEM; ++MemIdx)
 	{
-		FSharedMemoryMediaFrameMetadata* SharedMetadataPtr = static_cast<FSharedMemoryMediaFrameMetadata*>(SharedMemory[MemIdx]->GetAddress());
+		FSharedMemoryMediaFrameMetadata* SharedMemoryMetadataPtr = static_cast<FSharedMemoryMediaFrameMetadata*>(SharedMemory[MemIdx]->GetAddress());
+
+		// Make sure our current index is valid
+		if (!SharedMemoryMetadataPtr->IsReceiverConnectedToSlot(ReceiverId, ReceiverIndex[MemIdx]))
+		{
+			ReceiverIndex[MemIdx] = INDEX_NONE;
+		}
+
+		// Find a slot if we don't have one
+		
+		if (ReceiverIndex[MemIdx] == INDEX_NONE)
+		{
+			ReceiverIndex[MemIdx] = SharedMemoryMetadataPtr->ConnectToSlot(ReceiverId);
+
+			// If we don't have a valid slot at this point, we must return and try again later.
+			if (ReceiverIndex[MemIdx] == INDEX_NONE)
+			{
+				return;
+			}
+		}
 
 		// Let sender know that we're monitoring this shared memory.
 		// We expect the sender to not wait for acks if nobody is monitoring it.
-		SharedMetadataPtr->Receiver.KeepAliveShiftRegister = ~0; // all ones
+		if (!SharedMemoryMetadataPtr->KeepAlive(ReceiverIndex[MemIdx]))
+		{
+			// If keep alive failed, it means that we're not connected anymore.
+			ReceiverIndex[MemIdx] = INDEX_NONE;
+			return;
+		}
 	}
 
 	// Create a new media sample and populate it with any needed data
@@ -669,7 +671,9 @@ bool FSharedMemoryMediaPlayer::DetermineNextSourceFrameFreerunMode(uint64 FrameN
 		// Ack the skipped frames (that haven't been considered and therefore acked before)
 		if (LastSourceFrameNumberConsideredAtThisIdx < SenderMetadata[MemIdx].FrameNumber)
 		{
-			SharedMemoryData[MemIdx]->Receiver.FrameNumberAcked = SenderMetadata[MemIdx].FrameNumber;
+			check(ReceiverIndex[MemIdx] != INDEX_NONE); // We don't expect this function to be called if we didn't have a valid slot.
+
+			SharedMemoryData[MemIdx]->Receivers[ReceiverIndex[MemIdx]].FrameNumberAcked = SenderMetadata[MemIdx].FrameNumber;
 		}
 	}
 
@@ -762,7 +766,7 @@ void FSharedMemoryMediaPlayer::JustInTimeSampleRender()
 				FrameNumber = GFrameCounterRenderThread, 
 				SharedMemoryIdx, 
 				ExpectedFrameNumber, 
-				SharedMemoryData, 
+				SharedMemoryData,
 				this
 			](FRHICommandListImmediate& RHICmdList)
 			{
@@ -853,8 +857,14 @@ void FSharedMemoryMediaPlayer::JustInTimeSampleRender()
 
 	RunningTasksCount++;
 
-	UE::Tasks::Launch(UE_SOURCE_LOCATION,
-		[SharedMemoryData, ExpectedFrameNumber, SharedMemoryIdx, this]()
+	UE::Tasks::Launch(UE_SOURCE_LOCATION, 
+		[
+			SharedMemoryData, 
+			ExpectedFrameNumber, 
+			SharedMemoryIdx, 
+			ReceiverIdx = ReceiverIndex[SharedMemoryIdx], 
+			this
+	    ]()
 		{
 			TRACE_CPUPROFILER_EVENT_SCOPE(SharedMemoryMediaFrameAckTask);
 
@@ -882,7 +892,7 @@ void FSharedMemoryMediaPlayer::JustInTimeSampleRender()
 			bFrameAckFenceBusy[SharedMemoryIdx] = false;
 
 			// Ack the frame to the Media Capture, indicating that the cross gpu texture can be re-used.
-			SharedMemoryData->Receiver.FrameNumberAcked = ExpectedFrameNumber;
+			SharedMemoryData->Receivers[ReceiverIdx].FrameNumberAcked = ExpectedFrameNumber;
 		}
 	);
 }
diff --git a/Engine/Plugins/Runtime/nDisplay/Source/SharedMemoryMedia/Private/SharedMemoryMediaPlayer.h b/Engine/Plugins/Runtime/nDisplay/Source/SharedMemoryMedia/Private/SharedMemoryMediaPlayer.h
--- a/Engine/Plugins/Runtime/nDisplay/Source/SharedMemoryMedia/Private/SharedMemoryMediaPlayer.h
+++ b/Engine/Plugins/Runtime/nDisplay/Source/SharedMemoryMedia/Private/SharedMemoryMediaPlayer.h
@@ -232,18 +232,6 @@ protected:
 	/** DetermineNextSourceFrame implementration for Freerun Mode */
 	bool DetermineNextSourceFrameFreerunMode(uint64 FrameNumber, uint32& OutExpectedFrameNumber, uint32& OutSharedMemoryIdx);
 
-protected:
-
-	/** Used to make sure that players with the same UniqueName are not allowed to play at the same time, which is not supported.
-	 *  Only access from game thread.
-	 */
-	static TSet<FString> UniqueNamesRegistered;
-
-	/** Used by the player when it opens. The player should not play if the return value is false */
-	static bool RegisterUniqueName(FString& InUniqueName);
-
-	/** Used by the player when it closes. */
-	static void UnregisterUniqueName(FString& InUniqueName);
 
 protected:
 
@@ -256,6 +244,12 @@ protected:
 	/** Set of cross Gpu textures that the Media Capture allocates and the player opens */
 	FTextureRHIRef SharedCrossGpuTextures[NUMSHAREDMEM];
 
+	/** Receiver Index used to know what part of the shared metadata corresponds to this receiver */
+	int32 ReceiverIndex[NUMSHAREDMEM] = { 0 };
+
+	/** Unique number that identifies this receiver from any other */
+	FGuid ReceiverId = FGuid::NewGuid();
+
 	/** Shared memory allocated by the Media Capture for IPC purposes with the Media Player */
 	FPlatformMemory::FSharedMemoryRegion* SharedMemory[NUMSHAREDMEM] = { 0 };
 
diff --git a/Engine/Plugins/Runtime/nDisplay/Source/SharedMemoryMedia/Private/SharedMemoryMediaTypes.cpp b/Engine/Plugins/Runtime/nDisplay/Source/SharedMemoryMedia/Private/SharedMemoryMediaTypes.cpp
new file mode 100644
--- /dev/null
+++ b/Engine/Plugins/Runtime/nDisplay/Source/SharedMemoryMedia/Private/SharedMemoryMediaTypes.cpp
@@ -0,0 +1,151 @@
+// Copyright Epic Games, Inc. All Rights Reserved.
+
+#include "SharedMemoryMediaTypes.h"
+
+
+void FSharedMemoryMediaFrameMetadata::Initialize()
+{
+	for (int32 RxIdx = 0; RxIdx < MaxNumReceivers; ++RxIdx)
+	{
+		Receivers[RxIdx].Disconnect();
+	}
+}
+
+
+bool FSharedMemoryMediaFrameMetadata::IsSlotIndexValid(const int32 SlotIndex) const
+{
+	if (SlotIndex == INDEX_NONE)
+	{
+		return false;
+	}
+
+	check(SlotIndex >= 0 && SlotIndex < MaxNumReceivers);
+
+	return true;
+}
+
+
+bool FSharedMemoryMediaFrameMetadata::IsReceiverConnectedToSlot(const FGuid& Id, const int32 SlotIndex) const
+{
+	if (!IsSlotIndexValid(SlotIndex))
+	{
+		return false;
+	}
+
+	check(SlotIndex >= 0 && SlotIndex < MaxNumReceivers);
+
+	// KeepAliveCounter must not be zero and the id should be the owner of the slot
+	return (Receivers[SlotIndex].KeepAliveCounter && Receivers[SlotIndex].Id == Id);
+}
+
+
+int32 FSharedMemoryMediaFrameMetadata::ConnectToSlot(const FGuid& Id)
+{
+	for (uint32 RxIdx = 0; RxIdx < MaxNumReceivers; ++RxIdx)
+	{
+		uint8 ExpectedValue = 0;
+
+		const bool bReservedSlot = Receivers[RxIdx].KeepAliveCounter.compare_exchange_weak(ExpectedValue, KeepAliveCounterResetValue);
+
+		if (bReservedSlot)
+		{
+			Receivers[RxIdx].Id = Id;
+			return RxIdx;
+		}
+	}
+
+	return INDEX_NONE;
+}
+
+
+void FSharedMemoryMediaFrameMetadata::DisconnectReceiverFromSlot(const FGuid& Id, const int32 SlotIndex)
+{
+	// Check if it is already disconnected
+	if (!IsReceiverConnectedToSlot(Id, SlotIndex))
+	{
+		return;
+	}
+
+	Receivers[SlotIndex].Disconnect();
+}
+
+void FSharedMemoryMediaFrameMetadata::FReceiver::Disconnect()
+{
+	// Zero out the keep alive to effectively disconnect.
+	// 
+	// Note that there is a potential race condition, where after the connection verification
+	// above the KeepAlive got timed out and a new receiver with a different id connected,
+	// causing its premature disconnection.
+	KeepAliveCounter = 0;
+
+	// This avoids newly connected receivers from immediately stalling the sender.
+	FrameNumberAcked = ~0;
+}
+
+bool FSharedMemoryMediaFrameMetadata::KeepAlive(const int32 SlotIndex)
+{
+	if (!IsSlotIndexValid(SlotIndex))
+	{
+		return false;
+	}
+
+	// Only set the keepalive if it is not already zero. Because if it is already zero,
+	// it means that we got disconnected / timed out and need to find a slot again.
+
+	uint8 CurrentValue = Receivers[SlotIndex].KeepAliveCounter.load();
+
+	while (CurrentValue != 0) // loop until CurrentValue is zero or the exchange succeeds
+	{
+		// try to exchange KeepAliveCounter with KeepAliveCounterResetValue
+		if (Receivers[SlotIndex].KeepAliveCounter.compare_exchange_weak(CurrentValue, KeepAliveCounterResetValue))
+		{
+			// exchange succeeded
+			return true;
+		}
+
+		// Exchange failed, CurrentValue was updated with the actual value of KeepAliveCounter
+	}
+
+	// CurrentValue was zero, so we can't keep it alive anymore (it's dead).
+	return false;
+}
+
+
+bool FSharedMemoryMediaFrameMetadata::AllReceiversAckedFrameNumber(const uint32 FrameNumber) const
+{
+	for (int32 RxIdx = 0; RxIdx < MaxNumReceivers; ++RxIdx)
+	{
+		if (Receivers[RxIdx].KeepAliveCounter                     // Indicates an active receiver
+			&& (Receivers[RxIdx].FrameNumberAcked < FrameNumber)) // Indicates that the receiver hasn't released this frame number
+		{
+			return false;
+		}
+	}
+
+	return true;
+}
+
+
+void FSharedMemoryMediaFrameMetadata::DecrementKeepAlives()
+{
+	// Decrement the keep alive. The receiver must keep resetting it to avoid letting it expire.
+	for (int32 RxIdx = 0; RxIdx < MaxNumReceivers; ++RxIdx)
+	{
+		// We want to decrement the KeepAliveCounter, but making sure not to decrement it below zero, which would wrap around.
+		// So we first read its value, and if it is not zero then we try to decrement it. However if the value has changed when we try to 
+		// atomically write the decremented value, then we update the current value and try again. This is what the while loop does, and 
+		// how the compare_exchange_weak function behaves:
+		//
+		// compare_exchange_weak atomically compares the value of *this with that of expected. If those are bitwise-equal, replaces the former
+		// with desired (performs read-modify-write operation). Otherwise, loads the actual value stored in *this into expected (performs load operation).
+
+		uint8 CurrentValue = Receivers[RxIdx].KeepAliveCounter.load();
+
+		while (
+			CurrentValue // Current value must not be zero
+			&& !Receivers[RxIdx].KeepAliveCounter.compare_exchange_weak(CurrentValue, CurrentValue - 1) // Replace current value with its minus 1 value.
+			)
+		{
+		}
+	}
+}
diff --git a/Engine/Plugins/Runtime/nDisplay/Source/SharedMemoryMedia/Private/SharedMemoryMediaTypes.h b/Engine/Plugins/Runtime/nDisplay/Source/SharedMemoryMedia/Private/SharedMemoryMediaTypes.h
--- a/Engine/Plugins/Runtime/nDisplay/Source/SharedMemoryMedia/Private/SharedMemoryMediaTypes.h
+++ b/Engine/Plugins/Runtime/nDisplay/Source/SharedMemoryMedia/Private/SharedMemoryMediaTypes.h
@@ -33,7 +33,6 @@ namespace UE::SharedMemoryMedia
 
 	// Convenience ZeroGuid value.
 	static const FGuid ZeroGuid = FGuid(0, 0, 0, 0);
-
 }
 
 struct FSharedMemoryMediaTextureDescription
@@ -63,7 +62,14 @@ struct FSharedMemoryMediaTextureDescription
  */
 struct FSharedMemoryMediaFrameMetadata
 {
-	static constexpr uint32 MAGIC = 0x534D4D4D; // SMMM - Shared Memory Media Metadata
+	// SMMM - Shared Memory Media Metadata
+	static constexpr uint32 MAGIC = 0x534D4D4D; 
+
+	// Value that receivers should use when resetting the KeepAliveCounters. Corresponds to number of sender frames.
+	static constexpr uint8 KeepAliveCounterResetValue = 12; 
+
+	// Maximum number of concurrent receivers of the shared texture
+	static constexpr uint32 MaxNumReceivers = 4;
 
 	// Data that the sender wants to communicate
 	struct FSender
@@ -82,14 +88,50 @@ struct FSharedMemoryMediaFrameMetadata
 	struct FReceiver
 	{
 		// Frame number that the receiver is done reading. This signals the sender that it can re-use the texture.
-		uint32 FrameNumberAcked = 0;
+		std::atomic<uint32> FrameNumberAcked = ~0;
 
 		/** 
-		 * The receiver sets to ones this register. This signals the sender that it should wait for the related frame number ack before re-using the texture.
-		 * The sender will shift this register to the right each frame, as a way to time out the session in case the receiver disappears.
+		 * The receiver sets this register. This signals the sender that it should wait for the related frame number ack before re-using the texture.
+		 * The sender will decrement this register each frame, as a way to time out the session in case the receiver disappears without clearing it.
 		 */
-		uint8 KeepAliveShiftRegister = 0;
-	} Receiver;
+		std::atomic<uint8> KeepAliveCounter = 0;
+
+		/** Receiver id, used to detect collisions where two receivers think that they own the same slot */
+		FGuid Id;
+
+		// Disconnect from this slot. Meant to be called by the receiver.
+		void Disconnect();
+	};
+
+	// Array of possible receivers
+	FReceiver Receivers[MaxNumReceivers] = {};
+
+public:
+
+	// Initialize critical values
+	void Initialize();
+
+	// Returns true is the SlotId is valid
+	bool IsSlotIndexValid(const int32 SlotIndex) const;
+
+	// Checks if the receiver of the given Id owns the given slot
+	bool IsReceiverConnectedToSlot(const FGuid& Id, const int32 SlotIndex) const;
+
+	// Finds a slot for the give receiver. Returns the slot index, or INDEX_NONE if it could not find one.
+	int32 ConnectToSlot(const FGuid& Id);
+
+	// Give up control of the slot
+	void DisconnectReceiverFromSlot(const FGuid& Id, const int32 SlotIndex);
+
+	// Keeps the connection to the slot alive. True if succeeded.
+	bool KeepAlive(const int32 SlotIndex);
+
+	// Returns true only if all active receivers have acked the given frame number
+	bool AllReceiversAckedFrameNumber(const uint32 FrameNumber) const;
+
+	// Decrements all keepalives, which is intended to auto-disconnect stale receivers. 
+	// The receiver must keep resetting its keepalive to avoid letting it expire.
+	void DecrementKeepAlives();
 };
 
 
