nDisplay: Now logs the swap chain latency.

diff --git a/Engine/Plugins/Runtime/nDisplay/Source/DisplayCluster/Private/Render/Synchronization/Windows/DisplayClusterRenderSyncPolicyNvidiaWindows.cpp b/Engine/Plugins/Runtime/nDisplay/Source/DisplayCluster/Private/Render/Synchronization/Windows/DisplayClusterRenderSyncPolicyNvidiaWindows.cpp
--- a/Engine/Plugins/Runtime/nDisplay/Source/DisplayCluster/Private/Render/Synchronization/Windows/DisplayClusterRenderSyncPolicyNvidiaWindows.cpp
+++ b/Engine/Plugins/Runtime/nDisplay/Source/DisplayCluster/Private/Render/Synchronization/Windows/DisplayClusterRenderSyncPolicyNvidiaWindows.cpp
@@ -298,6 +298,16 @@ bool FDisplayClusterRenderSyncPolicyNvidia::InitializeNvidiaSwapLock()
 		UE_LOG(LogDisplayClusterRenderSync, Warning, TEXT("Couldn't set maximum frame latency"));
 	}
 
+	// Confirm frame latency
+	{
+		UINT MaxLatency = 0;
+
+		if (DXGISwapChain->GetMaximumFrameLatency(&MaxLatency) == S_OK)
+		{
+			UE_LOG(LogDisplayClusterRenderSync, Log, TEXT("Current maximum frame latency is %u"), MaxLatency);
+		}
+	}
+
 	NvU32 MaxGroups = 0;
 	NvU32 MaxBarriers = 0;
 
