nDisplay: Show error when getting swap chain frame latency.

diff --git a/Engine/Plugins/Runtime/nDisplay/Source/DisplayCluster/Private/Render/Synchronization/Windows/DisplayClusterRenderSyncPolicyNvidiaWindows.cpp b/Engine/Plugins/Runtime/nDisplay/Source/DisplayCluster/Private/Render/Synchronization/Windows/DisplayClusterRenderSyncPolicyNvidiaWindows.cpp
--- a/Engine/Plugins/Runtime/nDisplay/Source/DisplayCluster/Private/Render/Synchronization/Windows/DisplayClusterRenderSyncPolicyNvidiaWindows.cpp
+++ b/Engine/Plugins/Runtime/nDisplay/Source/DisplayCluster/Private/Render/Synchronization/Windows/DisplayClusterRenderSyncPolicyNvidiaWindows.cpp
@@ -293,18 +293,25 @@ bool FDisplayClusterRenderSyncPolicyNvidia::InitializeNvidiaSwapLock()
 	}
 
 	// Set frame latency
-	if (DXGISwapChain->SetMaximumFrameLatency(1) != S_OK)
 	{
-		UE_LOG(LogDisplayClusterRenderSync, Warning, TEXT("Couldn't set maximum frame latency"));
-	}
+		HRESULT Result = DXGISwapChain->SetMaximumFrameLatency(1);
 
-	// Confirm frame latency
-	{
-		UINT MaxLatency = 0;
+		if (Result != S_OK)
+		{
+			UE_LOG(LogDisplayClusterRenderSync, Warning, TEXT("Couldn't set maximum frame latency. Error: %x"), Result);
+		}
 
-		if (DXGISwapChain->GetMaximumFrameLatency(&MaxLatency) == S_OK)
+		UINT CurrentLatency = 0;
+
+		Result = DXGISwapChain->GetMaximumFrameLatency(&CurrentLatency);
+
+		if (Result == S_OK)
+		{
+			UE_LOG(LogDisplayClusterRenderSync, Log, TEXT("Swapchain frame latency: %u"), CurrentLatency);
+		}
+		else
 		{
-			UE_LOG(LogDisplayClusterRenderSync, Log, TEXT("Current maximum frame latency is %u"), MaxLatency);
+			UE_LOG(LogDisplayClusterRenderSync, Warning, TEXT("Couldn't get maximum frame latency. Error: %x"), Result);
 		}
 	}
 
