nDisplay: Added a mesh component to the DCRA that stores a 3D version of the stage's geometry map that line traces can be run against to compute distances to the stage's geometry from any point. This mesh can also be made visibile in the editor to allow debugging of the stage's generated geometry map

diff --git a/Engine/Plugins/Runtime/nDisplay/Source/DisplayCluster/Private/Game/EngineClasses/Scene/DisplayClusterRootActor.cpp b/Engine/Plugins/Runtime/nDisplay/Source/DisplayCluster/Private/Game/EngineClasses/Scene/DisplayClusterRootActor.cpp
--- a/Engine/Plugins/Runtime/nDisplay/Source/DisplayCluster/Private/Game/EngineClasses/Scene/DisplayClusterRootActor.cpp
+++ b/Engine/Plugins/Runtime/nDisplay/Source/DisplayCluster/Private/Game/EngineClasses/Scene/DisplayClusterRootActor.cpp
@@ -49,6 +49,7 @@
 #include "Components/DisplayClusterChromakeyCardStageActorComponent.h"
 #include "Components/DisplayClusterStageActorComponent.h"
 #include "Components/DisplayClusterStageGeometryComponent.h"
+#include "Components/DisplayClusterStageIsosphereComponent.h"
 #include "Components/LineBatchComponent.h"
 #include "UObject/Package.h"
 
@@ -114,6 +115,9 @@ ADisplayClusterRootActor::ADisplayClusterRootActor(const FObjectInitializer& Obj
 
 	StageGeometryComponent = CreateDefaultSubobject<UDisplayClusterStageGeometryComponent>(TEXT("DisplayClusterStageGeometry"));
 
+	StageIsosphereComponent = CreateDefaultSubobject<UDisplayClusterStageIsosphereComponent>(TEXT("DisplayClusterStageIsosphere"));
+	StageIsosphereComponent->AttachToComponent(GetRootComponent(), FAttachmentTransformRules::KeepRelativeTransform);
+
 	// Default nDisplay camera
 	DefaultViewPoint = CreateDefaultSubobject<UDisplayClusterCameraComponent>(TEXT("DefaultViewPoint"));
 	DefaultViewPoint->AttachToComponent(GetRootComponent(), FAttachmentTransformRules::KeepRelativeTransform);
@@ -946,6 +950,32 @@ bool ADisplayClusterRootActor::MakeStageActorFlushToWall(const TScriptInterface<
 	return false;
 }
 
+bool ADisplayClusterRootActor::GetDistanceToStageGeometry(const FVector& WorldPosition, const FVector& WorldDirection, float& OutDistance) const
+{
+	TRACE_CPUPROFILER_EVENT_SCOPE(ADisplayClusterRootActor::GetDistanceToStageGeometry);
+
+	if (StageGeometryComponent->IsGeometryMapValid())
+	{
+		FHitResult HitResult;
+
+		// Trace far enough to ensure the line trace will actually intersect any side of the stage isosphere from any world position
+		const float TraceDistance = FVector::Distance(WorldPosition, GetActorLocation()) + StageGeometryComponent->GetStageBoundingRadius();
+		const FVector TraceEnd = WorldPosition + WorldDirection.GetSafeNormal() * TraceDistance;
+
+		FCollisionQueryParams Params(SCENE_QUERY_STAT(DisplayClusterStageTrace), true);
+		if (StageIsosphereComponent->LineTraceComponent(HitResult, WorldPosition, TraceEnd, Params))
+		{
+			OutDistance = HitResult.Distance;
+			return true;
+		}
+
+		OutDistance = 0;
+		return false;
+	}
+
+	return false;
+}
+
 bool ADisplayClusterRootActor::IsBlueprint() const
 {
 	for (UClass* Class = GetClass(); Class; Class = Class->GetSuperClass())
diff --git a/Engine/Plugins/Runtime/nDisplay/Source/DisplayCluster/Private/Game/EngineClasses/Scene/DisplayClusterRootActorEditor.cpp b/Engine/Plugins/Runtime/nDisplay/Source/DisplayCluster/Private/Game/EngineClasses/Scene/DisplayClusterRootActorEditor.cpp
--- a/Engine/Plugins/Runtime/nDisplay/Source/DisplayCluster/Private/Game/EngineClasses/Scene/DisplayClusterRootActorEditor.cpp
+++ b/Engine/Plugins/Runtime/nDisplay/Source/DisplayCluster/Private/Game/EngineClasses/Scene/DisplayClusterRootActorEditor.cpp
@@ -44,6 +44,7 @@
 #include "Render/Viewport/DisplayClusterViewportManager.h"
 #include "Render/Viewport/IDisplayClusterViewport.h"
 #include "TextureResource.h"
+#include "Components/DisplayClusterStageIsosphereComponent.h"
 
 //////////////////////////////////////////////////////////////////////////////////////////////
 // IN-EDITOR STUFF
@@ -343,8 +344,6 @@ void ADisplayClusterRootActor::RerunConstructionScripts_Editor()
 	InitializeRootActor();
 
 	UpdateInnerFrustumPriority();
-
-	StageGeometryComponent->Invalidate();
 }
 
 void ADisplayClusterRootActor::UpdateInnerFrustumPriority()
@@ -549,6 +548,10 @@ void ADisplayClusterRootActor::PostEditChangeProperty(FPropertyChangedEvent& Pro
 		ResetEntireClusterPreviewRendering();
 		bReinitializeActor = false;
 	}
+	else if (PropertyName == GET_MEMBER_NAME_CHECKED(ADisplayClusterRootActor, bPreviewStageGeometryMesh))
+	{
+		StageIsosphereComponent->SetVisibility(bPreviewStageGeometryMesh);
+	}
 
 	if (bReinitializeActor)
 	{
@@ -581,9 +584,26 @@ void ADisplayClusterRootActor::HandleAssetReload(const EPackageReloadPhase InPac
 void ADisplayClusterRootActor::OnEndObjectMovement(UObject& InObject)
 {
 	// If any of this stage actor's components have been moved, invalidate the stage geometry map
-	if (InObject.IsA<USceneComponent>() && Cast<USceneComponent>(&InObject)->GetOwner() == this)
+	if (USceneComponent* SceneComponent = Cast<USceneComponent>(&InObject))
 	{
-		StageGeometryComponent->Invalidate();
+		if (SceneComponent->GetOwner() == this)
+		{
+			// Check to see if the object being moved is a part of the stage's geometry or a view origin
+			TArray<FString> ProjectionMeshNames;
+			if (UDisplayClusterConfigurationData* Config = GetConfigData())
+			{
+				Config->GetReferencedMeshNames(ProjectionMeshNames);
+			}
+
+			const bool bIsProjectionMesh = SceneComponent->IsA<UStaticMeshComponent>() && ProjectionMeshNames.Contains(SceneComponent->GetName());
+			const bool bIsScreen = SceneComponent->IsA<UDisplayClusterScreenComponent>();
+			const bool bIsViewOrigin = SceneComponent->IsA<UDisplayClusterCameraComponent>();
+
+			if (bIsProjectionMesh || bIsScreen || bIsViewOrigin)
+			{
+				StageGeometryComponent->Invalidate();
+			}
+		}
 	}
 }
 
diff --git a/Engine/Plugins/Runtime/nDisplay/Source/DisplayCluster/Private/Game/EngineClasses/Scene/DisplayClusterStageGeometryComponent.cpp b/Engine/Plugins/Runtime/nDisplay/Source/DisplayCluster/Private/Game/EngineClasses/Scene/DisplayClusterStageGeometryComponent.cpp
--- a/Engine/Plugins/Runtime/nDisplay/Source/DisplayCluster/Private/Game/EngineClasses/Scene/DisplayClusterStageGeometryComponent.cpp
+++ b/Engine/Plugins/Runtime/nDisplay/Source/DisplayCluster/Private/Game/EngineClasses/Scene/DisplayClusterStageGeometryComponent.cpp
@@ -9,6 +9,7 @@
 #include "ProceduralMeshComponent.h"
 #include "TextureResource.h"
 #include "Components/DisplayClusterScreenComponent.h"
+#include "Components/DisplayClusterStageIsosphereComponent.h"
 #include "Engine/TextureRenderTarget2D.h"
 
 const uint32 UDisplayClusterStageGeometryComponent::GeometryMapSize = 512;
@@ -29,6 +30,7 @@ void UDisplayClusterStageGeometryComponent::TickComponent(float DeltaTime, ELeve
 	if (bUpdateGeometryMap)
 	{
 		RedrawGeometryMap();
+		UpdateStageIsosphere();
 	}
 }
 
@@ -37,6 +39,7 @@ void UDisplayClusterStageGeometryComponent::Invalidate(bool bForceImmediateRedra
 	if (bForceImmediateRedraw)
 	{
 		RedrawGeometryMap();
+		UpdateStageIsosphere();
 	}
 	else
 	{
@@ -102,36 +105,58 @@ bool UDisplayClusterStageGeometryComponent::GetStageDistanceAndNormal(const FVec
 	return false;
 }
 
-bool UDisplayClusterStageGeometryComponent::MorphProceduralMesh(UProceduralMeshComponent* InProceduralMeshComponent)
+bool UDisplayClusterStageGeometryComponent::MorphProceduralMesh(UProceduralMeshComponent* InProceduralMeshComponent, bool bSyncMeshLocation)
 {
+	TRACE_CPUPROFILER_EVENT_SCOPE(UDisplayClusterStageGeometryComponent::MorphProceduralMesh);
+
 	if (bGeometryMapLoaded)
 	{
-		FProcMeshSection& Section = *InProceduralMeshComponent->GetProcMeshSection(0);
-
-		for (int32 Index = 0; Index < Section.ProcVertexBuffer.Num(); ++Index)
+		if (USceneComponent* CommonViewPoint = GetCommonViewPoint())
 		{
-			FProcMeshVertex& Vertex = Section.ProcVertexBuffer[Index];
+			if (bSyncMeshLocation)
+			{
+				InProceduralMeshComponent->SetWorldLocation(CommonViewPoint->GetComponentLocation());
+			}
 
-			const FVector VertexDirection = Vertex.Position.GetSafeNormal();
-			float Distance;
-			FVector Normal;
-			GetStageDistanceAndNormal(VertexDirection, Distance, Normal);
+			if (FProcMeshSection* Section = InProceduralMeshComponent->GetProcMeshSection(0))
+			{
+				for (int32 Index = 0; Index < Section->ProcVertexBuffer.Num(); ++Index)
+				{
+					FProcMeshVertex& Vertex = Section->ProcVertexBuffer[Index];
 
-			const FVector NewPosition = VertexDirection * Distance;
-			Vertex.Position = NewPosition;
+					const FVector VertexDirection = Vertex.Position.GetSafeNormal();
+					float Distance;
+					FVector Normal;
+					GetStageDistanceAndNormal(VertexDirection, Distance, Normal);
 
-			const FMatrix RadialBasis = FRotationMatrix::MakeFromX(VertexDirection);
+					Vertex.Position = VertexDirection * Distance;
 
-			const FVector WorldNormal = RadialBasis.TransformVector(Normal);
-			Vertex.Normal = WorldNormal;
-		}
+					const FMatrix RadialBasis = FRotationMatrix::MakeFromX(VertexDirection);
 
-		InProceduralMeshComponent->SetProcMeshSection(0, Section);
+					const FVector WorldNormal = RadialBasis.TransformVector(Normal);
+					Vertex.Normal = WorldNormal;
+				}
+
+				InProceduralMeshComponent->SetProcMeshSection(0, *Section);
+
+				return true;
+			}
+		}
 	}
 
 	return false;
 }
 
+USceneComponent* UDisplayClusterStageGeometryComponent::GetCommonViewPoint() const
+{
+	if (ADisplayClusterRootActor* RootActor = Cast<ADisplayClusterRootActor>(GetOwner()))
+	{
+		return RootActor->GetCommonViewPoint();
+	}
+
+	return nullptr;
+}
+
 UTextureRenderTarget2D* UDisplayClusterStageGeometryComponent::CreateRenderTarget()
 {
 	UTextureRenderTarget2D* RenderTarget = NewObject<UTextureRenderTarget2D>(this);
@@ -201,6 +226,18 @@ void UDisplayClusterStageGeometryComponent::UpdateStageGeometry()
 	}
 }
 
+void UDisplayClusterStageGeometryComponent::UpdateStageIsosphere()
+{
+	if (AActor* RootActor = GetOwner())
+	{
+		if (UDisplayClusterStageIsosphereComponent* IsosphereComponent = RootActor->FindComponentByClass<UDisplayClusterStageIsosphereComponent>())
+		{
+			IsosphereComponent->ResetIsosphere();
+			MorphProceduralMesh(IsosphereComponent, true);
+		}
+	}
+}
+
 void UDisplayClusterStageGeometryComponent::GenerateGeometryMap(bool bIsNorthMap)
 {
 	if (ADisplayClusterRootActor* RootActor = Cast<ADisplayClusterRootActor>(GetOwner()))
diff --git a/Engine/Plugins/Runtime/nDisplay/Source/DisplayCluster/Private/Game/EngineClasses/Scene/DisplayClusterStageIsosphereComponent.cpp b/Engine/Plugins/Runtime/nDisplay/Source/DisplayCluster/Private/Game/EngineClasses/Scene/DisplayClusterStageIsosphereComponent.cpp
new file mode 100644
--- /dev/null
+++ b/Engine/Plugins/Runtime/nDisplay/Source/DisplayCluster/Private/Game/EngineClasses/Scene/DisplayClusterStageIsosphereComponent.cpp
@@ -0,0 +1,83 @@
+// Copyright Epic Games, Inc. All Rights Reserved.
+
+#include "Components/DisplayClusterStageIsosphereComponent.h"
+
+#include "Engine/StaticMesh.h"
+#include "KismetProceduralMeshLibrary.h"
+#include "UObject/ConstructorHelpers.h"
+
+UDisplayClusterStageIsosphereComponent::UDisplayClusterStageIsosphereComponent(const FObjectInitializer& ObjectInitializer)
+	: Super(ObjectInitializer)
+{
+	ConstructorHelpers::FObjectFinder<UStaticMesh>  IsosphereMeshFinder(TEXT("/nDisplay/Meshes/SM_IcoSphere.SM_IcoSphere"));
+
+	IsosphereMesh = IsosphereMeshFinder.Object;
+
+	SetHiddenInGame(true);
+	SetVisibility(false);
+}
+
+void UDisplayClusterStageIsosphereComponent::OnComponentCreated()
+{
+	Super::OnComponentCreated();
+
+	// Load the isosphere's geometry into the procedural mesh
+	if (ensure(IsosphereMesh))
+	{
+		IsosphereMesh->bAllowCPUAccess = true;
+
+		const int32 NumSections = IsosphereMesh->GetNumSections(0);
+		for (int32 SectionIndex = 0; SectionIndex < NumSections; SectionIndex++)
+		{
+			TArray<FVector> Vertices;
+			TArray<int32> Triangles;
+			TArray<FVector> Normals;
+			TArray<FVector2D> UVs;
+			TArray<FProcMeshTangent> Tangents;
+
+			UKismetProceduralMeshLibrary::GetSectionFromStaticMesh(IsosphereMesh, 0, SectionIndex, Vertices, Triangles, Normals, UVs, Tangents);
+
+			TArray<FVector2D> EmptyUVs;
+			TArray<FLinearColor> EmptyColors;
+			CreateMeshSection_LinearColor(SectionIndex, Vertices, Triangles, Normals, UVs, EmptyUVs, EmptyUVs, EmptyUVs, EmptyColors, Tangents, true);
+
+			for (int32 Index = 0; Index < IsosphereMesh->GetStaticMaterials().Num(); ++Index)
+			{
+				if (UMaterialInterface* MaterialInterface = IsosphereMesh->GetStaticMaterials()[Index].MaterialInterface)
+				{
+					SetMaterial(Index, MaterialInterface);
+				}
+			}
+		}
+	}
+}
+
+void UDisplayClusterStageIsosphereComponent::ResetIsosphere()
+{
+	if (ensure(IsosphereMesh))
+	{
+		const int32 NumSections = IsosphereMesh->GetNumSections(0);
+		for (int32 SectionIndex = 0; SectionIndex < NumSections; SectionIndex++)
+		{
+			TArray<FVector> Vertices;
+			TArray<int32> Triangles;
+			TArray<FVector> Normals;
+			TArray<FVector2D> UVs;
+			TArray<FProcMeshTangent> Tangents;
+
+			UKismetProceduralMeshLibrary::GetSectionFromStaticMesh(IsosphereMesh, 0, SectionIndex, Vertices, Triangles, Normals, UVs, Tangents);
+
+			if (FProcMeshSection* Section = GetProcMeshSection(SectionIndex))
+			{
+				for (int32 Index = 0; Index < Vertices.Num(); ++Index)
+				{
+					FProcMeshVertex& Vertex = Section->ProcVertexBuffer[Index];
+					Vertex.Position = Vertices[Index];
+					Vertex.Normal = Normals[Index];
+				}
+
+				SetProcMeshSection(0, *Section);
+			}
+		}
+	}
+}
diff --git a/Engine/Plugins/Runtime/nDisplay/Source/DisplayCluster/Public/Components/DisplayClusterStageGeometryComponent.h b/Engine/Plugins/Runtime/nDisplay/Source/DisplayCluster/Public/Components/DisplayClusterStageGeometryComponent.h
--- a/Engine/Plugins/Runtime/nDisplay/Source/DisplayCluster/Public/Components/DisplayClusterStageGeometryComponent.h
+++ b/Engine/Plugins/Runtime/nDisplay/Source/DisplayCluster/Public/Components/DisplayClusterStageGeometryComponent.h
@@ -50,6 +50,10 @@ public:
 	virtual void TickComponent(float DeltaTime, ELevelTick TickType, FActorComponentTickFunction* ThisTickFunction) override;
 	//~ End UObject interface
 
+	/** Gets whether the stage geometry map is valid and usable */
+	UFUNCTION(BlueprintCallable, Category = "NDisplay")
+	bool IsGeometryMapValid() const { return bGeometryMapLoaded; }
+
 	/**
 	 * Invalidates and regenerates the cached stage geometry map
 	 * @param bForceImmediateRedraw Indicates whether the geometry map is regenerated immediately or is queued to redraw on the next tick
@@ -75,10 +79,15 @@ public:
 
 	/**
 	 * Morphs the specified procedural mesh to match the stage's geometry map
+	 * @param InProceduralMeshComponent - The mesh component to morph
+	 * @param bSyncMeshLocation - When true, the procedural mesh's world location will be synced to the view origin's world location
 	 */
-	bool MorphProceduralMesh(UProceduralMeshComponent* InProceduralMeshComponent);
+	bool MorphProceduralMesh(UProceduralMeshComponent* InProceduralMeshComponent, bool bSyncMeshLocation = false);
 
 private:
+	/** Gets the parent stage's common view point, or null if none were found */
+	USceneComponent* GetCommonViewPoint() const;
+
 	/** Creates a new render target that can be used to render the stage's geometry map */
 	UTextureRenderTarget2D* CreateRenderTarget();
 
@@ -88,6 +97,9 @@ private:
 	/** Updates the stage's current geometry with the renderer */
 	void UpdateStageGeometry();
 
+	/** Updates any isosphere that is found on the stage to match the geometry map */
+	void UpdateStageIsosphere();
+
 	/** Renders the stage's geometry to the specified geometry map */
 	void GenerateGeometryMap(bool bIsNorthMap);
 
diff --git a/Engine/Plugins/Runtime/nDisplay/Source/DisplayCluster/Public/Components/DisplayClusterStageIsosphereComponent.h b/Engine/Plugins/Runtime/nDisplay/Source/DisplayCluster/Public/Components/DisplayClusterStageIsosphereComponent.h
new file mode 100644
--- /dev/null
+++ b/Engine/Plugins/Runtime/nDisplay/Source/DisplayCluster/Public/Components/DisplayClusterStageIsosphereComponent.h
@@ -0,0 +1,30 @@
+// Copyright Epic Games, Inc. All Rights Reserved.
+
+#pragma once
+
+#include "CoreMinimal.h"
+#include "ProceduralMeshComponent.h"
+
+#include "DisplayClusterStageIsosphereComponent.generated.h"
+
+/** A procedural mesh component generated by morphing an isosphere to match the stage's geometry map, providing a 3D version of the map that can be ray traced against */
+UCLASS()
+class DISPLAYCLUSTER_API UDisplayClusterStageIsosphereComponent : public UProceduralMeshComponent
+{
+	GENERATED_BODY()
+
+public:
+	UDisplayClusterStageIsosphereComponent(const FObjectInitializer& ObjectInitializer);
+
+	virtual void OnComponentCreated() override;
+
+	/**
+	 * Resets the vertices of the mesh component to be an isosphere
+	 */
+	void ResetIsosphere();
+
+private:
+	/** A static mesh containing an isosphere that can be morphed */
+	UPROPERTY(Transient)
+	TObjectPtr<UStaticMesh> IsosphereMesh;
+};
\ No newline at end of file
diff --git a/Engine/Plugins/Runtime/nDisplay/Source/DisplayCluster/Public/DisplayClusterRootActor.h b/Engine/Plugins/Runtime/nDisplay/Source/DisplayCluster/Public/DisplayClusterRootActor.h
--- a/Engine/Plugins/Runtime/nDisplay/Source/DisplayCluster/Public/DisplayClusterRootActor.h
+++ b/Engine/Plugins/Runtime/nDisplay/Source/DisplayCluster/Public/DisplayClusterRootActor.h
@@ -40,6 +40,7 @@ class UDisplayClusterConfigurationData;
 class UDisplayClusterCameraComponent;
 class UDisplayClusterOriginComponent;
 class UDisplayClusterStageGeometryComponent;
+class UDisplayClusterStageIsosphereComponent;
 class UDisplayClusterSyncTickComponent;
 class UProceduralMeshComponent;
 
@@ -163,6 +164,16 @@ public:
 	UFUNCTION(BlueprintCallable, Category = "NDisplay")
 	bool MakeStageActorFlushToWall(const TScriptInterface<IDisplayClusterStageActor>& StageActor, double DesiredOffsetFromFlush = 0.0f);
 
+	/**
+	 * Gets the distance from a world position to the stage's geometry along the specified direction, if there is an intersection
+	 * @param WorldPosition - The world position to measure the distance from
+	 * @param WorldDirection - The direction to find the distance to the geometry along
+	 * @param OutDistance - The distance to the stage geometry from the specified point
+	 * @return True if an intersection point from WorldPosition along WorldDirection was found, false if not
+	 */
+	UFUNCTION(BlueprintCallable, Category = "NDisplay")
+	bool GetDistanceToStageGeometry(const FVector& WorldPosition, const FVector& WorldDirection, float& OutDistance) const;
+
 	UFUNCTION(BlueprintGetter)
 	UDisplayClusterStageGeometryComponent* GetStageGeometryComponent() const { return StageGeometryComponent; }
 
@@ -364,7 +375,11 @@ private:
 	/** Component that stores the stage's geometry map, which is used to make objects flush with the stage's walls and ceilings */
 	UPROPERTY()
 	TObjectPtr<UDisplayClusterStageGeometryComponent> StageGeometryComponent;
-	
+
+	/** Component that stores a 3D representation of the stage's geometry map, which can be used to perform ray traces against the processed stage geometry */
+	UPROPERTY()
+	TObjectPtr<UDisplayClusterStageIsosphereComponent> StageIsosphereComponent;
+
 private:
 	// Current operation mode
 	EDisplayClusterOperationMode OperationMode;
@@ -497,6 +512,13 @@ public:
 	UPROPERTY(EditDefaultsOnly, Category = "Editor Preview", DisplayName = "Default Display Device")
 	FName DefaultDisplayDeviceName;
 
+
+#if WITH_EDITORONLY_DATA
+	/** Toggles the visibility of the stage's geometry mesh, a smooth, continuous mesh generated and processed from the stage's geometry */
+	UPROPERTY(EditInstanceOnly, Category = "Editor Preview", AdvancedDisplay)
+	bool bPreviewStageGeometryMesh = false;
+#endif
+
 protected:
 	/** The default display device to use for preview rendering */
 	UPROPERTY(Transient, NonTransactional)
