nDisplay: Reimplementing fix made for UE-201200 because it produces now another issue (UE-201985)

diff --git a/Engine/Plugins/Runtime/nDisplay/Source/DisplayCluster/Private/Render/Viewport/Preview/DisplayClusterViewportPreviewMesh.cpp b/Engine/Plugins/Runtime/nDisplay/Source/DisplayCluster/Private/Render/Viewport/Preview/DisplayClusterViewportPreviewMesh.cpp
--- a/Engine/Plugins/Runtime/nDisplay/Source/DisplayCluster/Private/Render/Viewport/Preview/DisplayClusterViewportPreviewMesh.cpp
+++ b/Engine/Plugins/Runtime/nDisplay/Source/DisplayCluster/Private/Render/Viewport/Preview/DisplayClusterViewportPreviewMesh.cpp
@@ -14,14 +14,17 @@ namespace UE::DisplayCluster::ViewportPreviewMesh
 	/** Reset the mesh material to default values from its archetype. */
 	static inline bool RestoreMeshMaterialsFromArchetype(UMeshComponent* InMeshComponent)
 	{
-		if (const UMeshComponent* MeshArchetype = Cast<UMeshComponent>(InMeshComponent->GetArchetype()))
+		if (IsValid(InMeshComponent))
 		{
-			// Retrieve material on the preview mesh from an archetype or OrigMaterial.
-			if (UMaterialInterface* OrigMaterial = MeshArchetype->OverrideMaterials.IsValidIndex(0) ? MeshArchetype->OverrideMaterials[0] : nullptr)
+			if (const UMeshComponent* MeshArchetype = Cast<UMeshComponent>(InMeshComponent->GetArchetype()))
 			{
-				InMeshComponent->SetMaterial(0, OrigMaterial);
+				// Retrieve material on the preview mesh from an archetype or OrigMaterial.
+				if (UMaterialInterface* OrigMaterial = MeshArchetype->OverrideMaterials.IsValidIndex(0) ? MeshArchetype->OverrideMaterials[0] : nullptr)
+				{
+					InMeshComponent->SetMaterial(0, OrigMaterial);
 
-				return true;
+					return true;
+				}
 			}
 		}
 
@@ -34,12 +37,12 @@ namespace UE::DisplayCluster::ViewportPreviewMesh
 		TObjectPtr<UMaterial> OutMaterial = nullptr;
 
 		// First get the material from the ViewPoint component (WarpPolicy)
-		OutMaterial = ViewPointComponent ? ViewPointComponent->GetDisplayDeviceMaterial(InMeshType, InMaterialType) : nullptr;
+		OutMaterial = IsValid(ViewPointComponent) ? ViewPointComponent->GetDisplayDeviceMaterial(InMeshType, InMaterialType) : nullptr;
 
 		// Finally, get the material from the DisplayDevice
 		if (!OutMaterial)
 		{
-			OutMaterial = InDisplayDeviceComponent ? InDisplayDeviceComponent->GetDisplayDeviceMaterial(InMeshType, InMaterialType) : nullptr;
+			OutMaterial = IsValid(InDisplayDeviceComponent) ? InDisplayDeviceComponent->GetDisplayDeviceMaterial(InMeshType, InMaterialType) : nullptr;
 		}
 
 		// Ignore deleted materials
@@ -113,7 +116,7 @@ void FDisplayClusterViewportPreviewMesh::Update(FDisplayClusterViewport* InViewp
 	if (UMeshComponent* MeshComponent = GetMeshComponent())
 	{
 		// Update material instance and assign to the  mesh
-		if (GetMaterialInstance() == nullptr)
+		if (!CurrentMaterialPtr.IsValid())
 		{
 			CurrentMaterialPtr = InMeshMaterial;
 
@@ -122,33 +125,20 @@ void FDisplayClusterViewportPreviewMesh::Update(FDisplayClusterViewport* InViewp
 		}
 
 		// The material must be assigned to the mesh at each tick, without any conditions in case it can be changed externally.
-		MeshComponent->SetMaterial(0, GetMaterialInstance());
+		MeshComponent->SetMaterial(0, MaterialInstancePtr.Get());
 	}
 
 	// Handling the material overlay logic for the preview mesh:
 	UpdateOverlayMaterial(InViewport);
 }
 
-void FDisplayClusterViewportPreviewMesh::AddReferencedObjects(FReferenceCollector& Collector)
-{
-	TObjectPtr<UObject> OwnedObjects[] = { MeshComponentPtr, OrigOverlayMaterial, MaterialInstancePtr, CurrentMaterialPtr, DefaultMaterialPtr };
-
-	for (TObjectPtr<UObject>& OwnedObject : OwnedObjects)
-	{
-		if (OwnedObject)
-		{
-			Collector.AddReferencedObject(OwnedObject);
-		}
-	}
-}
-
 void FDisplayClusterViewportPreviewMesh::SetCustomOverlayMaterial(UMeshComponent* InMeshComponent, UMaterialInterface* InOverlayMaterial)
 {
 	if (InMeshComponent)
 	{
-		if (OrigOverlayMaterial == nullptr)
+		if (!OrigOverlayMaterialPtr.IsValid())
 		{
-			OrigOverlayMaterial = InMeshComponent->GetOverlayMaterial();
+			OrigOverlayMaterialPtr = InMeshComponent->GetOverlayMaterial();
 		}
 
 		InMeshComponent->SetOverlayMaterial(InOverlayMaterial);
@@ -159,7 +149,7 @@ void FDisplayClusterViewportPreviewMesh::RestoreOverlayMaterial(UMeshComponent*
 {
 	if (InMeshComponent)
 	{
-		if (OrigOverlayMaterial)
+		if (UMaterialInterface* const OrigOverlayMaterial = OrigOverlayMaterialPtr.Get())
 		{
 			InMeshComponent->SetOverlayMaterial(OrigOverlayMaterial);
 		}
@@ -200,12 +190,8 @@ void FDisplayClusterViewportPreviewMesh::ReleaseMeshComponent(FDisplayClusterVie
 	using namespace UE::DisplayCluster::ViewportPreviewMesh;
 
 	UMeshComponent* MeshComponent = GetMeshComponent();
-	if (!MeshComponent && MeshComponentPtr)
+	if (InViewport && !MeshComponent)
 	{
-		// Restore materials on deleted mesh
-		RestoreOverlayMaterial(MeshComponentPtr);
-		RestoreMeshMaterialsFromArchetype(MeshComponentPtr);
-
 		// The mesh was destroyed earlier, (re-running build scripts inside RootActor), but we need to update the new mesh component to.
 		MeshComponent = GetOrCreatePreviewMeshComponent(InViewport, bIsRootActorMeshComponent);
 	}
@@ -214,11 +200,15 @@ void FDisplayClusterViewportPreviewMesh::ReleaseMeshComponent(FDisplayClusterVie
 	{
 		EnumAddFlags(RuntimeFlags, EDisplayClusterViewportPreviewMeshFlags::HasDeletedMeshComponent);
 
+		UMaterial* const DefaultMaterial = DefaultMaterialPtr.Get();
+
 		// Restore materials on exists mesh
 		RestoreOverlayMaterial(MeshComponent);
-		if (!RestoreMeshMaterialsFromArchetype(MeshComponent) && DefaultMaterialPtr)
+
+		const bool bRestored = RestoreMeshMaterialsFromArchetype(MeshComponent);
+		if (!bRestored && DefaultMaterial)
 		{
-			MeshComponent->SetMaterial(0, DefaultMaterialPtr);
+			MeshComponent->SetMaterial(0, DefaultMaterial);
 		}
 
 		if (!bIsRootActorMeshComponent)
@@ -234,26 +224,20 @@ void FDisplayClusterViewportPreviewMesh::ReleaseMeshComponent(FDisplayClusterVie
 		}
 	}
 
-	MeshComponentPtr = nullptr;
+	MeshComponentPtr.Reset();
 }
 
 void FDisplayClusterViewportPreviewMesh::ReleaseMaterialInstance()
 {
 	// The material instance references the mesh, so it must also be deleted
-	if (UMaterialInstanceDynamic* MaterialInstance = GetMaterialInstance())
+	if (UMaterialInstanceDynamic* MaterialInstance = MaterialInstancePtr.Get())
 	{
 		EnumAddFlags(RuntimeFlags, EDisplayClusterViewportPreviewMeshFlags::HasDeletedMaterialInstance);
-
-		// Clear the material parameters to ensure that no pointers to the texture resources are being kept around
-		// Materials destroy their resources in BeginDestroy, so only clear the parameters BeginDestroy hasn't already been called
-		if (!MaterialInstance->HasAnyFlags(RF_BeginDestroyed))
-		{
-			MaterialInstance->ClearParameterValues();
-		}
+		MaterialInstance->ClearParameterValues();
 	}
 
-	MaterialInstancePtr = nullptr;
-	CurrentMaterialPtr = nullptr;
+	MaterialInstancePtr.Reset();
+	CurrentMaterialPtr.Reset();
 }
 
 bool FDisplayClusterViewportPreviewMesh::ShouldUseMeshComponent(FDisplayClusterViewport* InViewport) const
@@ -308,10 +292,10 @@ UMeshComponent* FDisplayClusterViewportPreviewMesh::GetOrCreatePreviewMeshCompon
 
 UMeshComponent* FDisplayClusterViewportPreviewMesh::GetMeshComponent() const
 {
-	return MeshComponentPtr;
+	return MeshComponentPtr.Get();
 }
 
 UMaterialInstanceDynamic* FDisplayClusterViewportPreviewMesh::GetMaterialInstance() const
 {
-	return MaterialInstancePtr;
+	return MaterialInstancePtr.Get();
 }
diff --git a/Engine/Plugins/Runtime/nDisplay/Source/DisplayCluster/Private/Render/Viewport/Preview/DisplayClusterViewportPreviewMesh.h b/Engine/Plugins/Runtime/nDisplay/Source/DisplayCluster/Private/Render/Viewport/Preview/DisplayClusterViewportPreviewMesh.h
--- a/Engine/Plugins/Runtime/nDisplay/Source/DisplayCluster/Private/Render/Viewport/Preview/DisplayClusterViewportPreviewMesh.h
+++ b/Engine/Plugins/Runtime/nDisplay/Source/DisplayCluster/Private/Render/Viewport/Preview/DisplayClusterViewportPreviewMesh.h
@@ -11,8 +11,6 @@
 
 #include "Misc/DisplayClusterObjectRef.h"
 
-#include "UObject/GCObject.h"
-
 
 class UDisplayClusterDisplayDeviceBaseComponent;
 class UDisplayClusterCameraComponent;
@@ -46,7 +44,6 @@ ENUM_CLASS_FLAGS(EDisplayClusterViewportPreviewMeshFlags);
 * Manage preview mesh of the viewport
 */
 class FDisplayClusterViewportPreviewMesh
-	: public FGCObject
 {
 public:
 	FDisplayClusterViewportPreviewMesh(const EDisplayClusterDisplayDeviceMeshType InMeshType, const TSharedRef<FDisplayClusterViewportConfiguration, ESPMode::ThreadSafe> InConfiguration)
@@ -81,15 +78,6 @@ public:
 		return EnumHasAnyFlags(RuntimeFlags, InMeshFlags);
 	}
 
-public:
-	//~ Begin FGCObject interface
-	virtual void AddReferencedObjects(FReferenceCollector& Collector) override;
-	virtual FString GetReferencerName() const override
-	{
-		return TEXT("FDisplayClusterViewportPreviewMesh");
-	}
-	//~ End FGCObject interface
-
 private:
 	/** Returns true if this mesh type is supported by the viewport projection policy and DCRA. */
 	bool ShouldUseMeshComponent(FDisplayClusterViewport* InViewport) const;
@@ -124,20 +112,20 @@ private:
 	EDisplayClusterViewportPreviewMeshFlags RuntimeFlags;
 
 	// Mesh component used for preview
-	TObjectPtr<UMeshComponent> MeshComponentPtr;
+	TWeakObjectPtr<UMeshComponent> MeshComponentPtr;
 
 	// This mesh component exists in DCRA and does not need to be deleted.
 	bool bIsRootActorMeshComponent = false;
 
 	// Overlay materials for the mesh can be customized when using preview rendering.
-	TObjectPtr<UMaterialInterface> OrigOverlayMaterial;
+	TWeakObjectPtr<UMaterialInterface> OrigOverlayMaterialPtr;
 
 	// Preview material used on the mesh
-	TObjectPtr<UMaterialInstanceDynamic> MaterialInstancePtr;
+	TWeakObjectPtr<UMaterialInstanceDynamic> MaterialInstancePtr;
 
 	// The current material assigned to the preview mesh
-	TObjectPtr<UMaterial> CurrentMaterialPtr;
+	TWeakObjectPtr<UMaterial> CurrentMaterialPtr;
 
 	// The default material defined in the DisplayDevice
-	TObjectPtr<UMaterial> DefaultMaterialPtr;
+	TWeakObjectPtr<UMaterial> DefaultMaterialPtr;
 };
