nDisplay: quick fix to UE-201985

diff --git a/Engine/Plugins/Runtime/nDisplay/Source/DisplayCluster/Private/Render/Viewport/Preview/DisplayClusterViewportPreviewMesh.cpp b/Engine/Plugins/Runtime/nDisplay/Source/DisplayCluster/Private/Render/Viewport/Preview/DisplayClusterViewportPreviewMesh.cpp
--- a/Engine/Plugins/Runtime/nDisplay/Source/DisplayCluster/Private/Render/Viewport/Preview/DisplayClusterViewportPreviewMesh.cpp
+++ b/Engine/Plugins/Runtime/nDisplay/Source/DisplayCluster/Private/Render/Viewport/Preview/DisplayClusterViewportPreviewMesh.cpp
@@ -116,7 +116,7 @@ void FDisplayClusterViewportPreviewMesh::Update(FDisplayClusterViewport* InViewp
 	if (UMeshComponent* MeshComponent = GetMeshComponent())
 	{
 		// Update material instance and assign to the  mesh
-		if (!CurrentMaterialPtr.IsValid())
+		if (!MaterialInstancePtr.IsValid())
 		{
 			CurrentMaterialPtr = InMeshMaterial;
 
