nDisplay: Implemented dynamic CoC offset into the engine's depth of field pipeline to help correct depth of field blur for ICVFX.

diff --git a/Engine/Plugins/Runtime/nDisplay/Source/DisplayCluster/Private/Game/EngineClasses/Scene/DisplayClusterICVFXCameraComponent.cpp b/Engine/Plugins/Runtime/nDisplay/Source/DisplayCluster/Private/Game/EngineClasses/Scene/DisplayClusterICVFXCameraComponent.cpp
--- a/Engine/Plugins/Runtime/nDisplay/Source/DisplayCluster/Private/Game/EngineClasses/Scene/DisplayClusterICVFXCameraComponent.cpp
+++ b/Engine/Plugins/Runtime/nDisplay/Source/DisplayCluster/Private/Game/EngineClasses/Scene/DisplayClusterICVFXCameraComponent.cpp
@@ -119,6 +119,25 @@ void UDisplayClusterICVFXCameraComponent::TickComponent(float DeltaTime, ELevelT
 	Super::TickComponent(DeltaTime, TickType, ThisTickFunction);
 
 	UpdateOverscanEstimatedFrameSize();
+
+	if (CameraSettings.CameraDepthOfField.bAutomaticallySetDistanceToWall)
+	{
+		TRACE_CPUPROFILER_EVENT_SCOPE_STR("UDisplayClusterICVFXCameraComponent Query Distance To Wall");
+
+		if (ADisplayClusterRootActor* RootActor = Cast<ADisplayClusterRootActor>(GetOwner()))
+		{
+			const FVector CameraLocation = GetComponentLocation();
+			const FVector CameraDirection = GetComponentRotation().RotateVector(FVector::XAxisVector);
+			float DistanceToWall = 0.0;
+
+			// For now, do a single trace from the center of the camera to the stage geometry.
+			// Alternative methods of obtaining wall distance, such as averaging multiple points, can be performed here
+			if (RootActor->GetDistanceToStageGeometry(CameraLocation, CameraDirection, DistanceToWall))
+			{
+				CameraSettings.CameraDepthOfField.DistanceToWall = DistanceToWall;
+			}
+		}
+	}
 }
 
 void UDisplayClusterICVFXCameraComponent::UpdateOverscanEstimatedFrameSize()
@@ -197,6 +216,17 @@ FDisplayClusterViewport_CameraMotionBlur UDisplayClusterICVFXCameraComponent::Ge
 	return OutParameters;
 }
 
+FDisplayClusterViewport_CameraDepthOfField UDisplayClusterICVFXCameraComponent::GetDepthOfFieldParameters()
+{
+	FDisplayClusterViewport_CameraDepthOfField OutParameters;
+
+	OutParameters.bEnableDepthOfFieldCompensation = CameraSettings.CameraDepthOfField.bEnableDepthOfFieldCompensation;
+	OutParameters.DistanceToWall = CameraSettings.CameraDepthOfField.DistanceToWall;
+	OutParameters.DistanceToWallOffset = CameraSettings.CameraDepthOfField.DistanceToWallOffset;
+
+	return OutParameters;
+}
+
 void UDisplayClusterICVFXCameraComponent::OnRegister()
 {
 	Super::OnRegister();
diff --git a/Engine/Plugins/Runtime/nDisplay/Source/DisplayCluster/Private/Render/Viewport/Configuration/DisplayClusterViewportConfigurationHelpers_ICVFX.cpp b/Engine/Plugins/Runtime/nDisplay/Source/DisplayCluster/Private/Render/Viewport/Configuration/DisplayClusterViewportConfigurationHelpers_ICVFX.cpp
--- a/Engine/Plugins/Runtime/nDisplay/Source/DisplayCluster/Private/Render/Viewport/Configuration/DisplayClusterViewportConfigurationHelpers_ICVFX.cpp
+++ b/Engine/Plugins/Runtime/nDisplay/Source/DisplayCluster/Private/Render/Viewport/Configuration/DisplayClusterViewportConfigurationHelpers_ICVFX.cpp
@@ -505,6 +505,9 @@ void FDisplayClusterViewportConfigurationHelpers_ICVFX::UpdateCameraViewportSett
 	// Motion blur:
 	DstViewport.UpdateConfiguration_CameraMotionBlur(InCameraComponent.GetMotionBlurParameters());
 
+	// Depth of field
+	DstViewport.UpdateConfiguration_CameraDepthOfField(InCameraComponent.GetDepthOfFieldParameters());
+
 	// FDisplayClusterConfigurationICVFX_CameraSettings
 	InOutRenderSettings.CameraId.Empty();
 
diff --git a/Engine/Plugins/Runtime/nDisplay/Source/DisplayCluster/Private/Render/Viewport/Containers/DisplayClusterViewport_CameraDepthOfField.cpp b/Engine/Plugins/Runtime/nDisplay/Source/DisplayCluster/Private/Render/Viewport/Containers/DisplayClusterViewport_CameraDepthOfField.cpp
new file mode 100644
--- /dev/null
+++ b/Engine/Plugins/Runtime/nDisplay/Source/DisplayCluster/Private/Render/Viewport/Containers/DisplayClusterViewport_CameraDepthOfField.cpp
@@ -0,0 +1,11 @@
+// Copyright Epic Games, Inc. All Rights Reserved.
+
+#include "Render/Viewport/Containers/DisplayClusterViewport_CameraDepthOfField.h"
+
+#include "SceneView.h"
+
+void FDisplayClusterViewport_CameraDepthOfField::SetupSceneView(FSceneView& InOutView) const
+{
+	InOutView.bEnableDynamicCocOffset = bEnableDepthOfFieldCompensation;
+	InOutView.InFocusDistance = DistanceToWall + DistanceToWallOffset;
+}
\ No newline at end of file
diff --git a/Engine/Plugins/Runtime/nDisplay/Source/DisplayCluster/Private/Render/Viewport/DisplayClusterViewport.cpp b/Engine/Plugins/Runtime/nDisplay/Source/DisplayCluster/Private/Render/Viewport/DisplayClusterViewport.cpp
--- a/Engine/Plugins/Runtime/nDisplay/Source/DisplayCluster/Private/Render/Viewport/DisplayClusterViewport.cpp
+++ b/Engine/Plugins/Runtime/nDisplay/Source/DisplayCluster/Private/Render/Viewport/DisplayClusterViewport.cpp
@@ -421,6 +421,9 @@ void FDisplayClusterViewport::SetupSceneView(uint32 ContextNum, class UWorld* Wo
 	// Handle Motion blur parameters
 	CameraMotionBlur.SetupSceneView(Contexts[ContextNum], InOutView);
 
+	// Handle depth of field parameters
+	CameraDepthOfField.SetupSceneView(InOutView);
+
 	// Handle DisplayDevice
 	if (UDisplayClusterDisplayDeviceBaseComponent* InDisplayDeviceComponent = GetDisplayDeviceComponent(Configuration->GetPreviewSettings().DisplayDeviceRootActorType))
 	{
diff --git a/Engine/Plugins/Runtime/nDisplay/Source/DisplayCluster/Private/Render/Viewport/DisplayClusterViewport.h b/Engine/Plugins/Runtime/nDisplay/Source/DisplayCluster/Private/Render/Viewport/DisplayClusterViewport.h
--- a/Engine/Plugins/Runtime/nDisplay/Source/DisplayCluster/Private/Render/Viewport/DisplayClusterViewport.h
+++ b/Engine/Plugins/Runtime/nDisplay/Source/DisplayCluster/Private/Render/Viewport/DisplayClusterViewport.h
@@ -6,6 +6,7 @@
 #include "Render/Viewport/Configuration/DisplayClusterViewportConfiguration.h"
 
 #include "Render/Viewport/Containers/ImplDisplayClusterViewport_CameraMotionBlur.h"
+#include "Render/Viewport/Containers/DisplayClusterViewport_CameraDepthOfField.h"
 #include "Render/Viewport/Containers/DisplayClusterViewport_CustomFrustumRuntimeSettings.h"
 #include "Render/Viewport/Containers/DisplayClusterViewport_OverscanRuntimeSettings.h"
 #include "Render/Viewport/Containers/DisplayClusterViewportRemap.h"
@@ -317,6 +318,9 @@ public:
 	/** setup CameraMotionBlur configuration for this viewport. */
 	void UpdateConfiguration_CameraMotionBlur(const struct FDisplayClusterViewport_CameraMotionBlur& InCameraMotionBlur);
 
+	/** setup camera depth of field configuration for this viewport. */
+	void UpdateConfiguration_CameraDepthOfField(const class FDisplayClusterViewport_CameraDepthOfField& InCameraDepthOfField);
+
 	/** setup PostRender mips configuration for this viewport. */
 	void UpdateConfiguration_PostRenderGenerateMips(const struct FDisplayClusterConfigurationPostRender_GenerateMips& InGenerateMips);
 
@@ -395,6 +399,9 @@ private:
 	// Additional features
 	FImplDisplayClusterViewport_CameraMotionBlur CameraMotionBlur;
 
+	// Depth of field settings
+	FDisplayClusterViewport_CameraDepthOfField CameraDepthOfField;
+
 	// Overscan rendering feature
 	FDisplayClusterViewport_OverscanRuntimeSettings      OverscanRuntimeSettings;
 
diff --git a/Engine/Plugins/Runtime/nDisplay/Source/DisplayCluster/Private/Render/Viewport/DisplayClusterViewport_Configuration.cpp b/Engine/Plugins/Runtime/nDisplay/Source/DisplayCluster/Private/Render/Viewport/DisplayClusterViewport_Configuration.cpp
--- a/Engine/Plugins/Runtime/nDisplay/Source/DisplayCluster/Private/Render/Viewport/DisplayClusterViewport_Configuration.cpp
+++ b/Engine/Plugins/Runtime/nDisplay/Source/DisplayCluster/Private/Render/Viewport/DisplayClusterViewport_Configuration.cpp
@@ -100,6 +100,11 @@ void FDisplayClusterViewport::UpdateConfiguration_CameraMotionBlur(const FDispla
 	CameraMotionBlur.BlurSetup = InCameraMotionBlur;
 }
 
+void FDisplayClusterViewport::UpdateConfiguration_CameraDepthOfField(const FDisplayClusterViewport_CameraDepthOfField& InCameraDepthOfField)
+{
+	CameraDepthOfField = InCameraDepthOfField;
+}
+
 void FDisplayClusterViewport::UpdateConfiguration_PostRenderOverride(const FDisplayClusterConfigurationPostRender_Override& InOverride)
 {
 	PostRenderSettings.Replace.TextureRHI.SafeRelease();
diff --git a/Engine/Plugins/Runtime/nDisplay/Source/DisplayCluster/Public/Components/DisplayClusterICVFXCameraComponent.h b/Engine/Plugins/Runtime/nDisplay/Source/DisplayCluster/Public/Components/DisplayClusterICVFXCameraComponent.h
--- a/Engine/Plugins/Runtime/nDisplay/Source/DisplayCluster/Public/Components/DisplayClusterICVFXCameraComponent.h
+++ b/Engine/Plugins/Runtime/nDisplay/Source/DisplayCluster/Public/Components/DisplayClusterICVFXCameraComponent.h
@@ -10,6 +10,7 @@
 #include "DisplayClusterConfigurationTypes_ICVFX.h"
 #include "DisplayClusterEditorPropertyReference.h"
 #include "Render/Viewport/Containers/DisplayClusterViewport_CameraMotionBlur.h"
+#include "Render/Viewport/Containers/DisplayClusterViewport_CameraDepthOfField.h"
 
 #include "DisplayClusterICVFXCameraComponent.generated.h"
 
@@ -47,6 +48,9 @@ public:
 public:
 	FDisplayClusterViewport_CameraMotionBlur GetMotionBlurParameters();
 
+	/** Gets the depth of field parameters to store on the display cluster viewport */
+	FDisplayClusterViewport_CameraDepthOfField GetDepthOfFieldParameters();
+
 	/**
 	 * Return the actual source camera, e.g. the camera component of the referenced cine camera.
 	 * Use GetCameraView() function to get viewinfo with actual camera position, postprocess and ICVFX postprocess
@@ -129,6 +133,10 @@ private:
 	UPROPERTY(EditAnywhere, Transient, Category = "In Camera VFX", meta = (PropertyPath = "CameraSettings.CameraMotionBlur"))
 	FDisplayClusterEditorPropertyReference CameraMotionBlurRef;
 
+	/** Exposed reference to the camera's inner depth of field settings */
+	UPROPERTY(EditAnywhere, Transient, Category = "In Camera VFX", meta = (PropertyPath = "CameraSettings.CameraDepthOfField"))
+	FDisplayClusterEditorPropertyReference CameraDepthOfFieldRef;
+
 	UPROPERTY(EditAnywhere, Transient, Category = "In Camera VFX", meta = (PropertyPath = "CameraSettings.CameraHideList"))
 	FDisplayClusterEditorPropertyReference CameraHideListRef;
 
diff --git a/Engine/Plugins/Runtime/nDisplay/Source/DisplayCluster/Public/Render/Viewport/Containers/DisplayClusterViewport_CameraDepthOfField.h b/Engine/Plugins/Runtime/nDisplay/Source/DisplayCluster/Public/Render/Viewport/Containers/DisplayClusterViewport_CameraDepthOfField.h
new file mode 100644
--- /dev/null
+++ b/Engine/Plugins/Runtime/nDisplay/Source/DisplayCluster/Public/Render/Viewport/Containers/DisplayClusterViewport_CameraDepthOfField.h
@@ -0,0 +1,27 @@
+// Copyright Epic Games, Inc. All Rights Reserved.
+
+#pragma once
+
+#include "CoreMinimal.h"
+
+class FSceneView;
+
+/**
+* Settings for the camera depth of field blur.
+*/
+class FDisplayClusterViewport_CameraDepthOfField
+{
+public:
+	/** Sets the appropriate depth of field settings in the scene view */
+	void SetupSceneView(FSceneView& InOutView) const;
+
+public:
+	/** Indicates whether depth of field correction is enabled for the ICVFX camera */
+	bool bEnableDepthOfFieldCompensation = false;
+
+	/** The distance from the ICVFX camera to the wall it is shooting against */
+	float DistanceToWall = 0.0;
+
+	/** An offset applied to the distance to wall value */
+	float DistanceToWallOffset = 0.0;
+};
\ No newline at end of file
diff --git a/Engine/Plugins/Runtime/nDisplay/Source/DisplayClusterConfiguration/Public/DisplayClusterConfigurationTypes_ICVFX.h b/Engine/Plugins/Runtime/nDisplay/Source/DisplayClusterConfiguration/Public/DisplayClusterConfigurationTypes_ICVFX.h
--- a/Engine/Plugins/Runtime/nDisplay/Source/DisplayClusterConfiguration/Public/DisplayClusterConfigurationTypes_ICVFX.h
+++ b/Engine/Plugins/Runtime/nDisplay/Source/DisplayClusterConfiguration/Public/DisplayClusterConfigurationTypes_ICVFX.h
@@ -565,6 +565,29 @@ public:
 	void SetupViewInfo(const FDisplayClusterConfigurationICVFX_StageSettings& InStageSettings, FMinimalViewInfo& InOutViewInfo) const;
 };
 
+USTRUCT(BlueprintType)
+struct DISPLAYCLUSTERCONFIGURATION_API FDisplayClusterConfigurationICVFX_CameraDepthOfField
+{
+	GENERATED_BODY()
+
+public:
+	/** Enables depth of field correction on the wall, which dynamically adjusts the size of the defocus circle of confusion to compensate for the real-world camera blur when shooting the wall */
+	UPROPERTY(BlueprintReadWrite, EditAnywhere, Category = NDisplay)
+	bool bEnableDepthOfFieldCompensation = false;
+
+	/** Allows the ICVFX camera to automatically compute its distance from the stage walls using ray casting every tick */
+	UPROPERTY(BlueprintReadWrite, EditAnywhere, Category = NDisplay)
+	bool bAutomaticallySetDistanceToWall = true;
+
+	/** The distance from the ICVFX camera to the wall it is pointing at */
+	UPROPERTY(BlueprintReadWrite, EditAnywhere, Category = NDisplay, meta=(EditCondition="!bAutomaticallySetDistanceToWall"))
+	float DistanceToWall = 0.0;
+
+	/** An offset applied to DistanceToWall (applied regardless of whether DistanceToWall is automatically set) */
+	UPROPERTY(BlueprintReadWrite, EditAnywhere, Category = NDisplay)
+	float DistanceToWallOffset = 0.0;
+};
+
 USTRUCT(BlueprintType)
 struct DISPLAYCLUSTERCONFIGURATION_API FDisplayClusterConfigurationICVFX_CameraSoftEdge
 {
@@ -741,6 +764,10 @@ public:
 	UPROPERTY(BlueprintReadWrite, BlueprintReadWrite, EditAnywhere, Category = "In Camera VFX")
 	FDisplayClusterConfigurationICVFX_CameraMotionBlur CameraMotionBlur;
 
+	/** Settings that control the depth of field blur applied to the ICVFX image */
+	UPROPERTY(BlueprintReadWrite, BlueprintReadWrite, EditAnywhere, Category = "In Camera VFX")
+	FDisplayClusterConfigurationICVFX_CameraDepthOfField CameraDepthOfField;
+
 	/** Configure global render settings for this viewport */
 	UPROPERTY(EditAnywhere, BlueprintReadWrite, Category = NDisplay)
 	FDisplayClusterConfigurationICVFX_CameraRenderSettings RenderSettings;
diff --git a/Engine/Shaders/Private/DiaphragmDOF/DOFCommon.ush b/Engine/Shaders/Private/DiaphragmDOF/DOFCommon.ush
--- a/Engine/Shaders/Private/DiaphragmDOF/DOFCommon.ush
+++ b/Engine/Shaders/Private/DiaphragmDOF/DOFCommon.ush
@@ -149,6 +149,8 @@ static const int2 kOffsetsSquare3x3[9] = {
 //------------------------------------------------------- PARAMETERS
 
 float CocInfinityRadius;
+float CocInFocusRadius;
+uint bCocEnableDynamicRadiusOffset;
 float CocMinRadius;
 float CocMaxRadius;
 float CocSqueeze;
@@ -167,6 +169,24 @@ RWTexture2D<float4> DebugOutput;
 
 //------------------------------------------------------- COMMON API.
 
+#if WITH_DYNAMIC_COC_OFFSET
+float GetCocOffset(float CocRadius)
+{
+	float DynamicOffset = 0.0;
+	BRANCH if (bCocEnableDynamicRadiusOffset)
+	{
+		// This curve minimizes the RMS error between a point that was blurred once by a disk blur and one
+		// that was blurred twice, where CocRadius + DynamicOffset is the size of the first blur and
+		// CocInFocusRadius in the size of the second blur, 
+		float B = 0.467743 + 7.89656 * pow(CocInFocusRadius, -1.89051);
+		float V = -(2.57186 + 0.142159 * CocInFocusRadius);
+		DynamicOffset = -CocInFocusRadius * (1 - pow(1 + exp(B * (CocInFocusRadius - CocRadius)), V));
+	}
+
+	return DynamicOffset;
+}
+#endif
+
 // Returns the coc radius
 float SceneDepthToCocRadius(float SceneDepth)
 {
@@ -175,6 +195,10 @@ float SceneDepthToCocRadius(float SceneDepth)
 	// Radius from the circle of confusion.
 	float CocRadius = ((SceneDepth - Focus) / SceneDepth) * CocInfinityRadius;
 
+#if WITH_DYNAMIC_COC_OFFSET
+	CocRadius = CocRadius + GetCocOffset(CocRadius);
+#endif
+
 	// Depth blur's radius.
 	float DepthBlurAbsRadius = (1.0 - exp2(-SceneDepth * DepthBlurExponent)) * DepthBlurRadius;
 
diff --git a/Engine/Source/Runtime/Engine/Public/SceneView.h b/Engine/Source/Runtime/Engine/Public/SceneView.h
--- a/Engine/Source/Runtime/Engine/Public/SceneView.h
+++ b/Engine/Source/Runtime/Engine/Public/SceneView.h
@@ -1554,6 +1554,12 @@ public:
 	FShaderResourceViewRHIRef LandscapeIndirectionBuffer;
 	FShaderResourceViewRHIRef LandscapePerComponentDataBuffer;
 
+	/** Enables a CoC offset whose offset value changes with pixel's scene depth */
+	bool bEnableDynamicCocOffset = false;
+
+	/** When dynamic CoC offset is enabled, this is the distance from camera at which objects will be in perfect focus (when the CoC offset is maximum so that the final CoC is 0) */
+	float InFocusDistance = 0.0;
+
 	/** Feature level for this scene */
 	const ERHIFeatureLevel::Type FeatureLevel;
 
diff --git a/Engine/Source/Runtime/Renderer/Private/PostProcess/DiaphragmDOF.cpp b/Engine/Source/Runtime/Renderer/Private/PostProcess/DiaphragmDOF.cpp
--- a/Engine/Source/Runtime/Renderer/Private/PostProcess/DiaphragmDOF.cpp
+++ b/Engine/Source/Runtime/Renderer/Private/PostProcess/DiaphragmDOF.cpp
@@ -428,6 +428,10 @@ public:
 	static void ModifyCompilationEnvironment(const FGlobalShaderPermutationParameters& Parameters, FShaderCompilerEnvironment& OutEnvironment)
 	{
 		OutEnvironment.SetDefine(TEXT("COC_TILE_SIZE"), kCocTileSize);
+
+		// If on console, don't compile any dynamic CoC offset to avoid performance differences
+		const bool bIsConsole = FDataDrivenShaderPlatformInfo::GetIsConsole(Parameters.Platform);
+		OutEnvironment.SetDefine(TEXT("WITH_DYNAMIC_COC_OFFSET"), bIsConsole ? 0 : 1);
 	}
 
 	FDiaphragmDOFShader() {}
@@ -633,6 +637,8 @@ END_SHADER_PARAMETER_STRUCT()
 
 BEGIN_SHADER_PARAMETER_STRUCT(FDOFCocModelShaderParameters, )
 	SHADER_PARAMETER(float, CocInfinityRadius)
+	SHADER_PARAMETER(float, CocInFocusRadius)
+	SHADER_PARAMETER(uint32, bCocEnableDynamicRadiusOffset)
 	SHADER_PARAMETER(float, CocMinRadius)
 	SHADER_PARAMETER(float, CocMaxRadius)
 	SHADER_PARAMETER(float, CocSqueeze)
@@ -647,6 +653,8 @@ void SetCocModelParameters(
 	float CocRadiusBasis = 1.0f)
 {
 	OutParameters->CocInfinityRadius = CocRadiusBasis * CocModel.InfinityBackgroundCocRadius;
+	OutParameters->CocInFocusRadius = CocRadiusBasis * CocModel.InFocusRadius;
+	OutParameters->bCocEnableDynamicRadiusOffset = CocModel.bEnableDynamicOffset;
 	OutParameters->CocMinRadius = CocRadiusBasis * CocModel.MinForegroundCocRadius;
 	OutParameters->CocMaxRadius = CocRadiusBasis * CocModel.MaxBackgroundCocRadius;
 	OutParameters->CocSqueeze = CocModel.Squeeze;
diff --git a/Engine/Source/Runtime/Renderer/Private/PostProcess/DiaphragmDOF.h b/Engine/Source/Runtime/Renderer/Private/PostProcess/DiaphragmDOF.h
--- a/Engine/Source/Runtime/Renderer/Private/PostProcess/DiaphragmDOF.h
+++ b/Engine/Source/Runtime/Renderer/Private/PostProcess/DiaphragmDOF.h
@@ -31,6 +31,12 @@ struct FPhysicalCocModel
 	// Unclamped resolution less background coc radius.
 	float InfinityBackgroundCocRadius;
 
+	/** Indicates whether a dynamic offset dependent on scene depth should be computed for every pixel */
+	bool bEnableDynamicOffset;
+
+	/** When dynamic offset is enabled, this is the coc radius at which objects will be perfectly sharp */
+	float InFocusRadius;
+
 	// Resolution less minimal foreground coc radius < 0.
 	float MinForegroundCocRadius;
 
@@ -59,6 +65,7 @@ struct FPhysicalCocModel
 	/** Returns limit(SceneDepthToCocRadius) for SceneDepth -> Infinity. */
 	FORCEINLINE float ComputeViewMaxBackgroundCocRadius(float HorizontalResolution) const
 	{
+		// Dynamic CoC offset is designed to go to zero at infinite distance, so we don't need to factor it into the max background radius
 		return FMath::Min(FMath::Max(InfinityBackgroundCocRadius, MaxDepthBlurRadius), MaxBackgroundCocRadius) * HorizontalResolution;
 	}
 	
@@ -70,6 +77,10 @@ struct FPhysicalCocModel
 	{
 		return DepthToResCocRadius(GNearClippingPlane, HorizontalResolution);
 	}
+
+private:
+	/** Gets the offset to the circle of confusion to apply for the specified radius */
+	float GetCocOffset(float CocRadius) const;
 };
 
 
diff --git a/Engine/Source/Runtime/Renderer/Private/PostProcess/DiaphragmDOFUtils.cpp b/Engine/Source/Runtime/Renderer/Private/PostProcess/DiaphragmDOFUtils.cpp
--- a/Engine/Source/Runtime/Renderer/Private/PostProcess/DiaphragmDOFUtils.cpp
+++ b/Engine/Source/Runtime/Renderer/Private/PostProcess/DiaphragmDOFUtils.cpp
@@ -140,17 +140,33 @@ void DiaphragmDOF::FPhysicalCocModel::Compile(const FViewInfo& View)
 
 		// Convert diameter in mm to resolution less radius on the filmback.
 		InfinityBackgroundCocRadius = DiameterInMM * 0.5f / SensorWidthInMM;
+
+		if (View.InFocusDistance > 0)
+		{
+			// For now, the dynamic CoC offset only handles cases where the in-focus radius is positive (the in-focus distance is further than the focal point)
+			// so clamp the in focus radius to always be positive
+			InFocusRadius = FMath::Max(InfinityBackgroundCocRadius * (1 - FocusDistance / View.InFocusDistance), 0.0f);
+			bEnableDynamicOffset = View.bEnableDynamicCocOffset;
+		}
+		else
+		{
+			InFocusRadius = 0.0;
+			bEnableDynamicOffset = false;
+		}
 	}
 	else
 	{
 		InfinityBackgroundCocRadius = 0.0f;
 		MinForegroundCocRadius = 0.0;
+		InFocusRadius = 0.0;
+		bEnableDynamicOffset = false;
 	}
 }
 
 float DiaphragmDOF::FPhysicalCocModel::DepthToResCocRadius(float SceneDepth, float HorizontalResolution) const
 {
-	float CocRadius = ((SceneDepth - FocusDistance) / SceneDepth) * InfinityBackgroundCocRadius;
+	float InitialCocRadius = ((SceneDepth - FocusDistance) / SceneDepth) * InfinityBackgroundCocRadius;
+	float CocRadius = InitialCocRadius + GetCocOffset(InitialCocRadius);
 
 	// Depth blur based.
 	float DepthBlurAbsRadius = (1.0 - FMath::Exp2(-SceneDepth * DepthBlurExponent)) * MaxDepthBlurRadius;
@@ -164,6 +180,19 @@ float DiaphragmDOF::FPhysicalCocModel::DepthToResCocRadius(float SceneDepth, flo
 	return HorizontalResolution * FMath::Clamp(ReturnCoc, MinForegroundCocRadius, MaxBackgroundCocRadius);
 }
 
+float DiaphragmDOF::FPhysicalCocModel::GetCocOffset(float CocRadius) const
+{
+	float DynamicOffset = 0.0;
+	if (bEnableDynamicOffset)
+	{
+		const float B = 0.467743 + 7.89656 * pow(InFocusRadius, -1.89051);
+		const float V = -(2.57186 + 0.142159 * InFocusRadius);
+		DynamicOffset = -InFocusRadius * (1 - pow(1 + exp(B * (InFocusRadius - CocRadius)), V));
+	}
+
+	return DynamicOffset;
+}
+
 void DiaphragmDOF::FBokehModel::Compile(const FViewInfo& View)
 {
 	{
