[nDisplay] Added more diagnostics to NVIDIA sync policy

- get/set latency every frame to see whether these two always fail, or it happens during some period(s) only
- added barrier after NvAPI_D3D1x_Present() call to see whether it fixes RHI thread alignment in time

diff --git a/Engine/Plugins/Runtime/nDisplay/Source/DisplayCluster/Private/Render/Synchronization/DisplayClusterRenderSyncPolicyNvidia.h b/Engine/Plugins/Runtime/nDisplay/Source/DisplayCluster/Private/Render/Synchronization/DisplayClusterRenderSyncPolicyNvidia.h
--- a/Engine/Plugins/Runtime/nDisplay/Source/DisplayCluster/Private/Render/Synchronization/DisplayClusterRenderSyncPolicyNvidia.h
+++ b/Engine/Plugins/Runtime/nDisplay/Source/DisplayCluster/Private/Render/Synchronization/DisplayClusterRenderSyncPolicyNvidia.h
@@ -36,8 +36,12 @@ public:
 	virtual bool SynchronizeClusterRendering(int32& InOutSyncInterval) override;
 
 private:
+	// Initializes NVIDIA internals
 	bool InitializeNvidiaSwapLock();
 
+	// Frame latency control
+	void SetMaximumFrameLatency(uint8 FrameLatency);
+
 private:
 	bool bNvApiInitialized = false;
 	bool bNvApiBarrierSet  = false;
diff --git a/Engine/Plugins/Runtime/nDisplay/Source/DisplayCluster/Private/Render/Synchronization/Linux/DisplayClusterRenderSyncPolicyNvidiaLinux.cpp b/Engine/Plugins/Runtime/nDisplay/Source/DisplayCluster/Private/Render/Synchronization/Linux/DisplayClusterRenderSyncPolicyNvidiaLinux.cpp
--- a/Engine/Plugins/Runtime/nDisplay/Source/DisplayCluster/Private/Render/Synchronization/Linux/DisplayClusterRenderSyncPolicyNvidiaLinux.cpp
+++ b/Engine/Plugins/Runtime/nDisplay/Source/DisplayCluster/Private/Render/Synchronization/Linux/DisplayClusterRenderSyncPolicyNvidiaLinux.cpp
@@ -22,6 +22,11 @@ bool FDisplayClusterRenderSyncPolicyNvidia::InitializeNvidiaSwapLock()
 	return true;
 }
 
+void FDisplayClusterRenderSyncPolicyNvidia::SetMaximumFrameLatency(uint8 FrameLatency)
+{
+	// not implemented yet
+}
+
 bool FDisplayClusterRenderSyncPolicyNvidia::SynchronizeClusterRendering(int32& InOutSyncInterval)
 {
 	// Wait unless the frame is rendered
diff --git a/Engine/Plugins/Runtime/nDisplay/Source/DisplayCluster/Private/Render/Synchronization/Windows/DisplayClusterRenderSyncPolicyNvidiaWindows.cpp b/Engine/Plugins/Runtime/nDisplay/Source/DisplayCluster/Private/Render/Synchronization/Windows/DisplayClusterRenderSyncPolicyNvidiaWindows.cpp
--- a/Engine/Plugins/Runtime/nDisplay/Source/DisplayCluster/Private/Render/Synchronization/Windows/DisplayClusterRenderSyncPolicyNvidiaWindows.cpp
+++ b/Engine/Plugins/Runtime/nDisplay/Source/DisplayCluster/Private/Render/Synchronization/Windows/DisplayClusterRenderSyncPolicyNvidiaWindows.cpp
@@ -60,6 +60,36 @@ static TAutoConsoleVariable<int32> CVarNvidiaSyncDiagnosticsCompletion(
 	,
 	ECVF_ReadOnly | ECVF_RenderThreadSafe
 );
+
+// In the recent logs (Nov-Dec, 2023), we see GetMaximumFrameLatency() and SetMaximumFrameLatency() don't work
+// properly. This is something weird and may indicate there is an issue outside of UE. Those methods are called
+// once during the policy initialization. To better understand the issue and probably see some correlation,
+// this cvar allows to focribly call set/get frame latency every frame, and log the results.
+static TAutoConsoleVariable<bool> CVarNvidiaSyncDiagnosticsUpdateLatencyEveryFrame(
+	TEXT("nDisplay.sync.nvidia.diag.UpdateLatencyEveryFrame"),
+	false,
+	TEXT("NVAPI diagnostics: performs get/set latency every frame, and logs the results\n")
+	TEXT("0 : disabled\n")
+	TEXT("1 : enabled\n")
+	,
+	ECVF_ReadOnly | ECVF_RenderThreadSafe
+);
+
+// Sometimes we see a weird pattern of NvAPI_D3D1x_Present() handling in the utraces. Those patterns
+// visualize something that we would not expect to see if synchronization works properly.
+// NvAPI_D3D1x_Present() may return in between of the vblanks, or it may return asynchronously to other nodes.
+// This cvar enables additional barrier synchronization step right after NvAPI_D3D1x_Present() call.
+// This should help keeping RHI threads aligned.
+static TAutoConsoleVariable<bool> CVarNvidiaSyncDiagnosticsPostPresentSync(
+	TEXT("nDisplay.sync.nvidia.diag.PostPresentSync"),
+	false,
+	TEXT("NVAPI diagnostics: sync nodes on a network barrier after frame presentation \n")
+	TEXT("0 : disabled\n")
+	TEXT("1 : enabled\n")
+	,
+	ECVF_ReadOnly | ECVF_RenderThreadSafe
+);
+
 // TEMPORARY DIAGNOSTICS END
 
 // CVarNvidiaPresentBarrierCountLimit is used to set the number of presentation cluster sync barriers before getting disabled.
@@ -144,6 +174,36 @@ bool FDisplayClusterRenderSyncPolicyNvidia::Initialize()
 	return bNvApiInitialized;
 }
 
+void FDisplayClusterRenderSyncPolicyNvidia::SetMaximumFrameLatency(uint8 FrameLatency)
+{
+	using namespace DisplayClusterRenderSyncPolicyNvidia_Data_Windows;
+
+	if (DXGISwapChain)
+	{
+		HRESULT Result = DXGISwapChain->SetMaximumFrameLatency(FrameLatency);
+
+		if (Result != S_OK)
+		{
+			UE_LOG(LogDisplayClusterRenderSync, Warning, TEXT("NVAPI Diag: Couldn't set maximum frame latency to %u. Error: %x"), FrameLatency, Result);
+		}
+
+#if 1 // This is not required actually, but used for diagnostics
+		UINT CurrentLatency = 0;
+
+		Result = DXGISwapChain->GetMaximumFrameLatency(&CurrentLatency);
+
+		if (Result == S_OK)
+		{
+			UE_LOG(LogDisplayClusterRenderSync, Log, TEXT("NVAPI Diag: Swapchain frame latency: %u"), CurrentLatency);
+		}
+		else
+		{
+			UE_LOG(LogDisplayClusterRenderSync, Warning, TEXT("NVAPI Diag: Couldn't get maximum frame latency. Error: %x"), Result);
+		}
+#endif
+	}
+}
+
 bool FDisplayClusterRenderSyncPolicyNvidia::SynchronizeClusterRendering(int32& InOutSyncInterval)
 {
 	using namespace DisplayClusterRenderSyncPolicyNvidia_Data_Windows;
@@ -182,6 +242,13 @@ bool FDisplayClusterRenderSyncPolicyNvidia::SynchronizeClusterRendering(int32& I
 		return true;
 	}
 
+	// NVAPI Diagnostics: trying to update maximum frame latency every frame and see
+	// if it works, and the API function does not fail.
+	if(CVarNvidiaSyncDiagnosticsUpdateLatencyEveryFrame.GetValueOnAnyThread())
+	{
+		SetMaximumFrameLatency(1);
+	}
+
 	// NVAPI Diagnostics: frame completion
 	// Wait unless frame rendering is finished
 	if (bNvDiagCompletion)
@@ -232,6 +299,14 @@ bool FDisplayClusterRenderSyncPolicyNvidia::SynchronizeClusterRendering(int32& I
 
 			// Present frame via NVIDIA API
 			const NvAPI_Status NvApiResult = NvAPI_D3D1x_Present(D3DDevice, DXGISwapChain, (UINT)InOutSyncInterval, (UINT)0);
+
+			// Regardless of NvAPI_D3D1x_Present() call result, synchronize clients on the barrier if requested. The following
+			// if-block may theoretically branch the execution paths of the nodes, so we might lose threads alignment.
+			if (CVarNvidiaSyncDiagnosticsPostPresentSync.GetValueOnAnyThread())
+			{
+				SyncOnBarrier();
+			}
+
 			if (NvApiResult != NVAPI_OK)
 			{
 				UE_LOG(LogDisplayClusterRenderSync, Warning, TEXT("NVAPI: An error occurred during frame presentation, error code 0x%x"), NvApiResult);
@@ -293,27 +368,7 @@ bool FDisplayClusterRenderSyncPolicyNvidia::InitializeNvidiaSwapLock()
 	}
 
 	// Set frame latency
-	{
-		HRESULT Result = DXGISwapChain->SetMaximumFrameLatency(1);
-
-		if (Result != S_OK)
-		{
-			UE_LOG(LogDisplayClusterRenderSync, Warning, TEXT("Couldn't set maximum frame latency. Error: %x"), Result);
-		}
-
-		UINT CurrentLatency = 0;
-
-		Result = DXGISwapChain->GetMaximumFrameLatency(&CurrentLatency);
-
-		if (Result == S_OK)
-		{
-			UE_LOG(LogDisplayClusterRenderSync, Log, TEXT("Swapchain frame latency: %u"), CurrentLatency);
-		}
-		else
-		{
-			UE_LOG(LogDisplayClusterRenderSync, Warning, TEXT("Couldn't get maximum frame latency. Error: %x"), Result);
-		}
-	}
+	SetMaximumFrameLatency(1);
 
 	NvU32 MaxGroups = 0;
 	NvU32 MaxBarriers = 0;
