nDisplay: Fixed issue where compile permutation flag was erroneously removed from the ICVFX hit proxy shader, preventing light card proxies from being rendered in the hit proxy pass, making them unclickable.

diff --git a/Engine/Plugins/Runtime/nDisplay/Source/DisplayClusterLightCardEditorShaders/Private/DisplayClusterMeshProjectionRenderer.cpp b/Engine/Plugins/Runtime/nDisplay/Source/DisplayClusterLightCardEditorShaders/Private/DisplayClusterMeshProjectionRenderer.cpp
--- a/Engine/Plugins/Runtime/nDisplay/Source/DisplayClusterLightCardEditorShaders/Private/DisplayClusterMeshProjectionRenderer.cpp
+++ b/Engine/Plugins/Runtime/nDisplay/Source/DisplayClusterLightCardEditorShaders/Private/DisplayClusterMeshProjectionRenderer.cpp
@@ -415,7 +415,8 @@ public:
 		return EnumHasAllFlags(Parameters.Flags, EShaderPermutationFlags::HasEditorOnlyData)
 			// and only compile for default materials or materials that are masked.
 			&&
-			( !Parameters.MaterialParameters.bWritesEveryPixel
+			(  Parameters.MaterialParameters.bIsSpecialEngineMaterial
+			|| !Parameters.MaterialParameters.bWritesEveryPixel
 			|| Parameters.MaterialParameters.bMaterialMayModifyMeshPosition
 			|| Parameters.MaterialParameters.bIsTwoSided);
 	}
