[nDisplay] New NVIDIA synchronization policy (present barrier)

- Added new NvidiaPresentBarrier sync policy
- Renamed old Nvidia policy to NvidiaSwapBarrier
- Common Nvidia sync features were taken out to an abstract class
- Added/updated auxiliary infrastructure

diff --git a/Engine/Plugins/Runtime/nDisplay/Source/DisplayCluster/Private/DisplayClusterCallbacks.h b/Engine/Plugins/Runtime/nDisplay/Source/DisplayCluster/Private/DisplayClusterCallbacks.h
--- a/Engine/Plugins/Runtime/nDisplay/Source/DisplayCluster/Private/DisplayClusterCallbacks.h
+++ b/Engine/Plugins/Runtime/nDisplay/Source/DisplayCluster/Private/DisplayClusterCallbacks.h
@@ -72,6 +72,11 @@ public:
 		return DisplayClusterPresentationPostSynchronizationEvent;
 	}
 
+	virtual FDisplayClusterFramePresentated_RHIThread& OnDisplayClusterFramePresented_RHIThread() override
+	{
+		return DisplayClusterFramePresentedEvent;
+	}
+
 	virtual FDisplayClusterFailoverNodeDown& OnDisplayClusterFailoverNodeDown() override
 	{
 		return DisplayClusterFailoverNodeDown;
@@ -147,6 +152,7 @@ private:
 
 	FDisplayClusterPresentationPreSynchronization_RHIThread  DisplayClusterPresentationPreSynchronizationEvent;
 	FDisplayClusterPresentationPostSynchronization_RHIThread DisplayClusterPresentationPostSynchronizationEvent;
+	FDisplayClusterFramePresentated_RHIThread                DisplayClusterFramePresentedEvent;
 
 	FDisplayClusterPostRenderViewFamily_RenderThread DisplayClusterPostRenderViewFamily;
 	FDisplayClusterPreWarp_RenderThread              DisplayClusterPreWarpEvent;
diff --git a/Engine/Plugins/Runtime/nDisplay/Source/DisplayCluster/Private/Render/Presentation/DisplayClusterPresentationBase.cpp b/Engine/Plugins/Runtime/nDisplay/Source/DisplayCluster/Private/Render/Presentation/DisplayClusterPresentationBase.cpp
--- a/Engine/Plugins/Runtime/nDisplay/Source/DisplayCluster/Private/Render/Presentation/DisplayClusterPresentationBase.cpp
+++ b/Engine/Plugins/Runtime/nDisplay/Source/DisplayCluster/Private/Render/Presentation/DisplayClusterPresentationBase.cpp
@@ -43,7 +43,7 @@ void FDisplayClusterPresentationBase::OnBackBufferResize()
 
 bool FDisplayClusterPresentationBase::Present(int32& InOutSyncInterval)
 {
-	TRACE_CPUPROFILER_EVENT_SCOPE(nDisplay RenderDevice::PresentationBase);
+	TRACE_CPUPROFILER_EVENT_SCOPE(nDisplay PresentationBase);
 
 	bool bNeedPresent = true;
 
@@ -63,3 +63,9 @@ bool FDisplayClusterPresentationBase::Present(int32& InOutSyncInterval)
 
 	return bNeedPresent;
 }
+
+void FDisplayClusterPresentationBase::PostPresent()
+{
+	// Notify native frame presentation has been done
+	GDisplayCluster->GetCallbacks().OnDisplayClusterFramePresented_RHIThread().Broadcast(true);
+}
diff --git a/Engine/Plugins/Runtime/nDisplay/Source/DisplayCluster/Private/Render/Presentation/DisplayClusterPresentationBase.h b/Engine/Plugins/Runtime/nDisplay/Source/DisplayCluster/Private/Render/Presentation/DisplayClusterPresentationBase.h
--- a/Engine/Plugins/Runtime/nDisplay/Source/DisplayCluster/Private/Render/Presentation/DisplayClusterPresentationBase.h
+++ b/Engine/Plugins/Runtime/nDisplay/Source/DisplayCluster/Private/Render/Presentation/DisplayClusterPresentationBase.h
@@ -33,6 +33,7 @@ protected:
 	{ return true; };
 
 	virtual bool Present(int32& InOutSyncInterval) override;
+	virtual void PostPresent() override;
 	virtual void OnBackBufferResize() override;
 
 protected:
diff --git a/Engine/Plugins/Runtime/nDisplay/Source/DisplayCluster/Private/Render/Synchronization/DisplayClusterRenderSyncHelper.h b/Engine/Plugins/Runtime/nDisplay/Source/DisplayCluster/Private/Render/Synchronization/DisplayClusterRenderSyncHelper.h
--- a/Engine/Plugins/Runtime/nDisplay/Source/DisplayCluster/Private/Render/Synchronization/DisplayClusterRenderSyncHelper.h
+++ b/Engine/Plugins/Runtime/nDisplay/Source/DisplayCluster/Private/Render/Synchronization/DisplayClusterRenderSyncHelper.h
@@ -23,6 +23,16 @@ public:
 	{
 		return false;
 	}
+
+	virtual bool GetMaximumFrameLatency(uint32& OutMaximumFrameLatency)
+	{
+		return false;
+	}
+
+	virtual bool SetMaximumFrameLatency(uint32 MaximumFrameLatency)
+	{
+		return false;
+	}
 };
 
 
@@ -51,6 +61,8 @@ private:
 	{
 		virtual bool WaitForVBlank() override;
 		virtual bool IsWaitForVBlankSupported() override;
+		virtual bool GetMaximumFrameLatency(uint32& OutMaximumFrameLatency) override;
+		virtual bool SetMaximumFrameLatency(uint32 MaximumFrameLatency) override;
 	};
 #endif
 
@@ -60,6 +72,8 @@ private:
 	{
 		virtual bool WaitForVBlank() override;
 		virtual bool IsWaitForVBlankSupported() override;
+		virtual bool GetMaximumFrameLatency(uint32& OutMaximumFrameLatency) override;
+		virtual bool SetMaximumFrameLatency(uint32 MaximumFrameLatency) override;
 	};
 
 private:
diff --git a/Engine/Plugins/Runtime/nDisplay/Source/DisplayCluster/Private/Render/Synchronization/DisplayClusterRenderSyncPolicyBase.cpp b/Engine/Plugins/Runtime/nDisplay/Source/DisplayCluster/Private/Render/Synchronization/DisplayClusterRenderSyncPolicyBase.cpp
--- a/Engine/Plugins/Runtime/nDisplay/Source/DisplayCluster/Private/Render/Synchronization/DisplayClusterRenderSyncPolicyBase.cpp
+++ b/Engine/Plugins/Runtime/nDisplay/Source/DisplayCluster/Private/Render/Synchronization/DisplayClusterRenderSyncPolicyBase.cpp
@@ -25,7 +25,7 @@ bool FDisplayClusterRenderSyncPolicyBase::Initialize()
 void FDisplayClusterRenderSyncPolicyBase::SyncOnBarrier()
 {
 	{
-		TRACE_CPUPROFILER_EVENT_SCOPE(nDisplay SyncPolicyBase::SyncBarrier);
+		TRACE_CPUPROFILER_EVENT_SCOPE(nDisplay SyncOnBarrier);
 		GDisplayCluster->GetPrivateClusterMgr()->GetClusterNodeController()->SyncOnBarrier();
 	}
 }
@@ -42,6 +42,16 @@ void FDisplayClusterRenderSyncPolicyBase::WaitForFrameCompletion()
 	}
 }
 
+bool FDisplayClusterRenderSyncPolicyBase::GetMaximumFrameLatency(uint32& OutMaximumFrameLatency)
+{
+	return FDisplayClusterRenderSyncHelper::Get().GetMaximumFrameLatency(OutMaximumFrameLatency);
+}
+
+bool FDisplayClusterRenderSyncPolicyBase::SetMaximumFrameLatency(uint32 MaximumFrameLatency)
+{
+	return FDisplayClusterRenderSyncHelper::Get().SetMaximumFrameLatency(MaximumFrameLatency);
+}
+
 bool FDisplayClusterRenderSyncPolicyBase::IsWaitForVBlankFeatureSupported()
 {
 	return FDisplayClusterRenderSyncHelper::Get().IsWaitForVBlankSupported();
diff --git a/Engine/Plugins/Runtime/nDisplay/Source/DisplayCluster/Private/Render/Synchronization/DisplayClusterRenderSyncPolicyBase.h b/Engine/Plugins/Runtime/nDisplay/Source/DisplayCluster/Private/Render/Synchronization/DisplayClusterRenderSyncPolicyBase.h
--- a/Engine/Plugins/Runtime/nDisplay/Source/DisplayCluster/Private/Render/Synchronization/DisplayClusterRenderSyncPolicyBase.h
+++ b/Engine/Plugins/Runtime/nDisplay/Source/DisplayCluster/Private/Render/Synchronization/DisplayClusterRenderSyncPolicyBase.h
@@ -19,28 +19,38 @@ public:
 	virtual ~FDisplayClusterRenderSyncPolicyBase() = default;
 
 public:
-	// Non-constructor initialization
+	//~ Begin IDisplayClusterRenderSyncPolicy interface
+
+	/** Non-constructor initialization */
 	virtual bool Initialize() override;
 
-	// VBlank related functions
+	/** VBlank related functions */
 	virtual bool IsWaitForVBlankFeatureSupported() override;
 	virtual bool WaitForVBlank() override;
 
+	//~ End IDisplayClusterRenderSyncPolicy interface
+
 public:
-	// Performs Ethernet barrier synchronization with other nodes in the cluster
+	/** Performs Ethernet barrier synchronization with other nodes in the cluster */
 	void SyncOnBarrier();
 
-	// Returns policy configuration parameters
+	/** Returns policy configuration parameters */
 	const TMap<FString, FString>& GetParameters() const
 	{
 		return Parameters;
 	}
 
 protected:
-	// Wait for frame rendering completion using a GPU fence
+	/** Wait for frame rendering completion using a GPU fence */
 	virtual void WaitForFrameCompletion();
 
+	/** Get maximum frame latency */
+	bool GetMaximumFrameLatency(uint32& OutMaximumFrameLatency);
+
+	/** Set maximum frame latency */
+	bool SetMaximumFrameLatency(uint32 MaximumFrameLatency);
+
 private:
-	// Sync policy configuration parameters
+	/** Sync policy configuration parameters */
 	TMap<FString, FString> Parameters;
 };
diff --git a/Engine/Plugins/Runtime/nDisplay/Source/DisplayCluster/Private/Render/Synchronization/DisplayClusterRenderSyncPolicyFactoryInternal.cpp b/Engine/Plugins/Runtime/nDisplay/Source/DisplayCluster/Private/Render/Synchronization/DisplayClusterRenderSyncPolicyFactoryInternal.cpp
--- a/Engine/Plugins/Runtime/nDisplay/Source/DisplayCluster/Private/Render/Synchronization/DisplayClusterRenderSyncPolicyFactoryInternal.cpp
+++ b/Engine/Plugins/Runtime/nDisplay/Source/DisplayCluster/Private/Render/Synchronization/DisplayClusterRenderSyncPolicyFactoryInternal.cpp
@@ -5,7 +5,8 @@
 #include "Render/Synchronization/DisplayClusterRenderSyncPolicyNone.h"
 #include "Render/Synchronization/DisplayClusterRenderSyncPolicyEthernet.h"
 #include "Render/Synchronization/DisplayClusterRenderSyncPolicyEthernetBarrier.h"
-#include "Render/Synchronization/DisplayClusterRenderSyncPolicyNvidia.h"
+#include "Render/Synchronization/DisplayClusterRenderSyncPolicyNvidiaSwapBarrier.h"
+#include "Render/Synchronization/DisplayClusterRenderSyncPolicyNvidiaPresentBarrier.h"
 
 #include "Misc/DisplayClusterLog.h"
 #include "Misc/DisplayClusterStrings.h"
@@ -13,6 +14,17 @@
 #include "RHI.h"
 #include "RHIDefinitions.h"
 
+
+// This cvar allows to replace the original NVIDIA sync policy with the new Present Barrier (DX12 only).
+// While in testing, it's controled by this cvar. Later on, Present Barrier policy will be exposed to
+// the GUI configurator as well as other policies are.
+static TAutoConsoleVariable<bool> CVarUseNvidiaPresentBarrierPolicy(
+	TEXT("nDisplay.sync.nvidia.UsePresentBarrierPolicy"),
+	false,
+	TEXT("NVIDIA Present Barrier synchronization policy (experimental)"),
+	ECVF_ReadOnly
+);
+
 TSharedPtr<IDisplayClusterRenderSyncPolicy> FDisplayClusterRenderSyncPolicyFactoryInternal::Create(const FString& InPolicyType, const TMap<FString, FString>& Parameters)
 {
 	const ERHIInterfaceType RHIType = RHIGetInterfaceType();
@@ -50,9 +62,14 @@ TSharedPtr<IDisplayClusterRenderSyncPolicy> FDisplayClusterRenderSyncPolicyFacto
 	else if (InPolicyType.Equals(DisplayClusterConfigurationStrings::config::cluster::render_sync::Nvidia, ESearchCase::IgnoreCase))
 	{
 #if PLATFORM_WINDOWS
-		if (RHIType == ERHIInterfaceType::D3D11 || RHIType == ERHIInterfaceType::D3D12)
+		if (RHIType == ERHIInterfaceType::D3D12 && CVarUseNvidiaPresentBarrierPolicy.GetValueOnGameThread())
+		{
+			UE_LOG(LogDisplayClusterRenderSync, Log, TEXT("Experimental NVIDIA synchronization approach will be used!"));
+			return MakeShared<FDisplayClusterRenderSyncPolicyNvidiaPresentBarrier>(Parameters);
+		}
+		else if (RHIType == ERHIInterfaceType::D3D11 || RHIType == ERHIInterfaceType::D3D12)
 		{
-			return MakeShared<FDisplayClusterRenderSyncPolicyNvidia>(Parameters);
+			return MakeShared<FDisplayClusterRenderSyncPolicyNvidiaSwapBarrier>(Parameters);
 		}
 		else if (RHIType == ERHIInterfaceType::Vulkan)
 		{
diff --git a/Engine/Plugins/Runtime/nDisplay/Source/DisplayCluster/Private/Render/Synchronization/DisplayClusterRenderSyncPolicyNvidia.h b/Engine/Plugins/Runtime/nDisplay/Source/DisplayCluster/Private/Render/Synchronization/DisplayClusterRenderSyncPolicyNvidia.h
deleted file mode 100644
diff --git a/Engine/Plugins/Runtime/nDisplay/Source/DisplayCluster/Private/Render/Synchronization/DisplayClusterRenderSyncPolicyNvidiaBase.cpp b/Engine/Plugins/Runtime/nDisplay/Source/DisplayCluster/Private/Render/Synchronization/DisplayClusterRenderSyncPolicyNvidiaBase.cpp
new file mode 100644
--- /dev/null
+++ b/Engine/Plugins/Runtime/nDisplay/Source/DisplayCluster/Private/Render/Synchronization/DisplayClusterRenderSyncPolicyNvidiaBase.cpp
@@ -0,0 +1,130 @@
+// Copyright Epic Games, Inc. All Rights Reserved.
+
+#include "Render/Synchronization/DisplayClusterRenderSyncPolicyNvidiaBase.h"
+
+#include "Misc/DisplayClusterLog.h"
+
+#if WITH_NVAPI
+	THIRD_PARTY_INCLUDES_START
+	#include "nvapi.h"
+	THIRD_PARTY_INCLUDES_END
+#endif // WITH_NVAPI
+
+
+// Enables aligning of the RHI threads on an Ethernet barrier before initializing
+// NVIDIA synchronization subsystem.
+static TAutoConsoleVariable<bool> CVarNvidiaSyncPreInitAlignment(
+	TEXT("nDisplay.sync.nvidia.PreInitAlignment"),
+	true,
+	TEXT("NVIDIA sync: Aligns RHI threads before initializing synchronization subsystem."),
+	ECVF_ReadOnly | ECVF_RenderThreadSafe
+);
+
+// Used to set the number of presentation cluster sync barriers before getting disabled.
+// It is meant to compensate for the fact that the node may not enter the Swap Group barrier on the first frame,
+// so an additional barrier may be necessary to ensure sync across the cluster.
+// This CVar may become obsolete when the driver can be queried to distinguish between pending and actually
+// having joined the swap group, which is not possible at the time, and the only indication of this is by looking at the
+// HUD that option 8 in ConfigureDriver.exe provides.
+static TAutoConsoleVariable<int32> CVarNvidiaSyncPrePresentAlignmentLimit(
+	TEXT("nDisplay.sync.nvidia.PrePresentAlignmentLimit"),
+	120,
+	TEXT("Sets a limit to the number of times (frames) a sync barrier is enforced before calling nvidia presentation.\n")
+	TEXT("Defaults to 120, which corresponds to 5s@24fps or 2s@60fps\n")
+	TEXT("    N < 0 : No limit\n")
+	TEXT("    N = 0 : Disabled\n")
+	TEXT("    N > 0 : The presentation sync barrier will only happen N times (frames)\n"),
+	ECVF_ReadOnly | ECVF_RenderThreadSafe
+);
+
+// Used to help the cluster start up in sync at the expense of performance. It forces the synchronization procedure
+// to happen after frame is done on the GPU side. But it stops doing that after the specified number of frames
+// to that it can run without the performance hit.
+static TAutoConsoleVariable<int32> CVarNvidiaSyncFrameCompletionLimit(
+	TEXT("nDisplay.sync.nvidia.FrameCompletionLimit"),
+	120,
+	TEXT("Sets a limit to the number of frames it waits for GPU completion.\n")
+	TEXT("Defaults to 120, which corresponds to 5s@24fps or 2s@60fps\n")
+	TEXT("    N < 0 : No limit\n")
+	TEXT("    N = 0 : Disabled\n")
+	TEXT("    N > 0 : The completion will only happen for N times (frames)\n"),
+	ECVF_ReadOnly | ECVF_RenderThreadSafe
+);
+
+// Sometimes we see a weird pattern of NvAPI_D3D1x_Present() handling in the utraces. Those patterns
+// visualize something that we would not expect to see if synchronization works properly.
+// NvAPI_D3D1x_Present() may return in between of the vblanks, or it may return asynchronously to other nodes.
+// This cvar enables additional barrier synchronization step right after NvAPI_D3D1x_Present() call.
+// This should help keeping RHI threads aligned.
+static TAutoConsoleVariable<bool> CVarNvidiaSyncPostPresentAlignment(
+	TEXT("nDisplay.sync.nvidia.PostPresentAlignment"),
+	false,
+	TEXT("Sync nodes on a network barrier after frame presentation\n"),
+	ECVF_ReadOnly | ECVF_RenderThreadSafe
+);
+
+// In the recent logs (Nov-Dec, 2023), we see GetMaximumFrameLatency() and SetMaximumFrameLatency() don't work
+// properly. This is something weird and may indicate there is an issue outside of UE. Those methods are called
+// once during the policy initialization. To better understand the issue and probably see some correlation,
+// this cvar allows to focribly call set/get frame latency every frame, and log the results.
+TAutoConsoleVariable<bool> CVarNvidiaSyncForceLatencyUpdateEveryFrame(
+	TEXT("nDisplay.sync.nvidia.ForceLatencyUpdateEveryFrame"),
+	false,
+	TEXT("Performs get/set latency every frame, and logs the results"),
+	ECVF_RenderThreadSafe
+);
+
+// Allows nvidia sync policy implementations to print their own statistics every frame.
+TAutoConsoleVariable<bool> CVarNvidiaSyncPrintStatsEveryFrame(
+	TEXT("nDisplay.sync.nvidia.ForceStatsEveryFrame"),
+	false,
+	TEXT("Enables sync stats every frame"),
+	ECVF_RenderThreadSafe
+);
+
+
+
+FDisplayClusterRenderSyncPolicyNvidiaBase::FDisplayClusterRenderSyncPolicyNvidiaBase(const TMap<FString, FString>& Parameters)
+	: Super(Parameters)
+	, bCfgPreInitAlignment(CVarNvidiaSyncPreInitAlignment.GetValueOnAnyThread())
+	, CfgPrePresentAlignmentLimit(CVarNvidiaSyncPrePresentAlignmentLimit.GetValueOnAnyThread())
+	, CfgFrameCompletionLimit(CVarNvidiaSyncFrameCompletionLimit.GetValueOnAnyThread())
+	, bCfgPostPresentAlignment(CVarNvidiaSyncPostPresentAlignment.GetValueOnAnyThread())
+{
+#if WITH_NVAPI
+	// Initialize NVIDIA library
+	const NvAPI_Status NvApiResult = NvAPI_Initialize();
+	if (NvApiResult == NVAPI_OK)
+	{
+		bNvLibraryInitialized = true;
+	}
+	else
+	{
+		UE_LOG(LogDisplayClusterRenderSync, Error, TEXT("NvAPI_Initialize() failed, error 0x%x"), NvApiResult);
+	}
+#endif // WITH_NVAPI
+
+	UE_LOG(LogDisplayClusterRenderSync, Log,
+		TEXT("NVIDIA sync configuration: PreInitAlignment=%d, PrePresentAlignmentLimit=%d, FrameCompletionLimit=%d, PostPresentAlignment=%d"),
+		bCfgPreInitAlignment ? 1 : 0,
+		CfgPrePresentAlignmentLimit,
+		CfgFrameCompletionLimit,
+		bCfgPostPresentAlignment ? 1 : 0
+	);
+}
+
+FDisplayClusterRenderSyncPolicyNvidiaBase::~FDisplayClusterRenderSyncPolicyNvidiaBase()
+{
+	// Release NVIDIA library
+	if (bNvLibraryInitialized)
+	{
+#if WITH_NVAPI
+		NvAPI_Unload();
+#endif // WITH_NVAPI
+	}
+}
+
+bool FDisplayClusterRenderSyncPolicyNvidiaBase::Initialize()
+{
+	return bNvLibraryInitialized;
+}
diff --git a/Engine/Plugins/Runtime/nDisplay/Source/DisplayCluster/Private/Render/Synchronization/DisplayClusterRenderSyncPolicyNvidiaBase.h b/Engine/Plugins/Runtime/nDisplay/Source/DisplayCluster/Private/Render/Synchronization/DisplayClusterRenderSyncPolicyNvidiaBase.h
new file mode 100644
--- /dev/null
+++ b/Engine/Plugins/Runtime/nDisplay/Source/DisplayCluster/Private/Render/Synchronization/DisplayClusterRenderSyncPolicyNvidiaBase.h
@@ -0,0 +1,64 @@
+// Copyright Epic Games, Inc. All Rights Reserved.
+
+#pragma once
+
+#include "Render/Synchronization/DisplayClusterRenderSyncPolicyBase.h"
+
+#include "HAL/IConsoleManager.h"
+
+
+// CVars available for change in runtime
+extern TAutoConsoleVariable<bool> CVarNvidiaSyncForceLatencyUpdateEveryFrame;
+extern TAutoConsoleVariable<bool> CVarNvidiaSyncPrintStatsEveryFrame;
+
+
+/**
+ * Abstract NVIDIA synchronization policy.
+ * 
+ * It's used as a base class for all NVIDIA sync policies. Encapsulates common NVIDIA synchronization features.
+ */
+class FDisplayClusterRenderSyncPolicyNvidiaBase
+	: public FDisplayClusterRenderSyncPolicyBase
+{
+private:
+	using Super = FDisplayClusterRenderSyncPolicyBase;
+
+public:
+	FDisplayClusterRenderSyncPolicyNvidiaBase(const TMap<FString, FString>& Parameters);
+	virtual ~FDisplayClusterRenderSyncPolicyNvidiaBase();
+
+public:
+	//~ Begin IDisplayClusterRenderSyncPolicy interface
+
+	/** Initialization */
+	virtual bool Initialize() override;
+
+	//~ End IDisplayClusterRenderSyncPolicy interface
+
+protected:
+	/** Whether NVIDIA API library was initialized */
+	bool bNvLibraryInitialized = false;
+
+	/** Whether NVIDIA sync initialization function was called (to prevent multiple initialization calls) */
+	bool bNvSyncInitializationCalled = false;
+
+	/** Whether NVIDIA sync initialization succeeded */
+	bool bNvSyncInitializedSuccessfully = false;
+
+protected:
+	/** CVar config: synchronous initialization (alignment on a barrier before initializing NVIDIA synchronization) */
+	const bool bCfgPreInitAlignment;
+
+	/** CVar config: pre-present threads alignment on a network barrier */
+	const int32 CfgPrePresentAlignmentLimit;
+	/** Pre-present sync counter */
+	int32 PrePresentAlignmentCounter = 0;
+
+	/** CVar config: frame completion limit */
+	const int32 CfgFrameCompletionLimit;
+	/** Frame completion counter */
+	int32 FrameCompletionCounter = 0;
+
+	/** CVar config: post-present threads alignment on a network barrier */
+	const bool bCfgPostPresentAlignment;
+};
diff --git a/Engine/Plugins/Runtime/nDisplay/Source/DisplayCluster/Private/Render/Synchronization/DisplayClusterRenderSyncPolicyNvidiaPresentBarrier.h b/Engine/Plugins/Runtime/nDisplay/Source/DisplayCluster/Private/Render/Synchronization/DisplayClusterRenderSyncPolicyNvidiaPresentBarrier.h
new file mode 100644
--- /dev/null
+++ b/Engine/Plugins/Runtime/nDisplay/Source/DisplayCluster/Private/Render/Synchronization/DisplayClusterRenderSyncPolicyNvidiaPresentBarrier.h
@@ -0,0 +1,46 @@
+// Copyright Epic Games, Inc. All Rights Reserved.
+
+#pragma once
+
+#include "Render/Synchronization/DisplayClusterRenderSyncPolicyNvidiaBase.h"
+#include "DisplayClusterConfigurationStrings.h"
+
+
+/**
+ * NVIDIA Present Barrier synchronization policy (new NVDIA sync approach)
+ */
+class FDisplayClusterRenderSyncPolicyNvidiaPresentBarrier
+	: public FDisplayClusterRenderSyncPolicyNvidiaBase
+{
+private:
+	using Super = FDisplayClusterRenderSyncPolicyNvidiaBase;
+
+public:
+	FDisplayClusterRenderSyncPolicyNvidiaPresentBarrier(const TMap<FString, FString>& Parameters);
+	virtual ~FDisplayClusterRenderSyncPolicyNvidiaPresentBarrier();
+
+	//~ Begin IDisplayClusterRenderSyncPolicy interface
+public:
+
+	/** Returns name of this policy implementation */
+	virtual FName GetName() const override
+	{
+		return FName(DisplayClusterConfigurationStrings::config::cluster::render_sync::NvidiaAdvanced);
+	}
+
+	/** Performs rendering and frame output synchronization */
+	virtual bool SynchronizeClusterRendering(int32& InOutSyncInterval) override;
+
+	//~ End IDisplayClusterRenderSyncPolicy interface
+
+private:
+	/** Handles frame presentation events */
+	void OnFramePresented(bool bNativePresent);
+
+	/** Prints present barrier stats to the log */
+	void LogPresentBarrierStats();
+
+private:
+	/** Initializes NVIDIA present barrier */
+	bool InitializePresentBarrier();
+};
diff --git a/Engine/Plugins/Runtime/nDisplay/Source/DisplayCluster/Private/Render/Synchronization/DisplayClusterRenderSyncPolicyNvidiaSwapBarrier.h b/Engine/Plugins/Runtime/nDisplay/Source/DisplayCluster/Private/Render/Synchronization/DisplayClusterRenderSyncPolicyNvidiaSwapBarrier.h
new file mode 100644
--- /dev/null
+++ b/Engine/Plugins/Runtime/nDisplay/Source/DisplayCluster/Private/Render/Synchronization/DisplayClusterRenderSyncPolicyNvidiaSwapBarrier.h
@@ -0,0 +1,46 @@
+// Copyright Epic Games, Inc. All Rights Reserved.
+
+#pragma once
+
+#include "Render/Synchronization/DisplayClusterRenderSyncPolicyNvidiaBase.h"
+#include "DisplayClusterConfigurationStrings.h"
+
+
+/**
+ * NVIDIA Swap Barrier synchronization policy (old NVDIA sync approach)
+ */
+class FDisplayClusterRenderSyncPolicyNvidiaSwapBarrier
+	: public FDisplayClusterRenderSyncPolicyNvidiaBase
+{
+private:
+	using Super = FDisplayClusterRenderSyncPolicyNvidiaBase;
+
+public:
+	FDisplayClusterRenderSyncPolicyNvidiaSwapBarrier(const TMap<FString, FString>& Parameters);
+	virtual ~FDisplayClusterRenderSyncPolicyNvidiaSwapBarrier();
+
+	//~ Begin IDisplayClusterRenderSyncPolicy interface
+public:
+
+	/** Returns name of this policy implementation */
+	virtual FName GetName() const override
+	{
+		return FName(DisplayClusterConfigurationStrings::config::cluster::render_sync::Nvidia);
+	}
+
+	/** Performs rendering and frame output synchronization */
+	virtual bool SynchronizeClusterRendering(int32& InOutSyncInterval) override;
+
+	//~ End IDisplayClusterRenderSyncPolicy interface
+
+private:
+	/** Initializes NVIDIA internals */
+	bool InitializeNvidiaSwapLock();
+
+private:
+	/** Policy config parameter: Swap group ID */
+	uint32 RequestedGroup   = 1;
+
+	/** Policy config parameter: Swap barrier ID */
+	uint32 RequestedBarrier = 1;
+};
diff --git a/Engine/Plugins/Runtime/nDisplay/Source/DisplayCluster/Private/Render/Synchronization/Linux/DisplayClusterRenderSyncHelperLinux.cpp b/Engine/Plugins/Runtime/nDisplay/Source/DisplayCluster/Private/Render/Synchronization/Linux/DisplayClusterRenderSyncHelperLinux.cpp
--- a/Engine/Plugins/Runtime/nDisplay/Source/DisplayCluster/Private/Render/Synchronization/Linux/DisplayClusterRenderSyncHelperLinux.cpp
+++ b/Engine/Plugins/Runtime/nDisplay/Source/DisplayCluster/Private/Render/Synchronization/Linux/DisplayClusterRenderSyncHelperLinux.cpp
@@ -36,3 +36,13 @@ bool FDisplayClusterRenderSyncHelper::FDCRSHelperVulkan::WaitForVBlank()
 {
 	return false;
 }
+
+bool FDisplayClusterRenderSyncHelper::FDCRSHelperVulkan::GetMaximumFrameLatency(uint32& OutMaximumFrameLatency)
+{
+	return false;
+}
+
+bool FDisplayClusterRenderSyncHelper::FDCRSHelperVulkan::SetMaximumFrameLatency(uint32 FrameLatency)
+{
+	return false;
+}
diff --git a/Engine/Plugins/Runtime/nDisplay/Source/DisplayCluster/Private/Render/Synchronization/Linux/DisplayClusterRenderSyncPolicyNvidiaLinux.cpp b/Engine/Plugins/Runtime/nDisplay/Source/DisplayCluster/Private/Render/Synchronization/Linux/DisplayClusterRenderSyncPolicyNvidiaLinux.cpp
deleted file mode 100644
diff --git a/Engine/Plugins/Runtime/nDisplay/Source/DisplayCluster/Private/Render/Synchronization/Linux/DisplayClusterRenderSyncPolicyNvidiaPresentBarrierLinux.cpp b/Engine/Plugins/Runtime/nDisplay/Source/DisplayCluster/Private/Render/Synchronization/Linux/DisplayClusterRenderSyncPolicyNvidiaPresentBarrierLinux.cpp
new file mode 100644
--- /dev/null
+++ b/Engine/Plugins/Runtime/nDisplay/Source/DisplayCluster/Private/Render/Synchronization/Linux/DisplayClusterRenderSyncPolicyNvidiaPresentBarrierLinux.cpp
@@ -0,0 +1,36 @@
+// Copyright Epic Games, Inc. All Rights Reserved.
+
+#include "Render/Synchronization/DisplayClusterRenderSyncPolicyNvidiaPresentBarrier.h"
+
+
+FDisplayClusterRenderSyncPolicyNvidiaPresentBarrier::FDisplayClusterRenderSyncPolicyNvidiaPresentBarrier(const TMap<FString, FString>& Parameters)
+	: Super(Parameters)
+{
+}
+
+FDisplayClusterRenderSyncPolicyNvidiaPresentBarrier::~FDisplayClusterRenderSyncPolicyNvidiaPresentBarrier()
+{
+}
+
+bool FDisplayClusterRenderSyncPolicyNvidiaPresentBarrier::InitializePresentBarrier()
+{
+	return true;
+}
+
+bool FDisplayClusterRenderSyncPolicyNvidiaPresentBarrier::SynchronizeClusterRendering(int32& InOutSyncInterval)
+{
+	// Wait unless the frame is rendered
+	WaitForFrameCompletion();
+	// Sync on the network barrier
+	SyncOnBarrier();
+	// Ask engine to present the frame
+	return true;
+}
+
+void FDisplayClusterRenderSyncPolicyNvidiaPresentBarrier::OnFramePresented(bool bNativePresent)
+{
+}
+
+void FDisplayClusterRenderSyncPolicyNvidiaPresentBarrier::LogPresentBarrierStats()
+{
+}
diff --git a/Engine/Plugins/Runtime/nDisplay/Source/DisplayCluster/Private/Render/Synchronization/Linux/DisplayClusterRenderSyncPolicyNvidiaSwapBarrierLinux.cpp b/Engine/Plugins/Runtime/nDisplay/Source/DisplayCluster/Private/Render/Synchronization/Linux/DisplayClusterRenderSyncPolicyNvidiaSwapBarrierLinux.cpp
new file mode 100644
--- /dev/null
+++ b/Engine/Plugins/Runtime/nDisplay/Source/DisplayCluster/Private/Render/Synchronization/Linux/DisplayClusterRenderSyncPolicyNvidiaSwapBarrierLinux.cpp
@@ -0,0 +1,28 @@
+// Copyright Epic Games, Inc. All Rights Reserved.
+
+#include "Render/Synchronization/DisplayClusterRenderSyncPolicyNvidiaSwapBarrier.h"
+
+
+FDisplayClusterRenderSyncPolicyNvidiaSwapBarrier::FDisplayClusterRenderSyncPolicyNvidiaSwapBarrier(const TMap<FString, FString>& Parameters)
+	: Super(Parameters)
+{
+}
+
+FDisplayClusterRenderSyncPolicyNvidiaSwapBarrier::~FDisplayClusterRenderSyncPolicyNvidiaSwapBarrier()
+{
+}
+
+bool FDisplayClusterRenderSyncPolicyNvidiaSwapBarrier::InitializeNvidiaSwapLock()
+{
+	return true;
+}
+
+bool FDisplayClusterRenderSyncPolicyNvidiaSwapBarrier::SynchronizeClusterRendering(int32& InOutSyncInterval)
+{
+	// Wait unless the frame is rendered
+	WaitForFrameCompletion();
+	// Sync on the network barrier
+	SyncOnBarrier();
+	// Ask engine to present the frame
+	return true;
+}
diff --git a/Engine/Plugins/Runtime/nDisplay/Source/DisplayCluster/Private/Render/Synchronization/Windows/DisplayClusterRenderSyncHelperWindows.cpp b/Engine/Plugins/Runtime/nDisplay/Source/DisplayCluster/Private/Render/Synchronization/Windows/DisplayClusterRenderSyncHelperWindows.cpp
--- a/Engine/Plugins/Runtime/nDisplay/Source/DisplayCluster/Private/Render/Synchronization/Windows/DisplayClusterRenderSyncHelperWindows.cpp
+++ b/Engine/Plugins/Runtime/nDisplay/Source/DisplayCluster/Private/Render/Synchronization/Windows/DisplayClusterRenderSyncHelperWindows.cpp
@@ -2,6 +2,7 @@
 
 #include "Render/Synchronization/DisplayClusterRenderSyncHelper.h"
 
+#include "Misc/DisplayClusterLog.h"
 #include "Misc/DisplayClusterStrings.h"
 
 #include "Engine/Engine.h"
@@ -10,12 +11,12 @@
 #include "DynamicRHI.h"
 #include "UnrealClient.h"
 
-#pragma warning(push)
-#pragma warning (disable : 4005)
 #include "Windows/AllowWindowsPlatformTypes.h"
-#include "dxgi.h"
+THIRD_PARTY_INCLUDES_START
+#include "dxgi1_3.h"
+THIRD_PARTY_INCLUDES_END
 #include "Windows/HideWindowsPlatformTypes.h"
-#pragma warning(pop)
+
 
 
 TUniquePtr<IDisplayClusterRenderSyncHelper> FDisplayClusterRenderSyncHelper::CreateHelper()
@@ -74,6 +75,60 @@ bool FDisplayClusterRenderSyncHelper::FDCRSHelperDX::WaitForVBlank()
 	return false;
 }
 
+bool FDisplayClusterRenderSyncHelper::FDCRSHelperDX::GetMaximumFrameLatency(uint32& OutMaximumFrameLatency)
+{
+	if (GEngine && GEngine->GameViewport && GEngine->GameViewport->Viewport)
+	{
+		if (FRHIViewport* RHIViewport = GEngine->GameViewport->Viewport->GetViewportRHI().GetReference())
+		{
+			if (IDXGISwapChain2* DXGISwapChain = static_cast<IDXGISwapChain2*>(RHIViewport->GetNativeSwapChain()))
+			{
+				const HRESULT Result = DXGISwapChain->GetMaximumFrameLatency(&OutMaximumFrameLatency);
+
+				if (Result == S_OK)
+				{
+					UE_LOG(LogDisplayClusterRenderSync, Verbose, TEXT("Swapchain frame latency: %u"), OutMaximumFrameLatency);
+				}
+				else
+				{
+					UE_LOG(LogDisplayClusterRenderSync, Warning, TEXT("Couldn't get maximum frame latency. Error: %x"), Result);
+				}
+
+				return Result == S_OK;
+			}
+		}
+	}
+
+	return false;
+}
+
+bool FDisplayClusterRenderSyncHelper::FDCRSHelperDX::SetMaximumFrameLatency(uint32 MaximumFrameLatency)
+{
+	if (GEngine && GEngine->GameViewport && GEngine->GameViewport->Viewport)
+	{
+		if (FRHIViewport* RHIViewport = GEngine->GameViewport->Viewport->GetViewportRHI().GetReference())
+		{
+			if (IDXGISwapChain2* DXGISwapChain = static_cast<IDXGISwapChain2*>(RHIViewport->GetNativeSwapChain()))
+			{
+				const HRESULT Result = DXGISwapChain->SetMaximumFrameLatency(MaximumFrameLatency);
+
+				if (Result == S_OK)
+				{
+					UE_LOG(LogDisplayClusterRenderSync, Verbose, TEXT("Swapchain frame latency was set to: %u"), MaximumFrameLatency);
+				}
+				else
+				{
+					UE_LOG(LogDisplayClusterRenderSync, Warning, TEXT("Couldn't set maximum frame latency to %u. Error: %x"), MaximumFrameLatency, Result);
+				}
+
+				return Result == S_OK;
+			}
+		}
+	}
+
+	return false;
+}
+
 bool FDisplayClusterRenderSyncHelper::FDCRSHelperVulkan::IsWaitForVBlankSupported()
 {
 	return false;
@@ -83,3 +138,13 @@ bool FDisplayClusterRenderSyncHelper::FDCRSHelperVulkan::WaitForVBlank()
 {
 	return false;
 }
+
+bool FDisplayClusterRenderSyncHelper::FDCRSHelperVulkan::GetMaximumFrameLatency(uint32& OutMaximumFrameLatency)
+{
+	return false;
+}
+
+bool FDisplayClusterRenderSyncHelper::FDCRSHelperVulkan::SetMaximumFrameLatency(uint32 MaximumFrameLatency)
+{
+	return false;
+}
diff --git a/Engine/Plugins/Runtime/nDisplay/Source/DisplayCluster/Private/Render/Synchronization/Windows/DisplayClusterRenderSyncPolicyEthernetWindows.cpp b/Engine/Plugins/Runtime/nDisplay/Source/DisplayCluster/Private/Render/Synchronization/Windows/DisplayClusterRenderSyncPolicyEthernetWindows.cpp
--- a/Engine/Plugins/Runtime/nDisplay/Source/DisplayCluster/Private/Render/Synchronization/Windows/DisplayClusterRenderSyncPolicyEthernetWindows.cpp
+++ b/Engine/Plugins/Runtime/nDisplay/Source/DisplayCluster/Private/Render/Synchronization/Windows/DisplayClusterRenderSyncPolicyEthernetWindows.cpp
@@ -2,7 +2,9 @@
 
 #include "Render/Synchronization/DisplayClusterRenderSyncPolicyEthernet.h"
 
+#include "Misc/DisplayClusterGlobals.h"
 #include "Misc/DisplayClusterLog.h"
+#include "DisplayClusterCallbacks.h"
 #include "DisplayClusterConfigurationStrings.h"
 
 #include "Engine/Engine.h"
@@ -85,6 +87,12 @@ bool FDisplayClusterRenderSyncPolicyEthernet::SynchronizeClusterRendering(int32&
 		}
 	}
 
+	if (!bNeedEnginePresent)
+	{
+		// Notify custom presentation was done
+		GDisplayCluster->GetCallbacks().OnDisplayClusterFramePresented_RHIThread().Broadcast(false);
+	}
+
 	++FrameCounter;
 
 	return bNeedEnginePresent;
diff --git a/Engine/Plugins/Runtime/nDisplay/Source/DisplayCluster/Private/Render/Synchronization/Windows/DisplayClusterRenderSyncPolicyNvidiaPresentBarrierWindows.cpp b/Engine/Plugins/Runtime/nDisplay/Source/DisplayCluster/Private/Render/Synchronization/Windows/DisplayClusterRenderSyncPolicyNvidiaPresentBarrierWindows.cpp
new file mode 100644
--- /dev/null
+++ b/Engine/Plugins/Runtime/nDisplay/Source/DisplayCluster/Private/Render/Synchronization/Windows/DisplayClusterRenderSyncPolicyNvidiaPresentBarrierWindows.cpp
@@ -0,0 +1,342 @@
+// Copyright Epic Games, Inc. All Rights Reserved.
+
+#include "Render/Synchronization/DisplayClusterRenderSyncPolicyNvidiaPresentBarrier.h"
+
+#include "DisplayClusterCallbacks.h"
+
+#include "Misc/DisplayClusterGlobals.h"
+#include "Misc/DisplayClusterLog.h"
+
+#include "Engine/Engine.h"
+#include "Engine/GameViewportClient.h"
+#include "UnrealClient.h"
+
+#include "Windows/AllowWindowsPlatformTypes.h"
+THIRD_PARTY_INCLUDES_START
+	#include "dxgi1_3.h"
+	#include "d3d12.h"
+	#include "nvapi.h"
+	#include <combaseapi.h>
+THIRD_PARTY_INCLUDES_END
+#include "Windows/HideWindowsPlatformTypes.h"
+
+
+namespace DisplayClusterRenderSyncPolicyNvidiaPresentBarrier_Data_Windows
+{
+	// Active output device
+	ID3D12Device*    D3DDevice = nullptr;
+	// Active swapchain
+	IDXGISwapChain2* DXGISwapChain = nullptr;
+	// Synchronization fence
+	ID3D12Fence*     PresentBarrierFence = nullptr;
+	// Barrier client handle
+	NvPresentBarrierClientHandle NvBarrierClientHandle = nullptr;
+	// Backbuffer references to keep the resources alive
+	TArray<TRefCountPtr<ID3D12Resource>> BackBuffers;
+}
+
+
+FDisplayClusterRenderSyncPolicyNvidiaPresentBarrier::FDisplayClusterRenderSyncPolicyNvidiaPresentBarrier(const TMap<FString, FString>& Parameters)
+	: Super(Parameters)
+{
+	GDisplayCluster->GetCallbacks().OnDisplayClusterFramePresented_RHIThread().AddRaw(this, &FDisplayClusterRenderSyncPolicyNvidiaPresentBarrier::OnFramePresented);
+}
+
+FDisplayClusterRenderSyncPolicyNvidiaPresentBarrier::~FDisplayClusterRenderSyncPolicyNvidiaPresentBarrier()
+{
+	using namespace DisplayClusterRenderSyncPolicyNvidiaPresentBarrier_Data_Windows;
+
+	GDisplayCluster->GetCallbacks().OnDisplayClusterFramePresented_RHIThread().RemoveAll(this);
+
+	if (bNvLibraryInitialized)
+	{
+		// Release present barrier client
+		if (NvBarrierClientHandle)
+		{
+			NvAPI_LeavePresentBarrier(NvBarrierClientHandle);
+			NvAPI_DestroyPresentBarrierClient(NvBarrierClientHandle);
+		}
+
+		// Release backbuffer references if there are any
+		BackBuffers.Empty();
+	}
+}
+
+bool FDisplayClusterRenderSyncPolicyNvidiaPresentBarrier::SynchronizeClusterRendering(int32& InOutSyncInterval)
+{
+	// Initialize barriers at first call
+	if (!bNvSyncInitializationCalled)
+	{
+		// Initialzie NVIDIA present barrier
+		bNvSyncInitializedSuccessfully = InitializePresentBarrier();
+
+		// Print initial stats before presenting any frames
+		if (bNvSyncInitializedSuccessfully)
+		{
+			LogPresentBarrierStats();
+		}
+	}
+
+	// Update maximum frame latency every if requested
+	if (CVarNvidiaSyncForceLatencyUpdateEveryFrame.GetValueOnAnyThread())
+	{
+		SetMaximumFrameLatency(1);
+	}
+
+	// Wait until frame rendering is finished
+	if (CfgFrameCompletionLimit != 0 && FrameCompletionCounter < CfgFrameCompletionLimit)
+	{
+		{
+			TRACE_CPUPROFILER_EVENT_SCOPE(nDisplay FrameCompletion);
+			WaitForFrameCompletion();
+		}
+
+		if (++FrameCompletionCounter == CfgFrameCompletionLimit)
+		{
+			UE_LOG(LogDisplayClusterRenderSync, Log, TEXT("NVS_PB: Disabled completion after %d frames"), FrameCompletionCounter);
+		}
+	}
+
+	// Align all threads on the timescale before calling NvAPI_D3D1x_Present
+	// As a side-effect, it should avoid NvAPI_D3D1x_Present stuck on application kill
+	if (CfgPrePresentAlignmentLimit != 0 && PrePresentAlignmentCounter < CfgPrePresentAlignmentLimit)
+	{
+		{
+			TRACE_CPUPROFILER_EVENT_SCOPE(nDisplay PrePresentAlignment);
+			SyncOnBarrier();
+		}
+
+		if (++PrePresentAlignmentCounter == CfgPrePresentAlignmentLimit)
+		{
+			UE_LOG(LogDisplayClusterRenderSync, Log, TEXT("NVS_PB: Disabled PrePresent sync after %d frames"), PrePresentAlignmentCounter);
+		}
+	}
+
+	// This synchronization approach automatically hooks for present calls so we don't need
+	// to do anything special over here. By returning true, we ask engine to present frame.
+	// Present barrier works transparently.
+	return true;
+}
+
+bool FDisplayClusterRenderSyncPolicyNvidiaPresentBarrier::InitializePresentBarrier()
+{
+	using namespace DisplayClusterRenderSyncPolicyNvidiaPresentBarrier_Data_Windows;
+
+	// Remember the fact this function has been called
+	bNvSyncInitializationCalled = true;
+
+	// No sense to proceed if NVIDIA library has not been initialized
+	if (!bNvLibraryInitialized)
+	{
+		UE_LOG(LogDisplayClusterRenderSync, Error, TEXT("NVS_PB: Couldn't initialize present barrier because NVIDIA library has not been initialized."));
+		return false;
+	}
+
+	// Get D3D1XDevice
+	D3DDevice = static_cast<ID3D12Device*>(GDynamicRHI->RHIGetNativeDevice());
+	check(D3DDevice);
+
+	// Get IDXGISwapChain
+	DXGISwapChain = static_cast<IDXGISwapChain2*>(GEngine->GameViewport->Viewport->GetViewportRHI().GetReference()->GetNativeSwapChain());
+	check(DXGISwapChain);
+
+	if (!D3DDevice || !DXGISwapChain)
+	{
+		UE_LOG(LogDisplayClusterRenderSync, Error, TEXT("NVS_PB: Couldn't get DX device or swap chain while initializing present barrier sync policy"));
+		return false;
+	}
+
+	// Set frame latency
+	SetMaximumFrameLatency(1);
+
+	NvAPI_Status NVResult = NVAPI_ERROR;
+
+	// Check if present barrier is supported
+	{
+		bool bSupported = false;
+		NVResult = NvAPI_D3D12_QueryPresentBarrierSupport(D3DDevice, &bSupported);
+		if (NVResult != NVAPI_OK || !bSupported)
+		{
+			UE_LOG(LogDisplayClusterRenderSync, Error, TEXT("NVS_PB: Present barrier is not supported on current device. Error=%d."), NVResult);
+			return false;
+		}
+		else
+		{
+			UE_LOG(LogDisplayClusterRenderSync, Log, TEXT("NVS_PB: Present barrier is supported"));
+		}
+	}
+
+	// Create present barrier client handle
+	{
+		NVResult = NvAPI_D3D12_CreatePresentBarrierClient(D3DDevice, DXGISwapChain, &NvBarrierClientHandle);
+		if (NVResult != NVAPI_OK || !NvBarrierClientHandle)
+		{
+			UE_LOG(LogDisplayClusterRenderSync, Error, TEXT("NVS_PB: Couldn't initialize present barrier client. Error=%d."), NVResult);
+			return false;
+		}
+		else
+		{
+			UE_LOG(LogDisplayClusterRenderSync, Log, TEXT("NVS_PB: Present barrier client has been initialized"));
+		}
+	}
+
+	// Create sync fence
+	{
+		HRESULT HResult = D3DDevice->CreateFence(0, D3D12_FENCE_FLAG_NONE, IID_PPV_ARGS(&PresentBarrierFence));
+		if (FAILED(HResult))
+		{
+			UE_LOG(LogDisplayClusterRenderSync, Error, TEXT("NVS_PB: Couldn't create ID3D12Fence instance. Error=%d."), HResult);
+			return false;
+		}
+		else
+		{
+			UE_LOG(LogDisplayClusterRenderSync, Log, TEXT("NVS_PB: Successfully obtained an instance of ID3D12Fence"));
+		}
+	}
+
+	// Register present barrier resources
+	{
+		// Get swap chain info (we need backbuffers amount)
+		DXGI_SWAP_CHAIN_DESC SwapChainDesc;
+		HRESULT HResult = DXGISwapChain->GetDesc(&SwapChainDesc);
+		if (FAILED(HResult))
+		{
+			UE_LOG(LogDisplayClusterRenderSync, Error, TEXT("NVS_PB: Couldn't create SwapChain info. Error=%d."), HResult);
+			return false;
+		}
+
+		// Get backbuffer references
+		TArray<ID3D12Resource*> BackBuffersRaw;
+		BackBuffersRaw.SetNumZeroed(SwapChainDesc.BufferCount);
+		BackBuffers.SetNumZeroed(SwapChainDesc.BufferCount);
+		for (uint32 BufferIdx = 0; BufferIdx < SwapChainDesc.BufferCount; ++BufferIdx)
+		{
+			DXGISwapChain->GetBuffer(BufferIdx, IID_PPV_ARGS(BackBuffers[BufferIdx].GetInitReference()));
+			BackBuffersRaw[BufferIdx] = BackBuffers[BufferIdx];
+		}
+
+		// Register backbuffers
+		NVResult = NvAPI_D3D12_RegisterPresentBarrierResources(NvBarrierClientHandle, PresentBarrierFence, BackBuffersRaw.GetData(), SwapChainDesc.BufferCount);
+		if (NVResult != NVAPI_OK)
+		{
+			UE_LOG(LogDisplayClusterRenderSync, Error, TEXT("NVS_PB: Couldn't register barrier resources. Error=%d."), NVResult);
+			return false;
+		}
+		else
+		{
+			UE_LOG(LogDisplayClusterRenderSync, Log, TEXT("NVS_PB: Barrier resources have been registered"));
+		}
+	}
+
+	{
+		// Here we initialize NVIDIA sync on the same frame interval on the timescale
+		if (bCfgPreInitAlignment)
+		{
+			TRACE_CPUPROFILER_EVENT_SCOPE(nDisplay NVIDIA Sync Init);
+
+			// Align all the threads on the same timeline. It's very likely all the nodes got free on the same side of a potential
+			// upcoming V-blank. Normally, the barrier synchronization takes up to 500 microseconds which is about 3% of frame time.
+			// So the probability of a situation where some nodes left the barrier before V-blank while others after is really low.
+			// However, it's still possible. And cluster restart should likely be successfull.
+			SyncOnBarrier();
+
+			// Assuming all the nodes in the right place and have some time before V-blank. Let's wait untill it happend.
+			WaitForVBlank();
+
+			// We don't really know either we're at the very beginning of V-blank interval or in the end. Let's spend some time
+			// on the barrier to let V-blank interval over.
+			SyncOnBarrier();
+
+			// Yes, the task scheduler can break the logic above by allocating CPU resources to this thread like 20ms later which
+			// makes us behind of the next V-blank for 60Hz output. And we probably should play around thread priorities to reduce
+			// probability of that. But hope we're good here. Usually the systems that use NVIDIA sync approach have a lot of CPU
+			// cores and don't run any other applications that compete for CPU.
+			// Moreover, we will initialize NVIDIA sync for all cluster nodes before presenting first frame. So the initialization
+			// process should be fine anyway.
+		}
+	}
+
+	// Join the barrier
+	{
+		NV_JOIN_PRESENT_BARRIER_PARAMS Params = { };
+		Params.dwVersion = NV_JOIN_PRESENT_BARRIER_PARAMS_VER1;
+		NVResult = NvAPI_JoinPresentBarrier(NvBarrierClientHandle, &Params);
+		if (NVResult != NVAPI_OK)
+		{
+			UE_LOG(LogDisplayClusterRenderSync, Error, TEXT("NVS_PB: Couldn't join the barrier. Error=%d."), NVResult);
+			return false;
+		}
+		else
+		{
+			UE_LOG(LogDisplayClusterRenderSync, Log, TEXT("NVS_PB: Joined the barrier successfully"));
+		}
+	}
+
+	UE_LOG(LogDisplayClusterRenderSync, Log, TEXT("NVS_PB: Initialized successfully"));
+
+	return true;
+}
+
+void FDisplayClusterRenderSyncPolicyNvidiaPresentBarrier::OnFramePresented(bool bNativePresent)
+{
+	if (CVarNvidiaSyncPrintStatsEveryFrame.GetValueOnAnyThread())
+	{
+		LogPresentBarrierStats();
+	}
+}
+
+// A helper function to convert ID->TEXT for friendly logging
+static constexpr const TCHAR* GetSyncModeName(NV_PRESENT_BARRIER_SYNC_MODE SyncMode)
+{
+	switch (SyncMode)
+	{
+	case PRESENT_BARRIER_NOT_JOINED:
+		// The client hasn't joined presentBarrier
+		return TEXT("PRESENT_BARRIER_NOT_JOINED");
+
+	case PRESENT_BARRIER_SYNC_CLIENT:
+		// The client joined the presentBarrier, but is not synchronized with
+		// any other presentBarrier clients. This happens if the back buffers
+		// of this client are composited instead of being flipped out to screen
+		return TEXT("PRESENT_BARRIER_SYNC_CLIENT");
+
+	case PRESENT_BARRIER_SYNC_SYSTEM:
+		// The client joined the presentBarrier, and is synchronized with other
+		// presentBarrier clients within the system
+		return TEXT("PRESENT_BARRIER_SYNC_SYSTEM");
+
+	case PRESENT_BARRIER_SYNC_CLUSTER:
+		// The client joined the presentBarrier, and is synchronized with other
+		// clients within the system and across systems through QSync devices
+		return TEXT("PRESENT_BARRIER_SYNC_CLUSTER");
+
+	default:
+		return TEXT("UNKNOWN_STATE");
+	}
+}
+
+void FDisplayClusterRenderSyncPolicyNvidiaPresentBarrier::LogPresentBarrierStats()
+{
+	using namespace DisplayClusterRenderSyncPolicyNvidiaPresentBarrier_Data_Windows;
+
+	if (bNvLibraryInitialized && bNvSyncInitializedSuccessfully)
+	{
+		NV_PRESENT_BARRIER_FRAME_STATISTICS PBStats = { };
+		PBStats.dwVersion = NV_PRESENT_BARRIER_FRAME_STATICS_VER1;
+
+		const NvAPI_Status NVResult = NvAPI_QueryPresentBarrierFrameStatistics(NvBarrierClientHandle, &PBStats);
+		if (NVResult != NVAPI_OK)
+		{
+			UE_LOG(LogDisplayClusterRenderSync, Warning, TEXT("NVS_PB: Couldn't obtain present barrier stats"));
+			return;
+		}
+
+		UE_LOG(LogDisplayClusterRenderSync, Log, TEXT("NVS_PB: Stats: SyncMode=%s, PresentCount=%u, PresentInSyncCount=%u, FlipInSyncCount=%u, RefreshCount=%u")
+			, GetSyncModeName(PBStats.SyncMode)
+			, PBStats.PresentCount
+			, PBStats.PresentInSyncCount
+			, PBStats.FlipInSyncCount
+			, PBStats.RefreshCount
+			);
+	}
+}
diff --git a/Engine/Plugins/Runtime/nDisplay/Source/DisplayCluster/Private/Render/Synchronization/Windows/DisplayClusterRenderSyncPolicyNvidiaSwapBarrierWindows.cpp b/Engine/Plugins/Runtime/nDisplay/Source/DisplayCluster/Private/Render/Synchronization/Windows/DisplayClusterRenderSyncPolicyNvidiaSwapBarrierWindows.cpp
new file mode 100644
--- /dev/null
+++ b/Engine/Plugins/Runtime/nDisplay/Source/DisplayCluster/Private/Render/Synchronization/Windows/DisplayClusterRenderSyncPolicyNvidiaSwapBarrierWindows.cpp
@@ -0,0 +1,254 @@
+// Copyright Epic Games, Inc. All Rights Reserved.
+
+#include "Render/Synchronization/DisplayClusterRenderSyncPolicyNvidiaSwapBarrier.h"
+
+#include "DisplayClusterCallbacks.h"
+
+#include "Misc/DisplayClusterGlobals.h"
+#include "Misc/DisplayClusterHelpers.h"
+#include "Misc/DisplayClusterLog.h"
+#include "Misc/DisplayClusterTypesConverter.h"
+
+#include "Engine/Engine.h"
+#include "Engine/GameViewportClient.h"
+#include "UnrealClient.h"
+
+#include "Windows/AllowWindowsPlatformTypes.h"
+THIRD_PARTY_INCLUDES_START
+	#include "dxgi1_3.h"
+	#include "d3d12.h"
+	#include "nvapi.h"
+THIRD_PARTY_INCLUDES_END
+#include "Windows/HideWindowsPlatformTypes.h"
+
+
+namespace DisplayClusterRenderSyncPolicyNvidiaSwapBarrier_Data_Windows
+{
+	IUnknown* D3DDevice = nullptr;
+	IDXGISwapChain2* DXGISwapChain = nullptr;
+}
+
+
+FDisplayClusterRenderSyncPolicyNvidiaSwapBarrier::FDisplayClusterRenderSyncPolicyNvidiaSwapBarrier(const TMap<FString, FString>& Parameters)
+	: Super(Parameters)
+{
+}
+
+FDisplayClusterRenderSyncPolicyNvidiaSwapBarrier::~FDisplayClusterRenderSyncPolicyNvidiaSwapBarrier()
+{
+	using namespace DisplayClusterRenderSyncPolicyNvidiaSwapBarrier_Data_Windows;
+
+	if (bNvLibraryInitialized)
+	{
+		if (bNvSyncInitializedSuccessfully)
+		{
+			// Unbind from swap barrier
+			NvAPI_D3D1x_BindSwapBarrier(D3DDevice, RequestedGroup, 0);
+			// Leave swap group
+			NvAPI_D3D1x_JoinSwapGroup(D3DDevice, DXGISwapChain, 0, false);
+		}
+	}
+}
+
+bool FDisplayClusterRenderSyncPolicyNvidiaSwapBarrier::SynchronizeClusterRendering(int32& InOutSyncInterval)
+{
+	using namespace DisplayClusterRenderSyncPolicyNvidiaSwapBarrier_Data_Windows;
+
+	// Initialize barriers at first call
+	if (!bNvSyncInitializationCalled)
+	{
+		// Set up NVIDIA swap barrier
+		bNvSyncInitializedSuccessfully = InitializeNvidiaSwapLock();
+	}
+
+	// Check if all required objects are available
+	if (!D3DDevice || !DXGISwapChain)
+	{
+		UE_LOG(LogDisplayClusterRenderSync, Warning, TEXT("NVS_SB: Couldn't get DX resources, no swap synchronization will be performed"));
+		// Present frame on a higher level
+		return true;
+	}
+
+	// Update maximum frame latency every if requested
+	if(CVarNvidiaSyncForceLatencyUpdateEveryFrame.GetValueOnAnyThread())
+	{
+		SetMaximumFrameLatency(1);
+	}
+
+	// Wait until frame rendering is finished
+	if (CfgFrameCompletionLimit != 0 && FrameCompletionCounter < CfgFrameCompletionLimit)
+	{
+		{
+			TRACE_CPUPROFILER_EVENT_SCOPE(nDisplay FrameCompletion);
+			WaitForFrameCompletion();
+		}
+
+		if (++FrameCompletionCounter == CfgFrameCompletionLimit)
+		{
+			UE_LOG(LogDisplayClusterRenderSync, Log, TEXT("NVS_SB: Disabled completion after %d frames"), FrameCompletionCounter);
+		}
+	}
+
+	// Align all threads on the timescale before calling NvAPI_D3D1x_Present
+	// As a side-effect, it should avoid NvAPI_D3D1x_Present stuck on application kill
+	if (CfgPrePresentAlignmentLimit != 0 && PrePresentAlignmentCounter < CfgPrePresentAlignmentLimit)
+	{
+		{
+			TRACE_CPUPROFILER_EVENT_SCOPE(nDisplay PrePresentAlignment);
+			SyncOnBarrier();
+		}
+
+		if (++PrePresentAlignmentCounter == CfgPrePresentAlignmentLimit)
+		{
+			UE_LOG(LogDisplayClusterRenderSync, Log, TEXT("NVS_SB: Disabled PrePresent sync after %d frames"), PrePresentAlignmentCounter);
+		}
+	}
+
+	if (bNvSyncInitializedSuccessfully && D3DDevice && DXGISwapChain)
+	{
+		UE_LOG(LogDisplayClusterRenderSync, VeryVerbose, TEXT("NVS_SB: presenting the frame with sync..."));
+
+		{
+			NvAPI_Status NvApiResult = NVAPI_ERROR;
+
+			{
+				TRACE_CPUPROFILER_EVENT_SCOPE(nDisplay NvAPI_D3D1x_Present);
+
+				// Present frame via NVIDIA API
+				NvApiResult = NvAPI_D3D1x_Present(D3DDevice, DXGISwapChain, (UINT)InOutSyncInterval, (UINT)0);
+			}
+
+			// Notify custom presentation was done
+			if (NvApiResult == NVAPI_OK)
+			{
+				GDisplayCluster->GetCallbacks().OnDisplayClusterFramePresented_RHIThread().Broadcast(false);
+			}
+
+			// Regardless of NvAPI_D3D1x_Present() call result, synchronize clients on the barrier if requested. The following
+			// if-block may theoretically branch the execution paths of the nodes, so we might lose threads alignment.
+			if (bCfgPostPresentAlignment)
+			{
+				SyncOnBarrier();
+			}
+
+			if (NvApiResult != NVAPI_OK)
+			{
+				UE_LOG(LogDisplayClusterRenderSync, Warning, TEXT("NVS_SB: An error occurred during frame presentation, error code 0x%x"), NvApiResult);
+				// Present frame on a higher level
+				return true;
+			}
+		}
+
+		UE_LOG(LogDisplayClusterRenderSync, VeryVerbose, TEXT("NVS_SB: the frame has been presented successfully"));
+	}
+	else
+	{
+		UE_LOG(LogDisplayClusterRenderSync, Warning, TEXT("NVS_SB: Can't synchronize frame presentation"));
+		// Something went wrong, let the upper level present this frame
+		return true;
+	}
+
+	// We presented current frame so no need to present it on higher level
+	return false;
+}
+
+bool FDisplayClusterRenderSyncPolicyNvidiaSwapBarrier::InitializeNvidiaSwapLock()
+{
+	using namespace DisplayClusterRenderSyncPolicyNvidiaSwapBarrier_Data_Windows;
+
+	bNvSyncInitializationCalled = true;
+
+	// Get D3D1XDevice
+	D3DDevice = static_cast<IUnknown*>(GDynamicRHI->RHIGetNativeDevice());
+	check(D3DDevice);
+	
+	// Get IDXGISwapChain
+	DXGISwapChain = static_cast<IDXGISwapChain2*>(GEngine->GameViewport->Viewport->GetViewportRHI().GetReference()->GetNativeSwapChain());
+	check(DXGISwapChain);
+
+	if (!D3DDevice || !DXGISwapChain)
+	{
+		UE_LOG(LogDisplayClusterRenderSync, Error, TEXT("NVS_SB: Couldn't get DX context data for NVIDIA swap barrier initialization"));
+		return false;
+	}
+
+	// Set frame latency
+	SetMaximumFrameLatency(1);
+
+	NvU32 MaxGroups = 0;
+	NvU32 MaxBarriers = 0;
+
+	// Get amount of available groups and barriers
+	NvAPI_Status NvApiResult = NvAPI_D3D1x_QueryMaxSwapGroup(D3DDevice, &MaxGroups, &MaxBarriers);
+	if (NvApiResult != NVAPI_OK)
+	{
+		UE_LOG(LogDisplayClusterRenderSync, Error, TEXT("NVS_SB: Couldn't query group/barrier limits, error code 0x%x"), NvApiResult);
+		return false;
+	}
+
+	// Make sure resources are available
+	UE_LOG(LogDisplayClusterRenderSync, Log, TEXT("NVS_SB: max_groups=%d max_barriers=%d"), (int)MaxGroups, (int)MaxBarriers);
+	if (!(MaxGroups > 0 && MaxBarriers > 0))
+	{
+		UE_LOG(LogDisplayClusterRenderSync, Error, TEXT("NVS_SB: No available groups or barriers"));
+		return false;
+	}
+
+	// Get requested barrier/group from config
+	DisplayClusterHelpers::map::template ExtractValueFromString(GetParameters(), FString("SyncGroup"),   RequestedGroup);
+	DisplayClusterHelpers::map::template ExtractValueFromString(GetParameters(), FString("SyncBarrier"), RequestedBarrier);
+
+	// Here we initialize NVIDIA sync on the same frame interval on the timescale
+	if (bCfgPreInitAlignment)
+	{
+		TRACE_CPUPROFILER_EVENT_SCOPE(nDisplay NVIDIA Sync Init);
+
+		// Align all the threads on the same timeline. It's very likely all the nodes got free on the same side of a potential
+		// upcoming V-blank. Normally, the barrier synchronization takes up to 500 microseconds which is about 3% of frame time.
+		// So the probability of a situation where some nodes left the barrier before V-blank while others after is really low.
+		// However, it's still possible. And cluster restart should likely be successfull.
+		SyncOnBarrier();
+
+		// Assuming all the nodes in the right place and have some time before V-blank. Let's wait untill it happend.
+		WaitForVBlank();
+
+		// We don't really know either we're at the very beginning of V-blank interval or in the end. Let's spend some time
+		// on the barrier to let V-blank interval over.
+		SyncOnBarrier();
+
+		// Yes, the task scheduler can break the logic above by allocating CPU resources to this thread like 20ms later which
+		// makes us behind of the next V-blank for 60Hz output. And we probably should play around thread priorities to reduce
+		// probability of that. But hope we're good here. Usually the systems that use NVIDIA sync approach have a lot of CPU
+		// cores and don't run any other applications that compete for CPU.
+		// Moreover, we will initialize NVIDIA sync for all cluster nodes before presenting first frame. So the initialization
+		// process should be fine anyway.
+	}
+
+	// Join swap group
+	NvApiResult = NvAPI_D3D1x_JoinSwapGroup(D3DDevice, DXGISwapChain, RequestedGroup, true);
+	if (NvApiResult != NVAPI_OK)
+	{
+		UE_LOG(LogDisplayClusterRenderSync, Error, TEXT("NVS_SB: Couldn't join swap group %d, error code 0x%x"), RequestedGroup, NvApiResult);
+		return false;
+	}
+	else
+	{
+		UE_LOG(LogDisplayClusterRenderSync, Log, TEXT("NVS_SB: Successfully joined the swap group %d"), RequestedGroup);
+	}
+
+	// Bind to sync barrier
+	NvApiResult = NvAPI_D3D1x_BindSwapBarrier(D3DDevice, RequestedGroup, RequestedBarrier);
+	if (NvApiResult != NVAPI_OK)
+	{
+		UE_LOG(LogDisplayClusterRenderSync, Error, TEXT("NVS_SB: Couldn't bind group %d to swap barrier %d, error code 0x%x"), RequestedGroup, RequestedBarrier, NvApiResult);
+		return false;
+	}
+	else
+	{
+		UE_LOG(LogDisplayClusterRenderSync, Log, TEXT("NVS_SB: Successfully bound group %d to the swap barrier %d"), RequestedGroup, RequestedBarrier);
+	}
+
+	UE_LOG(LogDisplayClusterRenderSync, Log, TEXT("NVS_SB: Initialized successfully"));
+
+	return true;
+}
diff --git a/Engine/Plugins/Runtime/nDisplay/Source/DisplayCluster/Private/Render/Synchronization/Windows/DisplayClusterRenderSyncPolicyNvidiaWindows.cpp b/Engine/Plugins/Runtime/nDisplay/Source/DisplayCluster/Private/Render/Synchronization/Windows/DisplayClusterRenderSyncPolicyNvidiaWindows.cpp
deleted file mode 100644
diff --git a/Engine/Plugins/Runtime/nDisplay/Source/DisplayCluster/Public/IDisplayClusterCallbacks.h b/Engine/Plugins/Runtime/nDisplay/Source/DisplayCluster/Public/IDisplayClusterCallbacks.h
--- a/Engine/Plugins/Runtime/nDisplay/Source/DisplayCluster/Public/IDisplayClusterCallbacks.h
+++ b/Engine/Plugins/Runtime/nDisplay/Source/DisplayCluster/Public/IDisplayClusterCallbacks.h
@@ -71,6 +71,10 @@ public:
 	DECLARE_EVENT(IDisplayClusterCallbacks, FDisplayClusterPresentationPostSynchronization_RHIThread);
 	virtual FDisplayClusterPresentationPostSynchronization_RHIThread& OnDisplayClusterPresentationPostSynchronization_RHIThread() = 0;
 
+	/** Called after frame presentation **/
+	DECLARE_EVENT_OneParam(IDisplayClusterCallbacks, FDisplayClusterFramePresentated_RHIThread, bool bNativePresent);
+	virtual FDisplayClusterFramePresentated_RHIThread& OnDisplayClusterFramePresented_RHIThread() = 0;
+
 	/** Failover notification **/
 	DECLARE_EVENT_OneParam(IDisplayClusterCallbacks, FDisplayClusterFailoverNodeDown, const FString&);
 	virtual FDisplayClusterFailoverNodeDown& OnDisplayClusterFailoverNodeDown() = 0;
diff --git a/Engine/Plugins/Runtime/nDisplay/Source/DisplayClusterConfiguration/Public/DisplayClusterConfigurationStrings.h b/Engine/Plugins/Runtime/nDisplay/Source/DisplayClusterConfiguration/Public/DisplayClusterConfigurationStrings.h
--- a/Engine/Plugins/Runtime/nDisplay/Source/DisplayClusterConfiguration/Public/DisplayClusterConfigurationStrings.h
+++ b/Engine/Plugins/Runtime/nDisplay/Source/DisplayClusterConfiguration/Public/DisplayClusterConfigurationStrings.h
@@ -84,10 +84,14 @@ namespace DisplayClusterConfigurationStrings
 				static constexpr const TCHAR* Ethernet          = TEXT("ethernet");
 				static constexpr const TCHAR* EthernetBarrier   = TEXT("ethernet_barrier");
 
+				// NVIDIA Swap Barrier (old)
 				static constexpr const TCHAR* Nvidia            = TEXT("nvidia");
 				static constexpr const TCHAR* NvidiaSwapBarrier = TEXT("swap_barrier");
 				static constexpr const TCHAR* NvidiaSwapGroup   = TEXT("swap_group");
 
+				// NVIDIA Present Barrier (new)
+				static constexpr const TCHAR* NvidiaAdvanced    = TEXT("nvidia_pb");
+
 				// Always use 'none' for headless rendering
 				static constexpr const TCHAR* HeadlessRenderingSyncPolicy = None;
 			}
