[nDisplay] Fixed build issue (NVAPI identifiers not found)
- Wrapped NVAPI code in WITH_NVAPI macro
- Added PostPresent alignment

diff --git a/Engine/Plugins/Runtime/nDisplay/Source/DisplayCluster/DisplayCluster.Build.cs b/Engine/Plugins/Runtime/nDisplay/Source/DisplayCluster/DisplayCluster.Build.cs
--- a/Engine/Plugins/Runtime/nDisplay/Source/DisplayCluster/DisplayCluster.Build.cs
+++ b/Engine/Plugins/Runtime/nDisplay/Source/DisplayCluster/DisplayCluster.Build.cs
@@ -76,6 +76,9 @@ public class DisplayCluster : ModuleRules
 			AddEngineThirdPartyPrivateStaticDependencies(Target, "DX11", "DX12");
 		}
 
-		AddEngineThirdPartyPrivateStaticDependencies(Target, "NVAPI");
+		if(Target.Platform == UnrealTargetPlatform.Win64 || Target.Platform == UnrealTargetPlatform.Linux)
+		{
+			AddEngineThirdPartyPrivateStaticDependencies(Target, "NVAPI");
+		}
 	}
 }
diff --git a/Engine/Plugins/Runtime/nDisplay/Source/DisplayCluster/Private/Render/Synchronization/DisplayClusterRenderSyncPolicyNvidiaPresentBarrier.h b/Engine/Plugins/Runtime/nDisplay/Source/DisplayCluster/Private/Render/Synchronization/DisplayClusterRenderSyncPolicyNvidiaPresentBarrier.h
--- a/Engine/Plugins/Runtime/nDisplay/Source/DisplayCluster/Private/Render/Synchronization/DisplayClusterRenderSyncPolicyNvidiaPresentBarrier.h
+++ b/Engine/Plugins/Runtime/nDisplay/Source/DisplayCluster/Private/Render/Synchronization/DisplayClusterRenderSyncPolicyNvidiaPresentBarrier.h
@@ -25,7 +25,7 @@ public:
 	/** Returns name of this policy implementation */
 	virtual FName GetName() const override
 	{
-		return FName(DisplayClusterConfigurationStrings::config::cluster::render_sync::NvidiaAdvanced);
+		return FName(DisplayClusterConfigurationStrings::config::cluster::render_sync::NvidiaPresentBarrier);
 	}
 
 	/** Performs rendering and frame output synchronization */
diff --git a/Engine/Plugins/Runtime/nDisplay/Source/DisplayCluster/Private/Render/Synchronization/Windows/DisplayClusterRenderSyncPolicyNvidiaPresentBarrierWindows.cpp b/Engine/Plugins/Runtime/nDisplay/Source/DisplayCluster/Private/Render/Synchronization/Windows/DisplayClusterRenderSyncPolicyNvidiaPresentBarrierWindows.cpp
--- a/Engine/Plugins/Runtime/nDisplay/Source/DisplayCluster/Private/Render/Synchronization/Windows/DisplayClusterRenderSyncPolicyNvidiaPresentBarrierWindows.cpp
+++ b/Engine/Plugins/Runtime/nDisplay/Source/DisplayCluster/Private/Render/Synchronization/Windows/DisplayClusterRenderSyncPolicyNvidiaPresentBarrierWindows.cpp
@@ -48,6 +48,7 @@ FDisplayClusterRenderSyncPolicyNvidiaPresentBarrier::~FDisplayClusterRenderSyncP
 
 	GDisplayCluster->GetCallbacks().OnDisplayClusterFramePresented_RHIThread().RemoveAll(this);
 
+#if WITH_NVAPI
 	if (bNvLibraryInitialized)
 	{
 		// Release present barrier client
@@ -60,6 +61,7 @@ FDisplayClusterRenderSyncPolicyNvidiaPresentBarrier::~FDisplayClusterRenderSyncP
 		// Release backbuffer references if there are any
 		BackBuffers.Empty();
 	}
+#endif // WITH_NVAPI
 }
 
 bool FDisplayClusterRenderSyncPolicyNvidiaPresentBarrier::SynchronizeClusterRendering(int32& InOutSyncInterval)
@@ -149,6 +151,7 @@ bool FDisplayClusterRenderSyncPolicyNvidiaPresentBarrier::InitializePresentBarri
 	// Set frame latency
 	SetMaximumFrameLatency(1);
 
+#if WITH_NVAPI
 	NvAPI_Status NVResult = NVAPI_ERROR;
 
 	// Check if present barrier is supported
@@ -275,10 +278,21 @@ bool FDisplayClusterRenderSyncPolicyNvidiaPresentBarrier::InitializePresentBarri
 	UE_LOG(LogDisplayClusterRenderSync, Log, TEXT("NVS_PB: Initialized successfully"));
 
 	return true;
+#else
+	return false;
+#endif // WITH_NVAPI
 }
 
 void FDisplayClusterRenderSyncPolicyNvidiaPresentBarrier::OnFramePresented(bool bNativePresent)
 {
+	// Post-present alignment
+	if (bCfgPostPresentAlignment)
+	{
+		TRACE_CPUPROFILER_EVENT_SCOPE(nDisplay PostPresentAlignment);
+		SyncOnBarrier();
+	}
+
+	// Log stats
 	if (CVarNvidiaSyncPrintStatsEveryFrame.GetValueOnAnyThread())
 	{
 		LogPresentBarrierStats();
@@ -319,6 +333,7 @@ void FDisplayClusterRenderSyncPolicyNvidiaPresentBarrier::LogPresentBarrierStats
 {
 	using namespace DisplayClusterRenderSyncPolicyNvidiaPresentBarrier_Data_Windows;
 
+#if WITH_NVAPI
 	if (bNvLibraryInitialized && bNvSyncInitializedSuccessfully)
 	{
 		NV_PRESENT_BARRIER_FRAME_STATISTICS PBStats = { };
@@ -339,4 +354,5 @@ void FDisplayClusterRenderSyncPolicyNvidiaPresentBarrier::LogPresentBarrierStats
 			, PBStats.RefreshCount
 			);
 	}
+#endif // WITH_NVAPI
 }
diff --git a/Engine/Plugins/Runtime/nDisplay/Source/DisplayCluster/Private/Render/Synchronization/Windows/DisplayClusterRenderSyncPolicyNvidiaSwapBarrierWindows.cpp b/Engine/Plugins/Runtime/nDisplay/Source/DisplayCluster/Private/Render/Synchronization/Windows/DisplayClusterRenderSyncPolicyNvidiaSwapBarrierWindows.cpp
--- a/Engine/Plugins/Runtime/nDisplay/Source/DisplayCluster/Private/Render/Synchronization/Windows/DisplayClusterRenderSyncPolicyNvidiaSwapBarrierWindows.cpp
+++ b/Engine/Plugins/Runtime/nDisplay/Source/DisplayCluster/Private/Render/Synchronization/Windows/DisplayClusterRenderSyncPolicyNvidiaSwapBarrierWindows.cpp
@@ -38,6 +38,7 @@ FDisplayClusterRenderSyncPolicyNvidiaSwapBarrier::~FDisplayClusterRenderSyncPoli
 {
 	using namespace DisplayClusterRenderSyncPolicyNvidiaSwapBarrier_Data_Windows;
 
+#if WITH_NVAPI
 	if (bNvLibraryInitialized)
 	{
 		if (bNvSyncInitializedSuccessfully)
@@ -48,6 +49,7 @@ FDisplayClusterRenderSyncPolicyNvidiaSwapBarrier::~FDisplayClusterRenderSyncPoli
 			NvAPI_D3D1x_JoinSwapGroup(D3DDevice, DXGISwapChain, 0, false);
 		}
 	}
+#endif // WITH_NVAPI
 }
 
 bool FDisplayClusterRenderSyncPolicyNvidiaSwapBarrier::SynchronizeClusterRendering(int32& InOutSyncInterval)
@@ -104,6 +106,7 @@ bool FDisplayClusterRenderSyncPolicyNvidiaSwapBarrier::SynchronizeClusterRenderi
 		}
 	}
 
+#if WITH_NVAPI
 	if (bNvSyncInitializedSuccessfully && D3DDevice && DXGISwapChain)
 	{
 		UE_LOG(LogDisplayClusterRenderSync, VeryVerbose, TEXT("NVS_SB: presenting the frame with sync..."));
@@ -150,6 +153,10 @@ bool FDisplayClusterRenderSyncPolicyNvidiaSwapBarrier::SynchronizeClusterRenderi
 
 	// We presented current frame so no need to present it on higher level
 	return false;
+#else
+	// NVAPI isn't available. Ask engine to present.
+	return true;
+#endif // WITH_NVAPI
 }
 
 bool FDisplayClusterRenderSyncPolicyNvidiaSwapBarrier::InitializeNvidiaSwapLock()
@@ -178,6 +185,7 @@ bool FDisplayClusterRenderSyncPolicyNvidiaSwapBarrier::InitializeNvidiaSwapLock(
 	NvU32 MaxGroups = 0;
 	NvU32 MaxBarriers = 0;
 
+#if WITH_NVAPI
 	// Get amount of available groups and barriers
 	NvAPI_Status NvApiResult = NvAPI_D3D1x_QueryMaxSwapGroup(D3DDevice, &MaxGroups, &MaxBarriers);
 	if (NvApiResult != NVAPI_OK)
@@ -251,4 +259,7 @@ bool FDisplayClusterRenderSyncPolicyNvidiaSwapBarrier::InitializeNvidiaSwapLock(
 	UE_LOG(LogDisplayClusterRenderSync, Log, TEXT("NVS_SB: Initialized successfully"));
 
 	return true;
+#else
+	return false;
+#endif // WITH_NVAPI
 }
diff --git a/Engine/Plugins/Runtime/nDisplay/Source/DisplayClusterConfiguration/Public/DisplayClusterConfigurationStrings.h b/Engine/Plugins/Runtime/nDisplay/Source/DisplayClusterConfiguration/Public/DisplayClusterConfigurationStrings.h
--- a/Engine/Plugins/Runtime/nDisplay/Source/DisplayClusterConfiguration/Public/DisplayClusterConfigurationStrings.h
+++ b/Engine/Plugins/Runtime/nDisplay/Source/DisplayClusterConfiguration/Public/DisplayClusterConfigurationStrings.h
@@ -90,7 +90,7 @@ namespace DisplayClusterConfigurationStrings
 				static constexpr const TCHAR* NvidiaSwapGroup   = TEXT("swap_group");
 
 				// NVIDIA Present Barrier (new)
-				static constexpr const TCHAR* NvidiaAdvanced    = TEXT("nvidia_pb");
+				static constexpr const TCHAR* NvidiaPresentBarrier = TEXT("nvidia_pb");
 
 				// Always use 'none' for headless rendering
 				static constexpr const TCHAR* HeadlessRenderingSyncPolicy = None;
