DisplayCluster: Stress unity fix for nvapi and missing d3d12 defines.

#rnx

diff --git a/Engine/Plugins/Runtime/nDisplay/Source/DisplayCluster/Private/Render/Synchronization/DisplayClusterRenderSyncPolicyNvidiaBase.cpp b/Engine/Plugins/Runtime/nDisplay/Source/DisplayCluster/Private/Render/Synchronization/DisplayClusterRenderSyncPolicyNvidiaBase.cpp
--- a/Engine/Plugins/Runtime/nDisplay/Source/DisplayCluster/Private/Render/Synchronization/DisplayClusterRenderSyncPolicyNvidiaBase.cpp
+++ b/Engine/Plugins/Runtime/nDisplay/Source/DisplayCluster/Private/Render/Synchronization/DisplayClusterRenderSyncPolicyNvidiaBase.cpp
@@ -6,6 +6,9 @@
 
 #if WITH_NVAPI
 	THIRD_PARTY_INCLUDES_START
+	#include "dxgi1_3.h"
+	#include "d3d12.h"
+	// NOTE: The two previous includes prevent issues in (stress) unity builds where our usage of nvapi expects __d3d12_h__ to be defined.
 	#include "nvapi.h"
 	THIRD_PARTY_INCLUDES_END
 #endif // WITH_NVAPI
