[nDisplay] Added PostBarrierJoin Sleep for NVIDIA sync policies

diff --git a/Engine/Plugins/Runtime/nDisplay/Source/DisplayCluster/Private/Render/Synchronization/DisplayClusterRenderSyncPolicyNvidiaBase.cpp b/Engine/Plugins/Runtime/nDisplay/Source/DisplayCluster/Private/Render/Synchronization/DisplayClusterRenderSyncPolicyNvidiaBase.cpp
--- a/Engine/Plugins/Runtime/nDisplay/Source/DisplayCluster/Private/Render/Synchronization/DisplayClusterRenderSyncPolicyNvidiaBase.cpp
+++ b/Engine/Plugins/Runtime/nDisplay/Source/DisplayCluster/Private/Render/Synchronization/DisplayClusterRenderSyncPolicyNvidiaBase.cpp
@@ -19,7 +19,18 @@
 static TAutoConsoleVariable<bool> CVarNvidiaSyncPreInitAlignment(
 	TEXT("nDisplay.sync.nvidia.PreInitAlignment"),
 	true,
-	TEXT("NVIDIA sync: Aligns RHI threads before initializing synchronization subsystem."),
+	TEXT("Aligns RHI threads before initializing synchronization subsystem."),
+	ECVF_ReadOnly | ECVF_RenderThreadSafe
+);
+
+// Forces delay after joining the barrier. The idea is to give nvidia sync clients some time
+// to switch to sync state before calling present. 
+static TAutoConsoleVariable<float> CVarNvidiaSyncPostBarrierJoinSleep(
+	TEXT("nDisplay.sync.nvidia.PostBarrierJoinSleep"),
+	0.f,
+	TEXT("Forces thread delay after joining the barrier\n")
+	TEXT("    N <= 0 : Disabled\n")
+	TEXT("    N  > 0 : Sleep time (seconds)\n"),
 	ECVF_ReadOnly | ECVF_RenderThreadSafe
 );
 
@@ -90,6 +101,7 @@ TAutoConsoleVariable<bool> CVarNvidiaSyncPrintStatsEveryFrame(
 FDisplayClusterRenderSyncPolicyNvidiaBase::FDisplayClusterRenderSyncPolicyNvidiaBase(const TMap<FString, FString>& Parameters)
 	: Super(Parameters)
 	, bCfgPreInitAlignment(CVarNvidiaSyncPreInitAlignment.GetValueOnAnyThread())
+	, CfgPostBarrierJoinSleep(CVarNvidiaSyncPostBarrierJoinSleep.GetValueOnAnyThread())
 	, CfgPrePresentAlignmentLimit(CVarNvidiaSyncPrePresentAlignmentLimit.GetValueOnAnyThread())
 	, CfgFrameCompletionLimit(CVarNvidiaSyncFrameCompletionLimit.GetValueOnAnyThread())
 	, bCfgPostPresentAlignment(CVarNvidiaSyncPostPresentAlignment.GetValueOnAnyThread())
@@ -108,8 +120,9 @@ FDisplayClusterRenderSyncPolicyNvidiaBase::FDisplayClusterRenderSyncPolicyNvidia
 #endif // WITH_NVAPI
 
 	UE_LOG(LogDisplayClusterRenderSync, Log,
-		TEXT("NVIDIA sync configuration: PreInitAlignment=%d, PrePresentAlignmentLimit=%d, FrameCompletionLimit=%d, PostPresentAlignment=%d"),
+		TEXT("NVIDIA sync configuration: PreInitAlignment=%d, PostBarrierJoinSleep=%f, PrePresentAlignmentLimit=%d, FrameCompletionLimit=%d, PostPresentAlignment=%d"),
 		bCfgPreInitAlignment ? 1 : 0,
+		CfgPostBarrierJoinSleep,
 		CfgPrePresentAlignmentLimit,
 		CfgFrameCompletionLimit,
 		bCfgPostPresentAlignment ? 1 : 0
diff --git a/Engine/Plugins/Runtime/nDisplay/Source/DisplayCluster/Private/Render/Synchronization/DisplayClusterRenderSyncPolicyNvidiaBase.h b/Engine/Plugins/Runtime/nDisplay/Source/DisplayCluster/Private/Render/Synchronization/DisplayClusterRenderSyncPolicyNvidiaBase.h
--- a/Engine/Plugins/Runtime/nDisplay/Source/DisplayCluster/Private/Render/Synchronization/DisplayClusterRenderSyncPolicyNvidiaBase.h
+++ b/Engine/Plugins/Runtime/nDisplay/Source/DisplayCluster/Private/Render/Synchronization/DisplayClusterRenderSyncPolicyNvidiaBase.h
@@ -49,6 +49,9 @@ protected:
 	/** CVar config: synchronous initialization (alignment on a barrier before initializing NVIDIA synchronization) */
 	const bool bCfgPreInitAlignment;
 
+	/** CVar config: Force delay after joining the barrier */
+	const float CfgPostBarrierJoinSleep;
+
 	/** CVar config: pre-present threads alignment on a network barrier */
 	const int32 CfgPrePresentAlignmentLimit;
 	/** Pre-present sync counter */
diff --git a/Engine/Plugins/Runtime/nDisplay/Source/DisplayCluster/Private/Render/Synchronization/Windows/DisplayClusterRenderSyncPolicyNvidiaPresentBarrierWindows.cpp b/Engine/Plugins/Runtime/nDisplay/Source/DisplayCluster/Private/Render/Synchronization/Windows/DisplayClusterRenderSyncPolicyNvidiaPresentBarrierWindows.cpp
--- a/Engine/Plugins/Runtime/nDisplay/Source/DisplayCluster/Private/Render/Synchronization/Windows/DisplayClusterRenderSyncPolicyNvidiaPresentBarrierWindows.cpp
+++ b/Engine/Plugins/Runtime/nDisplay/Source/DisplayCluster/Private/Render/Synchronization/Windows/DisplayClusterRenderSyncPolicyNvidiaPresentBarrierWindows.cpp
@@ -275,6 +275,13 @@ bool FDisplayClusterRenderSyncPolicyNvidiaPresentBarrier::InitializePresentBarri
 		}
 	}
 
+	// Force sleep if requested
+	if(CfgPostBarrierJoinSleep > 0.f)
+	{
+		TRACE_CPUPROFILER_EVENT_SCOPE(nDisplay PostBarrierJoinSleep);
+		FPlatformProcess::SleepNoStats(CfgPostBarrierJoinSleep);
+	}
+
 	UE_LOG(LogDisplayClusterRenderSync, Log, TEXT("NVS_PB: Initialized successfully"));
 
 	return true;
diff --git a/Engine/Plugins/Runtime/nDisplay/Source/DisplayCluster/Private/Render/Synchronization/Windows/DisplayClusterRenderSyncPolicyNvidiaSwapBarrierWindows.cpp b/Engine/Plugins/Runtime/nDisplay/Source/DisplayCluster/Private/Render/Synchronization/Windows/DisplayClusterRenderSyncPolicyNvidiaSwapBarrierWindows.cpp
--- a/Engine/Plugins/Runtime/nDisplay/Source/DisplayCluster/Private/Render/Synchronization/Windows/DisplayClusterRenderSyncPolicyNvidiaSwapBarrierWindows.cpp
+++ b/Engine/Plugins/Runtime/nDisplay/Source/DisplayCluster/Private/Render/Synchronization/Windows/DisplayClusterRenderSyncPolicyNvidiaSwapBarrierWindows.cpp
@@ -256,6 +256,13 @@ bool FDisplayClusterRenderSyncPolicyNvidiaSwapBarrier::InitializeNvidiaSwapLock(
 		UE_LOG(LogDisplayClusterRenderSync, Log, TEXT("NVS_SB: Successfully bound group %d to the swap barrier %d"), RequestedGroup, RequestedBarrier);
 	}
 
+	// Force sleep if requested
+	if (CfgPostBarrierJoinSleep > 0.f)
+	{
+		TRACE_CPUPROFILER_EVENT_SCOPE(nDisplay PostBarrierJoinSleep);
+		FPlatformProcess::SleepNoStats(CfgPostBarrierJoinSleep);
+	}
+
 	UE_LOG(LogDisplayClusterRenderSync, Log, TEXT("NVS_SB: Initialized successfully"));
 
 	return true;
