Fix minor shader warnings in nDisplay

diff --git a/Engine/Plugins/Runtime/nDisplay/Shaders/Private/MPCDIOverlayShaders.usf b/Engine/Plugins/Runtime/nDisplay/Shaders/Private/MPCDIOverlayShaders.usf
--- a/Engine/Plugins/Runtime/nDisplay/Shaders/Private/MPCDIOverlayShaders.usf
+++ b/Engine/Plugins/Runtime/nDisplay/Shaders/Private/MPCDIOverlayShaders.usf
@@ -8,7 +8,7 @@ OutputVS IcvfxWarpVS(InputVS IN)
 	OutputVS OUT;
 
 #if MESH_WARP
-	OUT.PFMPosition = mul(float4(IN.Position.xyz,1.f), MeshToStageProjectionMatrix);
+	OUT.PFMPosition = mul(float4(IN.Position.xyz,1.f), MeshToStageProjectionMatrix).xyz;
 	OUT.UV_Chromakey = float4(IN.UV_Chromakey, 0.f, 1.f);
 	float4 ScreenPosition = float4(IN.UV.xy, 0, 1);
 #else
diff --git a/Engine/Plugins/Runtime/nDisplay/Shaders/Private/MPCDIShaders.usf b/Engine/Plugins/Runtime/nDisplay/Shaders/Private/MPCDIShaders.usf
--- a/Engine/Plugins/Runtime/nDisplay/Shaders/Private/MPCDIShaders.usf
+++ b/Engine/Plugins/Runtime/nDisplay/Shaders/Private/MPCDIShaders.usf
@@ -6,7 +6,7 @@ OutputVS WarpVS(InputVS IN)
 	OutputVS OUT;
 
 #if MESH_WARP
-	OUT.PFMPosition = mul(float4(IN.Position.xyz, 1.f), MeshToStageProjectionMatrix);
+	OUT.PFMPosition = mul(float4(IN.Position.xyz, 1.f), MeshToStageProjectionMatrix).xyz;
 	OUT.UV_Chromakey = float4(IN.UV_Chromakey, 0.f, 1.f);
 	float4 ScreenPosition = float4(IN.UV.xy, 0, 1);
 #else
diff --git a/Engine/Plugins/Runtime/nDisplay/Shaders/Private/PicpPostProcessShaders.usf b/Engine/Plugins/Runtime/nDisplay/Shaders/Private/PicpPostProcessShaders.usf
--- a/Engine/Plugins/Runtime/nDisplay/Shaders/Private/PicpPostProcessShaders.usf
+++ b/Engine/Plugins/Runtime/nDisplay/Shaders/Private/PicpPostProcessShaders.usf
@@ -40,7 +40,7 @@ SamplerState BilinearClampTextureSampler;
 
 OutputPS DirectComposePS(OutputVS IN)
 {
-	float4 Color = SrcTexture.Sample(BilinearClampTextureSampler, IN.UV);
+	float4 Color = SrcTexture.Sample(BilinearClampTextureSampler, IN.UV.xy);
 	Color.w = 1.0 - Color.w;
 
 	OutputPS Out;
@@ -57,7 +57,7 @@ float GaussianWeight1d(float u, float KernelRadius)
 
 OutputPS BlurPostProcessPS(OutputVS IN)
 {
-	float2 UV = IN.UV;
+	float2 UV = IN.UV.xy;
 	float4 OutColor;
 
 #if BLUR_DILATE
