[nDisplay] Updates to the BP API

- Moved nDisplay BP functions to the global API namespace instead of keeping them in the separate UObject interfaces
- Updated BP API categories for better structuring (user friendly)
- Deprecated IsModuleInitialized (it's no longer needed)
- Removed some outdated, unused and unnecessary template helpers that now produce a compilation error

diff --git a/Engine/Plugins/Runtime/nDisplay/Source/DisplayCluster/Private/Blueprints/DisplayClusterBlueprintAPIImpl.cpp b/Engine/Plugins/Runtime/nDisplay/Source/DisplayCluster/Private/Blueprints/DisplayClusterBlueprintAPIImpl.cpp
--- a/Engine/Plugins/Runtime/nDisplay/Source/DisplayCluster/Private/Blueprints/DisplayClusterBlueprintAPIImpl.cpp
+++ b/Engine/Plugins/Runtime/nDisplay/Source/DisplayCluster/Private/Blueprints/DisplayClusterBlueprintAPIImpl.cpp
@@ -1,155 +1,86 @@
 // Copyright Epic Games, Inc. All Rights Reserved.
 
 #include "Blueprints/DisplayClusterBlueprintAPIImpl.h"
+#include "Blueprints/DisplayClusterBlueprintLib.h"
 
-#include "IDisplayCluster.h"
-
-#include "Cluster/DisplayClusterClusterEvent.h"
-#include "Cluster/IDisplayClusterClusterManager.h"
-#include "Config/IDisplayClusterConfigManager.h"
-#include "Game/IDisplayClusterGameManager.h"
-#include "Render/IDisplayClusterRenderManager.h"
-#include "Render/Device/IDisplayClusterRenderDevice.h"
-#include "DisplayClusterSceneViewExtensions.h"
-
-#include "Misc/DisplayClusterHelpers.h"
-#include "Misc/DisplayClusterTypesConverter.h"
-#include "Misc/DisplayClusterGlobals.h"
-#include "Misc/DisplayClusterHelpers.h"
-#include "Misc/DisplayClusterLog.h"
-#include "Misc/DisplayClusterStrings.h"
-
-#include "DisplayClusterRootActor.h"
-#include "Components/DisplayClusterCameraComponent.h"
-#include "Components/DisplayClusterScreenComponent.h"
-
-#include "Render/Projection/IDisplayClusterProjectionPolicy.h"
-
-#include "DisplayClusterConfigurationTypes.h"
-#include "DisplayClusterEnums.h"
-
-#include "Math/IntRect.h"
-
-#include "CineCameraComponent.h"
-
-#include "Render/Viewport/IDisplayClusterViewportManager.h"
-#include "Render/Viewport/IDisplayClusterViewport.h"
 
 // Cluster runtime API
 
-bool UDisplayClusterBlueprintAPIImpl::IsModuleInitialized() const
-{
-	const bool bInitialized = IDisplayCluster::Get().IsModuleInitialized();
-	UE_LOG(LogDisplayClusterBlueprint, Verbose, TEXT("IsModuleInitialized - %s"), *DisplayClusterTypesConverter::template ToString(bInitialized));
-	return bInitialized;
-}
-
 EDisplayClusterOperationMode UDisplayClusterBlueprintAPIImpl::GetOperationMode() const
 {
-	const EDisplayClusterOperationMode OpMode = IDisplayCluster::Get().GetOperationMode();
-	UE_LOG(LogDisplayClusterBlueprint, Verbose, TEXT("Operation mode - %s"), *DisplayClusterTypesConverter::template ToString(OpMode));
-	return OpMode;
+	return UDisplayClusterBlueprintLib::GetOperationMode();
 }
 
 ADisplayClusterRootActor* UDisplayClusterBlueprintAPIImpl::GetRootActor() const
 {
-	ADisplayClusterRootActor* const RootActor = IDisplayCluster::Get().GetGameMgr()->GetRootActor();
-	UE_LOG(LogDisplayClusterBlueprint, Verbose, TEXT("GetRootActor - %s"), RootActor ? *RootActor->GetHumanReadableName() : DisplayClusterStrings::log::NotFound);
-	return RootActor;
+	return UDisplayClusterBlueprintLib::GetRootActor();
 }
 
 // Local node runtime API
 
 FString UDisplayClusterBlueprintAPIImpl::GetNodeId() const
 {
-	const FString NodeId = IDisplayCluster::Get().GetClusterMgr()->GetNodeId();
-	UE_LOG(LogDisplayClusterBlueprint, Verbose, TEXT("GetNodeId - NodeId=%s"), *NodeId);
-	return NodeId;
+	return UDisplayClusterBlueprintLib::GetNodeId();
 }
 
 void UDisplayClusterBlueprintAPIImpl::GetActiveNodeIds(TArray<FString>& OutNodeIds) const
 {
-	UE_LOG(LogDisplayClusterBlueprint, Verbose, TEXT("GetNodeIds"));
-	IDisplayCluster::Get().GetClusterMgr()->GetNodeIds(OutNodeIds);
+	UDisplayClusterBlueprintLib::GetActiveNodeIds(OutNodeIds);
 }
 
 int32 UDisplayClusterBlueprintAPIImpl::GetActiveNodesAmount() const
 {
-	const int32 NodesAmount = IDisplayCluster::Get().GetClusterMgr()->GetNodesAmount();
-	UE_LOG(LogDisplayClusterBlueprint, Verbose, TEXT("GetNodesAmount - %d"), NodesAmount);
-	return NodesAmount;
+	return UDisplayClusterBlueprintLib::GetActiveNodesAmount();
 }
 
 bool UDisplayClusterBlueprintAPIImpl::IsPrimary() const
 {
-	const bool bIsPrimary = IDisplayCluster::Get().GetClusterMgr()->IsPrimary();
-	UE_LOG(LogDisplayClusterBlueprint, Verbose, TEXT("IsPrimary - %s"), *DisplayClusterTypesConverter::template ToString(bIsPrimary));
-	return bIsPrimary;
+	return UDisplayClusterBlueprintLib::IsPrimary();
 }
 
 bool UDisplayClusterBlueprintAPIImpl::IsSecondary() const
 {
-	const bool bIsSecondary = IDisplayCluster::Get().GetClusterMgr()->IsSecondary();
-	UE_LOG(LogDisplayClusterBlueprint, Verbose, TEXT("IsSecondary - %s"), *DisplayClusterTypesConverter::template ToString(bIsSecondary));
-	return bIsSecondary;
+	return UDisplayClusterBlueprintLib::IsSecondary();
 }
 
 bool UDisplayClusterBlueprintAPIImpl::IsBackup() const
 {
-	const bool bIsBackup = IDisplayCluster::Get().GetClusterMgr()->IsBackup();
-	UE_LOG(LogDisplayClusterBlueprint, Verbose, TEXT("IsBackup - %s"), *DisplayClusterTypesConverter::template ToString(bIsBackup));
-	return bIsBackup;
+	return UDisplayClusterBlueprintLib::IsBackup();
 }
 
 EDisplayClusterNodeRole UDisplayClusterBlueprintAPIImpl::GetClusterRole() const
 {
-	const EDisplayClusterNodeRole ClusterRole = IDisplayCluster::Get().GetClusterMgr()->GetClusterRole();
-	UE_LOG(LogDisplayClusterBlueprint, Verbose, TEXT("GetClusterRole - %d"), static_cast<int32>(ClusterRole));
-	return ClusterRole;
+	return UDisplayClusterBlueprintLib::GetClusterRole();
 }
 
 // Cluster events API
 
 void UDisplayClusterBlueprintAPIImpl::AddClusterEventListener(TScriptInterface<IDisplayClusterClusterEventListener> Listener)
 {
-	UE_LOG(LogDisplayClusterBlueprint, Verbose, TEXT("AddClusterEventListener - adding cluster event listener..."));
-	IDisplayCluster::Get().GetClusterMgr()->AddClusterEventListener(Listener);
+	UDisplayClusterBlueprintLib::AddClusterEventListener(Listener);
 }
 
 void UDisplayClusterBlueprintAPIImpl::RemoveClusterEventListener(TScriptInterface<IDisplayClusterClusterEventListener> Listener)
 {
-	UE_LOG(LogDisplayClusterBlueprint, Verbose, TEXT("RemoveClusterEventListener - removing cluster event listener..."));
-	IDisplayCluster::Get().GetClusterMgr()->RemoveClusterEventListener(Listener);
+	UDisplayClusterBlueprintLib::RemoveClusterEventListener(Listener);
 }
 
 void UDisplayClusterBlueprintAPIImpl::EmitClusterEventJson(const FDisplayClusterClusterEventJson& Event, bool bPrimaryOnly)
 {
-	UE_LOG(LogDisplayClusterBlueprint, Verbose, TEXT("EmitClusterEventJson - emitting cluster event: bPrimaryOnly='%s' Category='%s' Type='%s' Name='%s'"),
-		*DisplayClusterTypesConverter::template ToString(bPrimaryOnly), *Event.Category, *Event.Type, *Event.Name);
-
-	IDisplayCluster::Get().GetClusterMgr()->EmitClusterEventJson(Event, bPrimaryOnly);
+	UDisplayClusterBlueprintLib::EmitClusterEventJson(Event, bPrimaryOnly);
 }
 
 void UDisplayClusterBlueprintAPIImpl::EmitClusterEventBinary(const FDisplayClusterClusterEventBinary& Event, bool bPrimaryOnly)
 {
-	UE_LOG(LogDisplayClusterBlueprint, Verbose, TEXT("EmitClusterEventBinary - emitting cluster event: bPrimaryOnly='%s' EventId='%d'"),
-		*DisplayClusterTypesConverter::template ToString(bPrimaryOnly), Event.EventId);
-
-	IDisplayCluster::Get().GetClusterMgr()->EmitClusterEventBinary(Event, bPrimaryOnly);
+	UDisplayClusterBlueprintLib::EmitClusterEventBinary(Event, bPrimaryOnly);
 }
 
 void UDisplayClusterBlueprintAPIImpl::SendClusterEventJsonTo(const FString& Address, const int32 Port, const FDisplayClusterClusterEventJson& Event, bool bPrimaryOnly)
 {
-	checkSlow(Port > 0 && Port <= UINT16_MAX);
-
-	UE_LOG(LogDisplayClusterBlueprint, Verbose, TEXT("SendClusterEventJsonTo - sending json event to %s:%d, Category='%s' Type='%s' Name='%s'"), *Address, Port, *Event.Category, *Event.Type, *Event.Name);
-	IDisplayCluster::Get().GetClusterMgr()->SendClusterEventTo(Address, static_cast<uint16>(Port), Event, bPrimaryOnly);
+	UDisplayClusterBlueprintLib::SendClusterEventJsonTo(Address, Port, Event, bPrimaryOnly);
 }
 
 void UDisplayClusterBlueprintAPIImpl::SendClusterEventBinaryTo(const FString& Address, const int32 Port, const FDisplayClusterClusterEventBinary& Event, bool bPrimaryOnly)
 {
-	checkSlow(Port > 0 && Port <= UINT16_MAX);
-
-	UE_LOG(LogDisplayClusterBlueprint, Verbose, TEXT("SendClusterEventBinaryTo - sending binary event to %s:%d, EventId='%d'"), *Address, Port, Event.EventId);
-	IDisplayCluster::Get().GetClusterMgr()->SendClusterEventTo(Address, static_cast<uint16>(Port), Event, bPrimaryOnly);
+	UDisplayClusterBlueprintLib::SendClusterEventBinaryTo(Address, Port, Event, bPrimaryOnly);
 }
diff --git a/Engine/Plugins/Runtime/nDisplay/Source/DisplayCluster/Private/Blueprints/DisplayClusterBlueprintAPIImpl.h b/Engine/Plugins/Runtime/nDisplay/Source/DisplayCluster/Private/Blueprints/DisplayClusterBlueprintAPIImpl.h
--- a/Engine/Plugins/Runtime/nDisplay/Source/DisplayCluster/Private/Blueprints/DisplayClusterBlueprintAPIImpl.h
+++ b/Engine/Plugins/Runtime/nDisplay/Source/DisplayCluster/Private/Blueprints/DisplayClusterBlueprintAPIImpl.h
@@ -20,59 +20,78 @@ class DISPLAYCLUSTER_API UDisplayClusterBlueprintAPIImpl
 
 public:
 	// DisplayCluster module API
-	UFUNCTION(BlueprintCallable, meta = (DisplayName = "Is module initialized"), Category = "NDisplay")
-	virtual bool IsModuleInitialized() const override;
+	UE_DEPRECATED(5.4, "This function has been deprecated and will be removed soon.")
+	UFUNCTION(BlueprintCallable, meta = (DeprecatedFunction, DeprecationMessage = "This function has been deprecated and will be removed soon."))
+	virtual bool IsModuleInitialized() const override
+	{
+		return true;
+	}
 
 public:
 	// Runtime Cluster API
-	UFUNCTION(BlueprintCallable, meta = (DisplayName = "Get operation mode"), Category = "NDisplay")
+	UE_DEPRECATED(5.4, "This function has been moved to UDisplayClusterBlueprintLib.")
+	UFUNCTION(BlueprintCallable, meta = (DeprecatedFunction, DeprecationMessage = "This function is now available in the main blueprint functions list under 'nDisplay' section."))
 	virtual EDisplayClusterOperationMode GetOperationMode() const override;
 
-	UFUNCTION(BlueprintCallable, meta = (DisplayName = "Get root actor"), Category = "NDisplay|Game")
+	UE_DEPRECATED(5.4, "This function has been moved to UDisplayClusterBlueprintLib.")
+	UFUNCTION(BlueprintCallable, meta = (DeprecatedFunction, DeprecationMessage = "This function is now available in the main blueprint functions list under 'nDisplay' section."))
 	virtual ADisplayClusterRootActor* GetRootActor() const override;
 
 	// Runtime local node API
-	UFUNCTION(BlueprintCallable, meta = (DisplayName = "Get node ID"), Category = "NDisplay|Cluster")
+	UE_DEPRECATED(5.4, "This function has been moved to UDisplayClusterBlueprintLib.")
+	UFUNCTION(BlueprintCallable, meta = (DeprecatedFunction, DeprecationMessage = "This function is now available in the main blueprint functions list under 'nDisplay' section."))
 	virtual FString GetNodeId() const override;
 	
 	/** Returns List of the active nodes in the runtime cluster node in a cluster. */
-	UFUNCTION(BlueprintCallable, meta = (DisplayName = "Get cluster node IDs"), Category = "NDisplay|Cluster")
+	UE_DEPRECATED(5.4, "This function has been moved to UDisplayClusterBlueprintLib.")
+	UFUNCTION(BlueprintCallable, meta = (DeprecatedFunction, DeprecationMessage = "This function is now available in the main blueprint functions list under 'nDisplay' section."))
 	virtual void GetActiveNodeIds(TArray<FString>& OutNodeIds) const override;
 	
 
 	/** Returns amount of active nodes in a cluster. */
-	UFUNCTION(BlueprintCallable, meta = (DisplayName = "Get nodes amount"), Category = "NDisplay|Cluster")
+	UE_DEPRECATED(5.4, "This function has been moved to UDisplayClusterBlueprintLib.")
+	UFUNCTION(BlueprintCallable, meta = (DeprecatedFunction, DeprecationMessage = "This function is now available in the main blueprint functions list under 'nDisplay' section."))
 	virtual int32 GetActiveNodesAmount() const override;
 
-	UFUNCTION(BlueprintCallable, meta = (DisplayName = "Get cluster role"), Category = "NDisplay|Cluster")
+	UE_DEPRECATED(5.4, "This function has been moved to UDisplayClusterBlueprintLib.")
+	UFUNCTION(BlueprintCallable, meta = (DeprecatedFunction, DeprecationMessage = "This function is now available in the main blueprint functions list under 'nDisplay' section."))
 	virtual EDisplayClusterNodeRole GetClusterRole() const override;
 
-	UFUNCTION(BlueprintCallable, meta = (DisplayName = "Is a primary node"), Category = "NDisplay|Cluster")
+	UE_DEPRECATED(5.4, "This function has been moved to UDisplayClusterBlueprintLib.")
+	UFUNCTION(BlueprintCallable, meta = (DeprecatedFunction, DeprecationMessage = "This function is now available in the main blueprint functions list under 'nDisplay' section."))
 	virtual bool IsPrimary() const override;
 	
-	UFUNCTION(BlueprintCallable, meta = (DisplayName = "Is a secondary node"), Category = "NDisplay|Cluster")
+	UE_DEPRECATED(5.4, "This function has been moved to UDisplayClusterBlueprintLib.")
+	UFUNCTION(BlueprintCallable, meta = (DeprecatedFunction, DeprecationMessage = "This function is now available in the main blueprint functions list under 'nDisplay' section."))
 	virtual bool IsSecondary() const override;
 
 	/** Returns true if current node is a backup node in a cluster. */
-	UFUNCTION(BlueprintCallable, meta = (DisplayName = "Is backup node"), Category = "NDisplay|Cluster")
+	UE_DEPRECATED(5.4, "This function has been moved to UDisplayClusterBlueprintLib.")
+	UFUNCTION(BlueprintCallable, meta = (DeprecatedFunction, DeprecationMessage = "This function is now available in the main blueprint functions list under 'nDisplay' section."))
 	virtual bool IsBackup() const override;
 
 	// Cluster events API
-	UFUNCTION(BlueprintCallable, meta = (DisplayName = "Add cluster event listener"), Category = "NDisplay|Cluster")
+	UE_DEPRECATED(5.4, "This function has been moved to UDisplayClusterBlueprintLib.")
+	UFUNCTION(BlueprintCallable, meta = (DeprecatedFunction, DeprecationMessage = "This function is now available in the main blueprint functions list under 'nDisplay' section."))
 	virtual void AddClusterEventListener(TScriptInterface<IDisplayClusterClusterEventListener> Listener) override;
 
-	UFUNCTION(BlueprintCallable, meta = (DisplayName = "Remove cluster event listener"), Category = "NDisplay|Cluster")
+	UE_DEPRECATED(5.4, "This function has been moved to UDisplayClusterBlueprintLib.")
+	UFUNCTION(BlueprintCallable, meta = (DeprecatedFunction, DeprecationMessage = "This function is now available in the main blueprint functions list under 'nDisplay' section."))
 	virtual void RemoveClusterEventListener(TScriptInterface<IDisplayClusterClusterEventListener> Listener) override;
 
-	UFUNCTION(BlueprintCallable, meta = (DisplayName = "Emit JSON cluster event"), Category = "NDisplay|Cluster")
+	UE_DEPRECATED(5.4, "This function has been moved to UDisplayClusterBlueprintLib.")
+	UFUNCTION(BlueprintCallable, meta = (DeprecatedFunction, DeprecationMessage = "This function is now available in the main blueprint functions list under 'nDisplay' section."))
 	virtual void EmitClusterEventJson(const FDisplayClusterClusterEventJson& Event, bool bPrimaryOnly) override;
 
-	UFUNCTION(BlueprintCallable, meta = (DisplayName = "Emit binary cluster event"), Category = "NDisplay|Cluster")
+	UE_DEPRECATED(5.4, "This function has been moved to UDisplayClusterBlueprintLib.")
+	UFUNCTION(BlueprintCallable, meta = (DeprecatedFunction, DeprecationMessage = "This function is now available in the main blueprint functions list under 'nDisplay' section."))
 	virtual void EmitClusterEventBinary(const FDisplayClusterClusterEventBinary& Event, bool bPrimaryOnly) override;
 
-	UFUNCTION(BlueprintCallable, meta = (DisplayName = "Emit JSON cluster event"), Category = "NDisplay|Cluster")
+	UE_DEPRECATED(5.4, "This function has been moved to UDisplayClusterBlueprintLib.")
+	UFUNCTION(BlueprintCallable, meta = (DeprecatedFunction, DeprecationMessage = "This function is now available in the main blueprint functions list under 'nDisplay' section."))
 	virtual void SendClusterEventJsonTo(const FString& Address, const int32 Port, const FDisplayClusterClusterEventJson& Event, bool bPrimaryOnly) override;
 
-	UFUNCTION(BlueprintCallable, meta = (DisplayName = "Emit binary cluster event"), Category = "NDisplay|Cluster")
+	UE_DEPRECATED(5.4, "This function has been moved to UDisplayClusterBlueprintLib.")
+	UFUNCTION(BlueprintCallable, meta = (DeprecatedFunction, DeprecationMessage = "This function is now available in the main blueprint functions list under 'nDisplay' section."))
 	virtual void SendClusterEventBinaryTo(const FString& Address, const int32 Port, const FDisplayClusterClusterEventBinary& Event, bool bPrimaryOnly) override;
 };
diff --git a/Engine/Plugins/Runtime/nDisplay/Source/DisplayCluster/Private/Blueprints/DisplayClusterBlueprintLib.cpp b/Engine/Plugins/Runtime/nDisplay/Source/DisplayCluster/Private/Blueprints/DisplayClusterBlueprintLib.cpp
--- a/Engine/Plugins/Runtime/nDisplay/Source/DisplayCluster/Private/Blueprints/DisplayClusterBlueprintLib.cpp
+++ b/Engine/Plugins/Runtime/nDisplay/Source/DisplayCluster/Private/Blueprints/DisplayClusterBlueprintLib.cpp
@@ -2,13 +2,18 @@
 
 #include "Blueprints/DisplayClusterBlueprintLib.h"
 #include "Blueprints/DisplayClusterBlueprintAPIImpl.h"
-#include "UObject/Package.h"
+
+#include "Cluster/IDisplayClusterClusterManager.h"
+#include "Components/DisplayClusterICVFXCameraComponent.h"
+#include "Game/IDisplayClusterGameManager.h"
+#include "Misc/DisplayClusterLog.h"
+
+#include "IDisplayCluster.h"
 
 #include "DisplayClusterChromakeyCardActor.h"
 #include "DisplayClusterLightCardActor.h"
 #include "DisplayClusterRootActor.h"
 #include "DisplayClusterConfigurationTypes.h"
-#include "Components/DisplayClusterICVFXCameraComponent.h"
 
 #if WITH_EDITOR
 #include "ScopedTransaction.h"
@@ -16,17 +21,110 @@
 
 #define LOCTEXT_NAMESPACE "UDisplayClusterBlueprintLib"
 
+
 UDisplayClusterBlueprintLib::UDisplayClusterBlueprintLib(class FObjectInitializer const & ObjectInitializer)
 	: Super(ObjectInitializer)
 {
 }
 
+// [DEPRECATED]
 void UDisplayClusterBlueprintLib::GetAPI(TScriptInterface<IDisplayClusterBlueprintAPI>& OutAPI)
 {
 	static UDisplayClusterBlueprintAPIImpl* Obj = NewObject<UDisplayClusterBlueprintAPIImpl>(GetTransientPackage(), NAME_None, RF_MarkAsRootSet);
 	OutAPI = Obj;
 }
 
+EDisplayClusterOperationMode UDisplayClusterBlueprintLib::GetOperationMode()
+{
+	return IDisplayCluster::Get().GetOperationMode();
+}
+
+ADisplayClusterRootActor* UDisplayClusterBlueprintLib::GetRootActor()
+{
+	return IDisplayCluster::Get().GetGameMgr()->GetRootActor();
+}
+
+FString UDisplayClusterBlueprintLib::GetNodeId()
+{
+	return IDisplayCluster::Get().GetClusterMgr()->GetNodeId();
+}
+
+void UDisplayClusterBlueprintLib::GetActiveNodeIds(TArray<FString>& OutNodeIds)
+{
+	IDisplayCluster::Get().GetClusterMgr()->GetNodeIds(OutNodeIds);
+}
+
+int32 UDisplayClusterBlueprintLib::GetActiveNodesAmount()
+{
+	return IDisplayCluster::Get().GetClusterMgr()->GetNodesAmount();
+}
+
+bool UDisplayClusterBlueprintLib::IsPrimary()
+{
+	return IDisplayCluster::Get().GetClusterMgr()->IsPrimary();
+}
+
+bool UDisplayClusterBlueprintLib::IsSecondary()
+{
+	return IDisplayCluster::Get().GetClusterMgr()->IsSecondary();
+}
+
+bool UDisplayClusterBlueprintLib::IsBackup()
+{
+	return IDisplayCluster::Get().GetClusterMgr()->IsBackup();
+}
+
+EDisplayClusterNodeRole UDisplayClusterBlueprintLib::GetClusterRole()
+{
+	return IDisplayCluster::Get().GetClusterMgr()->GetClusterRole();
+}
+
+void UDisplayClusterBlueprintLib::AddClusterEventListener(TScriptInterface<IDisplayClusterClusterEventListener> Listener)
+{
+	IDisplayCluster::Get().GetClusterMgr()->AddClusterEventListener(Listener);
+}
+
+void UDisplayClusterBlueprintLib::RemoveClusterEventListener(TScriptInterface<IDisplayClusterClusterEventListener> Listener)
+{
+	IDisplayCluster::Get().GetClusterMgr()->RemoveClusterEventListener(Listener);
+}
+
+void UDisplayClusterBlueprintLib::EmitClusterEventJson(const FDisplayClusterClusterEventJson& Event, bool bPrimaryOnly)
+{
+	IDisplayCluster::Get().GetClusterMgr()->EmitClusterEventJson(Event, bPrimaryOnly);
+}
+
+void UDisplayClusterBlueprintLib::EmitClusterEventBinary(const FDisplayClusterClusterEventBinary& Event, bool bPrimaryOnly)
+{
+	IDisplayCluster::Get().GetClusterMgr()->EmitClusterEventBinary(Event, bPrimaryOnly);
+}
+
+void UDisplayClusterBlueprintLib::SendClusterEventJsonTo(const FString& Address, const int32 Port, const FDisplayClusterClusterEventJson& Event, bool bPrimaryOnly)
+{
+	if (Port >= 0 && Port <= 0xffff)
+	{
+		const uint16 PortNumber = static_cast<uint16>(Port);
+		IDisplayCluster::Get().GetClusterMgr()->SendClusterEventTo(Address, PortNumber, Event, bPrimaryOnly);
+	}
+	else
+	{
+		UE_LOG(LogDisplayClusterBlueprint, Warning, TEXT("Wrong port number: %d"), Port);
+	}
+}
+
+void UDisplayClusterBlueprintLib::SendClusterEventBinaryTo(const FString& Address, const int32 Port, const FDisplayClusterClusterEventBinary& Event, bool bPrimaryOnly)
+{
+	if (Port >= 0 && Port <= 0xffff)
+	{
+		const uint16 PortNumber = static_cast<uint16>(Port);
+		IDisplayCluster::Get().GetClusterMgr()->SendClusterEventTo(Address, PortNumber, Event, bPrimaryOnly);
+	}
+	else
+	{
+		UE_LOG(LogDisplayClusterBlueprint, Warning, TEXT("Wrong port number: %d"), Port);
+	}
+}
+
 ADisplayClusterLightCardActor* UDisplayClusterBlueprintLib::CreateLightCard(ADisplayClusterRootActor* RootActor)
 {
 	if (!RootActor)
diff --git a/Engine/Plugins/Runtime/nDisplay/Source/DisplayCluster/Private/DisplayClusterModule.cpp b/Engine/Plugins/Runtime/nDisplay/Source/DisplayCluster/Private/DisplayClusterModule.cpp
--- a/Engine/Plugins/Runtime/nDisplay/Source/DisplayCluster/Private/DisplayClusterModule.cpp
+++ b/Engine/Plugins/Runtime/nDisplay/Source/DisplayCluster/Private/DisplayClusterModule.cpp
@@ -104,9 +104,6 @@ bool FDisplayClusterModule::Init(EDisplayClusterOperationMode OperationMode)
 		UE_LOG(LogDisplayClusterModule, Error, TEXT("An error occurred during internal initialization"));
 	}
 
-	// Set internal initialization flag
-	bIsModuleInitialized = bResult;
-
 	return bResult;
 }
 
diff --git a/Engine/Plugins/Runtime/nDisplay/Source/DisplayCluster/Private/DisplayClusterModule.h b/Engine/Plugins/Runtime/nDisplay/Source/DisplayCluster/Private/DisplayClusterModule.h
--- a/Engine/Plugins/Runtime/nDisplay/Source/DisplayCluster/Private/DisplayClusterModule.h
+++ b/Engine/Plugins/Runtime/nDisplay/Source/DisplayCluster/Private/DisplayClusterModule.h
@@ -27,11 +27,6 @@ public:
 	//////////////////////////////////////////////////////////////////////////////////////////////
 	// IDisplayCluster
 	//////////////////////////////////////////////////////////////////////////////////////////////
-	virtual bool IsModuleInitialized() const override
-	{
-		return bIsModuleInitialized;
-	}
-	
 	virtual EDisplayClusterOperationMode GetOperationMode() const override
 	{
 		return CurrentOperationMode;
@@ -84,10 +79,6 @@ private:
 #endif
 
 private:
-	// Is module initialized.
-	// This flag is not the same as EDisplayClusterOperationMode::Disabled which is used when we turn off the DC functionality in a game mode.
-	bool bIsModuleInitialized = false;
-
 	// DisplayCluster subsystems
 	IPDisplayClusterClusterManager*   MgrCluster   = nullptr;
 	IPDisplayClusterRenderManager*    MgrRender    = nullptr;
diff --git a/Engine/Plugins/Runtime/nDisplay/Source/DisplayCluster/Private/Game/EngineClasses/Scene/DisplayClusterSceneComponentSync.cpp b/Engine/Plugins/Runtime/nDisplay/Source/DisplayCluster/Private/Game/EngineClasses/Scene/DisplayClusterSceneComponentSync.cpp
--- a/Engine/Plugins/Runtime/nDisplay/Source/DisplayCluster/Private/Game/EngineClasses/Scene/DisplayClusterSceneComponentSync.cpp
+++ b/Engine/Plugins/Runtime/nDisplay/Source/DisplayCluster/Private/Game/EngineClasses/Scene/DisplayClusterSceneComponentSync.cpp
@@ -2,6 +2,8 @@
 
 #include "Components/DisplayClusterSceneComponentSync.h"
 
+#include "DisplayClusterEnums.h"
+
 #include "GameFramework/Actor.h"
 
 #include "Cluster/IPDisplayClusterClusterManager.h"
@@ -21,7 +23,8 @@ void UDisplayClusterSceneComponentSync::BeginPlay()
 {
 	Super::BeginPlay();
 
-	if (!GDisplayCluster->IsModuleInitialized())
+	const EDisplayClusterOperationMode OpMode = GDisplayCluster->GetOperationMode();
+	if (OpMode != EDisplayClusterOperationMode::Cluster)
 	{
 		return;
 	}
@@ -48,7 +51,8 @@ void UDisplayClusterSceneComponentSync::BeginPlay()
 
 void UDisplayClusterSceneComponentSync::EndPlay(const EEndPlayReason::Type EndPlayReason)
 {
-	if (GDisplayCluster->IsModuleInitialized())
+	const EDisplayClusterOperationMode OpMode = GDisplayCluster->GetOperationMode();
+	if (OpMode == EDisplayClusterOperationMode::Cluster)
 	{
 		if (GameMgr && GameMgr->IsDisplayClusterActive())
 		{
diff --git a/Engine/Plugins/Runtime/nDisplay/Source/DisplayCluster/Public/Blueprints/DisplayClusterBlueprintLib.h b/Engine/Plugins/Runtime/nDisplay/Source/DisplayCluster/Public/Blueprints/DisplayClusterBlueprintLib.h
--- a/Engine/Plugins/Runtime/nDisplay/Source/DisplayCluster/Public/Blueprints/DisplayClusterBlueprintLib.h
+++ b/Engine/Plugins/Runtime/nDisplay/Source/DisplayCluster/Public/Blueprints/DisplayClusterBlueprintLib.h
@@ -3,6 +3,7 @@
 #pragma once
 
 #include "CoreMinimal.h"
+#include "DisplayClusterEnums.h"
 #include "Kismet/BlueprintFunctionLibrary.h"
 #include "Blueprints/IDisplayClusterBlueprintAPI.h"
 #include "DisplayClusterBlueprintLib.generated.h"
@@ -10,6 +11,10 @@
 class ADisplayClusterChromakeyCardActor;
 class ADisplayClusterLightCardActor;
 class ADisplayClusterRootActor;
+class IDisplayClusterClusterEventListener;
+struct FDisplayClusterClusterEventBinary;
+struct FDisplayClusterClusterEventJson;
+
 
 /**
  * Blueprint API function library
@@ -21,23 +26,91 @@ class UDisplayClusterBlueprintLib
 	GENERATED_UCLASS_BODY()
 
 public:
-	/** Return Display Cluster API interface. */
-	UFUNCTION(BlueprintPure, meta = (DisplayName = "DisplayCluster Module API"), Category = "NDisplay")
+
+	/** [DEPRECATED] Return Display Cluster API interface. */
+	UE_DEPRECATED(5.4, "GetAPI function has been deprecated. All functions are now available in UDisplayClusterBlueprintLib.")
+	UFUNCTION(BlueprintPure, meta = (DeprecatedFunction, DeprecationMessage = "GetAPI has been deprecated. All functions are now availalbe in the main blueprint functions list under 'nDisplay' category."))
 	static void GetAPI(TScriptInterface<IDisplayClusterBlueprintAPI>& OutAPI);
 
+public:
+
+	/** Returns current operation mode. */
+	UFUNCTION(BlueprintPure, meta = (DisplayName = "Get Operation Mode"), Category = "NDisplay")
+	static DISPLAYCLUSTER_API EDisplayClusterOperationMode GetOperationMode();
+
+	/** Returns currently active root actor. */
+	UFUNCTION(BlueprintPure, meta = (DisplayName = "Get Active Root Actor"), Category = "NDisplay")
+	static DISPLAYCLUSTER_API ADisplayClusterRootActor* GetRootActor();
+
+public:
+
+	/** Returns Id of the current node in a cluster. */
+	UFUNCTION(BlueprintPure, meta = (DisplayName = "Get Node ID"), Category = "NDisplay|Cluster")
+	static DISPLAYCLUSTER_API FString GetNodeId();
+
+	/** Returns List of the active nodes in the runtime cluster node in a cluster. */
+	UFUNCTION(BlueprintPure, meta = (DisplayName = "Get Cluster Node IDs"), Category = "NDisplay|Cluster")
+	static DISPLAYCLUSTER_API void GetActiveNodeIds(TArray<FString>& OutNodeIds);
+
+	/** Returns amount of active nodes in a cluster. */
+	UFUNCTION(BlueprintPure, meta = (DisplayName = "Get Cluster Nodes Amount"), Category = "NDisplay|Cluster")
+	static DISPLAYCLUSTER_API int32 GetActiveNodesAmount();
+
+	/** Returns true if current node is a primary node in a cluster. */
+	UFUNCTION(BlueprintPure, meta = (DisplayName = "Is Primary Node"), Category = "NDisplay|Cluster")
+	static DISPLAYCLUSTER_API bool IsPrimary();
+
+	/** Returns true if current node is a secondary node in a cluster. */
+	UFUNCTION(BlueprintPure, meta = (DisplayName = "Is Secondary Node"), Category = "NDisplay|Cluster")
+	static DISPLAYCLUSTER_API bool IsSecondary();
+
+	/** Returns true if current node is a backup node in a cluster. */
+	UFUNCTION(BlueprintPure, meta = (DisplayName = "Is Backup Node"), Category = "NDisplay|Cluster")
+	static DISPLAYCLUSTER_API bool IsBackup();
+
+	/** Returns the role of the current cluster node. */
+	UFUNCTION(BlueprintPure, meta = (DisplayName = "Get Cluster Node Role"), Category = "NDisplay|Cluster")
+	static DISPLAYCLUSTER_API EDisplayClusterNodeRole GetClusterRole();
+
+	/** Adds cluster event listener. */
+	UFUNCTION(BlueprintCallable, meta = (DisplayName = "Add Cluster Event Listener"), Category = "NDisplay|Cluster|Events")
+	static DISPLAYCLUSTER_API void AddClusterEventListener(TScriptInterface<IDisplayClusterClusterEventListener> Listener);
+
+	/** Removes cluster event listener. */
+	UFUNCTION(BlueprintCallable, meta = (DisplayName = "Remove Cluster Event Listener"), Category = "NDisplay|Cluster|Events")
+	static DISPLAYCLUSTER_API void RemoveClusterEventListener(TScriptInterface<IDisplayClusterClusterEventListener> Listener);
+
+	/** Emits JSON cluster event. */
+	UFUNCTION(BlueprintCallable, meta = (DisplayName = "Emit JSON Cluster Event"), Category = "NDisplay|Cluster|Events")
+	static DISPLAYCLUSTER_API void EmitClusterEventJson(const FDisplayClusterClusterEventJson& Event, bool bPrimaryOnly);
+
+	/** Emits binary cluster event. */
+	UFUNCTION(BlueprintCallable, meta = (DisplayName = "Emit Binary Cluster Event"), Category = "NDisplay|Cluster|Events")
+	static DISPLAYCLUSTER_API void EmitClusterEventBinary(const FDisplayClusterClusterEventBinary& Event, bool bPrimaryOnly);
+
+	/** Sends JSON cluster event to a specific target (may be outside of the cluster). */
+	UFUNCTION(BlueprintCallable, meta = (DisplayName = "Send JSON Event To Host"), Category = "NDisplay|Cluster|Events")
+	static DISPLAYCLUSTER_API void SendClusterEventJsonTo(const FString& Address, const int32 Port, const FDisplayClusterClusterEventJson& Event, bool bPrimaryOnly);
+
+	/** Sends binary cluster event to a specific target (may be outside of the cluster). */
+	UFUNCTION(BlueprintCallable, meta = (DisplayName = "Send Binary Event To Host"), Category = "NDisplay|Cluster|Events")
+	static DISPLAYCLUSTER_API void SendClusterEventBinaryTo(const FString& Address, const int32 Port, const FDisplayClusterClusterEventBinary& Event, bool bPrimaryOnly);
+
+public:
+
 	/** Create a new light card parented to the given nDisplay root actor. */
-	UFUNCTION(BlueprintPure, meta = (DisplayName = "DisplayCluster Module API"), Category = "NDisplay")
+	UFUNCTION(BlueprintPure, meta = (DisplayName = "Create Light Card"), Category = "NDisplay|LC & CK")
 	static ADisplayClusterLightCardActor* CreateLightCard(ADisplayClusterRootActor* RootActor);
 
 	/** Create duplicates of a list of existing light cards. */
-	UFUNCTION(BlueprintPure, meta = (DisplayName = "DisplayCluster Module API"), Category = "NDisplay")
+	UFUNCTION(BlueprintPure, meta = (DisplayName = "Duplicate Light Cards"), Category = "NDisplay|LC & CK")
 	static void DuplicateLightCards(TArray<ADisplayClusterLightCardActor*> OriginalLightcards, TArray<ADisplayClusterLightCardActor*>& OutNewLightCards);
 
 	/** Gets a list of all light card actors on the level linked to the specified root actor. */
-	UFUNCTION(BlueprintPure, meta = (DisplayName = "DisplayCluster Module API"), Category = "NDisplay")
+	UFUNCTION(BlueprintPure, meta = (DisplayName = "Find Light Cards"), Category = "NDisplay|LC & CK")
 	static DISPLAYCLUSTER_API void FindLightCardsForRootActor(const ADisplayClusterRootActor* RootActor, TSet<ADisplayClusterLightCardActor*>& OutLightCards);
 
 	/** Gets a list of all chromakey card actors on the level linked to the specified root actor. */
-	UFUNCTION(BlueprintPure, meta = (DisplayName = "DisplayCluster Module API"), Category = "NDisplay")
+	UFUNCTION(BlueprintPure, meta = (DisplayName = "Find Chromakey Cards"), Category = "NDisplay|LC & CK")
 	static DISPLAYCLUSTER_API void FindChromakeyCardsForRootActor(const ADisplayClusterRootActor* RootActor, TSet<ADisplayClusterChromakeyCardActor*>& OutChromakeyCards);
 };
diff --git a/Engine/Plugins/Runtime/nDisplay/Source/DisplayCluster/Public/Blueprints/IDisplayClusterBlueprintAPI.h b/Engine/Plugins/Runtime/nDisplay/Source/DisplayCluster/Public/Blueprints/IDisplayClusterBlueprintAPI.h
--- a/Engine/Plugins/Runtime/nDisplay/Source/DisplayCluster/Public/Blueprints/IDisplayClusterBlueprintAPI.h
+++ b/Engine/Plugins/Runtime/nDisplay/Source/DisplayCluster/Public/Blueprints/IDisplayClusterBlueprintAPI.h
@@ -41,71 +41,87 @@ class DISPLAYCLUSTER_API IDisplayClusterBlueprintAPI
 
 public:
 	/* Returns true if the module has been initialized. */
-	UFUNCTION(BlueprintCallable, meta = (DisplayName = "Is module initialized"), Category = "NDisplay")
+	UE_DEPRECATED(5.4, "This function has been deprecated and will be removed soon.")
+	UFUNCTION(BlueprintCallable, meta = (DeprecatedFunction, DeprecationMessage = "This function has been deprecated and will be removed soon."))
 	virtual bool IsModuleInitialized() const = 0;
 
 	/** Returns current operation mode. */
-	UFUNCTION(BlueprintCallable, meta = (DisplayName = "Get operation mode"), Category = "NDisplay")
+	UE_DEPRECATED(5.4, "This function has been moved to UDisplayClusterBlueprintLib.")
+	UFUNCTION(BlueprintCallable, meta = (DeprecatedFunction, DeprecationMessage = "This function is now available in the main blueprint functions list under 'nDisplay' section."))
 	virtual EDisplayClusterOperationMode GetOperationMode() const = 0;
 
 public:
 	/** Returns DisplayCluster root actor. */
-	UFUNCTION(BlueprintCallable, meta = (DisplayName = "Get root actor"), Category = "NDisplay|Game")
+	UE_DEPRECATED(5.4, "This function has been moved to UDisplayClusterBlueprintLib.")
+	UFUNCTION(BlueprintCallable, meta = (DeprecatedFunction, DeprecationMessage = "This function is now available in the main blueprint functions list under 'nDisplay' section."))
 	virtual ADisplayClusterRootActor* GetRootActor() const = 0;
 
 	// Runtime local node API
 
 	/** Returns Id of the current node in a cluster. */
-	UFUNCTION(BlueprintCallable, meta = (DisplayName = "Get node ID"), Category = "NDisplay|Cluster")
+	UE_DEPRECATED(5.4, "This function has been moved to UDisplayClusterBlueprintLib.")
+	UFUNCTION(BlueprintCallable, meta = (DeprecatedFunction, DeprecationMessage = "This function is now available in the main blueprint functions list under 'nDisplay' section."))
 	virtual FString GetNodeId() const = 0;
 
 	/** Returns List of the active nodes in the runtime cluster node in a cluster. */
-	UFUNCTION(BlueprintCallable, meta = (DisplayName = "Get cluster node IDs"), Category = "NDisplay|Cluster")
+	UE_DEPRECATED(5.4, "This function has been moved to UDisplayClusterBlueprintLib.")
+	UFUNCTION(BlueprintCallable, meta = (DeprecatedFunction, DeprecationMessage = "This function is now available in the main blueprint functions list under 'nDisplay' section."))
 	virtual void GetActiveNodeIds(TArray<FString>& OutNodeIds) const = 0;
 
 	/** Returns amount of active nodes in a cluster. */
-	UFUNCTION(BlueprintCallable, meta = (DisplayName = "Get nodes amount"), Category = "NDisplay|Cluster")
+	UE_DEPRECATED(5.4, "This function has been moved to UDisplayClusterBlueprintLib.")
+	UFUNCTION(BlueprintCallable, meta = (DeprecatedFunction, DeprecationMessage = "This function is now available in the main blueprint functions list under 'nDisplay' section."))
 	virtual int32 GetActiveNodesAmount() const = 0;
 
 	/** Returns true if current node is a primary node in a cluster. */
-	UFUNCTION(BlueprintCallable, meta = (DisplayName = "Is a primary node"), Category = "NDisplay|Cluster")
+	UE_DEPRECATED(5.4, "This function has been moved to UDisplayClusterBlueprintLib.")
+	UFUNCTION(BlueprintCallable, meta = (DeprecatedFunction, DeprecationMessage = "This function is now available in the main blueprint functions list under 'nDisplay' section."))
 	virtual bool IsPrimary() const = 0;
 
 	/** Returns true if current node is a secondary node in a cluster. */
-	UFUNCTION(BlueprintCallable, meta = (DisplayName = "Is a secondary node"), Category = "NDisplay|Cluster")
+	UE_DEPRECATED(5.4, "This function has been moved to UDisplayClusterBlueprintLib.")
+	UFUNCTION(BlueprintCallable, meta = (DeprecatedFunction, DeprecationMessage = "This function is now available in the main blueprint functions list under 'nDisplay' section."))
 	virtual bool IsSecondary() const = 0;
 
 	/** Returns true if current node is a backup node in a cluster. */
-	UFUNCTION(BlueprintCallable, meta = (DisplayName = "Is a backup node"), Category = "NDisplay|Cluster")
+	UE_DEPRECATED(5.4, "This function has been moved to UDisplayClusterBlueprintLib.")
+	UFUNCTION(BlueprintCallable, meta = (DeprecatedFunction, DeprecationMessage = "This function is now available in the main blueprint functions list under 'nDisplay' section."))
 	virtual bool IsBackup() const = 0;
 
 	/** Returns the role of the current cluster node. */
-	UFUNCTION(BlueprintCallable, meta = (DisplayName = "Get cluster role"), Category = "NDisplay|Cluster")
+	UE_DEPRECATED(5.4, "This function has been moved to UDisplayClusterBlueprintLib.")
+	UFUNCTION(BlueprintCallable, meta = (DeprecatedFunction, DeprecationMessage = "This function is now available in the main blueprint functions list under 'nDisplay' section."))
 	virtual EDisplayClusterNodeRole GetClusterRole() const = 0;
 
 	// cluster events api
 
 	/** Adds cluster event listener. */
-	UFUNCTION(BlueprintCallable, meta = (DisplayName = "Add cluster event listener"), Category = "NDisplay|Cluster")
+	UE_DEPRECATED(5.4, "This function has been moved to UDisplayClusterBlueprintLib.")
+	UFUNCTION(BlueprintCallable, meta = (DeprecatedFunction, DeprecationMessage = "This function is now available in the main blueprint functions list under 'nDisplay' section."))
 	virtual void AddClusterEventListener(TScriptInterface<IDisplayClusterClusterEventListener> Listener) = 0;
 
 	/** Removes cluster event listener. */
-	UFUNCTION(BlueprintCallable, meta = (DisplayName = "Remove cluster event listener"), Category = "NDisplay|Cluster")
+	UE_DEPRECATED(5.4, "This function has been moved to UDisplayClusterBlueprintLib.")
+	UFUNCTION(BlueprintCallable, meta = (DeprecatedFunction, DeprecationMessage = "This function is now available in the main blueprint functions list under 'nDisplay' section."))
 	virtual void RemoveClusterEventListener(TScriptInterface<IDisplayClusterClusterEventListener> Listener) = 0;
 
 	/** Emits JSON cluster event. */
-	UFUNCTION(BlueprintCallable, meta = (DisplayName = "Emit JSON cluster event"), Category = "NDisplay|Cluster")
+	UE_DEPRECATED(5.4, "This function has been moved to UDisplayClusterBlueprintLib.")
+	UFUNCTION(BlueprintCallable, meta = (DeprecatedFunction, DeprecationMessage = "This function is now available in the main blueprint functions list under 'nDisplay' section."))
 	virtual void EmitClusterEventJson(const FDisplayClusterClusterEventJson& Event, bool bPrimaryOnly) = 0;
 
 	/** Emits binary cluster event. */
-	UFUNCTION(BlueprintCallable, meta = (DisplayName = "Emit binary cluster event"), Category = "NDisplay|Cluster")
+	UE_DEPRECATED(5.4, "This function has been moved to UDisplayClusterBlueprintLib.")
+	UFUNCTION(BlueprintCallable, meta = (DeprecatedFunction, DeprecationMessage = "This function is now available in the main blueprint functions list under 'nDisplay' section."))
 	virtual void EmitClusterEventBinary(const FDisplayClusterClusterEventBinary& Event, bool bPrimaryOnly) = 0;
 
 	/** Sends JSON cluster event to a specific target (outside of the cluster). */
-	UFUNCTION(BlueprintCallable, meta = (DisplayName = "Send JSON event to a specific host"), Category = "NDisplay|Cluster")
+	UE_DEPRECATED(5.4, "This function has been moved to UDisplayClusterBlueprintLib.")
+	UFUNCTION(BlueprintCallable, meta = (DeprecatedFunction, DeprecationMessage = "This function is now available in the main blueprint functions list under 'nDisplay' section."))
 	virtual void SendClusterEventJsonTo(const FString& Address, const int32 Port, const FDisplayClusterClusterEventJson& Event, bool bPrimaryOnly) = 0;
 
 	/** Sends binary cluster event to a specific target (outside of the cluster). */
-	UFUNCTION(BlueprintCallable, meta = (DisplayName = "Send binary event to a specific host"), Category = "NDisplay|Cluster")
+	UE_DEPRECATED(5.4, "This function has been moved to UDisplayClusterBlueprintLib.")
+	UFUNCTION(BlueprintCallable, meta = (DeprecatedFunction, DeprecationMessage = "This function is now available in the main blueprint functions list under 'nDisplay' section."))
 	virtual void SendClusterEventBinaryTo(const FString& Address, const int32 Port, const FDisplayClusterClusterEventBinary& Event, bool bPrimaryOnly) = 0;
 };
diff --git a/Engine/Plugins/Runtime/nDisplay/Source/DisplayCluster/Public/Components/DisplayClusterScreenComponent.h b/Engine/Plugins/Runtime/nDisplay/Source/DisplayCluster/Public/Components/DisplayClusterScreenComponent.h
--- a/Engine/Plugins/Runtime/nDisplay/Source/DisplayCluster/Public/Components/DisplayClusterScreenComponent.h
+++ b/Engine/Plugins/Runtime/nDisplay/Source/DisplayCluster/Public/Components/DisplayClusterScreenComponent.h
@@ -22,11 +22,11 @@ public:
 
 public:
 	/** Return the screen size adjusted by its transform scale. */
-	UFUNCTION(BlueprintCallable, meta = (DisplayName = "Get screen size"), Category = "NDisplay")
+	UFUNCTION(BlueprintCallable, meta = (DisplayName = "Get screen size"), Category = "NDisplay|Components")
 	FVector2D GetScreenSize() const;
 
 	/** Set screen size (update transform scale). */
-	UFUNCTION(BlueprintCallable, meta = (DisplayName = "Set screen size"), Category = "NDisplay")
+	UFUNCTION(BlueprintCallable, meta = (DisplayName = "Set screen size"), Category = "NDisplay|Components")
 	void SetScreenSize(const FVector2D& Size);
 
 	virtual void Serialize(FArchive& Ar) override;
diff --git a/Engine/Plugins/Runtime/nDisplay/Source/DisplayCluster/Public/DisplayClusterRootActor.h b/Engine/Plugins/Runtime/nDisplay/Source/DisplayCluster/Public/DisplayClusterRootActor.h
--- a/Engine/Plugins/Runtime/nDisplay/Source/DisplayCluster/Public/DisplayClusterRootActor.h
+++ b/Engine/Plugins/Runtime/nDisplay/Source/DisplayCluster/Public/DisplayClusterRootActor.h
@@ -158,10 +158,10 @@ protected:
 	void SetLightCardOwnership();
 
 public:
-	UFUNCTION(BlueprintCallable, Category = "NDisplay")
+	UFUNCTION(BlueprintCallable, Category = "NDisplay|DCRA")
 	bool GetFlushPositionAndNormal(const FVector& WorldPosition, FVector& OutPosition, FVector& OutNormal);
 
-	UFUNCTION(BlueprintCallable, Category = "NDisplay")
+	UFUNCTION(BlueprintCallable, Category = "NDisplay|DCRA")
 	bool MakeStageActorFlushToWall(const TScriptInterface<IDisplayClusterStageActor>& StageActor, double DesiredOffsetFromFlush = 0.0f);
 
 	/**
@@ -171,7 +171,7 @@ public:
 	 * @param OutDistance - The distance to the stage geometry from the specified point
 	 * @return True if an intersection point from WorldPosition along WorldDirection was found, false if not
 	 */
-	UFUNCTION(BlueprintCallable, Category = "NDisplay")
+	UFUNCTION(BlueprintCallable, Category = "NDisplay|Stage")
 	bool GetDistanceToStageGeometry(const FVector& WorldPosition, const FVector& WorldDirection, float& OutDistance) const;
 
 	UFUNCTION(BlueprintGetter)
@@ -191,7 +191,7 @@ public:
 	 * Get the view origin most commonly used by viewports in this cluster.
 	 * If no viewports override the camera, this returns the default camera, or if there isn't one, the actor's root component.
 	 */
-	UFUNCTION(BlueprintCallable, Category = "NDisplay|Components")
+	UFUNCTION(BlueprintCallable, Category = "NDisplay|DCRA")
 	USceneComponent* GetCommonViewPoint() const;
 
 	UFUNCTION(BlueprintCallable, Category = "NDisplay|Render")
diff --git a/Engine/Plugins/Runtime/nDisplay/Source/DisplayCluster/Public/IDisplayCluster.h b/Engine/Plugins/Runtime/nDisplay/Source/DisplayCluster/Public/IDisplayCluster.h
--- a/Engine/Plugins/Runtime/nDisplay/Source/DisplayCluster/Public/IDisplayCluster.h
+++ b/Engine/Plugins/Runtime/nDisplay/Source/DisplayCluster/Public/IDisplayCluster.h
@@ -52,11 +52,15 @@ public:
 
 
 	/**
-	* Checks if the module has been initialized.
+	* [DEPRECATED] Checks if the module has been initialized.
 	*
 	* @return Is initialized
 	*/
-	virtual bool IsModuleInitialized() const = 0;
+	UE_DEPRECATED(5.4, "This function has been deprecated and will be removed soon.")
+	virtual bool IsModuleInitialized() const
+	{
+		return true;
+	}
 
 	/**
 	* Returns current operation mode.
diff --git a/Engine/Plugins/Runtime/nDisplay/Source/DisplayCluster/Public/Misc/DisplayClusterHelpers.h b/Engine/Plugins/Runtime/nDisplay/Source/DisplayCluster/Public/Misc/DisplayClusterHelpers.h
--- a/Engine/Plugins/Runtime/nDisplay/Source/DisplayCluster/Public/Misc/DisplayClusterHelpers.h
+++ b/Engine/Plugins/Runtime/nDisplay/Source/DisplayCluster/Public/Misc/DisplayClusterHelpers.h
@@ -351,71 +351,6 @@ namespace DisplayClusterHelpers
 		}
 	}
 
-	//////////////////////////////////////////////////////////////////////////////////////////////
-	// Array helpers
-	//////////////////////////////////////////////////////////////////////////////////////////////
-	namespace array
-	{
-		// Max element in array
-		template<typename T>
-		T max(const T* data, int size)
-		{
-			T result = data[0];
-			for (int i = 1; i < size; i++)
-				if (result < data[i])
-					result = data[i];
-			return result;
-		}
-
-		// Max element's index in array
-		template<typename T>
-		size_t max_idx(const T* data, int size)
-		{
-			size_t idx = 0;
-			T result = data[0];
-			for (int i = 1; i < size; i++)
-				if (result < data[i])
-				{
-					result = data[i];
-					idx = i;
-				}
-			return idx;
-		}
-
-		// Min element in array
-		template<typename T>
-		T min(const T* data, int size)
-		{
-			T result = data[0];
-			for (int i = 1; i < size; i++)
-				if (result > data[i])
-					result = data[i];
-			return result;
-		}
-
-		// Min element's index in array
-		template<typename T>
-		size_t min_idx(const T* data, int size)
-		{
-			size_t idx = 0;
-			T result = data[0];
-			for (int i = 1; i < size; i++)
-				if (result > data[i])
-				{
-					result = data[i];
-					idx = i;
-				}
-			return idx;
-		}
-
-		// Helper for array size
-		template <typename T, size_t n>
-		constexpr size_t array_size(const T(&)[n])
-		{
-			return n;
-		}
-	}
-
 	//////////////////////////////////////////////////////////////////////////////////////////////
 	// Game helpers
 	//////////////////////////////////////////////////////////////////////////////////////////////
diff --git a/Engine/Plugins/Runtime/nDisplay/Source/DisplayClusterProjection/Private/Blueprints/DisplayClusterProjectionBlueprintAPIImpl.cpp b/Engine/Plugins/Runtime/nDisplay/Source/DisplayClusterProjection/Private/Blueprints/DisplayClusterProjectionBlueprintAPIImpl.cpp
--- a/Engine/Plugins/Runtime/nDisplay/Source/DisplayClusterProjection/Private/Blueprints/DisplayClusterProjectionBlueprintAPIImpl.cpp
+++ b/Engine/Plugins/Runtime/nDisplay/Source/DisplayClusterProjection/Private/Blueprints/DisplayClusterProjectionBlueprintAPIImpl.cpp
@@ -1,39 +1,12 @@
 // Copyright Epic Games, Inc. All Rights Reserved.
 
 #include "Blueprints/DisplayClusterProjectionBlueprintAPIImpl.h"
-
-#include "IDisplayClusterProjection.h"
-#include "DisplayClusterProjectionLog.h"
-#include "DisplayClusterProjectionStrings.h"
-
-#include "Policy/Camera/DisplayClusterProjectionCameraPolicy.h"
-#include "Policy/Camera/DisplayClusterProjectionCameraPolicyFactory.h"
-
-#include "Camera/CameraComponent.h"
-
-#include "IDisplayCluster.h"
-#include "Render/IDisplayClusterRenderManager.h"
-#include "Render/Viewport/IDisplayClusterViewportManager.h"
-#include "Render/Viewport/IDisplayClusterViewport.h"
+#include "Blueprints/DisplayClusterProjectionBlueprintLib.h"
 
 //////////////////////////////////////////////////////////////////////////////////////////////
 // Policy: CAMERA
 //////////////////////////////////////////////////////////////////////////////////////////////
 void UDisplayClusterProjectionBlueprintAPIImpl::CameraPolicySetCamera(const FString& ViewportId, UCameraComponent* NewCamera, float FOVMultiplier)
 {
-	IDisplayClusterRenderManager* RenderManager = IDisplayCluster::Get().GetRenderMgr();
-	if (RenderManager && RenderManager->GetViewportManager())
-	{
-		IDisplayClusterViewport* Viewport = RenderManager->GetViewportManager()->FindViewport(ViewportId);
-		if (Viewport != nullptr)
-		{
-			// @todo: Add extra settings latter to BP call if required
-			FDisplayClusterProjectionCameraPolicySettings Settings;
-			Settings.FOVMultiplier = FOVMultiplier;
-			Settings.bCameraOverrideDefaults = true;
-
-			IDisplayClusterProjection::Get().CameraPolicySetCamera(Viewport->GetProjectionPolicy(), NewCamera, Settings);
-		}
-	}
+	UDisplayClusterProjectionBlueprintLib::CameraPolicySetCamera(ViewportId, NewCamera, FOVMultiplier);
 }
-
diff --git a/Engine/Plugins/Runtime/nDisplay/Source/DisplayClusterProjection/Private/Blueprints/DisplayClusterProjectionBlueprintAPIImpl.h b/Engine/Plugins/Runtime/nDisplay/Source/DisplayClusterProjection/Private/Blueprints/DisplayClusterProjectionBlueprintAPIImpl.h
--- a/Engine/Plugins/Runtime/nDisplay/Source/DisplayClusterProjection/Private/Blueprints/DisplayClusterProjectionBlueprintAPIImpl.h
+++ b/Engine/Plugins/Runtime/nDisplay/Source/DisplayClusterProjection/Private/Blueprints/DisplayClusterProjectionBlueprintAPIImpl.h
@@ -6,6 +6,8 @@
 #include "Blueprints/IDisplayClusterProjectionBlueprintAPI.h"
 #include "DisplayClusterProjectionBlueprintAPIImpl.generated.h"
 
+class UCameraComponent;
+
 
 /**
  * Blueprint API interface implementation
@@ -22,6 +24,7 @@ public:
 	// Policy: CAMERA
 	//////////////////////////////////////////////////////////////////////////////////////////////
 	/** Sets active camera component for camera policy */
-	UFUNCTION(BlueprintCallable, meta = (DisplayName = "Set camera"), Category = "NDisplayProjection|Camera")
+	UE_DEPRECATED(5.4, "This function has been moved to UDisplayClusterProjectionBlueprintLib.")
+	UFUNCTION(BlueprintCallable, meta = (DeprecatedFunction, DeprecationMessage = "This function is now available in the main blueprint functions list under 'nDisplay' section."))
 	virtual void CameraPolicySetCamera(const FString& ViewportId, UCameraComponent* NewCamera, float FOVMultiplier = 1.f) override;
 };
diff --git a/Engine/Plugins/Runtime/nDisplay/Source/DisplayClusterProjection/Private/Blueprints/DisplayClusterProjectionBlueprintLib.cpp b/Engine/Plugins/Runtime/nDisplay/Source/DisplayClusterProjection/Private/Blueprints/DisplayClusterProjectionBlueprintLib.cpp
--- a/Engine/Plugins/Runtime/nDisplay/Source/DisplayClusterProjection/Private/Blueprints/DisplayClusterProjectionBlueprintLib.cpp
+++ b/Engine/Plugins/Runtime/nDisplay/Source/DisplayClusterProjection/Private/Blueprints/DisplayClusterProjectionBlueprintLib.cpp
@@ -2,7 +2,14 @@
 
 #include "Blueprints/DisplayClusterProjectionBlueprintLib.h"
 #include "Blueprints/DisplayClusterProjectionBlueprintAPIImpl.h"
-#include "UObject/Package.h"
+
+#include "IDisplayCluster.h"
+#include "IDisplayClusterProjection.h"
+
+#include "Policy/Camera/DisplayClusterProjectionCameraPolicy.h"
+#include "Render/IDisplayClusterRenderManager.h"
+#include "Render/Viewport/IDisplayClusterViewportManager.h"
+#include "Render/Viewport/IDisplayClusterViewport.h"
 
 
 UDisplayClusterProjectionBlueprintLib::UDisplayClusterProjectionBlueprintLib(class FObjectInitializer const & ObjectInitializer)
@@ -10,8 +17,27 @@ UDisplayClusterProjectionBlueprintLib::UDisplayClusterProjectionBlueprintLib(cla
 {
 }
 
+// [DEPRECATED]
 void UDisplayClusterProjectionBlueprintLib::GetAPI(TScriptInterface<IDisplayClusterProjectionBlueprintAPI>& OutAPI)
 {
 	static UDisplayClusterProjectionBlueprintAPIImpl* Obj = NewObject<UDisplayClusterProjectionBlueprintAPIImpl>(GetTransientPackage(), NAME_None, RF_MarkAsRootSet);
 	OutAPI = Obj;
 }
+
+void UDisplayClusterProjectionBlueprintLib::CameraPolicySetCamera(const FString& ViewportId, UCameraComponent* NewCamera, float FOVMultiplier)
+{
+	IDisplayClusterRenderManager* RenderManager = IDisplayCluster::Get().GetRenderMgr();
+	if (RenderManager && RenderManager->GetViewportManager())
+	{
+		IDisplayClusterViewport* Viewport = RenderManager->GetViewportManager()->FindViewport(ViewportId);
+		if (Viewport != nullptr)
+		{
+			// @todo: Add extra settings latter to BP call if required
+			FDisplayClusterProjectionCameraPolicySettings Settings;
+			Settings.FOVMultiplier = FOVMultiplier;
+			Settings.bCameraOverrideDefaults = true;
+
+			IDisplayClusterProjection::Get().CameraPolicySetCamera(Viewport->GetProjectionPolicy(), NewCamera, Settings);
+		}
+	}
+}
diff --git a/Engine/Plugins/Runtime/nDisplay/Source/DisplayClusterProjection/Public/Blueprints/DisplayClusterProjectionBlueprintLib.h b/Engine/Plugins/Runtime/nDisplay/Source/DisplayClusterProjection/Public/Blueprints/DisplayClusterProjectionBlueprintLib.h
--- a/Engine/Plugins/Runtime/nDisplay/Source/DisplayClusterProjection/Public/Blueprints/DisplayClusterProjectionBlueprintLib.h
+++ b/Engine/Plugins/Runtime/nDisplay/Source/DisplayClusterProjection/Public/Blueprints/DisplayClusterProjectionBlueprintLib.h
@@ -18,7 +18,15 @@ class UDisplayClusterProjectionBlueprintLib
 	GENERATED_UCLASS_BODY()
 
 public:
-	/** Return Display Cluster API interface. */
-	UFUNCTION(BlueprintPure, meta = (DisplayName = "DisplayClusterProjection Module API"), Category = "NDisplay")
+
+	/** [DEPRECATED] Returns Display Cluster Projection API interface. */
+	UE_DEPRECATED(5.4, "GetAPI function has been deprecated. All functions are now available in UDisplayClusterProjectionBlueprintLib.")
+	UFUNCTION(BlueprintPure, meta = (DeprecatedFunction, DeprecationMessage = "GetAPI has been deprecated. All functions are now availalbe in the main blueprint functions list under 'nDisplay' category."))
 	static void GetAPI(TScriptInterface<IDisplayClusterProjectionBlueprintAPI>& OutAPI);
+
+public:
+
+	/** Sets camera up for 'camera' projection policy. */
+	UFUNCTION(BlueprintCallable, meta = (DisplayName = "Set camera"), Category = "NDisplay|Projection")
+	static void CameraPolicySetCamera(const FString& ViewportId, UCameraComponent* NewCamera, float FOVMultiplier = 1.f);
 };
diff --git a/Engine/Plugins/Runtime/nDisplay/Source/DisplayClusterProjection/Public/Blueprints/IDisplayClusterProjectionBlueprintAPI.h b/Engine/Plugins/Runtime/nDisplay/Source/DisplayClusterProjection/Public/Blueprints/IDisplayClusterProjectionBlueprintAPI.h
--- a/Engine/Plugins/Runtime/nDisplay/Source/DisplayClusterProjection/Public/Blueprints/IDisplayClusterProjectionBlueprintAPI.h
+++ b/Engine/Plugins/Runtime/nDisplay/Source/DisplayClusterProjection/Public/Blueprints/IDisplayClusterProjectionBlueprintAPI.h
@@ -29,6 +29,7 @@ public:
 	// Policy: CAMERA
 	//////////////////////////////////////////////////////////////////////////////////////////////
 	/** Sets active camera component for camera policy */
-	UFUNCTION(BlueprintCallable, meta = (DisplayName = "Set camera"), Category = "NDisplayProjection|Camera")
+	UE_DEPRECATED(5.4, "This function has been moved to UDisplayClusterProjectionBlueprintLib.")
+	UFUNCTION(BlueprintCallable, meta = (DeprecatedFunction, DeprecationMessage = "This function is now available in the main blueprint functions list under 'nDisplay' section."))
 	virtual void CameraPolicySetCamera(const FString& ViewportId, UCameraComponent* NewCamera, float FOVMultiplier = 1.f) = 0;
 };
