nDisplay: The CineCamera Filmback + CropSettings + SqueezeFactor settings should change the resolution of the ICVFX camera viewport.

diff --git a/Engine/Plugins/Runtime/nDisplay/Source/DisplayCluster/Private/Game/EngineClasses/Scene/DisplayClusterICVFXCameraComponent.cpp b/Engine/Plugins/Runtime/nDisplay/Source/DisplayCluster/Private/Game/EngineClasses/Scene/DisplayClusterICVFXCameraComponent.cpp
--- a/Engine/Plugins/Runtime/nDisplay/Source/DisplayCluster/Private/Game/EngineClasses/Scene/DisplayClusterICVFXCameraComponent.cpp
+++ b/Engine/Plugins/Runtime/nDisplay/Source/DisplayCluster/Private/Game/EngineClasses/Scene/DisplayClusterICVFXCameraComponent.cpp
@@ -150,8 +150,11 @@ void UDisplayClusterICVFXCameraComponent::UpdateOverscanEstimatedFrameSize()
 
 	const FDisplayClusterConfigurationICVFX_StageSettings& StageSettings = RootActor->GetStageSettings();
 
+	UCineCameraComponent* ActualCineCameraComponent = GetActualCineCameraComponent();
+	check(ActualCineCameraComponent);
+
 	const float CameraBufferRatio = CameraSettings.GetCameraBufferRatio(StageSettings);
-	const FIntPoint CameraFrameSize = CameraSettings.GetCameraFrameSize(StageSettings);
+	const FIntPoint CameraFrameSize = CameraSettings.GetCameraFrameSize(StageSettings, *ActualCineCameraComponent);
 
 	const FIntPoint InnerFrustumResolution(
 		CameraFrameSize.X * CameraBufferRatio,
diff --git a/Engine/Plugins/Runtime/nDisplay/Source/DisplayCluster/Private/Render/Viewport/Configuration/DisplayClusterViewportConfigurationHelpers_ICVFX.cpp b/Engine/Plugins/Runtime/nDisplay/Source/DisplayCluster/Private/Render/Viewport/Configuration/DisplayClusterViewportConfigurationHelpers_ICVFX.cpp
--- a/Engine/Plugins/Runtime/nDisplay/Source/DisplayCluster/Private/Render/Viewport/Configuration/DisplayClusterViewportConfigurationHelpers_ICVFX.cpp
+++ b/Engine/Plugins/Runtime/nDisplay/Source/DisplayCluster/Private/Render/Viewport/Configuration/DisplayClusterViewportConfigurationHelpers_ICVFX.cpp
@@ -456,12 +456,15 @@ FDisplayClusterShaderParameters_ICVFX::FCameraSettings FDisplayClusterViewportCo
 			Result.Resource.ViewportId = InCameraViewport.GetId();
 			Result.Local2WorldTransform = OriginComp->GetComponentTransform();
 
+			UCineCameraComponent* ActualCineCameraComponent = InCameraComponent.GetActualCineCameraComponent();
+			check(ActualCineCameraComponent);
+
 			// Get camera border settings
 			InCameraSettings.GetCameraBorder(*StageSettings, Result.InnerCameraBorderColor, Result.InnerCameraBorderThickness);
-			Result.InnerCameraFrameAspectRatio = InCameraSettings.GetCameraFrameAspectRatio(*StageSettings);
+			Result.InnerCameraFrameAspectRatio = InCameraSettings.GetCameraFrameAspectRatio(*StageSettings, *ActualCineCameraComponent);
 
 			// Soft edges
-			Result.SoftEdge = InCameraSettings.GetCameraSoftEdge(*StageSettings);
+			Result.SoftEdge = InCameraSettings.GetCameraSoftEdge(*StageSettings, *ActualCineCameraComponent);
 
 			// Rendering order for camera overlap
 			const FString InnerFrustumID = InCameraComponent.GetCameraUniqueId();
@@ -511,8 +514,11 @@ void FDisplayClusterViewportConfigurationHelpers_ICVFX::UpdateCameraViewportSett
 	// FDisplayClusterConfigurationICVFX_CameraSettings
 	InOutRenderSettings.CameraId.Empty();
 
+	UCineCameraComponent* ActualCineCameraComponent = InCameraComponent.GetActualCineCameraComponent();
+	check(ActualCineCameraComponent);
+
 	// UDisplayClusterConfigurationICVFX_CameraRenderSettings
-	const FIntPoint DesiredSize = InCameraSettings.GetCameraFrameSize(*StageSettings);
+	const FIntPoint DesiredSize = InCameraSettings.GetCameraFrameSize(*StageSettings, *ActualCineCameraComponent);
 
 	InOutRenderSettings.Rect = FDisplayClusterViewportHelpers::GetValidViewportRect(FIntRect(FIntPoint(0, 0), DesiredSize), DstViewport.GetId(), TEXT("Configuration Camera Frame Size"));
 
diff --git a/Engine/Plugins/Runtime/nDisplay/Source/DisplayClusterConfiguration/Private/DisplayClusterConfigurationTypes_ICVFX.cpp b/Engine/Plugins/Runtime/nDisplay/Source/DisplayClusterConfiguration/Private/DisplayClusterConfigurationTypes_ICVFX.cpp
--- a/Engine/Plugins/Runtime/nDisplay/Source/DisplayClusterConfiguration/Private/DisplayClusterConfigurationTypes_ICVFX.cpp
+++ b/Engine/Plugins/Runtime/nDisplay/Source/DisplayClusterConfiguration/Private/DisplayClusterConfigurationTypes_ICVFX.cpp
@@ -4,6 +4,7 @@
 #include "DisplayClusterConfigurationTypes.h"
 #include "IDisplayCluster.h"
 #include "Camera/CameraTypes.h"
+#include "CineCameraComponent.h"
 
 namespace UE::DisplayClusterConfiguration::ICVFX
 {
@@ -16,6 +17,14 @@ namespace UE::DisplayClusterConfiguration::ICVFX
 };
 using namespace UE::DisplayClusterConfiguration::ICVFX;
 
+int32 GDisplayClusterICVFXCameraAdoptResolution = 1;
+static FAutoConsoleVariableRef CVarGDisplayClusterICVFXCameraAdoptResolution(
+	TEXT("nDisplay.icvfx.camera.AdoptResolution"),
+	GDisplayClusterICVFXCameraAdoptResolution,
+	TEXT("Adopt camera viewport resolution with 'Filmback + CropSettings + SqueezeFactor' CineCamera settings.  (Default = 1)"),
+	ECVF_Default
+);
+
 ///////////////////////////////////////////////////////////////////////////////////////
 // FDisplayClusterConfigurationICVFX_ChromakeyMarkers
 ///////////////////////////////////////////////////////////////////////////////////////
@@ -265,26 +274,48 @@ void FDisplayClusterConfigurationICVFX_CameraSettings::SetupViewInfo(const FDisp
 	CameraMotionBlur.SetupViewInfo(InStageSettings, InOutViewInfo);
 }
 
-FIntPoint FDisplayClusterConfigurationICVFX_CameraSettings::GetCameraFrameSize(const FDisplayClusterConfigurationICVFX_StageSettings& InStageSettings) const
+FIntPoint FDisplayClusterConfigurationICVFX_CameraSettings::GetCameraFrameSize(const FDisplayClusterConfigurationICVFX_StageSettings& InStageSettings, const UCineCameraComponent& InCineCameraComponent) const
 {
-	if (RenderSettings.CustomFrameSize.bUseCustomSize)
+	const FIntPoint CameraFrameSize = RenderSettings.CustomFrameSize.bUseCustomSize
+		? FIntPoint(RenderSettings.CustomFrameSize.CustomWidth, RenderSettings.CustomFrameSize.CustomHeight)
+		: FIntPoint(InStageSettings.DefaultFrameSize.Width, InStageSettings.DefaultFrameSize.Height);
+
+	if (GDisplayClusterICVFXCameraAdoptResolution)
 	{
-		// Custom camera size
-		return FIntPoint(RenderSettings.CustomFrameSize.CustomWidth, RenderSettings.CustomFrameSize.CustomHeight);
+		// Get the size of the cinematic camera's cropped sensor:
+		const double CropedSensorWidth  = FMath::Tan(FMath::DegreesToRadians(InCineCameraComponent.GetHorizontalFieldOfView()) / 2.f) * 2.f * InCineCameraComponent.CurrentFocalLength;
+		const double CropedSensorHeight = FMath::Tan(FMath::DegreesToRadians(InCineCameraComponent.GetVerticalFieldOfView()) / 2.f) * 2.f * InCineCameraComponent.CurrentFocalLength;
+
+		// Get the ratio of the cinematic camera's cropped sensor size to the base sensor size.
+		const double CroppedSensorWidthRatio  = CropedSensorWidth / InCineCameraComponent.Filmback.SensorWidth;
+		const double CroppedSensorHeightRatio = CropedSensorHeight / InCineCameraComponent.Filmback.SensorHeight;
+
+		// Adapt camera resolution to the filmback sensor aspect ratio
+		// We keep the width, but adjust the height to match the aspect ratio of the Fimlmback sensor.
+		const double CameraFrameHeight = (InCineCameraComponent.Filmback.SensorHeight > 0.f && InCineCameraComponent.Filmback.SensorWidth > 0.f)
+			? CameraFrameSize.X / (InCineCameraComponent.Filmback.SensorWidth / InCineCameraComponent.Filmback.SensorHeight)
+			: CameraFrameSize.Y;
+
+		// Get cropped camera size
+		const FIntPoint CroppedCameraFrameSize(
+			FMath::RoundToInt(CameraFrameSize.X * CroppedSensorWidthRatio),
+			FMath::RoundToInt(CameraFrameHeight * CroppedSensorHeightRatio)
+		);
+
+		return CroppedCameraFrameSize;
 	}
 
-	// global camera size
-	return FIntPoint(InStageSettings.DefaultFrameSize.Width, InStageSettings.DefaultFrameSize.Height);
+	return CameraFrameSize;
 }
 
-float FDisplayClusterConfigurationICVFX_CameraSettings::GetCameraFrameAspectRatio(const FDisplayClusterConfigurationICVFX_StageSettings& InStageSettings) const
+float FDisplayClusterConfigurationICVFX_CameraSettings::GetCameraFrameAspectRatio(const FDisplayClusterConfigurationICVFX_StageSettings& InStageSettings, const UCineCameraComponent& InCineCameraComponent) const
 {
-	FIntPoint FrameSize = GetCameraFrameSize(InStageSettings);
+	FIntPoint FrameSize = GetCameraFrameSize(InStageSettings, InCineCameraComponent);
 
 	return (FrameSize.Y > 0 && FrameSize.X > 0) ? (float)FrameSize.X / float(FrameSize.Y) : 0;
 }
 
-FVector4 FDisplayClusterConfigurationICVFX_CameraSettings::GetCameraSoftEdge(const FDisplayClusterConfigurationICVFX_StageSettings& InStageSettings) const
+FVector4 FDisplayClusterConfigurationICVFX_CameraSettings::GetCameraSoftEdge(const FDisplayClusterConfigurationICVFX_StageSettings& InStageSettings, const UCineCameraComponent& InCineCameraComponent) const
 {
 	FVector4 ResultSoftEdge(ForceInitToZero);
 
@@ -315,7 +346,7 @@ FVector4 FDisplayClusterConfigurationICVFX_CameraSettings::GetCameraSoftEdge(con
 		if (CustomFrustum.Mode == EDisplayClusterConfigurationViewportCustomFrustumMode::Pixels)
 		{
 			const float CameraBufferRatio = GetCameraBufferRatio(InStageSettings);
-			const FIntPoint FrameSize = GetCameraFrameSize(InStageSettings);
+			const FIntPoint FrameSize = GetCameraFrameSize(InStageSettings, InCineCameraComponent);
 
 			const float  FrameWidth = FrameSize.X * CameraBufferRatio;
 			const float FrameHeight = FrameSize.Y * CameraBufferRatio;
diff --git a/Engine/Plugins/Runtime/nDisplay/Source/DisplayClusterConfiguration/Public/DisplayClusterConfigurationTypes_ICVFX.h b/Engine/Plugins/Runtime/nDisplay/Source/DisplayClusterConfiguration/Public/DisplayClusterConfigurationTypes_ICVFX.h
--- a/Engine/Plugins/Runtime/nDisplay/Source/DisplayClusterConfiguration/Public/DisplayClusterConfigurationTypes_ICVFX.h
+++ b/Engine/Plugins/Runtime/nDisplay/Source/DisplayClusterConfiguration/Public/DisplayClusterConfigurationTypes_ICVFX.h
@@ -21,6 +21,8 @@
 
 #include "DisplayClusterConfigurationTypes_ICVFX.generated.h"
 
+class UCineCameraComponent;
+
 USTRUCT(Blueprintable)
 struct DISPLAYCLUSTERCONFIGURATION_API FDisplayClusterConfigurationICVFX_LightcardCustomOCIO
 {
@@ -710,16 +712,16 @@ public:
 	bool IsChromakeyViewportSettingsEqual(const FString& InClusterNodeId1, const FString& InClusterNodeId2) const;
 
 	/** Return calculated soft edges values. */
-	FVector4 GetCameraSoftEdge(const FDisplayClusterConfigurationICVFX_StageSettings& InStageSettings) const;
+	FVector4 GetCameraSoftEdge(const FDisplayClusterConfigurationICVFX_StageSettings& InStageSettings, const UCineCameraComponent& InCineCameraComponent) const;
 
 	/** Get camera buffer ratio. */
 	float GetCameraBufferRatio(const FDisplayClusterConfigurationICVFX_StageSettings& InStageSettings) const;
 
 	/** Get camera frame size. */
-	FIntPoint GetCameraFrameSize(const FDisplayClusterConfigurationICVFX_StageSettings& InStageSettings) const;
+	FIntPoint GetCameraFrameSize(const FDisplayClusterConfigurationICVFX_StageSettings& InStageSettings, const UCineCameraComponent& InCineCameraComponent) const;
 
 	/** Get camera frame aspect ratio. */
-	float GetCameraFrameAspectRatio(const FDisplayClusterConfigurationICVFX_StageSettings& InStageSettings) const;
+	float GetCameraFrameAspectRatio(const FDisplayClusterConfigurationICVFX_StageSettings& InStageSettings, const UCineCameraComponent& InCineCameraComponent) const;
 
 	/** Get camera border settings. */
 	bool GetCameraBorder(const FDisplayClusterConfigurationICVFX_StageSettings& InStageSettings, FLinearColor& OutBorderColor, float& OutBorderThickness) const;
