nDisplay: Fixed math for ICVFX to use Origin correctly.

diff --git a/Engine/Plugins/Runtime/nDisplay/Source/DisplayCluster/Private/Render/Viewport/Configuration/DisplayClusterViewportConfigurationHelpers_ICVFX.cpp b/Engine/Plugins/Runtime/nDisplay/Source/DisplayCluster/Private/Render/Viewport/Configuration/DisplayClusterViewportConfigurationHelpers_ICVFX.cpp
--- a/Engine/Plugins/Runtime/nDisplay/Source/DisplayCluster/Private/Render/Viewport/Configuration/DisplayClusterViewportConfigurationHelpers_ICVFX.cpp
+++ b/Engine/Plugins/Runtime/nDisplay/Source/DisplayCluster/Private/Render/Viewport/Configuration/DisplayClusterViewportConfigurationHelpers_ICVFX.cpp
@@ -451,26 +451,22 @@ FDisplayClusterShaderParameters_ICVFX::FCameraSettings FDisplayClusterViewportCo
 
 	if (SceneRootActor && ConfigurationRootActor && StageSettings)
 	{
-		if (const USceneComponent* OriginComp = SceneRootActor->GetRootComponent())
-		{
-			Result.Resource.ViewportId = InCameraViewport.GetId();
-			Result.Local2WorldTransform = OriginComp->GetComponentTransform();
+		Result.Resource.ViewportId = InCameraViewport.GetId();
 
-			UCineCameraComponent* ActualCineCameraComponent = InCameraComponent.GetActualCineCameraComponent();
-			check(ActualCineCameraComponent);
+		UCineCameraComponent* ActualCineCameraComponent = InCameraComponent.GetActualCineCameraComponent();
+		check(ActualCineCameraComponent);
 
-			// Get camera border settings
-			InCameraSettings.GetCameraBorder(*StageSettings, Result.InnerCameraBorderColor, Result.InnerCameraBorderThickness);
-			Result.InnerCameraFrameAspectRatio = InCameraSettings.GetCameraFrameAspectRatio(*StageSettings, *ActualCineCameraComponent);
+		// Get camera border settings
+		InCameraSettings.GetCameraBorder(*StageSettings, Result.InnerCameraBorderColor, Result.InnerCameraBorderThickness);
+		Result.InnerCameraFrameAspectRatio = InCameraSettings.GetCameraFrameAspectRatio(*StageSettings, *ActualCineCameraComponent);
 
-			// Soft edges
-			Result.SoftEdge = InCameraSettings.GetCameraSoftEdge(*StageSettings, *ActualCineCameraComponent);
+		// Soft edges
+		Result.SoftEdge = InCameraSettings.GetCameraSoftEdge(*StageSettings, *ActualCineCameraComponent);
 
-			// Rendering order for camera overlap
-			const FString InnerFrustumID = InCameraComponent.GetCameraUniqueId();
-			const int32 CameraRenderOrder = ConfigurationRootActor->GetInnerFrustumPriority(InnerFrustumID);
-			Result.RenderOrder = (CameraRenderOrder < 0) ? InCameraSettings.RenderSettings.RenderOrder : CameraRenderOrder;
-		}
+		// Rendering order for camera overlap
+		const FString InnerFrustumID = InCameraComponent.GetCameraUniqueId();
+		const int32 CameraRenderOrder = ConfigurationRootActor->GetInnerFrustumPriority(InnerFrustumID);
+		Result.RenderOrder = (CameraRenderOrder < 0) ? InCameraSettings.RenderSettings.RenderOrder : CameraRenderOrder;
 	}
 
 	return Result;
diff --git a/Engine/Plugins/Runtime/nDisplay/Source/DisplayClusterProjection/Private/Policy/MPCDI/DisplayClusterProjectionMPCDIPolicy.cpp b/Engine/Plugins/Runtime/nDisplay/Source/DisplayClusterProjection/Private/Policy/MPCDI/DisplayClusterProjectionMPCDIPolicy.cpp
--- a/Engine/Plugins/Runtime/nDisplay/Source/DisplayClusterProjection/Private/Policy/MPCDI/DisplayClusterProjectionMPCDIPolicy.cpp
+++ b/Engine/Plugins/Runtime/nDisplay/Source/DisplayClusterProjection/Private/Policy/MPCDI/DisplayClusterProjectionMPCDIPolicy.cpp
@@ -282,6 +282,9 @@ bool FDisplayClusterProjectionMPCDIPolicy::CalculateView(IDisplayClusterViewport
 	// Readback warp context
 	WarpBlendContexts[InContextNum] = WarpBlendInterface->GetWarpData(InContextNum).WarpContext;
 
+	// Save Origin viewpoint transform to world space.
+	WarpBlendContexts[InContextNum].Origin2WorldTransform = WarpEye->World2LocalTransform;
+
 	// Transform viewpoint back to world space
 	InOutViewRotation = WarpEye->World2LocalTransform.TransformRotation(WarpBlendContexts[InContextNum].Rotation.Quaternion()).Rotator();
 	InOutViewLocation = WarpEye->World2LocalTransform.TransformPosition(WarpBlendContexts[InContextNum].Location);
@@ -356,7 +359,7 @@ void FDisplayClusterProjectionMPCDIPolicy::ApplyWarpBlend_RenderThread(FRHIComma
 		for (int32 ContextNum = 0; ContextNum < InputTextures.Num(); ContextNum++)
 		{
 			// Update referenced resources and math in the ShaderICVFX.
-			ShaderICVFX.IterateViewportResourcesByPredicate([ContextNum, ViewportManagerProxyPtr, &ShaderICVFX](FDisplayClusterShaderParametersICVFX_ViewportResource& ViewportResourceIt)
+			ShaderICVFX.IterateViewportResourcesByPredicate([ContextNum, ViewportManagerProxyPtr, &ShaderICVFX, WarpContext = WarpBlendContexts_Proxy[ContextNum]](FDisplayClusterShaderParametersICVFX_ViewportResource& ViewportResourceIt)
 			{
 				// reset prev resource reference
 				ViewportResourceIt.Texture = nullptr;
@@ -392,7 +395,7 @@ void FDisplayClusterProjectionMPCDIPolicy::ApplyWarpBlend_RenderThread(FRHIComma
 							LocalSpaceViewProjection.ViewRotation = InContext.ViewRotation;
 							LocalSpaceViewProjection.PrjMatrix    = InContext.ProjectionMatrix;
 
-							CameraSettings->SetViewProjection(LocalSpaceViewProjection);
+							CameraSettings->SetViewProjection(LocalSpaceViewProjection, WarpContext.Origin2WorldTransform);
 						}
 					}
 				}
diff --git a/Engine/Plugins/Runtime/nDisplay/Source/DisplayClusterShaders/Public/ShaderParameters/DisplayClusterShaderParameters_ICVFX.h b/Engine/Plugins/Runtime/nDisplay/Source/DisplayClusterShaders/Public/ShaderParameters/DisplayClusterShaderParameters_ICVFX.h
--- a/Engine/Plugins/Runtime/nDisplay/Source/DisplayClusterShaders/Public/ShaderParameters/DisplayClusterShaderParameters_ICVFX.h
+++ b/Engine/Plugins/Runtime/nDisplay/Source/DisplayClusterShaders/Public/ShaderParameters/DisplayClusterShaderParameters_ICVFX.h
@@ -158,12 +158,13 @@ public:
 			return OverlapChromakeyMarkerTextureRHI.IsValid();
 		}
 
-		inline void SetViewProjection(const FDisplayClusterShaderParametersICVFX_CameraViewProjection& InLocalSpaceViewProjection)
+		inline void SetViewProjection(const FDisplayClusterShaderParametersICVFX_CameraViewProjection& InCameraViewProjection, const FTransform& InOrigin2WorldTransform)
 		{
-			// Support icvfx stereo - update context from camera for each eye
-			ViewProjection.ViewRotation = Local2WorldTransform.InverseTransformRotation(InLocalSpaceViewProjection.ViewRotation.Quaternion()).Rotator();
-			ViewProjection.ViewLocation = Local2WorldTransform.InverseTransformPosition(InLocalSpaceViewProjection.ViewLocation);
-			ViewProjection.PrjMatrix = InLocalSpaceViewProjection.PrjMatrix;
+			// Transforming the camera view from "world" space to "origin" space.
+			// The "origin" is the local space for the warp geometry.
+			ViewProjection.ViewRotation = InOrigin2WorldTransform.InverseTransformRotation(InCameraViewProjection.ViewRotation.Quaternion()).Rotator();
+			ViewProjection.ViewLocation = InOrigin2WorldTransform.InverseTransformPosition(InCameraViewProjection.ViewLocation);
+			ViewProjection.PrjMatrix = InCameraViewProjection.PrjMatrix;
 		}
 
 		/**
@@ -206,8 +207,6 @@ public:
 			InnerCameraBorderThickness  = InCameraSettings.InnerCameraBorderThickness;
 			InnerCameraFrameAspectRatio = InCameraSettings.InnerCameraFrameAspectRatio;
 
-			Local2WorldTransform = InCameraSettings.Local2WorldTransform;
-
 			ViewProjection = InCameraSettings.ViewProjection;
 
 			ChromakeySource = InCameraSettings.ChromakeySource;
@@ -237,8 +236,6 @@ public:
 		float InnerCameraBorderThickness = 0.1f;
 		float InnerCameraFrameAspectRatio = 1.0f;
 
-		FTransform Local2WorldTransform;
-
 		// Camera view projection data
 		FDisplayClusterShaderParametersICVFX_CameraViewProjection ViewProjection;
 
diff --git a/Engine/Plugins/Runtime/nDisplay/Source/DisplayClusterWarp/Public/Containers/DisplayClusterWarpContext.h b/Engine/Plugins/Runtime/nDisplay/Source/DisplayClusterWarp/Public/Containers/DisplayClusterWarpContext.h
--- a/Engine/Plugins/Runtime/nDisplay/Source/DisplayClusterWarp/Public/Containers/DisplayClusterWarpContext.h
+++ b/Engine/Plugins/Runtime/nDisplay/Source/DisplayClusterWarp/Public/Containers/DisplayClusterWarpContext.h
@@ -31,6 +31,9 @@ struct FDisplayClusterWarpContext
 
 	// From the mesh local space to cave
 	FMatrix  MeshToStageMatrix = FMatrix::Identity;
+
+	// Origin viewpoint transform to world space.
+	FTransform Origin2WorldTransform = FTransform::Identity;
 };
 
 /** Internal warp data for the one context of the viewport. */
