Fix windows symbols leaking out causing unity build warnings
#rnx

diff --git a/Engine/Plugins/Runtime/nDisplay/Source/DisplayCluster/Private/Render/Synchronization/DisplayClusterRenderSyncPolicyNvidiaBase.cpp b/Engine/Plugins/Runtime/nDisplay/Source/DisplayCluster/Private/Render/Synchronization/DisplayClusterRenderSyncPolicyNvidiaBase.cpp
--- a/Engine/Plugins/Runtime/nDisplay/Source/DisplayCluster/Private/Render/Synchronization/DisplayClusterRenderSyncPolicyNvidiaBase.cpp
+++ b/Engine/Plugins/Runtime/nDisplay/Source/DisplayCluster/Private/Render/Synchronization/DisplayClusterRenderSyncPolicyNvidiaBase.cpp
@@ -5,12 +5,14 @@
 #include "Misc/DisplayClusterLog.h"
 
 #if WITH_NVAPI
+	#include "Windows/AllowWindowsPlatformTypes.h"
 	THIRD_PARTY_INCLUDES_START
 	#include "dxgi1_3.h"
 	#include "d3d12.h"
 	// NOTE: The two previous includes prevent issues in (stress) unity builds where our usage of nvapi expects __d3d12_h__ to be defined.
 	#include "nvapi.h"
 	THIRD_PARTY_INCLUDES_END
+	#include "Windows/HideWindowsPlatformTypes.h"
 #endif // WITH_NVAPI
 
 
