[nDisplay] Fixed CIS build issues caused by 30509247

diff --git a/Engine/Plugins/Runtime/nDisplay/Source/DisplayCluster/Private/Blueprints/DisplayClusterBlueprintLib.cpp b/Engine/Plugins/Runtime/nDisplay/Source/DisplayCluster/Private/Blueprints/DisplayClusterBlueprintLib.cpp
--- a/Engine/Plugins/Runtime/nDisplay/Source/DisplayCluster/Private/Blueprints/DisplayClusterBlueprintLib.cpp
+++ b/Engine/Plugins/Runtime/nDisplay/Source/DisplayCluster/Private/Blueprints/DisplayClusterBlueprintLib.cpp
@@ -15,6 +15,8 @@
 #include "DisplayClusterRootActor.h"
 #include "DisplayClusterConfigurationTypes.h"
 
+#include "UObject/Package.h"
+
 #if WITH_EDITOR
 #include "ScopedTransaction.h"
 #endif
@@ -22,11 +24,6 @@
 #define LOCTEXT_NAMESPACE "UDisplayClusterBlueprintLib"
 
 
-UDisplayClusterBlueprintLib::UDisplayClusterBlueprintLib(class FObjectInitializer const & ObjectInitializer)
-	: Super(ObjectInitializer)
-{
-}
-
 // [DEPRECATED]
 void UDisplayClusterBlueprintLib::GetAPI(TScriptInterface<IDisplayClusterBlueprintAPI>& OutAPI)
 {
diff --git a/Engine/Plugins/Runtime/nDisplay/Source/DisplayCluster/Public/Blueprints/DisplayClusterBlueprintLib.h b/Engine/Plugins/Runtime/nDisplay/Source/DisplayCluster/Public/Blueprints/DisplayClusterBlueprintLib.h
--- a/Engine/Plugins/Runtime/nDisplay/Source/DisplayCluster/Public/Blueprints/DisplayClusterBlueprintLib.h
+++ b/Engine/Plugins/Runtime/nDisplay/Source/DisplayCluster/Public/Blueprints/DisplayClusterBlueprintLib.h
@@ -23,7 +23,7 @@ UCLASS()
 class UDisplayClusterBlueprintLib
 	: public UBlueprintFunctionLibrary
 {
-	GENERATED_UCLASS_BODY()
+	GENERATED_BODY()
 
 public:
 
diff --git a/Engine/Plugins/Runtime/nDisplay/Source/DisplayClusterProjection/Private/Blueprints/DisplayClusterProjectionBlueprintLib.cpp b/Engine/Plugins/Runtime/nDisplay/Source/DisplayClusterProjection/Private/Blueprints/DisplayClusterProjectionBlueprintLib.cpp
--- a/Engine/Plugins/Runtime/nDisplay/Source/DisplayClusterProjection/Private/Blueprints/DisplayClusterProjectionBlueprintLib.cpp
+++ b/Engine/Plugins/Runtime/nDisplay/Source/DisplayClusterProjection/Private/Blueprints/DisplayClusterProjectionBlueprintLib.cpp
@@ -11,11 +11,8 @@
 #include "Render/Viewport/IDisplayClusterViewportManager.h"
 #include "Render/Viewport/IDisplayClusterViewport.h"
 
+#include "UObject/Package.h"
 
-UDisplayClusterProjectionBlueprintLib::UDisplayClusterProjectionBlueprintLib(class FObjectInitializer const & ObjectInitializer)
-	: Super(ObjectInitializer)
-{
-}
 
 // [DEPRECATED]
 void UDisplayClusterProjectionBlueprintLib::GetAPI(TScriptInterface<IDisplayClusterProjectionBlueprintAPI>& OutAPI)
diff --git a/Engine/Plugins/Runtime/nDisplay/Source/DisplayClusterProjection/Public/Blueprints/DisplayClusterProjectionBlueprintLib.h b/Engine/Plugins/Runtime/nDisplay/Source/DisplayClusterProjection/Public/Blueprints/DisplayClusterProjectionBlueprintLib.h
--- a/Engine/Plugins/Runtime/nDisplay/Source/DisplayClusterProjection/Public/Blueprints/DisplayClusterProjectionBlueprintLib.h
+++ b/Engine/Plugins/Runtime/nDisplay/Source/DisplayClusterProjection/Public/Blueprints/DisplayClusterProjectionBlueprintLib.h
@@ -15,7 +15,7 @@ UCLASS()
 class UDisplayClusterProjectionBlueprintLib
 	: public UBlueprintFunctionLibrary
 {
-	GENERATED_UCLASS_BODY()
+	GENERATED_BODY()
 
 public:
 
