nDisplay: fixed issue with a warning from compilation test.

diff --git a/Engine/Plugins/Runtime/nDisplay/Source/DisplayClusterProjection/Private/Policy/EasyBlend/IDisplayClusterProjectionEasyBlendPolicyViewData.h b/Engine/Plugins/Runtime/nDisplay/Source/DisplayClusterProjection/Private/Policy/EasyBlend/IDisplayClusterProjectionEasyBlendPolicyViewData.h
--- a/Engine/Plugins/Runtime/nDisplay/Source/DisplayClusterProjection/Private/Policy/EasyBlend/IDisplayClusterProjectionEasyBlendPolicyViewData.h
+++ b/Engine/Plugins/Runtime/nDisplay/Source/DisplayClusterProjection/Private/Policy/EasyBlend/IDisplayClusterProjectionEasyBlendPolicyViewData.h
@@ -25,10 +25,7 @@ struct FDisplayClusterProjectionEasyBlendPolicyViewInfo
 class IDisplayClusterProjectionEasyBlendPolicyViewData
 {
 public:
-	virtual ~IDisplayClusterProjectionEasyBlendPolicyViewData()
-	{
-		ImplRelease();
-	}
+	virtual ~IDisplayClusterProjectionEasyBlendPolicyViewData() = default;
 
 	/** Creates and initializes the EasyBlend interface for the current RHI and OS.
 	* Returns nullptr if it is not implemented or cannot be initialized.
@@ -48,10 +45,6 @@ public:
 	virtual bool HasPreviewMesh() { return false; }
 	virtual bool GetPreviewMeshGeometry(const FDisplayClusterProjectionEasyBlendPolicyConfiguration& InEasyBlendConfiguration, FDisplayClusterProjectionEasyBlendGeometryExportData& OutMeshData) { return false; }
 
-protected:
-	/** Release EasyBlend view data from any thread. */
-	virtual void ImplRelease() { }
-
 protected:
 	// Lock access to EasyBlendMeshData
 	FCriticalSection EasyBlendMeshDataAccessCS;
diff --git a/Engine/Plugins/Runtime/nDisplay/Source/DisplayClusterProjection/Private/Policy/EasyBlend/Windows/DX11/DisplayClusterProjectionEasyBlendPolicyViewDataDX11.cpp b/Engine/Plugins/Runtime/nDisplay/Source/DisplayClusterProjection/Private/Policy/EasyBlend/Windows/DX11/DisplayClusterProjectionEasyBlendPolicyViewDataDX11.cpp
--- a/Engine/Plugins/Runtime/nDisplay/Source/DisplayClusterProjection/Private/Policy/EasyBlend/Windows/DX11/DisplayClusterProjectionEasyBlendPolicyViewDataDX11.cpp
+++ b/Engine/Plugins/Runtime/nDisplay/Source/DisplayClusterProjection/Private/Policy/EasyBlend/Windows/DX11/DisplayClusterProjectionEasyBlendPolicyViewDataDX11.cpp
@@ -17,6 +17,11 @@
 //------------------------------------------------------------------------------
 // FDisplayClusterProjectionEasyBlendPolicyViewDataDX11
 //------------------------------------------------------------------------------
+FDisplayClusterProjectionEasyBlendPolicyViewDataDX11::~FDisplayClusterProjectionEasyBlendPolicyViewDataDX11()
+{
+	ImplRelease();
+}
+
 void FDisplayClusterProjectionEasyBlendPolicyViewDataDX11::ImplRelease()
 {
 	if (EasyBlendMeshData.IsValid())
diff --git a/Engine/Plugins/Runtime/nDisplay/Source/DisplayClusterProjection/Private/Policy/EasyBlend/Windows/DX11/DisplayClusterProjectionEasyBlendPolicyViewDataDX11.h b/Engine/Plugins/Runtime/nDisplay/Source/DisplayClusterProjection/Private/Policy/EasyBlend/Windows/DX11/DisplayClusterProjectionEasyBlendPolicyViewDataDX11.h
--- a/Engine/Plugins/Runtime/nDisplay/Source/DisplayClusterProjection/Private/Policy/EasyBlend/Windows/DX11/DisplayClusterProjectionEasyBlendPolicyViewDataDX11.h
+++ b/Engine/Plugins/Runtime/nDisplay/Source/DisplayClusterProjection/Private/Policy/EasyBlend/Windows/DX11/DisplayClusterProjectionEasyBlendPolicyViewDataDX11.h
@@ -12,15 +12,20 @@ class FDisplayClusterProjectionEasyBlendPolicyViewDataDX11
 	: public IDisplayClusterProjectionEasyBlendPolicyViewData
 	, public TSharedFromThis<FDisplayClusterProjectionEasyBlendPolicyViewDataDX11, ESPMode::ThreadSafe>
 {
+public:
+	~FDisplayClusterProjectionEasyBlendPolicyViewDataDX11();
+
 public:
 	//BEGIN ~IDisplayClusterProjectionEasyBlendPolicyViewData
 	virtual bool Initialize(const FDisplayClusterProjectionEasyBlendPolicyConfiguration& InEasyBlendConfiguration) override;
-	virtual void ImplRelease() override;
 
 	virtual bool CalculateWarpBlend(FDisplayClusterProjectionEasyBlendPolicyViewInfo& InOutViewInfo) override;
 	virtual bool ApplyWarpBlend_RenderThread(FRHICommandListImmediate& RHICmdList, const FDisplayClusterProjectionEasyBlendPolicyViewInfo& InViewInfo, FRHITexture2D* InputTexture, FRHITexture2D* OutputTexture, FRHIViewport* InRHIViewport) override;
 	//END ~~IDisplayClusterProjectionEasyBlendPolicyViewData
 
+protected:
+	void ImplRelease();
+
 private:
 	// The unique EasyBlend data that useed for warpblend
 	TUniquePtr<EasyBlend1SDKDX_Mesh> EasyBlendMeshData;
diff --git a/Engine/Plugins/Runtime/nDisplay/Source/DisplayClusterProjection/Private/Policy/EasyBlend/Windows/DX12/DisplayClusterProjectionEasyBlendPolicyViewDataDX12.cpp b/Engine/Plugins/Runtime/nDisplay/Source/DisplayClusterProjection/Private/Policy/EasyBlend/Windows/DX12/DisplayClusterProjectionEasyBlendPolicyViewDataDX12.cpp
--- a/Engine/Plugins/Runtime/nDisplay/Source/DisplayClusterProjection/Private/Policy/EasyBlend/Windows/DX12/DisplayClusterProjectionEasyBlendPolicyViewDataDX12.cpp
+++ b/Engine/Plugins/Runtime/nDisplay/Source/DisplayClusterProjection/Private/Policy/EasyBlend/Windows/DX12/DisplayClusterProjectionEasyBlendPolicyViewDataDX12.cpp
@@ -40,6 +40,11 @@ namespace UE::DisplayClusterProjection::EasyBlend
 //------------------------------------------------------------------------------
 // FDisplayClusterProjectionEasyBlendPolicyViewDataDX12
 //------------------------------------------------------------------------------
+FDisplayClusterProjectionEasyBlendPolicyViewDataDX12::~FDisplayClusterProjectionEasyBlendPolicyViewDataDX12()
+{
+	ImplRelease();
+}
+
 void FDisplayClusterProjectionEasyBlendPolicyViewDataDX12::ImplRelease()
 {
 	if (EasyBlendMeshData.IsValid())
diff --git a/Engine/Plugins/Runtime/nDisplay/Source/DisplayClusterProjection/Private/Policy/EasyBlend/Windows/DX12/DisplayClusterProjectionEasyBlendPolicyViewDataDX12.h b/Engine/Plugins/Runtime/nDisplay/Source/DisplayClusterProjection/Private/Policy/EasyBlend/Windows/DX12/DisplayClusterProjectionEasyBlendPolicyViewDataDX12.h
--- a/Engine/Plugins/Runtime/nDisplay/Source/DisplayClusterProjection/Private/Policy/EasyBlend/Windows/DX12/DisplayClusterProjectionEasyBlendPolicyViewDataDX12.h
+++ b/Engine/Plugins/Runtime/nDisplay/Source/DisplayClusterProjection/Private/Policy/EasyBlend/Windows/DX12/DisplayClusterProjectionEasyBlendPolicyViewDataDX12.h
@@ -12,11 +12,12 @@ class FDisplayClusterProjectionEasyBlendPolicyViewDataDX12
 	: public IDisplayClusterProjectionEasyBlendPolicyViewData
 	, public TSharedFromThis<FDisplayClusterProjectionEasyBlendPolicyViewDataDX12, ESPMode::ThreadSafe>
 {
+public:
+	~FDisplayClusterProjectionEasyBlendPolicyViewDataDX12();
+
 public:
 	//BEGIN ~IDisplayClusterProjectionEasyBlendPolicyViewData
 	virtual bool Initialize(const FDisplayClusterProjectionEasyBlendPolicyConfiguration& InEasyBlendConfiguration) override;
-	virtual void ImplRelease() override;
-
 	virtual bool CalculateWarpBlend(FDisplayClusterProjectionEasyBlendPolicyViewInfo& InOutViewInfo) override;
 	virtual bool ApplyWarpBlend_RenderThread(FRHICommandListImmediate& RHICmdList, const FDisplayClusterProjectionEasyBlendPolicyViewInfo& InViewInfo, FRHITexture2D* InputTexture, FRHITexture2D* OutputTexture, FRHIViewport* InRHIViewport) override;
 
@@ -24,6 +25,9 @@ public:
 	virtual bool GetPreviewMeshGeometry(const FDisplayClusterProjectionEasyBlendPolicyConfiguration& InEasyBlendConfiguration, FDisplayClusterProjectionEasyBlendGeometryExportData& OutMeshData) override;
 	//END ~~IDisplayClusterProjectionEasyBlendPolicyViewData
 
+protected:
+	void ImplRelease();
+
 private:
 	// The unique EasyBlend data that useed for warpblend
 	TUniquePtr<EasyBlendSDK_Mesh> EasyBlendMeshData;
