nDisplay: Fixes a bug in relative path logic for MPCDI and PFM projection policies.

diff --git a/Engine/Plugins/Runtime/nDisplay/Source/DisplayClusterProjection/Private/Policy/MPCDI/DisplayClusterProjectionMPCDIPolicy_Config.cpp b/Engine/Plugins/Runtime/nDisplay/Source/DisplayClusterProjection/Private/Policy/MPCDI/DisplayClusterProjectionMPCDIPolicy_Config.cpp
--- a/Engine/Plugins/Runtime/nDisplay/Source/DisplayClusterProjection/Private/Policy/MPCDI/DisplayClusterProjectionMPCDIPolicy_Config.cpp
+++ b/Engine/Plugins/Runtime/nDisplay/Source/DisplayClusterProjection/Private/Policy/MPCDI/DisplayClusterProjectionMPCDIPolicy_Config.cpp
@@ -63,7 +63,7 @@ bool FDisplayClusterProjectionMPCDIPolicy_ConfigParser::ImplGetMPCDIConfig()
 	if (DisplayClusterHelpers::map::template ExtractValue(ConfigParameters, DisplayClusterProjectionStrings::cfg::mpcdi::File, LocalMPCDIFileName))
 	{
 		UE_LOG(LogDisplayClusterProjectionMPCDI, Verbose, TEXT("Found mpcdi file name for %s:%s - %s"), *BufferId, *RegionId, *LocalMPCDIFileName);
-		MPCDIFileName = LocalMPCDIFileName;
+		MPCDIFileName = DisplayClusterHelpers::filesystem::GetFullPathForConfigResource(LocalMPCDIFileName);
 	}
 
 	if (MPCDIFileName.IsEmpty())
@@ -137,7 +137,7 @@ bool FDisplayClusterProjectionMPCDIPolicy_ConfigParser::ImplGetPFMConfig()
 	if (DisplayClusterHelpers::map::template ExtractValue(ConfigParameters, DisplayClusterProjectionStrings::cfg::mpcdi::FilePFM, LocalPFMFile))
 	{
 		UE_LOG(LogDisplayClusterProjectionMPCDI, Verbose, TEXT("Found Argument '%s'='%s'"), DisplayClusterProjectionStrings::cfg::mpcdi::FilePFM, *LocalPFMFile);
-		PFMFile = LocalPFMFile;
+		PFMFile = DisplayClusterHelpers::filesystem::GetFullPathForConfigResource(LocalPFMFile);
 	}
 
 	if (PFMFile.IsEmpty())
@@ -190,7 +190,7 @@ bool FDisplayClusterProjectionMPCDIPolicy_ConfigParser::ImplGetPFMConfig()
 	if (DisplayClusterHelpers::map::template ExtractValueFromString(ConfigParameters, DisplayClusterProjectionStrings::cfg::mpcdi::FileAlpha, LocalAlphaFile))
 	{
 		UE_LOG(LogDisplayClusterProjectionMPCDI, Verbose, TEXT("Found external AlphaMap file - %s"), *LocalAlphaFile);
-		AlphaFile = LocalAlphaFile;
+		AlphaFile = DisplayClusterHelpers::filesystem::GetFullPathForConfigResource(LocalAlphaFile);
 	}
 
 	AlphaGamma = 1;
@@ -204,7 +204,7 @@ bool FDisplayClusterProjectionMPCDIPolicy_ConfigParser::ImplGetPFMConfig()
 	if (DisplayClusterHelpers::map::template ExtractValueFromString(ConfigParameters, DisplayClusterProjectionStrings::cfg::mpcdi::FileBeta, LocalBetaFile))
 	{
 		UE_LOG(LogDisplayClusterProjectionMPCDI, Verbose, TEXT("Found external BetaMap file - %s"), *LocalBetaFile);
-		BetaFile = LocalBetaFile;
+		BetaFile = DisplayClusterHelpers::filesystem::GetFullPathForConfigResource(LocalBetaFile);
 	}
 
 	return true;
