[nDispaly] Fixed (workaround) crash in a child DCRA BP while compiling its parent BP

diff --git a/Engine/Plugins/Runtime/nDisplay/Source/DisplayClusterConfigurator/Private/DisplayClusterConfiguratorCompiler.cpp b/Engine/Plugins/Runtime/nDisplay/Source/DisplayClusterConfigurator/Private/DisplayClusterConfiguratorCompiler.cpp
--- a/Engine/Plugins/Runtime/nDisplay/Source/DisplayClusterConfigurator/Private/DisplayClusterConfiguratorCompiler.cpp
+++ b/Engine/Plugins/Runtime/nDisplay/Source/DisplayClusterConfigurator/Private/DisplayClusterConfiguratorCompiler.cpp
@@ -262,11 +262,25 @@ void FDisplayClusterConfiguratorKismetCompilerContext::ValidateConfiguration()
 	
 	for (const TPair<FString, TObjectPtr<UDisplayClusterConfigurationClusterNode>>& ClusterNode : BlueprintData->Cluster->Nodes)
 	{
+		// This is a temporary workaround to prevent crashes while dealing with DCRA BP hierarchies. It looks like
+		// we don't have a solution to propagate data changes to the assets that have not been instantiated.
+		if (ClusterNode.Key.IsEmpty() || !ClusterNode.Value)
+		{
+			continue;
+		}
+
 		// Validate viewports
 		if (ClusterNode.Value->Viewports.Num() > 0)
 		{
 			for (const TPair<FString, TObjectPtr<UDisplayClusterConfigurationViewport>>& Viewport : ClusterNode.Value->Viewports)
 			{
+				// This is a temporary workaround to prevent crashes while dealing with DCRA BP hierarchies. It looks like
+				// we don't have a solution to propagate data changes to the assets that have not been instantiated.
+				if (Viewport.Key.IsEmpty() || !Viewport.Value)
+				{
+					continue;
+				}
+
 				const FString ViewportName = Viewport.Value->GetName();
 
 				// Check that no two viewports have the same name
