nDisplay:
- Fixed the issue of asynchronous client nodes connection to the server.
- bPendingMode flag added to freeze Tick execution until the nodes receive a control cluster command - NodeSyncEvent
- the algorithm for listen server initialization improved and several flags added (bSkipServerReplicateActors, bTickingThrottleEnabled, bClusterHasConnected) to block replication before the game start: added checks for number of connected nodes and their state as multiplayer clients.
- cluster disconnects and re-connects handling added: clearing of packets queues and forcing of client nodes shutdown by special cluster command
- packets slice execution command improved. Locking of replication implemented until large size packets are processed (these packets are split and marked as Reliable)

diff --git a/Engine/Plugins/Runtime/nDisplay/Source/DisplayCluster/Private/Game/EngineClasses/Basics/DisplayClusterGameEngine.cpp b/Engine/Plugins/Runtime/nDisplay/Source/DisplayCluster/Private/Game/EngineClasses/Basics/DisplayClusterGameEngine.cpp
--- a/Engine/Plugins/Runtime/nDisplay/Source/DisplayCluster/Private/Game/EngineClasses/Basics/DisplayClusterGameEngine.cpp
+++ b/Engine/Plugins/Runtime/nDisplay/Source/DisplayCluster/Private/Game/EngineClasses/Basics/DisplayClusterGameEngine.cpp
@@ -361,7 +361,7 @@ void UDisplayClusterGameEngine::Tick(float DeltaSeconds, bool bIdleMode)
 
 		// Perform UGameEngine::Tick() calls for scene actors
 		UE_LOG(LogDisplayClusterEngine, Verbose, TEXT("Perform UGameEngine::Tick()"));
-		Super::Tick(DeltaSeconds, bIdleMode);
+		Super::Tick(DeltaSeconds, bIdleMode || bForcedTickIdleMode);
 
 		// Perform PostTick for DisplayCluster module
 		UE_LOG(LogDisplayClusterEngine, Verbose, TEXT("Perform PostTick()"));
@@ -521,6 +521,19 @@ EBrowseReturnVal::Type UDisplayClusterGameEngine::BrowseLoadMap(FWorldContext& W
 			WorldContext.PendingNetGame->Initialize(URL); //-V595
 			WorldContext.PendingNetGame->InitNetDriver(); //-V595
 
+			UNetDriver* PendingNetDriver = WorldContext.PendingNetGame->GetNetDriver();
+
+			if (IsDisplayCluster && PendingNetDriver)
+			{
+				const bool bIsDisplayClusterNetDriver = PendingNetDriver->GetClass()->GetName().Equals(TEXT("DisplayClusterNetDriver"));
+				
+				if (bIsDisplayClusterNetDriver)
+				{
+					// Force tick idle mode for multiplayer connections
+					bForcedTickIdleMode = true;
+				}
+			}
+
 			if (!WorldContext.PendingNetGame)
 			{
 				// If the inital packet sent in InitNetDriver results in a socket error, HandleDisconnect() and CancelPending() may be called, which will null the PendingNetGame.
diff --git a/Engine/Plugins/Runtime/nDisplay/Source/DisplayCluster/Public/DisplayClusterGameEngine.h b/Engine/Plugins/Runtime/nDisplay/Source/DisplayCluster/Public/DisplayClusterGameEngine.h
--- a/Engine/Plugins/Runtime/nDisplay/Source/DisplayCluster/Public/DisplayClusterGameEngine.h
+++ b/Engine/Plugins/Runtime/nDisplay/Source/DisplayCluster/Public/DisplayClusterGameEngine.h
@@ -43,6 +43,14 @@ public:
 		return OperationMode;
 	}
 
+	/** Resets forced idle mode flag
+	 * Used by DisplayClusterNetDriver when all nodes are connected
+	 */
+	void ResetForcedIdleMode()
+	{
+		bForcedTickIdleMode = false;
+	}
+
 protected:
 	virtual void UpdateTimeAndHandleMaxTickRate() override;
 
@@ -80,4 +88,10 @@ private:
 	EDisplayClusterRunningMode   RunningMode = EDisplayClusterRunningMode::Startup;
 
 	TMap<FString, TSet<FString>> SyncMap;
+
+	// Flag used to force nodes to skip frame rendering
+	// This is needed when multiplayer connections being established in the cluster,
+	// system have to wait for all nodes to be connected before proceeding with rendering otherwise cluster will deadlock itself between GT/RT barriers
+	// Enforced automatically into true on BrowseLoadMap for multiplayer connections
+	bool bForcedTickIdleMode = false;
 };
diff --git a/Engine/Plugins/Runtime/nDisplay/Source/DisplayClusterReplication/Private/DisplayClusterNetConnection.cpp b/Engine/Plugins/Runtime/nDisplay/Source/DisplayClusterReplication/Private/DisplayClusterNetConnection.cpp
--- a/Engine/Plugins/Runtime/nDisplay/Source/DisplayClusterReplication/Private/DisplayClusterNetConnection.cpp
+++ b/Engine/Plugins/Runtime/nDisplay/Source/DisplayClusterReplication/Private/DisplayClusterNetConnection.cpp
@@ -39,6 +39,9 @@ UDisplayClusterNetConnection::UDisplayClusterNetConnection(const FObjectInitiali
 	bIsClusterConnection(false),
 	bSynchronousMode(false)
 {
+	// By default, networking settings are set for internet based gameplay with a limited bandwidth capacity
+	// Those settings are not working well in high speed LAN environment
+	SetUnlimitedBunchSizeAllowed(true);
 }
 
 void UDisplayClusterNetConnection::SetClientLoginState(const EClientLoginState::Type NewState)
@@ -93,7 +96,7 @@ void UDisplayClusterNetConnection::SetClientLoginState(const EClientLoginState::
 		}
 		else
 		{
-			UE_LOG(LogDisplayClusterNetDriver, Verbose, TEXT("Cluster connection joined: Client [%u]; Node [%s]; RemoteAddr: %s"), ClientId, *NodeName, *GetRemoteAddr()->ToString(false));
+			UE_LOG(LogDisplayClusterNetDriver, Verbose, TEXT("Cluster connection joined: Client [%u]; Node [%s]; RemoteAddr: %s"), ClientId, *NodeName, *NodeAddress);
 		}
 
 		DisplayClusterNetDriver->AddNodeConnection(this);
diff --git a/Engine/Plugins/Runtime/nDisplay/Source/DisplayClusterReplication/Private/DisplayClusterNetDriver.cpp b/Engine/Plugins/Runtime/nDisplay/Source/DisplayClusterReplication/Private/DisplayClusterNetDriver.cpp
--- a/Engine/Plugins/Runtime/nDisplay/Source/DisplayClusterReplication/Private/DisplayClusterNetDriver.cpp
+++ b/Engine/Plugins/Runtime/nDisplay/Source/DisplayClusterReplication/Private/DisplayClusterNetDriver.cpp
@@ -6,6 +6,7 @@
 
 #include "DisplayClusterNetDriver.h"
 #include "DisplayClusterNetDriverLog.h"
+#include "DisplayClusterGameEngine.h"
 
 #include "Engine/ActorChannel.h"
 #include "Engine/ChildConnection.h"
@@ -37,12 +38,19 @@
 
 #include "UObject/Package.h"
 
+#include "Math/NumericLimits.h"
+
 UDisplayClusterNetDriver::UDisplayClusterNetDriver(const FObjectInitializer& ObjectInitializer) :
 	Super(ObjectInitializer),
 	bClusterHasConnected(false),
-	bLastBunchWasAcked(false),
-	bConnectionViewersAreReady(false)
+	bLastBunchWasAcked(true),
+	bConnectionViewersAreReady(false),
+	ListenClusterId(0),
+	ListenClusterNodesNum(0)
 {
+	bSkipServerReplicateActors = true;
+	bTickingThrottleEnabled = false;
+
 	// Preallocate elements to prevent allocs at runtime
 	const int NumReservedElements = 32;
 	ReadyOutPackets.Reserve(NumReservedElements);
@@ -108,6 +116,9 @@ bool UDisplayClusterNetDriver::InitListen(FNetworkNotify* InNotify, FURL& Listen
 
 		if (ClusterNetworkDriverHelper->RegisterClusterEventsBinaryClient(ClusterId, NodeAddress, NodePort))
 		{
+			ListenClusterId = ClusterId;
+			ListenClusterNodesNum = FCString::Atoi(ClusterNodesNumString);
+
 			UE_LOG(LogDisplayClusterNetDriver, Verbose, TEXT("Registered primary node for cluster [%u] node [%s:%u]"), ClusterId, *NodeAddress, NodePort);
 		}
 	}
@@ -128,7 +139,12 @@ void UDisplayClusterNetDriver::TickDispatch(float DeltaTime)
 	// Register binary client when ALL nodes connected
 	if (!bClusterHasConnected)
 	{
-		UE_LOG(LogDisplayClusterNetDriver, VeryVerbose, TEXT("Primary Node connecitons = [%u]"), PrimaryNodeConnections.Num());
+		UE_LOG(LogDisplayClusterNetDriver, VeryVerbose, TEXT("Primary Node connections = [%u]"), PrimaryNodeConnections.Num());
+
+		if (World->GetNetMode() == NM_ListenServer && (ListenClusterNodesNum == NodeConnections.Num() + 1))
+		{
+			bClusterHasConnected = NotifyClusterAsReadyForSync(ListenClusterId);
+		}
 
 		for (UDisplayClusterNetConnection* PrimaryNodeConnectionIt : PrimaryNodeConnections)
 		{
@@ -152,29 +168,7 @@ void UDisplayClusterNetDriver::TickDispatch(float DeltaTime)
 				UE_LOG(LogDisplayClusterNetDriver, VeryVerbose, TEXT("Registered primary node for cluster [%u] node [%s]"), PrimaryNodeConnectionIt->ClusterId, *PrimaryNodeConnectionIt->NodeName);
 			}
 
-			// Notify primary node that current node connection is ready
-			FDisplayClusterClusterEventBinary NetworkDriverSyncEvent;
-			GenerateClusterCommandsEvent(NetworkDriverSyncEvent, NodeSyncEvent);
-
-			// send command 
-			const bool bEventWasSent = ClusterNetworkDriverHelper->SendCommandToCluster(PrimaryNodeConnectionIt->ClusterId, NetworkDriverSyncEvent);
-
-			// If event was sent to primary cluster node
-			if (bEventWasSent)
-			{
-				bClusterHasConnected = true;
-				// Mark all node connections related to current ClusterId as ready
-				for (UDisplayClusterNetConnection* ClusterNodeConnectionIt : NodeConnections)
-				{
-					if (ClusterNodeConnectionIt->ClusterId == PrimaryNodeConnectionIt->ClusterId)
-					{
-						ClusterNodeConnectionIt->bSynchronousMode = true;
-						SyncConnections.Add(ClusterNodeConnectionIt);
-					}
-				}
-
-				UE_LOG(LogDisplayClusterNetDriver, VeryVerbose, TEXT("Setting synchronous mode in cluster [%u]"), PrimaryNodeConnectionIt->ClusterId);
-			}
+			bClusterHasConnected = NotifyClusterAsReadyForSync(PrimaryNodeConnectionIt->ClusterId);
 		}
 	}
 }
@@ -220,7 +214,6 @@ void UDisplayClusterNetDriver::RemoveNodeConnection(UDisplayClusterNetConnection
 	// Reset packet queues
 	if (SyncConnections.IsEmpty())
 	{
-		LastAckedPacket.Reset();
 		OutPacketsQueues.Reset();
 		bClusterHasConnected = false;
 	}
@@ -254,6 +247,17 @@ void UDisplayClusterNetDriver::HandleEvent(FDisplayClusterClusterEventBinary con
 		// Client received command to start synchrosonus packets processing
 		ClusterServerConnection->bSynchronousMode = true;
 		ClusterServerConnection->ClientId = TextKeyUtil::HashString(ClusterServerConnection->Challenge);
+		
+		UDisplayClusterGameEngine* ClusterGameEngine = Cast<UDisplayClusterGameEngine>(GEngine);
+		
+		if (ClusterGameEngine != nullptr)
+		{
+			ClusterGameEngine->ResetForcedIdleMode();
+		}
+		else
+		{
+			UE_LOG(LogDisplayClusterNetDriver, Error, TEXT("Can't reset internal idle mode, GEngine is not UDisplayClusterGameEngine instance"));
+		}
 
 		UE_LOG(LogDisplayClusterNetDriver, VeryVerbose, TEXT("Client node self-introducing: %u"), ClusterServerConnection->ClientId);
 	}
@@ -293,6 +297,28 @@ void UDisplayClusterNetDriver::HandleEvent(FDisplayClusterClusterEventBinary con
 
 void UDisplayClusterNetDriver::TickFlush(float DeltaSeconds)
 {
+	if (IsServer())
+	{
+		int32 ClusterNumNodes = -1;
+
+		for (UDisplayClusterNetConnection* SyncConnectionIt : SyncConnections)
+		{
+			ClusterNumNodes = SyncConnectionIt->ClusterNodesNum;
+		}
+
+		if (World->GetNetMode() == NM_ListenServer)
+		{
+			bSkipServerReplicateActors = (SyncConnections.Num() != ClusterNumNodes - 1);
+		}
+		else
+		{
+			bSkipServerReplicateActors = (SyncConnections.Num() != ClusterNumNodes);
+		}
+
+		UE_LOG(LogDisplayClusterNetDriver, VeryVerbose, TEXT("bSkipServerReplicateActors %u; ClusterNumNodes: %d, SyncConnections.Num(): %d"), bSkipServerReplicateActors,
+			ClusterNumNodes, SyncConnections.Num());
+	}
+
 	Super::TickFlush(DeltaSeconds);
 
 	// Check if we the server or if any synchronous connections exists for processing
@@ -310,6 +336,11 @@ void UDisplayClusterNetDriver::TickFlush(float DeltaSeconds)
 		{
 			const UDisplayClusterConfigurationData* ConfigData = ConfigMgr->GetConfig();
 
+			if (ConfigData == nullptr)
+			{
+				return;
+			}
+
 			if (SyncConnections.Num() < static_cast<int32>(ConfigData->GetNumberOfClusterNodes() - 1))
 			{
 				return;
@@ -317,36 +348,55 @@ void UDisplayClusterNetDriver::TickFlush(float DeltaSeconds)
 		}
 	}
 
-	bLastBunchWasAcked = false;
-	
-	int32 NumPacketsInQueue = 0;
+	// Check if next slice can be formed
+	bool bCanFormSlice = true;
 
-	// Fill queues of last send packets for each connection, update last acked packet for each connection
 	for (UDisplayClusterNetConnection* SyncConnectionIt : SyncConnections)
 	{
-		int32 ConnectionUniqueId = SyncConnectionIt->GetUniqueID();
-		int32 LastOutPacketId = SyncConnectionIt->LastOut.PacketId;
+		const int32 ConnectionUniqueId = SyncConnectionIt->GetUniqueID();
+		const int32 LastOutPacketId = SyncConnectionIt->LastOut.PacketId;
 
-		// Add last send packet of current Connection to queue
 		TDeque<int32>& ChannelPacketQueue = OutPacketsQueues.FindOrAdd(ConnectionUniqueId);
-		ChannelPacketQueue.PushLast(LastOutPacketId);
 
-		// Get size of queue. Size will be the same for each Connection
-		NumPacketsInQueue = ChannelPacketQueue.Num();
-
-		// Fill LastAckedPacket map, which contains last acked packet for each Connection
-		int32& LastAckedPacketIdRef = LastAckedPacket.FindOrAdd(ConnectionUniqueId, 0);
-		int32 ConnectionAckId = SyncConnectionIt->GetLastNotifiedPacket();
-
-		if (ConnectionAckId > LastAckedPacketIdRef)
+		if (!ChannelPacketQueue.IsEmpty())
 		{
-			LastAckedPacketIdRef = ConnectionAckId;
+			if (ChannelPacketQueue.Last() == LastOutPacketId)
+			{
+				bCanFormSlice = false;
+				break;
+			}
 		}
+	}
+
+	bLastBunchWasAcked = true;
+
+	for (UDisplayClusterNetConnection* SyncConnectionIt : SyncConnections)
+	{
+		const int32 LastOutPacketId = SyncConnectionIt->LastOut.PacketId;
 
 		// Driver waits until last sent bunch of any sync connection will be acked
-		bLastBunchWasAcked |= (LastOutPacketId - ConnectionAckId > 1);
+		bLastBunchWasAcked &= (LastOutPacketId - SyncConnectionIt->OutAckPacketId <= 1);
+	}
+
+	int32 NumPacketsInQueue = 0;
+
+	if (bCanFormSlice)
+	{
+		// Fill queues of last send packets for each connection, update last acked packet for each connection
+		for (UDisplayClusterNetConnection* SyncConnectionIt : SyncConnections)
+		{
+			const int32 ConnectionUniqueId = SyncConnectionIt->GetUniqueID();
+			const int32 LastOutPacketId = SyncConnectionIt->LastOut.PacketId;
+
+			// Add last send packet of current Connection to queue
+			TDeque<int32>& ChannelPacketQueue = OutPacketsQueues.FindOrAdd(ConnectionUniqueId);
+			ChannelPacketQueue.PushLast(LastOutPacketId);
+
+			// Get size of queue. Size will be the same for each Connection
+			NumPacketsInQueue = ChannelPacketQueue.Num();
 
-		UE_LOG(LogDisplayClusterNetDriver, VeryVerbose, TEXT("Client %u; Packet %i sent, %i acked"), SyncConnectionIt->ClientId, ConnectionAckId, LastAckedPacketIdRef);
+			UE_LOG(LogDisplayClusterNetDriver, VeryVerbose, TEXT("Client %u; Packet %i sent, %i acked"), ConnectionUniqueId, LastOutPacketId, SyncConnectionIt->OutAckPacketId);
+		}
 	}
 
 	// Find slice of packets which are ready for processing on clients
@@ -369,17 +419,13 @@ void UDisplayClusterNetDriver::TickFlush(float DeltaSeconds)
 			int32 CurrentConnectionId = SyncConnectionIt->GetUniqueID();
 			int32 CurrentPacketId = OutPacketsQueues[CurrentConnectionId][PacketIdx];
 
-			if ((CurrentPacketId > 0) && (LastAckedPacket[CurrentConnectionId] >= CurrentPacketId))
+			if ((CurrentPacketId > 0) && (CurrentPacketId <= SyncConnectionIt->OutAckPacketId))
 			{
 				ReadyOutPackets.Add(CurrentConnectionId, CurrentPacketId);
 			}
-
-			if (ReadyOutPackets.Num() == SyncConnections.Num())
-			{
-				bPacketsReadyForSync = true;
-				break;
-			}
 		}
+
+		bPacketsReadyForSync = (ReadyOutPackets.Num() == SyncConnections.Num());
 	}
 
 	// Slice of acked packets is found - form PacketParams and send cluster event to primary nodes
@@ -399,7 +445,7 @@ void UDisplayClusterNetDriver::TickFlush(float DeltaSeconds)
 			TDeque<int32>& ConnectionPacketQueue = OutPacketsQueues[UniqueID];
 			if (!ConnectionPacketQueue.IsEmpty())
 			{
-				for (int j = 0; j < PacketIdx; j++)
+				for (int32 ConnectionPacketIndex = 0; ConnectionPacketIndex <= PacketIdx; ConnectionPacketIndex++)
 				{
 					ConnectionPacketQueue.PopFirst();
 				}
@@ -422,7 +468,7 @@ void UDisplayClusterNetDriver::PreListUpdate(ConsiderListUpdateParams const& Upd
 {
 	bConnectionViewersAreReady = false;
 
-	if (bLastBunchWasAcked || NodeConnections.IsEmpty())
+	if (!bLastBunchWasAcked || NodeConnections.IsEmpty())
 	{
 		return;
 	}
@@ -483,6 +529,25 @@ void UDisplayClusterNetDriver::PreListUpdate(ConsiderListUpdateParams const& Upd
 		ConsiderList, bCPUSaturated, ClusterReplicationState.PriorityList, ClusterReplicationState.PriorityActors);
 
 	bConnectionViewersAreReady = true;
+
+	for (UDisplayClusterNetConnection* ClusterConnection : SyncConnections)
+	{
+		if (ClusterConnection->ViewTarget == nullptr)
+		{
+			bConnectionViewersAreReady = false;
+			break;
+		}
+
+		for (TObjectPtr<UChannel> Channel : ClusterConnection->Channels)
+		{
+			if (!IsValid(Channel))
+			{
+				continue;
+			}
+
+			bConnectionViewersAreReady &= !(Channel->bPausedUntilReliableACK && (Channel->NumOutRec > 0));
+		}
+	}
 }
 
 void UDisplayClusterNetDriver::PostListUpdate(ConsiderListUpdateParams const& UpdateParams, int& OutUpdated, const TArray<FNetworkObjectInfo*>& ConsiderList)
@@ -517,10 +582,12 @@ void UDisplayClusterNetDriver::PostListUpdate(ConsiderListUpdateParams const& Up
 		if (LastProcessedActors < ClusterReplicationState.MaxLastProcessedActor && ClusterReplicationState.MaxLastProcessedActor > 0)
 		{
 			// We do use generated ConnectionViewers above, but than use stored ConnectionViewers from primary nodes?
-			const int32 ActorStartIndex = LastProcessedActors + 1;
-			const int32 ActorsToProcess = ClusterReplicationState.MaxLastProcessedActor - ActorStartIndex + 1;
+			const int32 ActorStartIndex = FMath::Clamp<int32>(LastProcessedActors, 0, ClusterReplicationState.FinalSortedCount - 1);
+			const int32 ActorsToProcess = FMath::Min(ClusterReplicationState.MaxLastProcessedActor - ActorStartIndex, ClusterReplicationState.FinalSortedCount - ActorStartIndex);
 
 			constexpr bool bForceUpdate = true;
+			OutUpdated = 0;
+
 			TInterval<int32> ActorsIndexRange(ActorStartIndex, ActorsToProcess);
 			ServerReplicateActors_ProcessPrioritizedActorsRange(DisplayClusterConnection, PivotNodeConnectionViewers, ClusterReplicationState.PriorityActors, ActorsIndexRange,
 				OutUpdated, bForceUpdate);
@@ -548,7 +615,7 @@ void UDisplayClusterNetDriver::ListUpdate(ConsiderListUpdateParams const& Update
 		TArray<FNetViewer>& ConnectionViewers = WorldSettings->ReplicationViewers;
 
 		FActorPriority* PriorityList = nullptr;
-		FActorPriority** PriorityActors = nullptr;
+		FActorPriority** PriorityActors = nullptr;	
 
 		// Get a sorted list of actors for this connection
 		const int32 FinalSortedCount = ServerReplicateActors_PrioritizeActors(Connection, ConnectionViewers, ConsiderList, bCPUSaturated, PriorityList, PriorityActors);
@@ -571,17 +638,47 @@ void UDisplayClusterNetDriver::ListUpdate(ConsiderListUpdateParams const& Update
 		// Use single priority list per cluster which was calculated for Primary node
 		// Process the sorted list of actors for this connection
 		TInterval<int32> ActorsIndexRange(0, ClusterReplicationState.FinalSortedCount);
-		const int32 LastProcessedActor = ServerReplicateActors_ProcessPrioritizedActorsRange(Connection, PivotNodeConnectionViewers, ClusterReplicationState.PriorityActors,
+		
+		OutUpdated = 0;
+		int32 LastProcessedActor = ServerReplicateActors_ProcessPrioritizedActorsRange(Connection, PivotNodeConnectionViewers, ClusterReplicationState.PriorityActors,
 			ActorsIndexRange, OutUpdated);
 
 		// Cache last processed actors for equalization on PostListUpdate step
 		// MaxLastProcessed will be used for equalization
-		ClusterReplicationState.LastProcessedActors.Add(Connection->GetUniqueID(), LastProcessedActor);
-		ClusterReplicationState.MaxLastProcessedActor = FMath::Max<int32>(ClusterReplicationState.MaxLastProcessedActor, LastProcessedActor);
+		ClusterReplicationState.LastProcessedActors.Add(Connection->GetUniqueID(), OutUpdated);
+		ClusterReplicationState.MaxLastProcessedActor = FMath::Max<int32>(ClusterReplicationState.MaxLastProcessedActor, OutUpdated);
 	}
 }
 #endif
 
+
+bool UDisplayClusterNetDriver::NotifyClusterAsReadyForSync(int32 ClusterId)
+{
+	// Notify primary node that current node connection is ready
+	FDisplayClusterClusterEventBinary NetworkDriverSyncEvent;
+	GenerateClusterCommandsEvent(NetworkDriverSyncEvent, NodeSyncEvent);
+
+	// send command 
+	const bool bEventWasSent = ClusterNetworkDriverHelper->SendCommandToCluster(ClusterId, NetworkDriverSyncEvent);
+
+	// If event was sent to primary cluster node
+	if (bEventWasSent)
+	{
+		// Mark all node connections related to current ClusterId as ready
+		for (UDisplayClusterNetConnection* ClusterNodeConnectionIt : NodeConnections)
+		{
+			ClusterNodeConnectionIt->bSynchronousMode = true;
+			SyncConnections.Add(ClusterNodeConnectionIt);
+		}
+		
+		UE_LOG(LogDisplayClusterNetDriver, VeryVerbose, TEXT("Setting synchronous mode in cluster [%u]"), ClusterId);
+
+		return true;
+	}
+
+	return false;
+}
+
 void UDisplayClusterNetDriver::GenerateClusterCommandsEvent(FDisplayClusterClusterEventBinary& NetworkDriverSyncEvent, 
 	int32 EventId)
 {
diff --git a/Engine/Plugins/Runtime/nDisplay/Source/DisplayClusterReplication/Public/DisplayClusterNetConnection.h b/Engine/Plugins/Runtime/nDisplay/Source/DisplayClusterReplication/Public/DisplayClusterNetConnection.h
--- a/Engine/Plugins/Runtime/nDisplay/Source/DisplayClusterReplication/Public/DisplayClusterNetConnection.h
+++ b/Engine/Plugins/Runtime/nDisplay/Source/DisplayClusterReplication/Public/DisplayClusterNetConnection.h
@@ -58,9 +58,6 @@ public:
 	 */
 	void ProcessPacket(int32 PacketId);
 
-	/** Return last notified packet id */
-	int32 GetLastNotifiedPacket() const { return GetLastNotifiedPacketId(); }
-
 protected:
 	// Storage for InPackets map keys, prevents reallocs 
 	TArray<int32> PacketIDs;
diff --git a/Engine/Plugins/Runtime/nDisplay/Source/DisplayClusterReplication/Public/DisplayClusterNetDriver.h b/Engine/Plugins/Runtime/nDisplay/Source/DisplayClusterReplication/Public/DisplayClusterNetDriver.h
--- a/Engine/Plugins/Runtime/nDisplay/Source/DisplayClusterReplication/Public/DisplayClusterNetDriver.h
+++ b/Engine/Plugins/Runtime/nDisplay/Source/DisplayClusterReplication/Public/DisplayClusterNetDriver.h
@@ -78,9 +78,6 @@ public:
 	void RemoveNodeConnection(UDisplayClusterNetConnection* NetConnection);
 
 protected:
-	// Contains last max acked packet for specific connection (key - connectionID, value - packetID)
-	TMap<int32, int32> LastAckedPacket;
-
 	// Contains queue of non-processed packets for specific connection (key - connectionID, value - packetID)
 	TMap<int32, TDeque<int32>> OutPacketsQueues;
 
@@ -126,6 +123,12 @@ protected:
 	// Checks if connections viewers for Primary node were formed
 	bool bConnectionViewersAreReady;
 
+	// ClusterId for listen Server
+	int32 ListenClusterId;
+
+	// Num cluster nodes for listen Server
+	int32 ListenClusterNodesNum;
+
 	// Cluster event id, used to start synchrosonus packets processing
 	inline static const int NodeSyncEvent = GetTypeHash(FStringView(TEXT("nDCRNodeSyncEvent")));
 
@@ -138,6 +141,10 @@ protected:
 	void PostListUpdate(ConsiderListUpdateParams const& UpdateParams, int& OutUpdated, const TArray<FNetworkObjectInfo*>& ConsiderList);
 	void ListUpdate(ConsiderListUpdateParams const& UpdateParams, int& OutUpdated, const TArray<FNetworkObjectInfo*>& ConsiderList);
 #endif
+	
+	// Helper functions used by server side to notify cluster via cluster event that it is are ready for sync
+	bool NotifyClusterAsReadyForSync(int32 ClusterId);
+
 	// Helper functions for parameters serialization into binary event data
 	void GenerateClusterCommandsEvent(FDisplayClusterClusterEventBinary& NetworkDriverSyncEvent, int32 EventId);
 	void GenerateClusterCommandsEvent(FDisplayClusterClusterEventBinary& NetworkDriverSyncEvent, int32 EventId, const TMap<uint32, int32>& Parameters);
diff --git a/Engine/Plugins/VirtualProduction/Switchboard/Source/Switchboard/switchboard/devices/ndisplay/plugin_ndisplay.py b/Engine/Plugins/VirtualProduction/Switchboard/Source/Switchboard/switchboard/devices/ndisplay/plugin_ndisplay.py
--- a/Engine/Plugins/VirtualProduction/Switchboard/Source/Switchboard/switchboard/devices/ndisplay/plugin_ndisplay.py
+++ b/Engine/Plugins/VirtualProduction/Switchboard/Source/Switchboard/switchboard/devices/ndisplay/plugin_ndisplay.py
@@ -763,7 +763,7 @@ class DevicenDisplay(DeviceUnreal):
         if multiplayer_mode_name == 'Dedicated server':
             dedicated_server_address = DevicenDisplay.csettings["dedicated_server_address"].get_value()
 
-            if dedicated_server_address == self.address or is_local_address(self.address):
+            if dedicated_server_address == self.address or is_local_address(dedicated_server_address):
                 dedicated_server_address = '127.0.0.1'
 
             return get_client_args(
diff --git a/Engine/Source/Runtime/Engine/Classes/Engine/NetDriver.h b/Engine/Source/Runtime/Engine/Classes/Engine/NetDriver.h
--- a/Engine/Source/Runtime/Engine/Classes/Engine/NetDriver.h
+++ b/Engine/Source/Runtime/Engine/Classes/Engine/NetDriver.h
@@ -2148,6 +2148,8 @@ protected:
 	/** Creates a trace event that updates the name and properties of the associated Game Instance */
 	void NotifyGameInstanceUpdated();
 
+	/** Indicates whether ticking throttle is enabled for this instance of NetDriver */
+	bool bTickingThrottleEnabled = true;
 private:
 	// Only for ForwardRemoteFunction
 	friend FObjectReplicator;
diff --git a/Engine/Source/Runtime/Engine/Private/NetDriver.cpp b/Engine/Source/Runtime/Engine/Private/NetDriver.cpp
--- a/Engine/Source/Runtime/Engine/Private/NetDriver.cpp
+++ b/Engine/Source/Runtime/Engine/Private/NetDriver.cpp
@@ -4615,7 +4615,7 @@ int32 UNetDriver::ServerReplicateActors_PrepConnections( const float DeltaSecond
 
 	// by default only throttle update for listen servers unless specified on the commandline
 	static bool bForceClientTickingThrottle = FParse::Param( FCommandLine::Get(), TEXT( "limitclientticks" ) );
-	if ( bForceClientTickingThrottle || GetNetMode() == NM_ListenServer )
+	if ( (bForceClientTickingThrottle || GetNetMode() == NM_ListenServer) && bTickingThrottleEnabled )
 	{
 		// determine how many clients to tick this frame based on GEngine->NetTickRate (always tick at least one client), double for lan play
 		// FIXME: DeltaTimeOverflow is a static, and will conflict with other running net drivers, we investigate storing it on the driver itself!
