Fixed issue where any rotation on a stage's root component would cause the stage isosphere to be incorrectly rotated when generated

diff --git a/Engine/Plugins/Runtime/nDisplay/Source/DisplayCluster/Private/Game/EngineClasses/Scene/DisplayClusterStageGeometryComponent.cpp b/Engine/Plugins/Runtime/nDisplay/Source/DisplayCluster/Private/Game/EngineClasses/Scene/DisplayClusterStageGeometryComponent.cpp
--- a/Engine/Plugins/Runtime/nDisplay/Source/DisplayCluster/Private/Game/EngineClasses/Scene/DisplayClusterStageGeometryComponent.cpp
+++ b/Engine/Plugins/Runtime/nDisplay/Source/DisplayCluster/Private/Game/EngineClasses/Scene/DisplayClusterStageGeometryComponent.cpp
@@ -115,7 +115,9 @@ bool UDisplayClusterStageGeometryComponent::MorphProceduralMesh(UProceduralMeshC
 		{
 			if (bSyncMeshLocation)
 			{
-				InProceduralMeshComponent->SetWorldLocation(CommonViewPoint->GetComponentLocation());
+				// Align the procedural mesh with the view point that the stage geometry component was rendered from to ensure an accurate morph.
+				// The procedural mesh's world orientation must also be set to zero, as the geometry component render was computed in world coordinates
+				InProceduralMeshComponent->SetWorldLocationAndRotation(CommonViewPoint->GetComponentLocation(), FRotator::ZeroRotator);
 			}
 
 			if (FProcMeshSection* Section = InProceduralMeshComponent->GetProcMeshSection(0))
