[nDisplay] Fixed media tiling

Tiles:
  - Fixed bug that forced target tiles (media input) to always render
  - Fixed another bug. Now we render only required tiles.
  - Changed/fixed the way how we obtain viewport size to support tiling in media

Media gui:
  - Extended gui to properly support tiles
  - Separate media config types for viewports and backbuffer
  - Removed backbuffer media input
  - Restructured/simplified media initialization pipeline

diff --git a/Engine/Plugins/Runtime/nDisplay/Config/BasenDisplay.ini b/Engine/Plugins/Runtime/nDisplay/Config/BasenDisplay.ini
--- a/Engine/Plugins/Runtime/nDisplay/Config/BasenDisplay.ini
+++ b/Engine/Plugins/Runtime/nDisplay/Config/BasenDisplay.ini
@@ -8,3 +8,6 @@
 +ClassRedirects=(OldName="/Script/DisplayClusterMedia.SharedMemoryMediaCapture", NewName="/Script/SharedMemoryMedia.SharedMemoryMediaCapture")
 +ClassRedirects=(OldName="/Script/DisplayClusterMedia.SharedMemoryMediaSource", NewName="/Script/SharedMemoryMedia.SharedMemoryMediaSource")
 +ClassRedirects=(OldName="/Script/DisplayClusterMedia.SharedMemoryMediaPlayer", NewName="/Script/SharedMemoryMedia.SharedMemoryMediaPlayer")
+
+; 5.4
++StructRedirects=(OldName="/Script/DisplayClusterConfiguration.DisplayClusterConfigurationMedia", NewName="/Script/DisplayClusterConfiguration.DisplayClusterConfigurationMediaViewport")
diff --git a/Engine/Plugins/Runtime/nDisplay/Source/DisplayCluster/Private/Render/Viewport/Configuration/DisplayClusterViewportConfigurationHelpers_Tile.cpp b/Engine/Plugins/Runtime/nDisplay/Source/DisplayCluster/Private/Render/Viewport/Configuration/DisplayClusterViewportConfigurationHelpers_Tile.cpp
--- a/Engine/Plugins/Runtime/nDisplay/Source/DisplayCluster/Private/Render/Viewport/Configuration/DisplayClusterViewportConfigurationHelpers_Tile.cpp
+++ b/Engine/Plugins/Runtime/nDisplay/Source/DisplayCluster/Private/Render/Viewport/Configuration/DisplayClusterViewportConfigurationHelpers_Tile.cpp
@@ -3,21 +3,23 @@
 #include "Render/Viewport/Configuration/DisplayClusterViewportConfigurationHelpers_Tile.h"
 #include "Render/Viewport/Configuration/DisplayClusterViewportConfigurationHelpers.h"
 
-#include "DisplayClusterConfigurationTypes_Tile.h"
+#include "Misc/DisplayClusterGlobals.h"
+#include "Misc/DisplayClusterLog.h"
 
 #include "Render/Viewport/DisplayClusterViewport.h"
 #include "Render/Viewport/DisplayClusterViewportManager.h"
 #include "Render/Viewport/DisplayClusterViewportStrings.h"
 #include "Render/Viewport/DisplayClusterViewportHelpers.h"
-
-#include "DisplayClusterProjectionStrings.h"
-
-#include "Render/Viewport/DisplayClusterViewportStrings.h"
 #include "Render/Viewport/Configuration/DisplayClusterViewportConfiguration.h"
 
+#include "DisplayClusterEnums.h"
 #include "DisplayClusterConfigurationTypes_ICVFX.h"
+#include "DisplayClusterConfigurationTypes_Tile.h"
+#include "DisplayClusterProjectionStrings.h"
+
+#include "IDisplayCluster.h"
+#include "IDisplayClusterCallbacks.h"
 
-#include "Misc/DisplayClusterLog.h"
 
 ////////////////////////////////////////////////////////////////////////
 // FDisplayClusterViewportConfigurationHelpers_Tile
@@ -187,6 +189,27 @@ FDisplayClusterViewport* FDisplayClusterViewportConfigurationHelpers_Tile::GetOr
 		// Setup as tile.
 		InOutRenderSettings.TileSettings = FDisplayClusterViewport_TileSettings(InSourceViewport.GetId(), InTilePos, InTileSize);
 
+		// Allow external customers to configure media state
+		{
+			// By default, we set 'None' so the tiles can be rendered in editor for camera preview.
+			EDisplayClusterViewportMediaState NewMediaStates = EDisplayClusterViewportMediaState::None;
+
+			// But in cluster mode, we set 'Inactive' by default. When tiling is used, the tile viewport must either be rendered for media output
+			// or preserve internal buffer for a media input texture. However, it's possible the tile viewport has wrong or completely
+			// missing media output settings, or corresponding media device was not able to start properly for some reason.
+			// In this case the tile should not be rendered at all.
+			if (GDisplayCluster->GetOperationMode() == EDisplayClusterOperationMode::Cluster)
+			{
+				NewMediaStates = EDisplayClusterViewportMediaState::Inactive;
+			}
+			
+			// Now allow to override media state if anyone wants
+			IDisplayCluster::Get().GetCallbacks().OnDisplayClusterUpdateViewportMediaState().Broadcast(TileViewport, NewMediaStates);
+
+			// Update the media state for the new frame.
+			InOutRenderSettings.AssignMediaStates(NewMediaStates);
+		}
+
 		// Copy internal render settings from the source:
 		TileViewport->GetVisibilitySettingsImpl() = InSourceViewport.GetVisibilitySettingsImpl();
 		TileViewport->GetCameraMotionBlurImpl() = InSourceViewport.GetCameraMotionBlurImpl();
diff --git a/Engine/Plugins/Runtime/nDisplay/Source/DisplayCluster/Private/Render/Viewport/DisplayClusterViewport_Rules.cpp b/Engine/Plugins/Runtime/nDisplay/Source/DisplayCluster/Private/Render/Viewport/DisplayClusterViewport_Rules.cpp
--- a/Engine/Plugins/Runtime/nDisplay/Source/DisplayCluster/Private/Render/Viewport/DisplayClusterViewport_Rules.cpp
+++ b/Engine/Plugins/Runtime/nDisplay/Source/DisplayCluster/Private/Render/Viewport/DisplayClusterViewport_Rules.cpp
@@ -26,9 +26,6 @@
 
 #include "Render/DisplayDevice/Components/DisplayClusterDisplayDeviceBaseComponent.h"
 
-#include "IDisplayCluster.h"
-#include "IDisplayClusterCallbacks.h"
-
 #include "EngineUtils.h"
 #include "SceneManagement.h"
 #include "SceneView.h"
@@ -55,13 +52,6 @@ void FDisplayClusterViewport::ResetRuntimeParameters()
 
 	OverscanRuntimeSettings = FDisplayClusterViewport_OverscanRuntimeSettings();
 	CustomFrustumRuntimeSettings = FDisplayClusterViewport_CustomFrustumRuntimeSettings();
-	
-	// Obtain viewport media state from external multicast delegates.
-	EDisplayClusterViewportMediaState NewMediaStates = EDisplayClusterViewportMediaState::None;
-	IDisplayCluster::Get().GetCallbacks().OnDisplayClusterUpdateViewportMediaState().Broadcast(this, NewMediaStates);
-
-	// Update the media state for the new frame.
-	RenderSettings.AssignMediaStates(NewMediaStates);
 }
 
 bool FDisplayClusterViewport::IsInternalViewport() const
@@ -137,6 +127,12 @@ bool FDisplayClusterViewport::IsRenderEnabled() const
 		return false;
 	}
 
+	if (RenderSettings.HasAnyMediaStates(EDisplayClusterViewportMediaState::Inactive))
+	{
+		// Don't render inactive tiles
+		return false;
+	}
+
 	return true;
 }
 
@@ -355,4 +351,3 @@ uint8 FDisplayClusterViewport::GetPriority() const
 
 	return OutOrder;
 }
-
diff --git a/Engine/Plugins/Runtime/nDisplay/Source/DisplayCluster/Public/Render/Viewport/Containers/DisplayClusterViewport_Enums.h b/Engine/Plugins/Runtime/nDisplay/Source/DisplayCluster/Public/Render/Viewport/Containers/DisplayClusterViewport_Enums.h
--- a/Engine/Plugins/Runtime/nDisplay/Source/DisplayCluster/Public/Render/Viewport/Containers/DisplayClusterViewport_Enums.h
+++ b/Engine/Plugins/Runtime/nDisplay/Source/DisplayCluster/Public/Render/Viewport/Containers/DisplayClusterViewport_Enums.h
@@ -197,6 +197,9 @@ enum class EDisplayClusterViewportMediaState : uint8
 	Input = 1 << 4,
 	Input_ForceLateOCIOPass = 1 << 5,
 
+	// This viewport is inactive
+	Inactive = 1 << 6,
+
 	ForceLateOCIOPass = Capture_ForceLateOCIOPass | Input_ForceLateOCIOPass,
 };
 ENUM_CLASS_FLAGS(EDisplayClusterViewportMediaState);
diff --git a/Engine/Plugins/Runtime/nDisplay/Source/DisplayClusterConfiguration/Private/DisplayClusterConfigurationTypes.cpp b/Engine/Plugins/Runtime/nDisplay/Source/DisplayClusterConfiguration/Private/DisplayClusterConfigurationTypes.cpp
--- a/Engine/Plugins/Runtime/nDisplay/Source/DisplayClusterConfiguration/Private/DisplayClusterConfigurationTypes.cpp
+++ b/Engine/Plugins/Runtime/nDisplay/Source/DisplayClusterConfiguration/Private/DisplayClusterConfigurationTypes.cpp
@@ -325,6 +325,14 @@ void UDisplayClusterConfigurationClusterNode::PostLoad()
 			Media.MediaOutput = nullptr;
 			Media.OutputSyncPolicy = nullptr;
 		}
+
+		if (MediaSettings.MediaOutputs.IsEmpty() && Media.MediaOutputs.Num() > 0)
+		{
+			MediaSettings.bEnable = Media.bEnable;
+			MediaSettings.MediaOutputs = Media.MediaOutputs;
+			Media.bEnable = false;
+			Media.MediaOutputs.Empty();
+		}
 	PRAGMA_ENABLE_DEPRECATION_WARNINGS
 #endif // WITH_EDITOR
 }
diff --git a/Engine/Plugins/Runtime/nDisplay/Source/DisplayClusterConfiguration/Private/DisplayClusterConfigurationTypes_ICVFX.cpp b/Engine/Plugins/Runtime/nDisplay/Source/DisplayClusterConfiguration/Private/DisplayClusterConfigurationTypes_ICVFX.cpp
--- a/Engine/Plugins/Runtime/nDisplay/Source/DisplayClusterConfiguration/Private/DisplayClusterConfigurationTypes_ICVFX.cpp
+++ b/Engine/Plugins/Runtime/nDisplay/Source/DisplayClusterConfiguration/Private/DisplayClusterConfigurationTypes_ICVFX.cpp
@@ -213,7 +213,7 @@ bool FDisplayClusterConfigurationICVFX_CameraSettings::IsICVFXEnabled(const UDis
 
 	// Get backbuffer media settings
 	const UDisplayClusterConfigurationClusterNode* const NodeCfg = InConfigurationData.Cluster->GetNode(InClusterNodeId);
-	const FDisplayClusterConfigurationMedia* BackbufferMediaSettings = NodeCfg ? &NodeCfg->Media : nullptr;
+	const FDisplayClusterConfigurationMediaNodeBackbuffer* BackbufferMediaSettings = NodeCfg ? &NodeCfg->MediaSettings : nullptr;
 
 	// Second condition to render offscreen: the backbuffer has media output assigned.
 	// This means the whole frame including ICVFX cameras need to be rendered.
diff --git a/Engine/Plugins/Runtime/nDisplay/Source/DisplayClusterConfiguration/Private/DisplayClusterConfigurationTypes_Media.cpp b/Engine/Plugins/Runtime/nDisplay/Source/DisplayClusterConfiguration/Private/DisplayClusterConfigurationTypes_Media.cpp
--- a/Engine/Plugins/Runtime/nDisplay/Source/DisplayClusterConfiguration/Private/DisplayClusterConfigurationTypes_Media.cpp
+++ b/Engine/Plugins/Runtime/nDisplay/Source/DisplayClusterConfiguration/Private/DisplayClusterConfigurationTypes_Media.cpp
@@ -2,26 +2,39 @@
 
 #include "DisplayClusterConfigurationTypes_Media.h"
 
+
+///////////////////////////////////////////////////
+// FDisplayClusterConfigurationMediaNodeBackbuffer
+
+bool FDisplayClusterConfigurationMediaNodeBackbuffer::IsMediaOutputAssigned() const
+{
+	// Just see if any media output instances are assigned
+	const bool bAnyMediaAssigned = MediaOutputs.ContainsByPredicate([](const FDisplayClusterConfigurationMediaOutput& Item)
+		{
+			return IsValid(Item.MediaOutput);
+		});
+
+	return bAnyMediaAssigned;
+}
+
+
 ///////////////////////////////////////////////////
-// FDisplayClusterConfigurationMedia
+// FDisplayClusterConfigurationMediaViewport
 
-bool FDisplayClusterConfigurationMedia::IsMediaInputAssigned() const
+bool FDisplayClusterConfigurationMediaViewport::IsMediaInputAssigned() const
 {
 	return IsValid(MediaInput.MediaSource);
 }
 
-bool FDisplayClusterConfigurationMedia::IsMediaOutputAssigned() const
+bool FDisplayClusterConfigurationMediaViewport::IsMediaOutputAssigned() const
 {
-	// Return true if we have at least one media output set
-	for (const FDisplayClusterConfigurationMediaOutput& MediaOutputItem : MediaOutputs)
-	{
-		if (IsValid(MediaOutputItem.MediaOutput))
+	// Just see if any media output instances are assigned
+	const bool bAnyMediaAssigned = MediaOutputs.ContainsByPredicate([](const FDisplayClusterConfigurationMediaOutput& Item)
 		{
-			return true;
-		}
-	}
+			return IsValid(Item.MediaOutput);
+		});
 
-	return false;
+	return bAnyMediaAssigned;
 }
 
 
@@ -80,34 +93,48 @@ TArray<FDisplayClusterConfigurationMediaOutputGroup> FDisplayClusterConfiguratio
 		});
 }
 
-UMediaSource* FDisplayClusterConfigurationMediaICVFX::GetMediaSourceForTiles(const FString& NodeId, TArray<FIntPoint>& OutTiles) const
+/** Returns all tiles bound to a specific cluster node. */
+bool FDisplayClusterConfigurationMediaICVFX::GetMediaInputTiles(const FString& NodeId, TArray<FDisplayClusterConfigurationMediaUniformTileInput>& OutInputTiles) const
 {
-	OutTiles.Empty();
-
-	// Look up for a group that contains node ID specified
-	for (const FDisplayClusterConfigurationMediaInputGroup& MediaInputGroup : MediaInputGroups)
+	// Here we iterate through tiled media input groups, and look for a group that contains a cluster node ID specified.
+	for (const FDisplayClusterConfigurationMediaTiledInputGroup& TiledInputGroup : TiledMediaInputGroups)
 	{
-		if (IsValid(MediaInputGroup.MediaSource))
-		{
-			if (const FDisplayClusterConfigurationClusterNodeTilesReferenceList* TilesReferenceList = MediaInputGroup.ClusterNodesWithTiles.Find(NodeId))
+		// Look for NodeId in the group
+		const bool bFoundGroupWithNodeId = TiledInputGroup.ClusterNodes.ItemNames.ContainsByPredicate([NodeId](const FString& Item)
 			{
-				for(const FDisplayClusterConfigurationTileIndex& TileIndexIt : TilesReferenceList->Tiles)
-				{
-					OutTiles.AddUnique(FIntPoint(TileIndexIt.TileX, TileIndexIt.TileY));
-				}
-			}
+				return Item.Equals(NodeId, ESearchCase::IgnoreCase);
+			});
 
-			return MediaInputGroup.MediaSource;
+		// If found
+		if (bFoundGroupWithNodeId)
+		{
+			OutInputTiles = TiledInputGroup.Tiles;
+			return true;
 		}
 	}
 
-	return nullptr;
+	return false;
 }
 
-TArray<FDisplayClusterConfigurationMediaOutputGroup> FDisplayClusterConfigurationMediaICVFX::GetMediaOutputGroupsForTiles(const FString& NodeId) const
+/** Returns all tiles bound to a specific cluster node. */
+bool FDisplayClusterConfigurationMediaICVFX::GetMediaOutputTiles(const FString& NodeId, TArray<FDisplayClusterConfigurationMediaUniformTileOutput>& OutOutputTiles) const
 {
-	return MediaOutputGroups.FilterByPredicate([NodeId](const FDisplayClusterConfigurationMediaOutputGroup& Item)
+	// Here we iterate through tiled media output groups, and look for a group that contains a cluster node ID specified.
+	for (const FDisplayClusterConfigurationMediaTiledOutputGroup& TiledOutputGroup : TiledMediaOutputGroups)
+	{
+		// Look for NodeId in the group
+		const bool bFoundGroupWithNodeId = TiledOutputGroup.ClusterNodes.ItemNames.ContainsByPredicate([NodeId](const FString& Item)
+			{
+				return Item.Equals(NodeId, ESearchCase::IgnoreCase);
+			});
+
+		// If found
+		if (bFoundGroupWithNodeId)
 		{
-			return Item.ClusterNodesWithTiles.Contains(NodeId);
-		});
-}
\ No newline at end of file
+			OutOutputTiles = TiledOutputGroup.Tiles;
+			return true;
+		}
+	}
+
+	return false;
+}
diff --git a/Engine/Plugins/Runtime/nDisplay/Source/DisplayClusterConfiguration/Public/DisplayClusterConfigurationTypes.h b/Engine/Plugins/Runtime/nDisplay/Source/DisplayClusterConfiguration/Public/DisplayClusterConfigurationTypes.h
--- a/Engine/Plugins/Runtime/nDisplay/Source/DisplayClusterConfiguration/Public/DisplayClusterConfigurationTypes.h
+++ b/Engine/Plugins/Runtime/nDisplay/Source/DisplayClusterConfiguration/Public/DisplayClusterConfigurationTypes.h
@@ -358,8 +358,8 @@ public:
 	TMap<FString, FDisplayClusterConfigurationPostprocess> Postprocess;
 
 	// Media settings
-	UPROPERTY(EditAnywhere, BlueprintReadWrite, Category = "Media")
-	FDisplayClusterConfigurationMedia Media;
+	UPROPERTY(EditAnywhere, BlueprintReadWrite, Category = "Media", meta = (DisplayName = "Media"))
+	FDisplayClusterConfigurationMediaNodeBackbuffer MediaSettings;
 
 #if WITH_EDITORONLY_DATA
 	UPROPERTY(EditDefaultsOnly, Category = "Configuration", meta = (HideProperty))
@@ -373,6 +373,14 @@ public:
 	FDisplayClusterConfigurationExternalImage PreviewImage;
 #endif
 
+public:
+
+#if WITH_EDITORONLY_DATA
+	UE_DEPRECATED(5.4, "This property has been deprecated. Please use 'MediaSettings'.")
+	UPROPERTY(meta = (DeprecatedProperty, DeprecationMessage = "This property has been deprecated. Please use 'MediaSettings'."))
+	FDisplayClusterConfigurationMediaViewport Media;
+#endif
+
 protected:
 	virtual void GetObjectsToExport(TArray<UObject*>& OutObjects) override;
 };
diff --git a/Engine/Plugins/Runtime/nDisplay/Source/DisplayClusterConfiguration/Public/DisplayClusterConfigurationTypes_Base.h b/Engine/Plugins/Runtime/nDisplay/Source/DisplayClusterConfiguration/Public/DisplayClusterConfigurationTypes_Base.h
--- a/Engine/Plugins/Runtime/nDisplay/Source/DisplayClusterConfiguration/Public/DisplayClusterConfigurationTypes_Base.h
+++ b/Engine/Plugins/Runtime/nDisplay/Source/DisplayClusterConfiguration/Public/DisplayClusterConfigurationTypes_Base.h
@@ -191,69 +191,8 @@ USTRUCT(BlueprintType)
 struct DISPLAYCLUSTERCONFIGURATION_API FDisplayClusterConfigurationClusterItemReferenceList
 {
 	GENERATED_BODY()
-public:
-	bool Contains(const FString& InName) const
-	{
-		return ItemNames.ContainsByPredicate([InName](const FString& ItemIt)
-			{
-				return ItemIt.Equals(InName, ESearchCase::IgnoreCase);
-			});
-	}
 
 public:
 	UPROPERTY(EditAnywhere, BlueprintReadOnly, Category = NDisplay)
 	TArray<FString> ItemNames;
 };
-
-USTRUCT(BlueprintType)
-struct DISPLAYCLUSTERCONFIGURATION_API FDisplayClusterConfigurationTileIndex
-{
-	GENERATED_BODY()
-
-public:
-	UPROPERTY(EditAnywhere, BlueprintReadOnly, Category = NDisplay)
-	int TileX;
-
-	UPROPERTY(EditAnywhere, BlueprintReadOnly, Category = NDisplay)
-	int TileY;
-};
-
-USTRUCT(BlueprintType)
-struct DISPLAYCLUSTERCONFIGURATION_API FDisplayClusterConfigurationClusterNodeTilesReferenceList
-{
-	GENERATED_BODY()
-
-public:
-	UPROPERTY(EditAnywhere, BlueprintReadOnly, Category = NDisplay)
-	FString NodeId;
-
-	UPROPERTY(EditAnywhere, BlueprintReadOnly, Category = NDisplay)
-	TArray<FDisplayClusterConfigurationTileIndex> Tiles;
-};
-
-USTRUCT(BlueprintType)
-struct DISPLAYCLUSTERCONFIGURATION_API FDisplayClusterConfigurationClusterTileItemReferenceList
-{
-	GENERATED_BODY()
-
-public:
-	bool Contains(const FString& InName) const
-	{
-		return Items.ContainsByPredicate([InName](const FDisplayClusterConfigurationClusterNodeTilesReferenceList& ItemIt)
-			{
-				return ItemIt.NodeId.Equals(InName, ESearchCase::IgnoreCase);
-			});
-	}
-
-	const FDisplayClusterConfigurationClusterNodeTilesReferenceList* Find(const FString& InNodeId) const
-	{
-		return Items.FindByPredicate([InNodeId](const FDisplayClusterConfigurationClusterNodeTilesReferenceList& Item)
-			{
-				return Item.NodeId.Equals(InNodeId, ESearchCase::IgnoreCase);
-			});
-	}
-
-public:
-	UPROPERTY(EditAnywhere, BlueprintReadOnly, Category = NDisplay)
-	TArray<FDisplayClusterConfigurationClusterNodeTilesReferenceList> Items;
-};
diff --git a/Engine/Plugins/Runtime/nDisplay/Source/DisplayClusterConfiguration/Public/DisplayClusterConfigurationTypes_Media.h b/Engine/Plugins/Runtime/nDisplay/Source/DisplayClusterConfiguration/Public/DisplayClusterConfigurationTypes_Media.h
--- a/Engine/Plugins/Runtime/nDisplay/Source/DisplayClusterConfiguration/Public/DisplayClusterConfigurationTypes_Media.h
+++ b/Engine/Plugins/Runtime/nDisplay/Source/DisplayClusterConfiguration/Public/DisplayClusterConfigurationTypes_Media.h
@@ -15,6 +15,17 @@
 #include "DisplayClusterConfigurationTypes_Media.generated.h"
 
 
+/*
+ * Media frame split types
+ */
+UENUM(BlueprintType)
+enum class EDisplayClusterConfigurationMediaSplitType : uint8
+{
+	FullFrame     UMETA(DisplayName = "Full Frame"),
+	UniformTiles  UMETA(DisplayName = "Uniform Tiles"),
+};
+
+
 /*
  * Media input item
  */
@@ -48,14 +59,37 @@ public:
 	TObjectPtr<UDisplayClusterMediaOutputSynchronizationPolicy> OutputSyncPolicy;
 };
 
+/*
+ * Media settings for node backbuffer
+ */
+USTRUCT(Blueprintable)
+struct DISPLAYCLUSTERCONFIGURATION_API FDisplayClusterConfigurationMediaNodeBackbuffer
+{
+	GENERATED_BODY()
+
+public:
+	/** Enable/disable media */
+	UPROPERTY(EditAnywhere, BlueprintReadWrite, Category = "Media")
+	bool bEnable = false;
+
+	/** Media outputs to use */
+	UPROPERTY(EditAnywhere, BlueprintReadWrite, Category = "Media")
+	TArray<FDisplayClusterConfigurationMediaOutput> MediaOutputs;
+
+public:
+
+	/** Returns true if at least one media output assigned */
+	bool IsMediaOutputAssigned() const;
+};
+
 
 PRAGMA_DISABLE_DEPRECATION_WARNINGS
 
 /*
- * Media settings for viewports and backbuffer
+ * Media settings for viewports
  */
 USTRUCT(Blueprintable)
-struct DISPLAYCLUSTERCONFIGURATION_API FDisplayClusterConfigurationMedia
+struct DISPLAYCLUSTERCONFIGURATION_API FDisplayClusterConfigurationMediaViewport
 {
 	GENERATED_BODY()
 
@@ -109,7 +143,7 @@ PRAGMA_ENABLE_DEPRECATION_WARNINGS
 
 
 /*
- * Media input group (ICVFX)
+ * Media input group (ICVFX, Full frame)
  */
 USTRUCT(Blueprintable)
 struct DISPLAYCLUSTERCONFIGURATION_API FDisplayClusterConfigurationMediaInputGroup
@@ -121,15 +155,11 @@ public:
 	/** Cluster nodes that use media source below */
 	UPROPERTY(EditAnywhere, BlueprintReadWrite, Category = Media, meta = (ClusterItemType = ClusterNodes))
 	FDisplayClusterConfigurationClusterItemReferenceList ClusterNodes;
-
-	/** Cluster nodes with tiles that use the media source below */
-	UPROPERTY(EditAnywhere, BlueprintReadWrite, Category = Media, meta = (ClusterItemType = ClusterNodes))
-	FDisplayClusterConfigurationClusterTileItemReferenceList ClusterNodesWithTiles;
 };
 
 
 /*
- * Media output group (ICVFX)
+ * Media output group (ICVFX, Full Frame)
  */
 USTRUCT(Blueprintable)
 struct DISPLAYCLUSTERCONFIGURATION_API FDisplayClusterConfigurationMediaOutputGroup
@@ -141,10 +171,94 @@ public:
 	/** Cluster nodes that export media via MediaOutput below */
 	UPROPERTY(EditAnywhere, BlueprintReadWrite, Category = Media, meta = (ClusterItemType = ClusterNodes))
 	FDisplayClusterConfigurationClusterItemReferenceList ClusterNodes;
+};
 
-	/** Cluster nodes with tiles that export media via MediaOutput below */
-	UPROPERTY(EditAnywhere, BlueprintReadWrite, Category = Media, meta = (ClusterItemType = ClusterNodes))
-	FDisplayClusterConfigurationClusterTileItemReferenceList ClusterNodesWithTiles;
+
+/**
+ * Uniform tile media input item. Maps a tile to a media source.
+ */
+USTRUCT(BlueprintType)
+struct DISPLAYCLUSTERCONFIGURATION_API FDisplayClusterConfigurationMediaUniformTileInput
+{
+	GENERATED_BODY()
+
+public:
+	/** Tile X location */
+	UPROPERTY(EditAnywhere, BlueprintReadOnly, Category = "Media", meta = (UIMin = 0, ClampMin = 0, UIMax = 4, ClampMax = 4))
+	int32 TileX = 0;
+
+	/** Tile Y location */
+	UPROPERTY(EditAnywhere, BlueprintReadOnly, Category = "Media", meta = (UIMin = 0, ClampMin = 0, UIMax = 4, ClampMax = 4))
+	int32 TileY = 0;
+
+	/** Media source to use */
+	UPROPERTY(Instanced, EditAnywhere, BlueprintReadWrite, Category = "Media")
+	TObjectPtr<UMediaSource> MediaSource;
+};
+
+
+/**
+ * Uniform tile media output item. Maps a tile to a media output.
+ */
+USTRUCT(BlueprintType)
+struct DISPLAYCLUSTERCONFIGURATION_API FDisplayClusterConfigurationMediaUniformTileOutput
+{
+	GENERATED_BODY()
+
+public:
+	/** Tile X location */
+	UPROPERTY(EditAnywhere, BlueprintReadOnly, Category = "Media", meta = (UIMin = 0, ClampMin = 0, UIMax = 4, ClampMax = 4))
+	int32 TileX = 0;
+
+	/** Tile Y location */
+	UPROPERTY(EditAnywhere, BlueprintReadOnly, Category = "Media", meta = (UIMin = 0, ClampMin = 0, UIMax = 4, ClampMax = 4))
+	int32 TileY = 0;
+
+	/** Media output to use */
+	UPROPERTY(Instanced, EditAnywhere, BlueprintReadWrite, Category = "Media")
+	TObjectPtr<UMediaOutput> MediaOutput;
+};
+
+
+/*
+ * Media input group (ICVFX, Tiled)
+ */
+USTRUCT(Blueprintable)
+struct DISPLAYCLUSTERCONFIGURATION_API FDisplayClusterConfigurationMediaTiledInputGroup
+{
+	GENERATED_BODY()
+
+public:
+	/** Cluster nodes that use media source below */
+	UPROPERTY(EditAnywhere, BlueprintReadWrite, Category = "Media", meta = (ClusterItemType = ClusterNodes))
+	FDisplayClusterConfigurationClusterItemReferenceList ClusterNodes;
+
+	/** Tile mapping. Maps tiles to media sources. */
+	UPROPERTY(EditAnywhere, BlueprintReadWrite, Category = "Media", meta = (ClusterItemType = ClusterNodes))
+	TArray<FDisplayClusterConfigurationMediaUniformTileInput> Tiles;
+};
+
+
+/*
+ * Media output group (ICVFX, Tiled)
+ */
+USTRUCT(Blueprintable)
+struct DISPLAYCLUSTERCONFIGURATION_API FDisplayClusterConfigurationMediaTiledOutputGroup
+{
+	GENERATED_BODY()
+
+public:
+	/** Cluster nodes that export media via MediaOutput below */
+	UPROPERTY(EditAnywhere, BlueprintReadWrite, Category = "Media", meta = (ClusterItemType = ClusterNodes))
+	FDisplayClusterConfigurationClusterItemReferenceList ClusterNodes;
+
+	/** Tile mapping. Maps tiles to media outputs. */
+	UPROPERTY(EditAnywhere, BlueprintReadWrite, Category = "Media", meta = (ClusterItemType = ClusterNodes))
+	TArray<FDisplayClusterConfigurationMediaUniformTileOutput> Tiles;
+
+	/** Media output synchronization policy */
+	UPROPERTY(Instanced, EditAnywhere, BlueprintReadWrite, Category = "Media", meta = (DisplayName = "Capture Synchronization"))
+	TObjectPtr<UDisplayClusterMediaOutputSynchronizationPolicy> OutputSyncPolicy;
 };
 
 
@@ -161,14 +275,26 @@ public:
 	UPROPERTY(EditAnywhere, BlueprintReadWrite, Category = "Media")
 	bool bEnable = false;
 
-	/** Media input mapping */
+	/** Media frame split type */
 	UPROPERTY(EditAnywhere, BlueprintReadWrite, Category = "Media")
+	EDisplayClusterConfigurationMediaSplitType SplitType = EDisplayClusterConfigurationMediaSplitType::FullFrame;
+
+	/** Media input mapping (Full frame) */
+	UPROPERTY(EditAnywhere, BlueprintReadWrite, Category = "Media", meta = (DisplayName = "Media Input Groups (Full Frame)", EditCondition = "SplitType == EDisplayClusterConfigurationMediaSplitType::FullFrame"))
 	TArray<FDisplayClusterConfigurationMediaInputGroup> MediaInputGroups;
 
-	/** Media output mapping */
-	UPROPERTY(EditAnywhere, BlueprintReadWrite, Category = "Media")
+	/** Media output mapping (Full frame) */
+	UPROPERTY(EditAnywhere, BlueprintReadWrite, Category = "Media", meta = (DisplayName = "Media Output Groups (Full Frame)", EditCondition = "SplitType == EDisplayClusterConfigurationMediaSplitType::FullFrame"))
 	TArray<FDisplayClusterConfigurationMediaOutputGroup> MediaOutputGroups;
 
+	/** Media input mapping (Tiled) */
+	UPROPERTY(EditAnywhere, BlueprintReadWrite, Category = "Media", meta = (DisplayName = "Media Input Groups (Tiled)", EditCondition = "SplitType == EDisplayClusterConfigurationMediaSplitType::UniformTiles"))
+	TArray<FDisplayClusterConfigurationMediaTiledInputGroup> TiledMediaInputGroups;
+
+	/** Media output mapping (Tiled) */
+	UPROPERTY(EditAnywhere, BlueprintReadWrite, Category = "Media", meta = (DisplayName = "Media Output Groups (Tiled)", EditCondition = "SplitType == EDisplayClusterConfigurationMediaSplitType::UniformTiles"))
+	TArray<FDisplayClusterConfigurationMediaTiledOutputGroup> TiledMediaOutputGroups;
+
 	/** Force late OCIO pass */
 	UPROPERTY(EditAnywhere, BlueprintReadWrite, Category = "Media", meta = (DisplayName = "Late OCIO Pass"))
 	bool bLateOCIOPass = false;
@@ -186,12 +312,11 @@ public:
 	/** Returns media outputs bound to a specific cluster node */
 	TArray<FDisplayClusterConfigurationMediaOutputGroup> GetMediaOutputGroups(const FString& NodeId) const;
 
+	/** Returns all the tiles bound to a specific cluster node. */
+	bool GetMediaInputTiles(const FString& NodeId, TArray<FDisplayClusterConfigurationMediaUniformTileInput>& OutInputTiles) const;
 
-	/** Returns media source bound to a specific cluster node with tiles. */
-	UMediaSource* GetMediaSourceForTiles(const FString& NodeId, TArray<FIntPoint>& OutTiles) const;
-
-	/** Returns media outputs bound to a specific cluster node */
-	TArray<FDisplayClusterConfigurationMediaOutputGroup> GetMediaOutputGroupsForTiles(const FString& NodeId) const;
+	/** Returns all the tiles bound to a specific cluster node. */
+	bool GetMediaOutputTiles(const FString& NodeId, TArray<FDisplayClusterConfigurationMediaUniformTileOutput>& OutOutputTiles) const;
 
 public:
 	UE_DEPRECATED(5.3, "This function has beend deprecated. Please use GetMediaOutputGroups.")
diff --git a/Engine/Plugins/Runtime/nDisplay/Source/DisplayClusterConfiguration/Public/DisplayClusterConfigurationTypes_Tile.h b/Engine/Plugins/Runtime/nDisplay/Source/DisplayClusterConfiguration/Public/DisplayClusterConfigurationTypes_Tile.h
--- a/Engine/Plugins/Runtime/nDisplay/Source/DisplayClusterConfiguration/Public/DisplayClusterConfigurationTypes_Tile.h
+++ b/Engine/Plugins/Runtime/nDisplay/Source/DisplayClusterConfiguration/Public/DisplayClusterConfigurationTypes_Tile.h
@@ -28,11 +28,11 @@ public:
 
 	/** The viewport will be horizontally partitioned into the specified number of tiles.  */
 	UPROPERTY(EditAnywhere, BlueprintReadWrite, Category = "Tile Rendering", meta = (DisplayName = "Tile X", ClampMin = "1", UIMin = "1", ClampMax = "4", UIMax = "4"))
-	int TileX = 1;
+	int32 TileX = 1;
 
 	/** The viewport will be vertically partitioned into the specified number of tiles.  */
 	UPROPERTY(EditAnywhere, BlueprintReadWrite, Category = "Tile Rendering", meta = (DisplayName = "Tile Y", ClampMin = "1", UIMin = "1", ClampMax = "4", UIMax = "4"))
-	int TileY = 1;
+	int32 TileY = 1;
 };
 
 USTRUCT(Blueprintable)
diff --git a/Engine/Plugins/Runtime/nDisplay/Source/DisplayClusterConfiguration/Public/DisplayClusterConfigurationTypes_Viewport.h b/Engine/Plugins/Runtime/nDisplay/Source/DisplayClusterConfiguration/Public/DisplayClusterConfigurationTypes_Viewport.h
--- a/Engine/Plugins/Runtime/nDisplay/Source/DisplayClusterConfiguration/Public/DisplayClusterConfigurationTypes_Viewport.h
+++ b/Engine/Plugins/Runtime/nDisplay/Source/DisplayClusterConfiguration/Public/DisplayClusterConfigurationTypes_Viewport.h
@@ -106,7 +106,7 @@ public:
 
 	// Media settings
 	UPROPERTY(EditAnywhere, BlueprintReadWrite, Category = "Configuration", meta = (DisplayName = "Media"))
-	FDisplayClusterConfigurationMedia Media;
+	FDisplayClusterConfigurationMediaViewport Media;
 
 	// Experimental: Support special frame builder mode - merge viewports to single viewfamily by group num
 	// [not implemented yet]
diff --git a/Engine/Plugins/Runtime/nDisplay/Source/DisplayClusterMedia/Private/Capture/DisplayClusterMediaCaptureViewport.cpp b/Engine/Plugins/Runtime/nDisplay/Source/DisplayClusterMedia/Private/Capture/DisplayClusterMediaCaptureViewport.cpp
--- a/Engine/Plugins/Runtime/nDisplay/Source/DisplayClusterMedia/Private/Capture/DisplayClusterMediaCaptureViewport.cpp
+++ b/Engine/Plugins/Runtime/nDisplay/Source/DisplayClusterMedia/Private/Capture/DisplayClusterMediaCaptureViewport.cpp
@@ -2,13 +2,19 @@
 
 #include "Capture/DisplayClusterMediaCaptureViewport.h"
 
-#include "Config/IDisplayClusterConfigManager.h"
-#include "DisplayClusterConfigurationTypes_Viewport.h"
 #include "IDisplayCluster.h"
 #include "IDisplayClusterCallbacks.h"
 
+#include "DisplayClusterConfigurationTypes_Viewport.h"
+
+#include "Config/IDisplayClusterConfigManager.h"
+
+#include "Render/IDisplayClusterRenderManager.h"
 #include "Render/Viewport/IDisplayClusterViewport.h"
+#include "Render/Viewport/IDisplayClusterViewportManager.h"
 #include "Render/Viewport/IDisplayClusterViewportManagerProxy.h"
+#include "Render/Viewport/Containers/DisplayClusterViewport_Context.h"
+#include "Render/Viewport/Containers/DisplayClusterViewport_RenderSettings.h"
 #include "RHICommandList.h"
 #include "RHIResources.h"
 
@@ -22,15 +28,12 @@ FDisplayClusterMediaCaptureViewport::FDisplayClusterMediaCaptureViewport(const F
 
 bool FDisplayClusterMediaCaptureViewport::StartCapture()
 {
-	// If capturing has started successfully, subscribe for rendering callbacks
+	// If capturing has started successfully, subscribe for the pipeline callbacks
 	if (FDisplayClusterMediaCaptureBase::StartCapture())
 	{
 		IDisplayCluster::Get().GetCallbacks().OnDisplayClusterPostRenderViewFamily_RenderThread().AddRaw(this, &FDisplayClusterMediaCaptureViewport::OnPostRenderViewFamily_RenderThread);
-
-		// Subscribes to viewport callback to raise media flags for viewport.
-		// Note: viewport is unaware of the media's configurations.
-		// Therefore, any future changes to the USTRUCT used by media do not affect the logic in the DisplayCluster module.
 		IDisplayCluster::Get().GetCallbacks().OnDisplayClusterUpdateViewportMediaState().AddRaw(this, &FDisplayClusterMediaCaptureViewport::OnUpdateViewportMediaState);
+
 		return true;
 	}
 
@@ -56,9 +59,9 @@ void FDisplayClusterMediaCaptureViewport::OnUpdateViewportMediaState(IDisplayClu
 	if (InViewport && InViewport->GetId().Equals(GetViewportId(), ESearchCase::IgnoreCase))
 	{
 		// Raise flags that this viewport will be captured by media.
-		InOutMediaState |= EDisplayClusterViewportMediaState::Capture;
+		InOutMediaState = EDisplayClusterViewportMediaState::Capture;
 
-		if (ForceLateOCIOPass)
+		if (bForceLateOCIOPass)
 		{
 			// Raise flags that this capture requires ForceLateOCIOPass.
 			InOutMediaState |= EDisplayClusterViewportMediaState::Capture_ForceLateOCIOPass;
@@ -68,12 +71,20 @@ void FDisplayClusterMediaCaptureViewport::OnUpdateViewportMediaState(IDisplayClu
 
 FIntPoint FDisplayClusterMediaCaptureViewport::GetCaptureSize() const
 {
-	if (IDisplayCluster::Get().GetConfigMgr())
+	// We need to get actual texture size for the viewport
+	if (const IDisplayClusterRenderManager* const RenderMgr = IDisplayCluster::Get().GetRenderMgr())
 	{
-		const UDisplayClusterConfigurationViewport* Viewport = IDisplayCluster::Get().GetConfigMgr()->GetLocalViewport(ViewportId);
-		if (ensure(Viewport))
+		if (const IDisplayClusterViewportManager* const ViewportMgr = RenderMgr->GetViewportManager())
 		{
-			return { Viewport->Region.W, Viewport->Region.H };
+			if (const IDisplayClusterViewport* const Viewport = ViewportMgr->FindViewport(ViewportId))
+			{
+				const TArray<FDisplayClusterViewport_Context>& Contexts = Viewport->GetContexts();
+				if (Contexts.Num() > 0)
+				{
+					const FIntPoint Size = Contexts[0].RenderTargetRect.Size();
+					return Size;
+				}
+			}
 		}
 	}
 
diff --git a/Engine/Plugins/Runtime/nDisplay/Source/DisplayClusterMedia/Private/Capture/DisplayClusterMediaCaptureViewport.h b/Engine/Plugins/Runtime/nDisplay/Source/DisplayClusterMedia/Private/Capture/DisplayClusterMediaCaptureViewport.h
--- a/Engine/Plugins/Runtime/nDisplay/Source/DisplayClusterMedia/Private/Capture/DisplayClusterMediaCaptureViewport.h
+++ b/Engine/Plugins/Runtime/nDisplay/Source/DisplayClusterMedia/Private/Capture/DisplayClusterMediaCaptureViewport.h
@@ -22,36 +22,34 @@ public:
 	FDisplayClusterMediaCaptureViewport(const FString& MediaId, const FString& ClusterNodeId, const FString& ViewportId, UMediaOutput* MediaOutput, UDisplayClusterMediaOutputSynchronizationPolicy* SyncPolicy = nullptr);
 
 public:
+	/** Start capturing */
 	virtual bool StartCapture() override;
+
+	/** Stop capturing */
 	virtual void StopCapture() override;
 
+	/** Returns viewport ID that is configured for capture */
 	const FString& GetViewportId() const
 	{
 		return ViewportId;
 	}
 
 protected:
-	
+	/** Returns texture size of a viewport assigned to capture */
 	virtual FIntPoint GetCaptureSize() const override;
 
 private:
+	/** PostRenderViewFamily callback handler where data is captured */
 	void OnPostRenderViewFamily_RenderThread(FRDGBuilder& GraphBuilder, const FSceneViewFamily& ViewFamily, const IDisplayClusterViewportProxy* ViewportProxy);
 
-	/** Media is completely independent of DC viewports. In addition, multiple media can use the same viewport in different ways.
-	* The viewport does not know about the configuration of the media. But the viewport needs some special media flags that change its internal logic.
-	* To get these flags, the viewport uses a multicast delegate to get all these flags from multiple media.
-	* On the media side, they must subscribe to this callback and raise these flags to change the viewport logic in the expected way.
-	*/
+	/** UpdateViewportMediaState callback to configure media state for a viewoprt */
 	void OnUpdateViewportMediaState(IDisplayClusterViewport* InViewport, EDisplayClusterViewportMediaState& InOutMediaState);
 
 public:
-	// Note: viewport is unaware of the media's configurations.
-	// Therefore, any future changes to the USTRUCT used by media do not affect the logic in the DisplayCluster module.
-	// These flags are passed to the viewport in the OnUpdateViewportMediaState() callback function.
-
-	// Force late OCIO pass
-	bool ForceLateOCIOPass = false;
+	/** Force late OCIO pass */
+	bool bForceLateOCIOPass = false;
 
 private:
+	/** Viewport ID assigned to capture */
 	const FString ViewportId;
 };
diff --git a/Engine/Plugins/Runtime/nDisplay/Source/DisplayClusterMedia/Private/DisplayClusterMediaCVars.cpp b/Engine/Plugins/Runtime/nDisplay/Source/DisplayClusterMedia/Private/DisplayClusterMediaCVars.cpp
new file mode 100644
--- /dev/null
+++ b/Engine/Plugins/Runtime/nDisplay/Source/DisplayClusterMedia/Private/DisplayClusterMediaCVars.cpp
@@ -0,0 +1,13 @@
+// Copyright Epic Games, Inc. All Rights Reserved.
+
+#include "DisplayClusterMediaCVars.h"
+
+
+TAutoConsoleVariable<bool> CVarMediaEnabled(
+	TEXT("nDisplay.Media.Enabled"),
+	true,
+	TEXT("nDisplay media subsystem\n")
+	TEXT("0 : Disabled\n")
+	TEXT("1 : Enabled\n"),
+	ECVF_ReadOnly
+);
diff --git a/Engine/Plugins/Runtime/nDisplay/Source/DisplayClusterMedia/Private/DisplayClusterMediaCVars.h b/Engine/Plugins/Runtime/nDisplay/Source/DisplayClusterMedia/Private/DisplayClusterMediaCVars.h
new file mode 100644
--- /dev/null
+++ b/Engine/Plugins/Runtime/nDisplay/Source/DisplayClusterMedia/Private/DisplayClusterMediaCVars.h
@@ -0,0 +1,9 @@
+// Copyright Epic Games, Inc. All Rights Reserved.
+
+#pragma once
+
+#include "HAL/IConsoleManager.h"
+
+// Enable/disable media.
+// May be useful if media needs to be disabled without doing any changes to the nDisplay configuration asset.
+extern TAutoConsoleVariable<bool> CVarMediaEnabled;
diff --git a/Engine/Plugins/Runtime/nDisplay/Source/DisplayClusterMedia/Private/DisplayClusterMediaHelpers.cpp b/Engine/Plugins/Runtime/nDisplay/Source/DisplayClusterMedia/Private/DisplayClusterMediaHelpers.cpp
--- a/Engine/Plugins/Runtime/nDisplay/Source/DisplayClusterMedia/Private/DisplayClusterMediaHelpers.cpp
+++ b/Engine/Plugins/Runtime/nDisplay/Source/DisplayClusterMedia/Private/DisplayClusterMediaHelpers.cpp
@@ -33,7 +33,7 @@ namespace DisplayClusterMediaHelpers
 						: FString::Printf(TEXT("%s_%s_%s_icvfx_input"), *NodeId, *DCRAName, *OwnerName);
 
 				default:
-					unimplemented();
+					checkNoEntry();
 				}
 			}
 			else if (DeviceType == EMediaDeviceType::Output)
diff --git a/Engine/Plugins/Runtime/nDisplay/Source/DisplayClusterMedia/Private/DisplayClusterMediaModule.cpp b/Engine/Plugins/Runtime/nDisplay/Source/DisplayClusterMedia/Private/DisplayClusterMediaModule.cpp
--- a/Engine/Plugins/Runtime/nDisplay/Source/DisplayClusterMedia/Private/DisplayClusterMediaModule.cpp
+++ b/Engine/Plugins/Runtime/nDisplay/Source/DisplayClusterMedia/Private/DisplayClusterMediaModule.cpp
@@ -1,6 +1,8 @@
 // Copyright Epic Games, Inc. All Rights Reserved.
 
 #include "DisplayClusterMediaModule.h"
+
+#include "DisplayClusterMediaCVars.h"
 #include "DisplayClusterMediaLog.h"
 #include "DisplayClusterMediaHelpers.h"
 
@@ -18,24 +20,12 @@
 #include "Capture/DisplayClusterMediaCaptureCamera.h"
 #include "Capture/DisplayClusterMediaCaptureNode.h"
 #include "Capture/DisplayClusterMediaCaptureViewport.h"
-#include "Input/DisplayClusterMediaInputNode.h"
 #include "Input/DisplayClusterMediaInputViewport.h"
 
 #include "IMediaModule.h"
 #include "Misc/CoreDelegates.h"
 
 
-static TAutoConsoleVariable<int32> CVarMediaEnabled(
-	TEXT("nDisplay.media.Enabled"),
-	1,
-	TEXT("nDisplay media subsystem\n")
-	TEXT("0 : Disabled\n")
-	TEXT("1 : Enabled\n")
-	,
-	ECVF_ReadOnly
-);
-
-
 void FDisplayClusterMediaModule::StartupModule()
 {
 	UE_LOG(LogDisplayClusterMedia, Log, TEXT("Starting module 'DisplayClusterMedia'..."));
@@ -79,7 +69,7 @@ void FDisplayClusterMediaModule::InitializeMedia()
 	}
 
 	// Check if media enabled
-	if (CVarMediaEnabled.GetValueOnGameThread() == 0)
+	if (!CVarMediaEnabled.GetValueOnGameThread())
 	{
 		UE_LOG(LogDisplayClusterMedia, Log, TEXT("nDisplay media subsytem is disabled by a cvar"));
 		return;
@@ -92,113 +82,22 @@ void FDisplayClusterMediaModule::InitializeMedia()
 	if (const ADisplayClusterRootActor* const RootActor = IDisplayCluster::Get().GetGameMgr()->GetRootActor())
 	{
 		const FString ClusterNodeId = IDisplayCluster::Get().GetClusterMgr()->GetNodeId();
+		const FString RootActorName = RootActor->GetName();
+
 		if (const UDisplayClusterConfigurationClusterNode* const ClusterNode = RootActor->GetConfigData()->Cluster->GetNode(ClusterNodeId))
 		{
-			///////////////////////////////
 			// Node backbuffer media setup
 			{
-				const FDisplayClusterConfigurationMedia& MediaSettings = ClusterNode->Media;
-
-				if (MediaSettings.bEnable)
-				{
-					// Media input
-					if (MediaSettings.IsMediaInputAssigned())
-					{
-						const FString MediaInputId = DisplayClusterMediaHelpers::MediaId::GenerateMediaId(
-							DisplayClusterMediaHelpers::MediaId::EMediaDeviceType::Input,
-							DisplayClusterMediaHelpers::MediaId::EMediaOwnerType::Backbuffer,
-							*ClusterNodeId, *RootActor->GetName(), FString(), 0);
-
-						UE_LOG(LogDisplayClusterMedia, Log, TEXT("Initializing backbuffer media input '%s'..."), *MediaInputId);
-
-						InputNode = MakeShared<FDisplayClusterMediaInputNode>(
-							MediaInputId, ClusterNodeId,
-							MediaSettings.MediaInput.MediaSource);
-					}
-
-					// Media capture
-					uint8 CaptureIdx = 0;
-					for (const FDisplayClusterConfigurationMediaOutput& MediaOutputItem : MediaSettings.MediaOutputs)
-					{
-						if (IsValid(MediaOutputItem.MediaOutput))
-						{
-							const FString MediaCaptureId = DisplayClusterMediaHelpers::MediaId::GenerateMediaId(
-								DisplayClusterMediaHelpers::MediaId::EMediaDeviceType::Output,
-								DisplayClusterMediaHelpers::MediaId::EMediaOwnerType::Backbuffer,
-								*ClusterNodeId, *RootActor->GetName(), FString(), CaptureIdx);
-
-							UE_LOG(LogDisplayClusterMedia, Log, TEXT("Initializing backbuffer media capture [%u]: '%s'"), CaptureIdx, *MediaCaptureId);
-
-							TSharedPtr<FDisplayClusterMediaCaptureNode> NewNodeCapture = MakeShared<FDisplayClusterMediaCaptureNode>(
-								MediaCaptureId, ClusterNodeId,
-								MediaOutputItem.MediaOutput,
-								MediaOutputItem.OutputSyncPolicy);
-
-							CaptureNode.Emplace(MediaCaptureId, NewNodeCapture);
-						}
-
-						++CaptureIdx;
-					}
-				}
+				InitializeBackbufferOutput(ClusterNode, RootActorName, ClusterNodeId);
 			}
 
-			///////////////////////////////
 			// Viewports media setup
 			for (const TPair<FString, TObjectPtr<UDisplayClusterConfigurationViewport>>& ViewportIt : ClusterNode->Viewports)
 			{
-				if (const UDisplayClusterConfigurationViewport* const Viewport = ViewportIt.Value)
-				{
-					const FDisplayClusterConfigurationMedia& MediaSettings = Viewport->RenderSettings.Media;
-
-					if (MediaSettings.bEnable)
-					{
-						// Media input
-						if (MediaSettings.IsMediaInputAssigned())
-						{
-							const FString MediaInputId = DisplayClusterMediaHelpers::MediaId::GenerateMediaId(
-								DisplayClusterMediaHelpers::MediaId::EMediaDeviceType::Input,
-								DisplayClusterMediaHelpers::MediaId::EMediaOwnerType::Viewport,
-								*ClusterNodeId, *RootActor->GetName(), ViewportIt.Key, 0);
-
-							UE_LOG(LogDisplayClusterMedia, Log, TEXT("Initializing viewport media input '%s' for viewport '%s'"), *MediaInputId, *ViewportIt.Key);
-
-							TSharedPtr<FDisplayClusterMediaInputViewport> NewViewportInput = MakeShared<FDisplayClusterMediaInputViewport>(
-								MediaInputId, ClusterNodeId,
-								ViewportIt.Key,
-								MediaSettings.MediaInput.MediaSource);
-
-							InputViewports.Emplace(MediaInputId, MoveTemp(NewViewportInput));
-						}
-
-						// Media capture
-						uint8 CaptureIdx = 0;
-						for (const FDisplayClusterConfigurationMediaOutput& MediaOutputItem : MediaSettings.MediaOutputs)
-						{
-							if (IsValid(MediaOutputItem.MediaOutput))
-							{
-								const FString MediaCaptureId = DisplayClusterMediaHelpers::MediaId::GenerateMediaId(
-									DisplayClusterMediaHelpers::MediaId::EMediaDeviceType::Output,
-									DisplayClusterMediaHelpers::MediaId::EMediaOwnerType::Viewport,
-									*ClusterNodeId, *RootActor->GetName(), ViewportIt.Key, CaptureIdx);
-
-								UE_LOG(LogDisplayClusterMedia, Log, TEXT("Initializing viewport capture [%u]: '%s' for viewport '%s'"), CaptureIdx, *MediaCaptureId, *ViewportIt.Key);
-
-								TSharedPtr<FDisplayClusterMediaCaptureViewport> NewViewportCapture = MakeShared<FDisplayClusterMediaCaptureViewport>(
-									MediaCaptureId, ClusterNodeId,
-									ViewportIt.Key,
-									MediaOutputItem.MediaOutput,
-									MediaOutputItem.OutputSyncPolicy);
-
-								CaptureViewports.Emplace(MediaCaptureId, MoveTemp(NewViewportCapture));
-							}
-
-							++CaptureIdx;
-						}
-					}
-				}
+				InitializeViewportInput(ViewportIt.Value, ViewportIt.Key, RootActorName, ClusterNodeId);
+				InitializeViewportOutput(ViewportIt.Value, ViewportIt.Key, RootActorName, ClusterNodeId);
 			}
 
-			///////////////////////////////
 			// ICVFX media setup
 			{
 				// Get all ICVFX camera components
@@ -209,141 +108,17 @@ void FDisplayClusterMediaModule::InitializeMedia()
 				{
 					const FDisplayClusterConfigurationMediaICVFX& MediaSettings = ICVFXCameraComponent->CameraSettings.RenderSettings.Media;
 
-					if (MediaSettings.bEnable)
+					// Full frame
+					if (MediaSettings.SplitType == EDisplayClusterConfigurationMediaSplitType::FullFrame)
+					{
+						InitializeICVFXCameraFullFrameInput(ICVFXCameraComponent, RootActorName, ClusterNodeId);
+						InitializeICVFXCameraFullFrameOutput(ICVFXCameraComponent, RootActorName, ClusterNodeId);
+					}
+					// Uniform tiles
+					else if (MediaSettings.SplitType == EDisplayClusterConfigurationMediaSplitType::UniformTiles)
 					{
-						const FString ICVFXCameraName = ICVFXCameraComponent->GetName();
-						const FString ICVFXViewportId = DisplayClusterMediaHelpers::GenerateICVFXViewportName(ClusterNodeId, ICVFXCameraName);
-
-						// Media input
-						if (UMediaSource* MediaSource = MediaSettings.GetMediaSource(ClusterNodeId))
-						{
-							const FString MediaInputId = DisplayClusterMediaHelpers::MediaId::GenerateMediaId(
-								DisplayClusterMediaHelpers::MediaId::EMediaDeviceType::Input,
-								DisplayClusterMediaHelpers::MediaId::EMediaOwnerType::ICVFXCamera,
-								*ClusterNodeId, *RootActor->GetName(), ICVFXCameraName, 0);
-
-							UE_LOG(LogDisplayClusterMedia, Log, TEXT("Initializing ICVFX media input '%s' for camera '%s'"), *MediaInputId, *ICVFXCameraName);
-
-							TSharedPtr<FDisplayClusterMediaInputViewport> NewICVFXInput = MakeShared<FDisplayClusterMediaInputViewport>(
-								MediaInputId, ClusterNodeId,
-								ICVFXViewportId,
-								MediaSource);
-
-							// Saves special media flags from the configuration to the media instance.
-							//
-							// Note: viewport is unaware of the media's configurations.
-							// Therefore, any future changes to the USTRUCT used by media do not affect the logic in the DisplayCluster module.
-							// These flags are passed to the viewport in the OnUpdateViewportMediaState() callback function.
-							NewICVFXInput->ForceLateOCIOPass = MediaSettings.bLateOCIOPass;
-
-							InputViewports.Emplace(MediaInputId, MoveTemp(NewICVFXInput));
-						}
-
-						// Media input for tile rendering
-						TArray<FIntPoint> InputTiles;
-						if (UMediaSource* MediaSource = MediaSettings.GetMediaSourceForTiles(ClusterNodeId, InputTiles))
-						{
-							for (const FIntPoint& TilePos : InputTiles)
-							{
-								const FString ICVFXViewportTileId = DisplayClusterMediaHelpers::GenerateTileViewportName(ICVFXViewportId, TilePos);
-
-								const FString MediaInputId = DisplayClusterMediaHelpers::MediaId::GenerateMediaId(
-									DisplayClusterMediaHelpers::MediaId::EMediaDeviceType::Input,
-									DisplayClusterMediaHelpers::MediaId::EMediaOwnerType::ICVFXCamera,
-									*ClusterNodeId, *RootActor->GetName(), ICVFXCameraName, 0, &TilePos);
-
-								UE_LOG(LogDisplayClusterMedia, Log, TEXT("Initializing ICVFX media input '%s' for camera '%s' tile '%d,%d'"), *MediaInputId, *ICVFXCameraName, TilePos.X, TilePos.Y);
-
-								TSharedPtr<FDisplayClusterMediaInputViewport> NewICVFXInput = MakeShared<FDisplayClusterMediaInputViewport>(
-									MediaInputId, ClusterNodeId,
-									ICVFXViewportTileId,
-									MediaSource);
-
-								// Saves special media flags from the configuration to the media instance.
-								//
-								// Note: viewport is unaware of the media's configurations.
-								// Therefore, any future changes to the USTRUCT used by media do not affect the logic in the DisplayCluster module.
-								// These flags are passed to the viewport in the OnUpdateViewportMediaState() callback function.
-								NewICVFXInput->ForceLateOCIOPass = MediaSettings.bLateOCIOPass;
-
-								InputViewports.Emplace(MediaInputId, MoveTemp(NewICVFXInput));
-							}
-						}
-
-						// Media capture
-						const TArray<FDisplayClusterConfigurationMediaOutputGroup> MediaOutputItems = MediaSettings.GetMediaOutputGroups(ClusterNodeId);
-						uint8 CaptureIdx = 0;
-						for (const FDisplayClusterConfigurationMediaOutputGroup& MediaOutputItem : MediaOutputItems)
-						{
-							if (IsValid(MediaOutputItem.MediaOutput))
-							{
-								const FString MediaCaptureId = DisplayClusterMediaHelpers::MediaId::GenerateMediaId(
-									DisplayClusterMediaHelpers::MediaId::EMediaDeviceType::Output,
-									DisplayClusterMediaHelpers::MediaId::EMediaOwnerType::ICVFXCamera,
-									*ClusterNodeId, *RootActor->GetName(), ICVFXCameraName, CaptureIdx);
-
-								UE_LOG(LogDisplayClusterMedia, Log, TEXT("Initializing ICVFX capture [%u]: '%s' for camera '%s'"), CaptureIdx, *MediaCaptureId, *ICVFXCameraName);
-
-								TSharedPtr<FDisplayClusterMediaCaptureViewport> NewICVFXCapture = MakeShared<FDisplayClusterMediaCaptureCamera>(
-									MediaCaptureId, ClusterNodeId,
-									ICVFXCameraName, ICVFXViewportId,
-									MediaOutputItem.MediaOutput,
-									MediaOutputItem.OutputSyncPolicy);
-
-								// Saves special media flags from the configuration to the media instance.
-								//
-								// Note: viewport is unaware of the media's configurations.
-								// Therefore, any future changes to the USTRUCT used by media do not affect the logic in the DisplayCluster module.
-								// These flags are passed to the viewport in the OnUpdateViewportMediaState() callback function.
-								NewICVFXCapture->ForceLateOCIOPass = MediaSettings.bLateOCIOPass;
-
-								CaptureViewports.Emplace(MediaCaptureId, MoveTemp(NewICVFXCapture));
-							}
-
-							++CaptureIdx;
-						}
-
-						// Media capture for tiles
-						const TArray<FDisplayClusterConfigurationMediaOutputGroup> MediaOutputTileItems = MediaSettings.GetMediaOutputGroupsForTiles(ClusterNodeId);
-						uint8 CaptureTileIdx = 0;
-						for (const FDisplayClusterConfigurationMediaOutputGroup& MediaOutputTileItem : MediaOutputTileItems)
-						{
-							if (IsValid(MediaOutputTileItem.MediaOutput))
-							{
-								if (const FDisplayClusterConfigurationClusterNodeTilesReferenceList* TilesReferenceList = MediaOutputTileItem.ClusterNodesWithTiles.Find(ClusterNodeId))
-								{
-									for (const FDisplayClusterConfigurationTileIndex& TileIndex : TilesReferenceList->Tiles)
-									{
-										const FIntPoint TilePos(TileIndex.TileX, TileIndex.TileY);
-										const FString ICVFXViewportTileId = DisplayClusterMediaHelpers::GenerateTileViewportName(ICVFXViewportId, TilePos);
-
-										const FString MediaCaptureId = DisplayClusterMediaHelpers::MediaId::GenerateMediaId(
-											DisplayClusterMediaHelpers::MediaId::EMediaDeviceType::Output,
-											DisplayClusterMediaHelpers::MediaId::EMediaOwnerType::ICVFXCamera,
-											*ClusterNodeId, *RootActor->GetName(), ICVFXCameraName, CaptureTileIdx, &TilePos);
-
-										UE_LOG(LogDisplayClusterMedia, Log, TEXT("Initializing ICVFX capture [%u]: '%s' for camera '%s' tile '%d,%d'"), CaptureTileIdx, *MediaCaptureId, *ICVFXCameraName, TilePos.X, TilePos.Y);
-
-										TSharedPtr<FDisplayClusterMediaCaptureViewport> NewICVFXCapture = MakeShared<FDisplayClusterMediaCaptureCamera>(
-											MediaCaptureId, ClusterNodeId,
-											ICVFXCameraName, ICVFXViewportTileId,
-											MediaOutputTileItem.MediaOutput,
-											MediaOutputTileItem.OutputSyncPolicy);
-
-										// Saves special media flags from the configuration to the media instance.
-										//
-										// Note: viewport is unaware of the media's configurations.
-										// Therefore, any future changes to the USTRUCT used by media do not affect the logic in the DisplayCluster module.
-										// These flags are passed to the viewport in the OnUpdateViewportMediaState() callback function.
-										NewICVFXCapture->ForceLateOCIOPass = MediaSettings.bLateOCIOPass;
-
-										CaptureViewports.Emplace(MediaCaptureId, MoveTemp(NewICVFXCapture));
-									}
-								}
-							}
-
-							++CaptureTileIdx;
-						}
+						InitializeICVFXCameraUniformTilesInput(ICVFXCameraComponent, RootActorName, ClusterNodeId);
+						InitializeICVFXCameraUniformTilesOutput(ICVFXCameraComponent, RootActorName, ClusterNodeId);
 					}
 				}
 			}
@@ -360,7 +135,6 @@ void FDisplayClusterMediaModule::ReleaseMedia()
 	CaptureNode.Reset();
 
 	InputViewports.Reset();
-	InputNode.Reset();
 
 	FrameQueue.Release();
 }
@@ -402,12 +176,6 @@ void FDisplayClusterMediaModule::PlayMedia()
 	{
 		MediaInput.Value->Play();
 	}
-
-	// Start playback to the backbuffer
-	if (InputNode)
-	{
-		InputNode->Play();
-	}
 }
 
 void FDisplayClusterMediaModule::StopMedia()
@@ -417,11 +185,293 @@ void FDisplayClusterMediaModule::StopMedia()
 	{
 		MediaInput.Value->Stop();
 	}
+}
+
+void FDisplayClusterMediaModule::InitializeBackbufferOutput(const UDisplayClusterConfigurationClusterNode* ClusterNode, const FString& RootActorName, const FString& ClusterNodeId)
+{
+	checkSlow(ClusterNode);
+
+	if (IsValid(ClusterNode))
+	{
+		const FDisplayClusterConfigurationMediaNodeBackbuffer& MediaSettings = ClusterNode->MediaSettings;
+
+		if (MediaSettings.bEnable)
+		{
+			uint8 CaptureIdx = 0;
+			for (const FDisplayClusterConfigurationMediaOutput& MediaOutputItem : MediaSettings.MediaOutputs)
+			{
+				if (IsValid(MediaOutputItem.MediaOutput))
+				{
+					const FString MediaCaptureId = DisplayClusterMediaHelpers::MediaId::GenerateMediaId(
+						DisplayClusterMediaHelpers::MediaId::EMediaDeviceType::Output,
+						DisplayClusterMediaHelpers::MediaId::EMediaOwnerType::Backbuffer,
+						ClusterNodeId, RootActorName, FString(), CaptureIdx);
+
+					UE_LOG(LogDisplayClusterMedia, Log, TEXT("Initializing backbuffer media capture [%u]: '%s'"), CaptureIdx, *MediaCaptureId);
+
+					TSharedPtr<FDisplayClusterMediaCaptureNode> NewNodeCapture = MakeShared<FDisplayClusterMediaCaptureNode>(
+						MediaCaptureId, ClusterNodeId,
+						MediaOutputItem.MediaOutput,
+						MediaOutputItem.OutputSyncPolicy);
+
+					CaptureNode.Emplace(MediaCaptureId, NewNodeCapture);
+				}
+
+				++CaptureIdx;
+			}
+		}
+	}
+}
+
+void FDisplayClusterMediaModule::InitializeViewportInput(const UDisplayClusterConfigurationViewport* Viewport, const FString& ViewportId, const FString& RootActorName, const FString& ClusterNodeId)
+{
+	checkSlow(Viewport);
+
+	if (IsValid(Viewport))
+	{
+		const FDisplayClusterConfigurationMediaViewport& MediaSettings = Viewport->RenderSettings.Media;
+
+		if (MediaSettings.bEnable)
+		{
+			if (MediaSettings.IsMediaInputAssigned())
+			{
+				const FString MediaInputId = DisplayClusterMediaHelpers::MediaId::GenerateMediaId(
+					DisplayClusterMediaHelpers::MediaId::EMediaDeviceType::Input,
+					DisplayClusterMediaHelpers::MediaId::EMediaOwnerType::Viewport,
+					ClusterNodeId, RootActorName, ViewportId, 0);
+
+				UE_LOG(LogDisplayClusterMedia, Log, TEXT("Initializing viewport media input '%s' for viewport '%s'"), *MediaInputId, *ViewportId);
+
+				TSharedPtr<FDisplayClusterMediaInputViewport> NewViewportInput = MakeShared<FDisplayClusterMediaInputViewport>(
+					MediaInputId, ClusterNodeId,
+					ViewportId,
+					MediaSettings.MediaInput.MediaSource);
+
+				InputViewports.Emplace(MediaInputId, MoveTemp(NewViewportInput));
+			}
+		}
+	}
+}
+
+void FDisplayClusterMediaModule::InitializeViewportOutput(const UDisplayClusterConfigurationViewport* Viewport, const FString& ViewportId, const FString& RootActorName, const FString& ClusterNodeId)
+{
+	checkSlow(Viewport);
+
+	if (IsValid(Viewport))
+	{
+		const FDisplayClusterConfigurationMediaViewport& MediaSettings = Viewport->RenderSettings.Media;
+
+		// Media capture
+		uint8 CaptureIdx = 0;
+		for (const FDisplayClusterConfigurationMediaOutput& MediaOutputItem : MediaSettings.MediaOutputs)
+		{
+			if (IsValid(MediaOutputItem.MediaOutput))
+			{
+				const FString MediaCaptureId = DisplayClusterMediaHelpers::MediaId::GenerateMediaId(
+					DisplayClusterMediaHelpers::MediaId::EMediaDeviceType::Output,
+					DisplayClusterMediaHelpers::MediaId::EMediaOwnerType::Viewport,
+					ClusterNodeId, RootActorName, ViewportId, CaptureIdx);
+
+				UE_LOG(LogDisplayClusterMedia, Log, TEXT("Initializing viewport capture [%u]: '%s' for viewport '%s'"), CaptureIdx, *MediaCaptureId, *ViewportId);
+
+				TSharedPtr<FDisplayClusterMediaCaptureViewport> NewViewportCapture = MakeShared<FDisplayClusterMediaCaptureViewport>(
+					MediaCaptureId, ClusterNodeId,
+					ViewportId,
+					MediaOutputItem.MediaOutput,
+					MediaOutputItem.OutputSyncPolicy);
+
+				CaptureViewports.Emplace(MediaCaptureId, MoveTemp(NewViewportCapture));
+			}
+
+			++CaptureIdx;
+		}
+	}
+}
+
+void FDisplayClusterMediaModule::InitializeICVFXCameraFullFrameInput(const UDisplayClusterICVFXCameraComponent* ICVFXCameraComponent, const FString& RootActorName, const FString& ClusterNodeId)
+{
+	checkSlow(ICVFXCameraComponent);
+
+	if (IsValid(ICVFXCameraComponent))
+	{
+		const FDisplayClusterConfigurationMediaICVFX& MediaSettings = ICVFXCameraComponent->CameraSettings.RenderSettings.Media;
+
+		if (MediaSettings.bEnable && MediaSettings.SplitType == EDisplayClusterConfigurationMediaSplitType::FullFrame)
+		{
+			const FString ICVFXCameraName = ICVFXCameraComponent->GetName();
+			const FString ICVFXViewportId = DisplayClusterMediaHelpers::GenerateICVFXViewportName(ClusterNodeId, ICVFXCameraName);
+
+			if (UMediaSource* MediaSource = MediaSettings.GetMediaSource(ClusterNodeId))
+			{
+				const FString MediaInputId = DisplayClusterMediaHelpers::MediaId::GenerateMediaId(
+					DisplayClusterMediaHelpers::MediaId::EMediaDeviceType::Input,
+					DisplayClusterMediaHelpers::MediaId::EMediaOwnerType::ICVFXCamera,
+					ClusterNodeId, RootActorName, ICVFXCameraName, 0);
+
+				UE_LOG(LogDisplayClusterMedia, Log, TEXT("Initializing ICVFX media input '%s' for camera '%s'"), *MediaInputId, *ICVFXCameraName);
+
+				TSharedPtr<FDisplayClusterMediaInputViewport> NewICVFXInput = MakeShared<FDisplayClusterMediaInputViewport>(
+					MediaInputId, ClusterNodeId,
+					ICVFXViewportId,
+					MediaSource);
+
+				NewICVFXInput->bForceLateOCIOPass = MediaSettings.bLateOCIOPass;
+
+				InputViewports.Emplace(MediaInputId, MoveTemp(NewICVFXInput));
+			}
+		}
+	}
+}
+
+void FDisplayClusterMediaModule::InitializeICVFXCameraFullFrameOutput(const UDisplayClusterICVFXCameraComponent* ICVFXCameraComponent, const FString& RootActorName, const FString& ClusterNodeId)
+{
+	checkSlow(ICVFXCameraComponent);
+
+	if (IsValid(ICVFXCameraComponent))
+	{
+		const FDisplayClusterConfigurationMediaICVFX& MediaSettings = ICVFXCameraComponent->CameraSettings.RenderSettings.Media;
+
+		if (MediaSettings.bEnable && MediaSettings.SplitType == EDisplayClusterConfigurationMediaSplitType::FullFrame)
+		{
+			const FString ICVFXCameraName = ICVFXCameraComponent->GetName();
+			const FString ICVFXViewportId = DisplayClusterMediaHelpers::GenerateICVFXViewportName(ClusterNodeId, ICVFXCameraName);
+
+			// Media capture (full frame)
+			const TArray<FDisplayClusterConfigurationMediaOutputGroup> MediaOutputItems = MediaSettings.GetMediaOutputGroups(ClusterNodeId);
+			uint8 CaptureIdx = 0;
+			for (const FDisplayClusterConfigurationMediaOutputGroup& MediaOutputItem : MediaOutputItems)
+			{
+				if (IsValid(MediaOutputItem.MediaOutput))
+				{
+					const FString MediaCaptureId = DisplayClusterMediaHelpers::MediaId::GenerateMediaId(
+						DisplayClusterMediaHelpers::MediaId::EMediaDeviceType::Output,
+						DisplayClusterMediaHelpers::MediaId::EMediaOwnerType::ICVFXCamera,
+						ClusterNodeId, RootActorName, ICVFXCameraName, CaptureIdx);
+
+					UE_LOG(LogDisplayClusterMedia, Log, TEXT("Initializing ICVFX capture [%u]: '%s' for camera '%s'"), CaptureIdx, *MediaCaptureId, *ICVFXCameraName);
+
+					TSharedPtr<FDisplayClusterMediaCaptureViewport> NewICVFXCapture = MakeShared<FDisplayClusterMediaCaptureCamera>(
+						MediaCaptureId, ClusterNodeId,
+						ICVFXCameraName, ICVFXViewportId,
+						MediaOutputItem.MediaOutput,
+						MediaOutputItem.OutputSyncPolicy);
+
+					NewICVFXCapture->bForceLateOCIOPass = MediaSettings.bLateOCIOPass;
+
+					CaptureViewports.Emplace(MediaCaptureId, MoveTemp(NewICVFXCapture));
+				}
+
+				++CaptureIdx;
+			}
+		}
+	}
+}
+
+void FDisplayClusterMediaModule::InitializeICVFXCameraUniformTilesInput(const UDisplayClusterICVFXCameraComponent* ICVFXCameraComponent, const FString& RootActorName, const FString& ClusterNodeId)
+{
+	checkSlow(ICVFXCameraComponent);
+
+	if (IsValid(ICVFXCameraComponent))
+	{
+		const FDisplayClusterConfigurationMediaICVFX& MediaSettings = ICVFXCameraComponent->CameraSettings.RenderSettings.Media;
+
+		if (MediaSettings.bEnable && MediaSettings.SplitType == EDisplayClusterConfigurationMediaSplitType::UniformTiles)
+		{
+			const FString ICVFXCameraName = ICVFXCameraComponent->GetName();
+			const FString ICVFXViewportId = DisplayClusterMediaHelpers::GenerateICVFXViewportName(ClusterNodeId, ICVFXCameraName);
+
+			// Find corresponsing media group
+			TArray<FDisplayClusterConfigurationMediaUniformTileInput> MediaInputTiles;
+			const bool bMediaInputGroupFound = MediaSettings.GetMediaInputTiles(ClusterNodeId, MediaInputTiles);
+
+			if (bMediaInputGroupFound)
+			{
+				uint8 Index = 0;
+				for (const FDisplayClusterConfigurationMediaUniformTileInput& MediaInputTile : MediaInputTiles)
+				{
+					if (IsValid(MediaInputTile.MediaSource))
+					{
+						const FIntPoint TileLoc{ MediaInputTile.TileX, MediaInputTile.TileY };
+
+						const FString ICVFXViewportTileId = DisplayClusterMediaHelpers::GenerateTileViewportName(ICVFXViewportId, TileLoc);
+
+						const FString MediaInputId = DisplayClusterMediaHelpers::MediaId::GenerateMediaId(
+							DisplayClusterMediaHelpers::MediaId::EMediaDeviceType::Input,
+							DisplayClusterMediaHelpers::MediaId::EMediaOwnerType::ICVFXCamera,
+							ClusterNodeId, RootActorName, ICVFXCameraName,
+							static_cast<uint8>(Index), &TileLoc);
+
+						UE_LOG(LogDisplayClusterMedia, Log, TEXT("Initializing ICVFX media input '%s' for camera '%s' tile '%d,%d'"),
+							*MediaInputId, *ICVFXCameraName, TileLoc.X, TileLoc.Y);
+
+						TSharedPtr<FDisplayClusterMediaInputViewport> NewICVFXTileInput = MakeShared<FDisplayClusterMediaInputViewport>(
+							MediaInputId, ClusterNodeId,
+							ICVFXViewportTileId,
+							MediaInputTile.MediaSource);
+
+						NewICVFXTileInput->bForceLateOCIOPass = MediaSettings.bLateOCIOPass;
+
+						InputViewports.Emplace(MediaInputId, MoveTemp(NewICVFXTileInput));
+					}
+
+					++Index;
+				}
+			}
+		}
+	}
+}
 
-	// Stop playback to the backbuffer
-	if (InputNode)
+void FDisplayClusterMediaModule::InitializeICVFXCameraUniformTilesOutput(const UDisplayClusterICVFXCameraComponent* ICVFXCameraComponent, const FString& RootActorName, const FString& ClusterNodeId)
+{
+	checkSlow(ICVFXCameraComponent);
+
+	if (IsValid(ICVFXCameraComponent))
 	{
-		InputNode->Stop();
+		const FDisplayClusterConfigurationMediaICVFX& MediaSettings = ICVFXCameraComponent->CameraSettings.RenderSettings.Media;
+
+		if (MediaSettings.bEnable && MediaSettings.SplitType == EDisplayClusterConfigurationMediaSplitType::UniformTiles)
+		{
+			const FString ICVFXCameraName = ICVFXCameraComponent->GetName();
+			const FString ICVFXViewportId = DisplayClusterMediaHelpers::GenerateICVFXViewportName(ClusterNodeId, ICVFXCameraName);
+
+			// Find corresponsing media group
+			TArray<FDisplayClusterConfigurationMediaUniformTileOutput> MediaOutputTiles;
+			const bool bMediaOutputGroupFound = MediaSettings.GetMediaOutputTiles(ClusterNodeId, MediaOutputTiles);
+
+			if (bMediaOutputGroupFound)
+			{
+				uint8 Index = 0;
+				for (const FDisplayClusterConfigurationMediaUniformTileOutput& MediaOutputTile : MediaOutputTiles)
+				{
+					if (IsValid(MediaOutputTile.MediaOutput))
+					{
+						const FIntPoint TileLoc{ MediaOutputTile.TileX, MediaOutputTile.TileY };
+
+						const FString ICVFXViewportTileId = DisplayClusterMediaHelpers::GenerateTileViewportName(ICVFXViewportId, TileLoc);
+
+						const FString MediaOutputId = DisplayClusterMediaHelpers::MediaId::GenerateMediaId(
+							DisplayClusterMediaHelpers::MediaId::EMediaDeviceType::Output,
+							DisplayClusterMediaHelpers::MediaId::EMediaOwnerType::ICVFXCamera,
+							ClusterNodeId, RootActorName, ICVFXCameraName,
+							static_cast<uint8>(Index), &TileLoc);
+
+						UE_LOG(LogDisplayClusterMedia, Log, TEXT("Initializing ICVFX media output '%s' for camera '%s' tile '%d,%d'"),
+							*MediaOutputId, *ICVFXCameraName, TileLoc.X, TileLoc.Y);
+
+						TSharedPtr<FDisplayClusterMediaCaptureViewport> NewICVFXTileOutput = MakeShared<FDisplayClusterMediaCaptureViewport>(
+							MediaOutputId, ClusterNodeId,
+							ICVFXViewportTileId,
+							MediaOutputTile.MediaOutput);
+
+						NewICVFXTileOutput->bForceLateOCIOPass = MediaSettings.bLateOCIOPass;
+
+						CaptureViewports.Emplace(MediaOutputId, MoveTemp(NewICVFXTileOutput));
+					}
+
+					++Index;
+				}
+			}
+		}
 	}
 }
 
diff --git a/Engine/Plugins/Runtime/nDisplay/Source/DisplayClusterMedia/Private/DisplayClusterMediaModule.h b/Engine/Plugins/Runtime/nDisplay/Source/DisplayClusterMedia/Private/DisplayClusterMediaModule.h
--- a/Engine/Plugins/Runtime/nDisplay/Source/DisplayClusterMedia/Private/DisplayClusterMediaModule.h
+++ b/Engine/Plugins/Runtime/nDisplay/Source/DisplayClusterMedia/Private/DisplayClusterMediaModule.h
@@ -9,10 +9,11 @@
 
 class FDisplayClusterMediaCaptureNode;
 class FDisplayClusterMediaCaptureViewport;
-class FDisplayClusterMediaInputNode;
 class FDisplayClusterMediaInputViewport;
 class IMediaPlayerFactory;
-
+class UDisplayClusterConfigurationClusterNode;
+class UDisplayClusterConfigurationViewport;
+class UDisplayClusterICVFXCameraComponent;
 
 /**
  * Media module
@@ -46,15 +47,43 @@ protected:
 	void StopMedia();
 
 private:
+	/** Backbuffer output initializer */
+	void InitializeBackbufferOutput(const UDisplayClusterConfigurationClusterNode* ClusterNode, const FString& RootActorName, const FString& ClusterNodeId);
+	
+	/** Viewport input initializer */
+	void InitializeViewportInput(const UDisplayClusterConfigurationViewport* Viewport, const FString& ViewportId, const FString& RootActorName, const FString& ClusterNodeId);
+
+	/** Viewport output initializer */
+	void InitializeViewportOutput(const UDisplayClusterConfigurationViewport* Viewport, const FString& ViewportId, const FString& RootActorName, const FString& ClusterNodeId);
+
+	/** ICVFX camera input initializer (full frame) */
+	void InitializeICVFXCameraFullFrameInput(const UDisplayClusterICVFXCameraComponent* ICVFXCameraComponent, const FString& RootActorName, const FString& ClusterNodeId);
+
+	/** ICVFX camera output initializer (full frame) */
+	void InitializeICVFXCameraFullFrameOutput(const UDisplayClusterICVFXCameraComponent* ICVFXCameraComponent, const FString& RootActorName, const FString& ClusterNodeId);
+
+	/** ICVFX camera input initializer (uniform tiles) */
+	void InitializeICVFXCameraUniformTilesInput(const UDisplayClusterICVFXCameraComponent* ICVFXCameraComponent, const FString& RootActorName, const FString& ClusterNodeId);
+
+	/** ICVFX camera output initializer (uniform tiles) */
+	void InitializeICVFXCameraUniformTilesOutput(const UDisplayClusterICVFXCameraComponent* ICVFXCameraComponent, const FString& RootActorName, const FString& ClusterNodeId);
+
+private:
+	/** CustomPresentSet event handler */
 	void OnCustomPresentSet();
+
+	/** EnginePreExit event handler */
 	void OnEnginePreExit();
 
 private:
+	/** Active media devices that capture viewports */
 	TMap<FString, TSharedPtr<FDisplayClusterMediaCaptureViewport>> CaptureViewports;
+
+	/** Active media devices that capture local backbuffer */
 	TMap<FString, TSharedPtr<FDisplayClusterMediaCaptureNode>>     CaptureNode;
 
+	/** Active media devices for viewport playback */
 	TMap<FString, TSharedPtr<FDisplayClusterMediaInputViewport>>   InputViewports;
-	TSharedPtr<FDisplayClusterMediaInputNode>                      InputNode;
 
 	// Latency queue
 	FDisplayClusterFrameQueue FrameQueue;
diff --git a/Engine/Plugins/Runtime/nDisplay/Source/DisplayClusterMedia/Private/Input/DisplayClusterMediaInputNode.cpp b/Engine/Plugins/Runtime/nDisplay/Source/DisplayClusterMedia/Private/Input/DisplayClusterMediaInputNode.cpp
deleted file mode 100644
diff --git a/Engine/Plugins/Runtime/nDisplay/Source/DisplayClusterMedia/Private/Input/DisplayClusterMediaInputNode.h b/Engine/Plugins/Runtime/nDisplay/Source/DisplayClusterMedia/Private/Input/DisplayClusterMediaInputNode.h
deleted file mode 100644
diff --git a/Engine/Plugins/Runtime/nDisplay/Source/DisplayClusterMedia/Private/Input/DisplayClusterMediaInputViewport.cpp b/Engine/Plugins/Runtime/nDisplay/Source/DisplayClusterMedia/Private/Input/DisplayClusterMediaInputViewport.cpp
--- a/Engine/Plugins/Runtime/nDisplay/Source/DisplayClusterMedia/Private/Input/DisplayClusterMediaInputViewport.cpp
+++ b/Engine/Plugins/Runtime/nDisplay/Source/DisplayClusterMedia/Private/Input/DisplayClusterMediaInputViewport.cpp
@@ -26,10 +26,6 @@ bool FDisplayClusterMediaInputViewport::Play()
 	if (FDisplayClusterMediaInputBase::Play())
 	{
 		IDisplayCluster::Get().GetCallbacks().OnDisplayClusterPostCrossGpuTransfer_RenderThread().AddRaw(this, &FDisplayClusterMediaInputViewport::PostCrossGpuTransfer_RenderThread);
-
-		// Subscribes to viewport callback to raise media flags for viewport.
-		// Note: viewport is unaware of the media's configurations.
-		// Therefore, any future changes to the USTRUCT used by media do not affect the logic in the DisplayCluster module.
 		IDisplayCluster::Get().GetCallbacks().OnDisplayClusterUpdateViewportMediaState().AddRaw(this, &FDisplayClusterMediaInputViewport::OnUpdateViewportMediaState);
 
 		return true;
@@ -57,9 +53,9 @@ void FDisplayClusterMediaInputViewport::OnUpdateViewportMediaState(IDisplayClust
 	if (InViewport && InViewport->GetId().Equals(GetViewportId(), ESearchCase::IgnoreCase))
 	{
 		// Raise flags that this viewport texture will be overridden by media.
-		InOutMediaState |= EDisplayClusterViewportMediaState::Input;
+		InOutMediaState = EDisplayClusterViewportMediaState::Input;
 
-		if (ForceLateOCIOPass)
+		if (bForceLateOCIOPass)
 		{
 			// Raise flags that this viewport requires ForceLateOCIOPass.
 			InOutMediaState |= EDisplayClusterViewportMediaState::Input_ForceLateOCIOPass;
diff --git a/Engine/Plugins/Runtime/nDisplay/Source/DisplayClusterMedia/Private/Input/DisplayClusterMediaInputViewport.h b/Engine/Plugins/Runtime/nDisplay/Source/DisplayClusterMedia/Private/Input/DisplayClusterMediaInputViewport.h
--- a/Engine/Plugins/Runtime/nDisplay/Source/DisplayClusterMedia/Private/Input/DisplayClusterMediaInputViewport.h
+++ b/Engine/Plugins/Runtime/nDisplay/Source/DisplayClusterMedia/Private/Input/DisplayClusterMediaInputViewport.h
@@ -21,32 +21,30 @@ public:
 	FDisplayClusterMediaInputViewport(const FString& MediaId, const FString& ClusterNodeId, const FString& ViewportId, UMediaSource* MediaSource);
 
 public:
+	/** Start playback */
 	virtual bool Play() override;
+
+	/** Stop playback */
 	virtual void Stop() override;
 
+	/** Returns viewport ID bound for playback */
 	const FString& GetViewportId() const
 	{
 		return ViewportId;
 	}
 
 private:
+	/** PostCrossGpuTransfer callback handler where media data is pushed into nDisplay internal buffers */
 	void PostCrossGpuTransfer_RenderThread(FRHICommandListImmediate& RHICmdList, const IDisplayClusterViewportManagerProxy* ViewportManagerProxy, FViewport* Viewport);
 
-	/** Media is completely independent of DC viewports. In addition, multiple media can use the same viewport in different ways.
-	* The viewport does not know about the configuration of the media. But the viewport needs some special media flags that change its internal logic.
-	* To get these flags, the viewport uses a multicast delegate to get all these flags from multiple media.
-	* On the media side, they must subscribe to this callback and raise these flags to change the viewport logic in the expected way.
-	*/
+	/** UpdateViewportMediaState callback to configure media state for a viewoprt */
 	void OnUpdateViewportMediaState(IDisplayClusterViewport* InViewport, EDisplayClusterViewportMediaState& InOutMediaState);
 
 public:
-	// Note: viewport is unaware of the media's configurations.
-	// Therefore, any future changes to the USTRUCT used by media do not affect the logic in the DisplayCluster module.
-	// These flags are passed to the viewport in the OnUpdateViewportMediaState() callback function.
-
-	// Force late OCIO pass
-	bool ForceLateOCIOPass = false;
+	/** Force late OCIO pass */
+	bool bForceLateOCIOPass = false;
 
 private:
+	/** Viewport ID assigned for this media input */
 	const FString ViewportId;
 };
diff --git a/Engine/Plugins/Runtime/nDisplay/Source/DisplayClusterMedia/Private/Synchronization/SyncPolicies/DisplayClusterMediaOutputSynchronizationPolicyEthernetBarrierBase.cpp b/Engine/Plugins/Runtime/nDisplay/Source/DisplayClusterMedia/Private/Synchronization/SyncPolicies/DisplayClusterMediaOutputSynchronizationPolicyEthernetBarrierBase.cpp
--- a/Engine/Plugins/Runtime/nDisplay/Source/DisplayClusterMedia/Private/Synchronization/SyncPolicies/DisplayClusterMediaOutputSynchronizationPolicyEthernetBarrierBase.cpp
+++ b/Engine/Plugins/Runtime/nDisplay/Source/DisplayClusterMedia/Private/Synchronization/SyncPolicies/DisplayClusterMediaOutputSynchronizationPolicyEthernetBarrierBase.cpp
@@ -225,7 +225,7 @@ void FDisplayClusterMediaOutputSynchronizationPolicyEthernetBarrierBaseHandler::
 				/////////////////////
 				// Backbuffer capture
 				{
-					const FDisplayClusterConfigurationMedia& MediaSettings = NodeIt.Value->Media;
+					const FDisplayClusterConfigurationMediaNodeBackbuffer& MediaSettings = NodeIt.Value->MediaSettings;
 
 					if (MediaSettings.bEnable)
 					{
@@ -257,7 +257,7 @@ void FDisplayClusterMediaOutputSynchronizationPolicyEthernetBarrierBaseHandler::
 					// Iterate over viewports
 					for (const TPair<FString, TObjectPtr<UDisplayClusterConfigurationViewport>>& ViewportIt : NodeIt.Value->Viewports)
 					{
-						const FDisplayClusterConfigurationMedia& MediaSettings = ViewportIt.Value->RenderSettings.Media;
+						const FDisplayClusterConfigurationMediaViewport& MediaSettings = ViewportIt.Value->RenderSettings.Media;
 
 						if (MediaSettings.bEnable)
 						{
