nDisplay: Changed the dynamic depth of field compensation from a hard coded curve to a LUT with a texture that can be provided through the ICVFX camera properties

#virtualized

diff --git a/Engine/Plugins/Runtime/nDisplay/Source/DisplayCluster/Private/Game/EngineClasses/Scene/DisplayClusterICVFXCameraComponent.cpp b/Engine/Plugins/Runtime/nDisplay/Source/DisplayCluster/Private/Game/EngineClasses/Scene/DisplayClusterICVFXCameraComponent.cpp
--- a/Engine/Plugins/Runtime/nDisplay/Source/DisplayCluster/Private/Game/EngineClasses/Scene/DisplayClusterICVFXCameraComponent.cpp
+++ b/Engine/Plugins/Runtime/nDisplay/Source/DisplayCluster/Private/Game/EngineClasses/Scene/DisplayClusterICVFXCameraComponent.cpp
@@ -226,6 +226,7 @@ FDisplayClusterViewport_CameraDepthOfField UDisplayClusterICVFXCameraComponent::
 	OutParameters.bEnableDepthOfFieldCompensation = CameraSettings.CameraDepthOfField.bEnableDepthOfFieldCompensation;
 	OutParameters.DistanceToWall = CameraSettings.CameraDepthOfField.DistanceToWall;
 	OutParameters.DistanceToWallOffset = CameraSettings.CameraDepthOfField.DistanceToWallOffset;
+	OutParameters.CompensationLUT = CameraSettings.CameraDepthOfField.DynamicCompensationLUT ? CameraSettings.CameraDepthOfField.DynamicCompensationLUT : CameraSettings.CameraDepthOfField.CompensationLUT;
 
 	return OutParameters;
 }
@@ -234,6 +235,8 @@ void UDisplayClusterICVFXCameraComponent::OnRegister()
 {
 	Super::OnRegister();
 
+	CameraSettings.CameraDepthOfField.UpdateDynamicCompensationLUT();
+
 #if WITH_EDITORONLY_DATA
 	// disable frustum for icvfx camera component
 	if (DrawFrustum != nullptr)
@@ -246,6 +249,36 @@ void UDisplayClusterICVFXCameraComponent::OnRegister()
 #endif
 }
 
+void UDisplayClusterICVFXCameraComponent::SetDepthOfFieldParameters(const FDisplayClusterConfigurationICVFX_CameraDepthOfField& NewDepthOfFieldParams)
+{
+	CameraSettings.CameraDepthOfField.bEnableDepthOfFieldCompensation = NewDepthOfFieldParams.bEnableDepthOfFieldCompensation;
+	CameraSettings.CameraDepthOfField.bAutomaticallySetDistanceToWall = NewDepthOfFieldParams.bAutomaticallySetDistanceToWall;
+	CameraSettings.CameraDepthOfField.DistanceToWallOffset = NewDepthOfFieldParams.DistanceToWallOffset;
+
+	if (!NewDepthOfFieldParams.bAutomaticallySetDistanceToWall)
+	{
+		CameraSettings.CameraDepthOfField.DistanceToWall = NewDepthOfFieldParams.DistanceToWall;
+	}
+
+	bool bGenerateNewLUT = false;
+	if (CameraSettings.CameraDepthOfField.DepthOfFieldGain != NewDepthOfFieldParams.DepthOfFieldGain)
+	{
+		CameraSettings.CameraDepthOfField.DepthOfFieldGain = NewDepthOfFieldParams.DepthOfFieldGain;
+		bGenerateNewLUT = true;
+	}
+
+	if (CameraSettings.CameraDepthOfField.CompensationLUT != NewDepthOfFieldParams.CompensationLUT)
+	{
+		CameraSettings.CameraDepthOfField.CompensationLUT = NewDepthOfFieldParams.CompensationLUT;
+		bGenerateNewLUT = true;
+	}
+
+	if (bGenerateNewLUT)
+	{
+		CameraSettings.CameraDepthOfField.UpdateDynamicCompensationLUT();
+	}
+}
+
 #if WITH_EDITORONLY_DATA
 void UDisplayClusterICVFXCameraComponent::PreEditChange(FProperty* PropertyThatWillChange)
 {
@@ -259,6 +292,13 @@ void UDisplayClusterICVFXCameraComponent::PostEditChangeProperty(FPropertyChange
 {
 	Super::PostEditChangeProperty(PropertyChangedEvent);
 
+	FName PropertyName = PropertyChangedEvent.GetPropertyName();
+	if (PropertyName == GET_MEMBER_NAME_CHECKED(FDisplayClusterConfigurationICVFX_CameraDepthOfField, CompensationLUT) ||
+		(PropertyName == GET_MEMBER_NAME_CHECKED(FDisplayClusterConfigurationICVFX_CameraDepthOfField, DepthOfFieldGain) && PropertyChangedEvent.ChangeType != EPropertyChangeType::Interactive))
+	{
+		CameraSettings.CameraDepthOfField.UpdateDynamicCompensationLUT();
+	}
+
 	UpdateICVFXPreviewState();
 }
 
diff --git a/Engine/Plugins/Runtime/nDisplay/Source/DisplayCluster/Private/Render/Viewport/Containers/DisplayClusterViewport_CameraDepthOfField.cpp b/Engine/Plugins/Runtime/nDisplay/Source/DisplayCluster/Private/Render/Viewport/Containers/DisplayClusterViewport_CameraDepthOfField.cpp
--- a/Engine/Plugins/Runtime/nDisplay/Source/DisplayCluster/Private/Render/Viewport/Containers/DisplayClusterViewport_CameraDepthOfField.cpp
+++ b/Engine/Plugins/Runtime/nDisplay/Source/DisplayCluster/Private/Render/Viewport/Containers/DisplayClusterViewport_CameraDepthOfField.cpp
@@ -2,10 +2,17 @@
 
 #include "Render/Viewport/Containers/DisplayClusterViewport_CameraDepthOfField.h"
 
+#include "Engine/Texture2D.h"
 #include "SceneView.h"
+#include "TextureResource.h"
 
 void FDisplayClusterViewport_CameraDepthOfField::SetupSceneView(FSceneView& InOutView) const
 {
 	InOutView.bEnableDynamicCocOffset = bEnableDepthOfFieldCompensation;
 	InOutView.InFocusDistance = DistanceToWall + DistanceToWallOffset;
+
+	if (CompensationLUT.IsValid())
+	{
+		InOutView.DynamicCocOffsetLUT = CompensationLUT->GetResource() ? CompensationLUT->GetResource()->GetTextureRHI() : nullptr;
+	}
 }
\ No newline at end of file
diff --git a/Engine/Plugins/Runtime/nDisplay/Source/DisplayCluster/Public/Components/DisplayClusterICVFXCameraComponent.h b/Engine/Plugins/Runtime/nDisplay/Source/DisplayCluster/Public/Components/DisplayClusterICVFXCameraComponent.h
--- a/Engine/Plugins/Runtime/nDisplay/Source/DisplayCluster/Public/Components/DisplayClusterICVFXCameraComponent.h
+++ b/Engine/Plugins/Runtime/nDisplay/Source/DisplayCluster/Public/Components/DisplayClusterICVFXCameraComponent.h
@@ -80,6 +80,10 @@ public:
 	// UActorComponent interface
 	virtual void OnRegister() override;
 
+	/** Sets new depth of field parameters and updates the dynamically generated compensation LUT if needed */
+	UFUNCTION(BlueprintCallable, Category = "In Camera VFX")
+	void SetDepthOfFieldParameters(const FDisplayClusterConfigurationICVFX_CameraDepthOfField& NewDepthOfFieldParams);
+
 private:
 	void UpdateOverscanEstimatedFrameSize();
 
diff --git a/Engine/Plugins/Runtime/nDisplay/Source/DisplayCluster/Public/Render/Viewport/Containers/DisplayClusterViewport_CameraDepthOfField.h b/Engine/Plugins/Runtime/nDisplay/Source/DisplayCluster/Public/Render/Viewport/Containers/DisplayClusterViewport_CameraDepthOfField.h
--- a/Engine/Plugins/Runtime/nDisplay/Source/DisplayCluster/Public/Render/Viewport/Containers/DisplayClusterViewport_CameraDepthOfField.h
+++ b/Engine/Plugins/Runtime/nDisplay/Source/DisplayCluster/Public/Render/Viewport/Containers/DisplayClusterViewport_CameraDepthOfField.h
@@ -3,8 +3,10 @@
 #pragma once
 
 #include "CoreMinimal.h"
+#include "UObject/WeakObjectPtr.h"
 
 class FSceneView;
+class UTexture2D;
 
 /**
 * Settings for the camera depth of field blur.
@@ -31,4 +33,7 @@ public:
 
 	/** An offset applied to the distance to wall value */
 	float DistanceToWallOffset = 0.0;
+
+	/** Look-up texture that encodes the specific amount of compensation used for each combination of wall distance and object distance */
+	TWeakObjectPtr<UTexture2D> CompensationLUT;
 };
\ No newline at end of file
diff --git a/Engine/Plugins/Runtime/nDisplay/Source/DisplayClusterConfiguration/Private/DisplayClusterConfigurationTypes_ICVFX.cpp b/Engine/Plugins/Runtime/nDisplay/Source/DisplayClusterConfiguration/Private/DisplayClusterConfigurationTypes_ICVFX.cpp
--- a/Engine/Plugins/Runtime/nDisplay/Source/DisplayClusterConfiguration/Private/DisplayClusterConfigurationTypes_ICVFX.cpp
+++ b/Engine/Plugins/Runtime/nDisplay/Source/DisplayClusterConfiguration/Private/DisplayClusterConfigurationTypes_ICVFX.cpp
@@ -581,6 +581,40 @@ bool FDisplayClusterConfigurationICVFX_ChromakeyRenderSettings::ShouldUseChromak
 	return ShowOnlyList.IsVisibilityListValid();
 }
 
+void FDisplayClusterConfigurationICVFX_CameraDepthOfField::UpdateDynamicCompensationLUT()
+{
+	if (CompensationLUT)
+	{
+		FSharedImageConstRef CPUTextureRef = CompensationLUT->GetCPUCopy();
+		if (CPUTextureRef.IsValid())
+		{
+			TArray64<uint8> Pixels(CPUTextureRef->RawData);
+			for (int32 Index = 0; Index < Pixels.Num(); ++Index)
+			{
+				// Scale the offset encoded in the LUT so that the final CoC when computed in the DoF pipeline is scaled by the gain.
+				// The actual new offset needed to accomplish this comes from the following equation:
+				// c * (CoC_obj + CoC_off) = CoC_obj + newOffset =>
+				// newOffset = (1 - c) * CoC_obj + c * CoC_off
+				const int32 ObjectCoC = Index % CPUTextureRef->GetHeight();
+				const int32 Offset = Pixels[Index] - 128;
+				const int32 ScaledOffset = (1 - DepthOfFieldGain) * ObjectCoC + DepthOfFieldGain * Offset;
+				Pixels[Index] = FMath::Clamp(ScaledOffset + 128, 0, 255);
+			}
+
+			// Texture format is assumed to be greyscale (PF_G8), and we must disable sRGB on the texture to ensure the raw byte value, which encodes
+			// the offset in pixels, is passed unmodified to the depth of field shader
+			if (UTexture2D* NewTexture = UTexture2D::CreateTransient(CPUTextureRef->GetWidth(), CPUTextureRef->GetHeight(), PF_G8, NAME_None, Pixels))
+			{
+				DynamicCompensationLUT = NewTexture;
+				DynamicCompensationLUT->SRGB = 0;
+				return;
+			}
+		}
+	}
+
+	DynamicCompensationLUT = nullptr;
+}
+
 void FDisplayClusterConfigurationICVFX_CameraCustomFrustum::SetupViewInfo(const FDisplayClusterConfigurationICVFX_StageSettings& InStageSettings, FMinimalViewInfo& InOutViewInfo) const
 {
 	// Since Circle of confusion is directly proportional to aperature, with wider FOV focal length needs to be shortened by the same amount as FOV.
diff --git a/Engine/Plugins/Runtime/nDisplay/Source/DisplayClusterConfiguration/Public/DisplayClusterConfigurationTypes_ICVFX.h b/Engine/Plugins/Runtime/nDisplay/Source/DisplayClusterConfiguration/Public/DisplayClusterConfigurationTypes_ICVFX.h
--- a/Engine/Plugins/Runtime/nDisplay/Source/DisplayClusterConfiguration/Public/DisplayClusterConfigurationTypes_ICVFX.h
+++ b/Engine/Plugins/Runtime/nDisplay/Source/DisplayClusterConfiguration/Public/DisplayClusterConfigurationTypes_ICVFX.h
@@ -605,6 +605,22 @@ public:
 	/** An offset applied to DistanceToWall (applied regardless of whether DistanceToWall is automatically set) */
 	UPROPERTY(BlueprintReadWrite, EditAnywhere, Category = NDisplay)
 	float DistanceToWallOffset = 0.0;
+
+	/** A gain factor that scales the amount of depth of field blur rendered on the wall */
+	UPROPERTY(BlueprintReadWrite, EditAnywhere, Category = NDisplay, meta = (UIMin=0.0, UIMax=4.0, ClampMin=0.0, ClampMax=4.0))
+	float DepthOfFieldGain = 1.0;
+
+	/** Look-up texture that encodes the specific amount of compensation used for each combination of wall distance and object distance */
+	UPROPERTY(BlueprintReadWrite, EditAnywhere, Category = NDisplay)
+	UTexture2D* CompensationLUT = nullptr;
+
+	/** Actual LUT to use with the depth of field pipeline, copied and modified from CompensationLUT */
+	UPROPERTY(Transient)
+	UTexture2D* DynamicCompensationLUT = nullptr;
+
+public:
+	/** Processes the compensation LUT by adding any needed DoF gain and writes the result to the dynamic compensation LUT texture */
+	void UpdateDynamicCompensationLUT();
 };
 
 USTRUCT(BlueprintType)
diff --git a/Engine/Shaders/Private/DiaphragmDOF/DOFCommon.ush b/Engine/Shaders/Private/DiaphragmDOF/DOFCommon.ush
--- a/Engine/Shaders/Private/DiaphragmDOF/DOFCommon.ush
+++ b/Engine/Shaders/Private/DiaphragmDOF/DOFCommon.ush
@@ -159,6 +159,9 @@ float CocInvSqueeze;
 float DepthBlurRadius;
 float DepthBlurExponent;
 
+Texture2D		DynamicRadiusOffsetLUT;
+SamplerState	DynamicRadiusOffsetLUTSampler;
+float2			DynamicRadiusOffsetLUTInvSize;
 
 //------------------------------------------------------- DEBUG OUTPUT
 
@@ -175,12 +178,19 @@ float GetCocOffset(float CocRadius)
 	float DynamicOffset = 0.0;
 	BRANCH if (bCocEnableDynamicRadiusOffset)
 	{
-		// This curve minimizes the RMS error between a point that was blurred once by a disk blur and one
-		// that was blurred twice, where CocRadius + DynamicOffset is the size of the first blur and
-		// CocInFocusRadius in the size of the second blur, 
-		float B = 0.467743 + 7.89656 * pow(CocInFocusRadius, -1.89051);
-		float V = -(2.57186 + 0.142159 * CocInFocusRadius);
-		DynamicOffset = -CocInFocusRadius * (1 - pow(1 + exp(B * (CocInFocusRadius - CocRadius)), V));
+		// By convention, the LUT is a 1-to-1 map from CoC size in pixels to CoC offset, so scale the CoC radii by the texture
+		// size to get the texel coordinate. If this assumption ever changes, we will need to scale this value
+		float2 TexCoords = float2(CocRadius, CocInFocusRadius) * DynamicRadiusOffsetLUTInvSize;
+
+		// Assume that the LUT is homogeneous, which means that the value of any pixel at s*uv is equal to s times the value at uv
+		// This allows us to compute what a pixel should be outside the bounds of the LUT by finding the largest integer
+		// scalar s such that uv is between 0 and 1, and then scaling the sampled pixel by that much
+		float WrapCoefficient = ceil(max(TexCoords.x, TexCoords.y));
+
+		// Also by convention, a value of 128 is used as the zero offset value, allowing support for positive and negative offsets, so scale by 255
+		// to put the pixel into byte space, then offset by 128 to get the final value. Also, enforce a maximum offset to be the initial CoC radius;
+		// the offset should only ever make the final CoC between the initial CoC and 0
+		DynamicOffset = -min((DynamicRadiusOffsetLUT.Sample(DynamicRadiusOffsetLUTSampler, TexCoords / WrapCoefficient).r * 255 - 128) * WrapCoefficient, CocRadius);
 	}
 
 	return DynamicOffset;
diff --git a/Engine/Source/Runtime/Engine/Public/SceneView.h b/Engine/Source/Runtime/Engine/Public/SceneView.h
--- a/Engine/Source/Runtime/Engine/Public/SceneView.h
+++ b/Engine/Source/Runtime/Engine/Public/SceneView.h
@@ -1593,6 +1593,9 @@ public:
 	/** When dynamic CoC offset is enabled, this is the distance from camera at which objects will be in perfect focus (when the CoC offset is maximum so that the final CoC is 0) */
 	float InFocusDistance = 0.0;
 
+	/** Lookup table for the dynamic CoC offset */
+	FTextureRHIRef DynamicCocOffsetLUT;
+
 	/** Feature level for this scene */
 	const ERHIFeatureLevel::Type FeatureLevel;
 
diff --git a/Engine/Source/Runtime/Renderer/Private/PostProcess/DiaphragmDOF.cpp b/Engine/Source/Runtime/Renderer/Private/PostProcess/DiaphragmDOF.cpp
--- a/Engine/Source/Runtime/Renderer/Private/PostProcess/DiaphragmDOF.cpp
+++ b/Engine/Source/Runtime/Renderer/Private/PostProcess/DiaphragmDOF.cpp
@@ -638,6 +638,9 @@ END_SHADER_PARAMETER_STRUCT()
 BEGIN_SHADER_PARAMETER_STRUCT(FDOFCocModelShaderParameters, )
 	SHADER_PARAMETER(float, CocInfinityRadius)
 	SHADER_PARAMETER(float, CocInFocusRadius)
+	SHADER_PARAMETER_RDG_TEXTURE(Texture2D, DynamicRadiusOffsetLUT)
+	SHADER_PARAMETER_SAMPLER(SamplerState, DynamicRadiusOffsetLUTSampler)
+	SHADER_PARAMETER(FVector2f, DynamicRadiusOffsetLUTInvSize)
 	SHADER_PARAMETER(uint32, bCocEnableDynamicRadiusOffset)
 	SHADER_PARAMETER(float, CocMinRadius)
 	SHADER_PARAMETER(float, CocMaxRadius)
@@ -648,6 +651,7 @@ BEGIN_SHADER_PARAMETER_STRUCT(FDOFCocModelShaderParameters, )
 END_SHADER_PARAMETER_STRUCT()
 
 void SetCocModelParameters(
+	FRDGBuilder& GraphBuilder,
 	FDOFCocModelShaderParameters* OutParameters,
 	const DiaphragmDOF::FPhysicalCocModel& CocModel,
 	float CocRadiusBasis = 1.0f)
@@ -661,6 +665,23 @@ void SetCocModelParameters(
 	OutParameters->CocInvSqueeze = 1.0f / CocModel.Squeeze;
 	OutParameters->DepthBlurRadius = CocRadiusBasis * CocModel.MaxDepthBlurRadius;
 	OutParameters->DepthBlurExponent = CocModel.DepthBlurExponent;
+
+	if (CocModel.DynamicRadiusOffsetLUT)
+	{
+		OutParameters->DynamicRadiusOffsetLUT = GraphBuilder.RegisterExternalTexture(CreateRenderTarget(CocModel.DynamicRadiusOffsetLUT, TEXT("DynamicRadiusOffsetLUT")));
+		OutParameters->DynamicRadiusOffsetLUTInvSize.X = 1.0f / CocModel.DynamicRadiusOffsetLUT->GetDesc().Extent.X;
+		OutParameters->DynamicRadiusOffsetLUTInvSize.Y = 1.0f / CocModel.DynamicRadiusOffsetLUT->GetDesc().Extent.Y;
+	}
+	else
+	{
+		const FRDGSystemTextures& SystemTextures = FRDGSystemTextures::Get(GraphBuilder);
+
+		OutParameters->DynamicRadiusOffsetLUT = SystemTextures.Black;
+		OutParameters->DynamicRadiusOffsetLUTInvSize.X = 1.0f;
+		OutParameters->DynamicRadiusOffsetLUTInvSize.Y = 1.0f;
+	}
+
+	OutParameters->DynamicRadiusOffsetLUTSampler = TStaticSamplerState<SF_Bilinear>::GetRHI();
 }
 
 
@@ -1600,7 +1621,7 @@ FRDGTextureRef DiaphragmDOF::AddPasses(
 		FDiaphragmDOFSetupCS::FParameters* PassParameters = GraphBuilder.AllocParameters<FDiaphragmDOFSetupCS::FParameters>();
 		{
 			PassParameters->CommonParameters = CommonParameters;
-			SetCocModelParameters(&PassParameters->CocModel, CocModel, EncodedCocRadiusBasis);
+			SetCocModelParameters(GraphBuilder, &PassParameters->CocModel, CocModel, EncodedCocRadiusBasis);
 			PassParameters->ViewportRect = FIntRect(FIntPoint::ZeroValue, PassViewSize);
 			PassParameters->SceneColorTexture = InputSceneColor;
 			PassParameters->SceneDepthTexture = SceneTextures.SceneDepthTexture;
@@ -2672,7 +2693,7 @@ FRDGTextureRef DiaphragmDOF::AddPasses(
 
 		FDiaphragmDOFRecombineCS::FParameters* PassParameters = GraphBuilder.AllocParameters<FDiaphragmDOFRecombineCS::FParameters>();
 		PassParameters->CommonParameters = CommonParameters;
-		SetCocModelParameters(&PassParameters->CocModel, CocModel, /* CocRadiusBasis = */ float(GatheringViewSize.X));
+		SetCocModelParameters(GraphBuilder, &PassParameters->CocModel, CocModel, /* CocRadiusBasis = */ float(GatheringViewSize.X));
 
 		PassParameters->ViewportRect = PassViewRect;
 		PassParameters->ViewportSize = FVector4f(PassViewRect.Width(), PassViewRect.Height(), 1.0f / PassViewRect.Width(), 1.0f / PassViewRect.Height());
diff --git a/Engine/Source/Runtime/Renderer/Private/PostProcess/DiaphragmDOF.h b/Engine/Source/Runtime/Renderer/Private/PostProcess/DiaphragmDOF.h
--- a/Engine/Source/Runtime/Renderer/Private/PostProcess/DiaphragmDOF.h
+++ b/Engine/Source/Runtime/Renderer/Private/PostProcess/DiaphragmDOF.h
@@ -37,6 +37,9 @@ struct FPhysicalCocModel
 	/** When dynamic offset is enabled, this is the coc radius at which objects will be perfectly sharp */
 	float InFocusRadius;
 
+	/** Radius offset lookup table */
+	FRHITexture2D* DynamicRadiusOffsetLUT;
+
 	// Resolution less minimal foreground coc radius < 0.
 	float MinForegroundCocRadius;
 
diff --git a/Engine/Source/Runtime/Renderer/Private/PostProcess/DiaphragmDOFUtils.cpp b/Engine/Source/Runtime/Renderer/Private/PostProcess/DiaphragmDOFUtils.cpp
--- a/Engine/Source/Runtime/Renderer/Private/PostProcess/DiaphragmDOFUtils.cpp
+++ b/Engine/Source/Runtime/Renderer/Private/PostProcess/DiaphragmDOFUtils.cpp
@@ -147,11 +147,13 @@ void DiaphragmDOF::FPhysicalCocModel::Compile(const FViewInfo& View)
 			// so clamp the in focus radius to always be positive
 			InFocusRadius = FMath::Max(InfinityBackgroundCocRadius * (1 - FocusDistance / View.InFocusDistance), 0.0f);
 			bEnableDynamicOffset = View.bEnableDynamicCocOffset;
+			DynamicRadiusOffsetLUT = View.bEnableDynamicCocOffset ? View.DynamicCocOffsetLUT : nullptr;
 		}
 		else
 		{
 			InFocusRadius = 0.0;
 			bEnableDynamicOffset = false;
+			DynamicRadiusOffsetLUT = nullptr;
 		}
 	}
 	else
@@ -160,6 +162,7 @@ void DiaphragmDOF::FPhysicalCocModel::Compile(const FViewInfo& View)
 		MinForegroundCocRadius = 0.0;
 		InFocusRadius = 0.0;
 		bEnableDynamicOffset = false;
+		DynamicRadiusOffsetLUT = nullptr;
 	}
 }
 
