[nDisplay] Pick a small media tiling fix from R29 stream

diff --git a/Engine/Plugins/Runtime/nDisplay/Source/DisplayClusterMedia/Private/Capture/DisplayClusterMediaCaptureViewport.cpp b/Engine/Plugins/Runtime/nDisplay/Source/DisplayClusterMedia/Private/Capture/DisplayClusterMediaCaptureViewport.cpp
--- a/Engine/Plugins/Runtime/nDisplay/Source/DisplayClusterMedia/Private/Capture/DisplayClusterMediaCaptureViewport.cpp
+++ b/Engine/Plugins/Runtime/nDisplay/Source/DisplayClusterMedia/Private/Capture/DisplayClusterMediaCaptureViewport.cpp
@@ -58,8 +58,11 @@ void FDisplayClusterMediaCaptureViewport::OnUpdateViewportMediaState(IDisplayClu
 	// In the future, after the media redesign, the DCRA name will also need to be checked here.
 	if (InViewport && InViewport->GetId().Equals(GetViewportId(), ESearchCase::IgnoreCase))
 	{
+		// Reset inactive state explicitly
+		InOutMediaState &= ~(EDisplayClusterViewportMediaState::Inactive);
+
 		// Raise flags that this viewport will be captured by media.
-		InOutMediaState = EDisplayClusterViewportMediaState::Capture;
+		InOutMediaState |= EDisplayClusterViewportMediaState::Capture;
 
 		if (bForceLateOCIOPass)
 		{
diff --git a/Engine/Plugins/Runtime/nDisplay/Source/DisplayClusterMedia/Private/Input/DisplayClusterMediaInputViewport.cpp b/Engine/Plugins/Runtime/nDisplay/Source/DisplayClusterMedia/Private/Input/DisplayClusterMediaInputViewport.cpp
--- a/Engine/Plugins/Runtime/nDisplay/Source/DisplayClusterMedia/Private/Input/DisplayClusterMediaInputViewport.cpp
+++ b/Engine/Plugins/Runtime/nDisplay/Source/DisplayClusterMedia/Private/Input/DisplayClusterMediaInputViewport.cpp
@@ -52,8 +52,11 @@ void FDisplayClusterMediaInputViewport::OnUpdateViewportMediaState(IDisplayClust
 	// In the future, after the media redesign, the DCRA name will also need to be checked here.
 	if (InViewport && InViewport->GetId().Equals(GetViewportId(), ESearchCase::IgnoreCase))
 	{
+		// Reset inactive state explicitly
+		InOutMediaState &= ~(EDisplayClusterViewportMediaState::Inactive);
+
 		// Raise flags that this viewport texture will be overridden by media.
-		InOutMediaState = EDisplayClusterViewportMediaState::Input;
+		InOutMediaState |= EDisplayClusterViewportMediaState::Input;
 
 		if (bForceLateOCIOPass)
 		{
