[nDisplay] Fix for the media tiling issue that prevented internal optimization for non-tiled viewports

diff --git a/Engine/Plugins/Runtime/nDisplay/Source/DisplayCluster/Private/Render/Viewport/Configuration/DisplayClusterViewportConfigurationHelpers.cpp b/Engine/Plugins/Runtime/nDisplay/Source/DisplayCluster/Private/Render/Viewport/Configuration/DisplayClusterViewportConfigurationHelpers.cpp
--- a/Engine/Plugins/Runtime/nDisplay/Source/DisplayCluster/Private/Render/Viewport/Configuration/DisplayClusterViewportConfigurationHelpers.cpp
+++ b/Engine/Plugins/Runtime/nDisplay/Source/DisplayCluster/Private/Render/Viewport/Configuration/DisplayClusterViewportConfigurationHelpers.cpp
@@ -18,6 +18,7 @@
 #include "DisplayClusterConfigurationTypes_PostRender.h"
 
 #include "IDisplayCluster.h"
+#include "IDisplayClusterCallbacks.h"
 #include "Cluster/IDisplayClusterClusterManager.h"
 
 #include "ShaderParameters/DisplayClusterShaderParameters_PostprocessBlur.h"
@@ -93,6 +94,18 @@ void FDisplayClusterViewportConfigurationHelpers::UpdateBaseViewportSetting(FDis
 	// Reset runtime flags from prev frame:
 	DstViewport.ResetRuntimeParameters();
 
+	// Allow external customers to configure media state
+	{
+		// By default, we set 'None' so it can be rendered from media POV
+		EDisplayClusterViewportMediaState NewMediaStates = EDisplayClusterViewportMediaState::None;
+
+		// Now allow to override media state if anyone wants
+		IDisplayCluster::Get().GetCallbacks().OnDisplayClusterUpdateViewportMediaState().Broadcast(&DstViewport, NewMediaStates);
+
+		// Update the media state for the new frame.
+		InOutRenderSettings.AssignMediaStates(NewMediaStates);
+	}
+
 	// UDisplayClusterConfigurationViewport
 	{
 		if (InConfigurationViewport.bAllowRendering == false)
diff --git a/Engine/Plugins/Runtime/nDisplay/Source/DisplayCluster/Private/Render/Viewport/Configuration/DisplayClusterViewportConfigurationHelpers_ICVFX.cpp b/Engine/Plugins/Runtime/nDisplay/Source/DisplayCluster/Private/Render/Viewport/Configuration/DisplayClusterViewportConfigurationHelpers_ICVFX.cpp
--- a/Engine/Plugins/Runtime/nDisplay/Source/DisplayCluster/Private/Render/Viewport/Configuration/DisplayClusterViewportConfigurationHelpers_ICVFX.cpp
+++ b/Engine/Plugins/Runtime/nDisplay/Source/DisplayCluster/Private/Render/Viewport/Configuration/DisplayClusterViewportConfigurationHelpers_ICVFX.cpp
@@ -22,6 +22,7 @@
 #include "DisplayClusterConfigurationTypes_PostRender.h"
 
 #include "IDisplayCluster.h"
+#include "IDisplayClusterCallbacks.h"
 #include "Cluster/IDisplayClusterClusterManager.h"
 
 #include "ShaderParameters/DisplayClusterShaderParameters_PostprocessBlur.h"
@@ -487,6 +488,18 @@ void FDisplayClusterViewportConfigurationHelpers_ICVFX::UpdateCameraViewportSett
 	// Reset runtime flags from prev frame:
 	DstViewport.ResetRuntimeParameters();
 
+	// Allow external customers to configure media state
+	{
+		// By default, we set 'None' so it can be rendered from media POV
+		EDisplayClusterViewportMediaState NewMediaStates = EDisplayClusterViewportMediaState::None;
+
+		// Now allow to override media state if anyone wants
+		IDisplayCluster::Get().GetCallbacks().OnDisplayClusterUpdateViewportMediaState().Broadcast(&DstViewport, NewMediaStates);
+
+		// Update the media state for the new frame.
+		InOutRenderSettings.AssignMediaStates(NewMediaStates);
+	}
+
 	// incamera textrure used as overlay
 	InOutRenderSettings.bVisible = false;
 
