[nDisplay] Reimplemented recent fixes for camera split (tiles rendering control)

- Based on Peter's CL (rolled back recent fixes based on 'Inactive' media state)
- Added list of nodes to allow render unbound tiles
- Extended IsRenderEnabled() rules to match various media configuration topologies

diff --git a/Engine/Plugins/Runtime/nDisplay/Source/DisplayCluster/Private/Render/Viewport/Configuration/DisplayClusterViewportConfigurationHelpers.cpp b/Engine/Plugins/Runtime/nDisplay/Source/DisplayCluster/Private/Render/Viewport/Configuration/DisplayClusterViewportConfigurationHelpers.cpp
--- a/Engine/Plugins/Runtime/nDisplay/Source/DisplayCluster/Private/Render/Viewport/Configuration/DisplayClusterViewportConfigurationHelpers.cpp
+++ b/Engine/Plugins/Runtime/nDisplay/Source/DisplayCluster/Private/Render/Viewport/Configuration/DisplayClusterViewportConfigurationHelpers.cpp
@@ -18,7 +18,6 @@
 #include "DisplayClusterConfigurationTypes_PostRender.h"
 
 #include "IDisplayCluster.h"
-#include "IDisplayClusterCallbacks.h"
 #include "Cluster/IDisplayClusterClusterManager.h"
 
 #include "ShaderParameters/DisplayClusterShaderParameters_PostprocessBlur.h"
@@ -94,18 +93,6 @@ void FDisplayClusterViewportConfigurationHelpers::UpdateBaseViewportSetting(FDis
 	// Reset runtime flags from prev frame:
 	DstViewport.ResetRuntimeParameters();
 
-	// Allow external customers to configure media state
-	{
-		// By default, we set 'None' so it can be rendered from media POV
-		EDisplayClusterViewportMediaState NewMediaStates = EDisplayClusterViewportMediaState::None;
-
-		// Now allow to override media state if anyone wants
-		IDisplayCluster::Get().GetCallbacks().OnDisplayClusterUpdateViewportMediaState().Broadcast(&DstViewport, NewMediaStates);
-
-		// Update the media state for the new frame.
-		InOutRenderSettings.AssignMediaStates(NewMediaStates);
-	}
-
 	// UDisplayClusterConfigurationViewport
 	{
 		if (InConfigurationViewport.bAllowRendering == false)
diff --git a/Engine/Plugins/Runtime/nDisplay/Source/DisplayCluster/Private/Render/Viewport/Configuration/DisplayClusterViewportConfigurationHelpers_ICVFX.cpp b/Engine/Plugins/Runtime/nDisplay/Source/DisplayCluster/Private/Render/Viewport/Configuration/DisplayClusterViewportConfigurationHelpers_ICVFX.cpp
--- a/Engine/Plugins/Runtime/nDisplay/Source/DisplayCluster/Private/Render/Viewport/Configuration/DisplayClusterViewportConfigurationHelpers_ICVFX.cpp
+++ b/Engine/Plugins/Runtime/nDisplay/Source/DisplayCluster/Private/Render/Viewport/Configuration/DisplayClusterViewportConfigurationHelpers_ICVFX.cpp
@@ -22,7 +22,6 @@
 #include "DisplayClusterConfigurationTypes_PostRender.h"
 
 #include "IDisplayCluster.h"
-#include "IDisplayClusterCallbacks.h"
 #include "Cluster/IDisplayClusterClusterManager.h"
 
 #include "ShaderParameters/DisplayClusterShaderParameters_PostprocessBlur.h"
@@ -488,18 +487,6 @@ void FDisplayClusterViewportConfigurationHelpers_ICVFX::UpdateCameraViewportSett
 	// Reset runtime flags from prev frame:
 	DstViewport.ResetRuntimeParameters();
 
-	// Allow external customers to configure media state
-	{
-		// By default, we set 'None' so it can be rendered from media POV
-		EDisplayClusterViewportMediaState NewMediaStates = EDisplayClusterViewportMediaState::None;
-
-		// Now allow to override media state if anyone wants
-		IDisplayCluster::Get().GetCallbacks().OnDisplayClusterUpdateViewportMediaState().Broadcast(&DstViewport, NewMediaStates);
-
-		// Update the media state for the new frame.
-		InOutRenderSettings.AssignMediaStates(NewMediaStates);
-	}
-
 	// incamera textrure used as overlay
 	InOutRenderSettings.bVisible = false;
 
diff --git a/Engine/Plugins/Runtime/nDisplay/Source/DisplayCluster/Private/Render/Viewport/Configuration/DisplayClusterViewportConfigurationHelpers_Tile.cpp b/Engine/Plugins/Runtime/nDisplay/Source/DisplayCluster/Private/Render/Viewport/Configuration/DisplayClusterViewportConfigurationHelpers_Tile.cpp
--- a/Engine/Plugins/Runtime/nDisplay/Source/DisplayCluster/Private/Render/Viewport/Configuration/DisplayClusterViewportConfigurationHelpers_Tile.cpp
+++ b/Engine/Plugins/Runtime/nDisplay/Source/DisplayCluster/Private/Render/Viewport/Configuration/DisplayClusterViewportConfigurationHelpers_Tile.cpp
@@ -17,10 +17,6 @@
 #include "DisplayClusterConfigurationTypes_Tile.h"
 #include "DisplayClusterProjectionStrings.h"
 
-#include "IDisplayCluster.h"
-#include "IDisplayClusterCallbacks.h"
-
-
 ////////////////////////////////////////////////////////////////////////
 // FDisplayClusterViewportConfigurationHelpers_Tile
 ////////////////////////////////////////////////////////////////////////
@@ -36,9 +32,20 @@ void FDisplayClusterViewportConfigurationHelpers_Tile::UpdateICVFXCameraViewport
 				const FDisplayClusterViewport_OverscanSettings& OverscanSettings = FDisplayClusterViewportConfigurationHelpers_Tile::GetTileOverscanSettings(InCameraTile.TileOverscan);
 				const FIntPoint Size(InCameraTile.TileSettings.TileX, InCameraTile.TileSettings.TileY);
 
+				// Find if this cluster node is allowed to render unbound tiles
+				const FString ThisNodeId = InSourceViewport.GetClusterNodeId();
+				const bool bAllowRenderUnbound = InCameraTile.ClusterNodesToRenderUnboundTiles.ItemNames.ContainsByPredicate([&ThisNodeId](const FString& Item)
+					{
+						return ThisNodeId.Equals(Item, ESearchCase::IgnoreCase);
+					});
+
+				// Prepare tile flags
+				EDisplayClusterViewportTileFlags TileFlags = EDisplayClusterViewportTileFlags::None;
+				TileFlags |= (bAllowRenderUnbound ? EDisplayClusterViewportTileFlags::AllowUnboundRender : EDisplayClusterViewportTileFlags::None);
+
 				// Set this viewport as the source for tile rendering.
 				FDisplayClusterViewport_TileSettings& OutTileSettings = InSourceViewport.GetRenderSettingsImpl().TileSettings;
-				OutTileSettings = FDisplayClusterViewport_TileSettings(Size, OverscanSettings);
+				OutTileSettings = FDisplayClusterViewport_TileSettings(Size, OverscanSettings, TileFlags);
 				OutTileSettings.bOptimizeTileOverscan = InCameraTile.TileOverscan.bOptimizeTileOverscan;
 
 				return;
@@ -151,13 +158,23 @@ FDisplayClusterViewport* FDisplayClusterViewportConfigurationHelpers_Tile::GetOr
 		// Gain direct access to internal resources of the NewViewport:
 		FDisplayClusterViewport_RenderSettings& InOutRenderSettings = TileViewport->GetRenderSettingsImpl();
 
-		// Reset runtime flags from prev frame:
+		// Reset runtime flags from prev frame.
+		// Also this function update media states.
 		TileViewport->ResetRuntimeParameters();
 
-		// Copy all the settings from the source viewport, but some of them still need to be overridden.
-		InOutRenderSettings = SourceRenderSettings;
+		// Override rendering settings from the source viewport, except for media states.
+		{
+			// Save the new media states after calling ResetRuntimeParameters().
+			const EDisplayClusterViewportMediaState TileMediaStates = TileViewport->GetRenderSettings().GetMediaStates();
+
+			// Copy all the settings from the source viewport, but some of them still need to be overridden.
+			InOutRenderSettings = SourceRenderSettings;
+
+			// Restore media states
+			InOutRenderSettings.AssignMediaStates(TileMediaStates);
+		}
 
-		// Dont show Tile composing viewports on frame target
+		// Don't show Tile composing viewports on frame target
 		InOutRenderSettings.bVisible = false;
 
 		// Disable custom frustum settings and override overscan settings.
@@ -186,29 +203,11 @@ FDisplayClusterViewport* FDisplayClusterViewportConfigurationHelpers_Tile::GetOr
 			}
 		}
 
-		// Setup as tile.
-		InOutRenderSettings.TileSettings = FDisplayClusterViewport_TileSettings(InSourceViewport.GetId(), InTilePos, InTileSize);
+		// Inherit tile flags from the source viewport
+		const EDisplayClusterViewportTileFlags TileFlags = InSourceViewport.GetRenderSettings().TileSettings.GetTileFlags();
 
-		// Allow external customers to configure media state
-		{
-			// By default, we set 'None' so the tiles can be rendered in editor for camera preview.
-			EDisplayClusterViewportMediaState NewMediaStates = EDisplayClusterViewportMediaState::None;
-
-			// But in cluster mode, we set 'Inactive' by default. When tiling is used, the tile viewport must either be rendered for media output
-			// or preserve internal buffer for a media input texture. However, it's possible the tile viewport has wrong or completely
-			// missing media output settings, or corresponding media device was not able to start properly for some reason.
-			// In this case the tile should not be rendered at all.
-			if (GDisplayCluster->GetOperationMode() == EDisplayClusterOperationMode::Cluster)
-			{
-				NewMediaStates = EDisplayClusterViewportMediaState::Inactive;
-			}
-			
-			// Now allow to override media state if anyone wants
-			IDisplayCluster::Get().GetCallbacks().OnDisplayClusterUpdateViewportMediaState().Broadcast(TileViewport, NewMediaStates);
-
-			// Update the media state for the new frame.
-			InOutRenderSettings.AssignMediaStates(NewMediaStates);
-		}
+		// Setup as tile.
+		InOutRenderSettings.TileSettings = FDisplayClusterViewport_TileSettings(InSourceViewport.GetId(), InTilePos, InTileSize, TileFlags);
 
 		// Copy internal render settings from the source:
 		TileViewport->GetVisibilitySettingsImpl() = InSourceViewport.GetVisibilitySettingsImpl();
diff --git a/Engine/Plugins/Runtime/nDisplay/Source/DisplayCluster/Private/Render/Viewport/DisplayClusterViewport_Rules.cpp b/Engine/Plugins/Runtime/nDisplay/Source/DisplayCluster/Private/Render/Viewport/DisplayClusterViewport_Rules.cpp
--- a/Engine/Plugins/Runtime/nDisplay/Source/DisplayCluster/Private/Render/Viewport/DisplayClusterViewport_Rules.cpp
+++ b/Engine/Plugins/Runtime/nDisplay/Source/DisplayCluster/Private/Render/Viewport/DisplayClusterViewport_Rules.cpp
@@ -26,6 +26,9 @@
 
 #include "Render/DisplayDevice/Components/DisplayClusterDisplayDeviceBaseComponent.h"
 
+#include "IDisplayCluster.h"
+#include "IDisplayClusterCallbacks.h"
+
 #include "EngineUtils.h"
 #include "SceneManagement.h"
 #include "SceneView.h"
@@ -35,7 +38,11 @@
 
 #include "LegacyScreenPercentageDriver.h"
 
+#include "Misc/CommandLine.h"
+#include "Misc/DisplayClusterGlobals.h"
 #include "Misc/DisplayClusterLog.h"
+#include "Misc/Parse.h"
+
 
 ///////////////////////////////////////////////////////////////////////////////////////
 // FDisplayClusterViewport
@@ -52,6 +59,12 @@ void FDisplayClusterViewport::ResetRuntimeParameters()
 
 	OverscanRuntimeSettings = FDisplayClusterViewport_OverscanRuntimeSettings();
 	CustomFrustumRuntimeSettings = FDisplayClusterViewport_CustomFrustumRuntimeSettings();
+
+	// Obtain viewport media state from external multicast delegates (This viewport can be used by multiple media).
+	EDisplayClusterViewportMediaState AllMediaStates = EDisplayClusterViewportMediaState::None;
+	IDisplayCluster::Get().GetCallbacks().OnDisplayClusterUpdateViewportMediaState().Broadcast(this, AllMediaStates);
+	// Update the media state for the new frame.
+	RenderSettings.AssignMediaStates(AllMediaStates);
 }
 
 bool FDisplayClusterViewport::IsInternalViewport() const
@@ -97,6 +110,7 @@ bool FDisplayClusterViewport::IsExternalRendering() const
 		return true;
 	}
 
+	// Do not render the viewport if it is used as a media input.
 	if (RenderSettings.HasAnyMediaStates(EDisplayClusterViewportMediaState::Input))
 	{
 		// This viewport is not rendered but gets the image from the media.
@@ -121,17 +135,56 @@ bool FDisplayClusterViewport::IsRenderEnabled() const
 		return true;
 	}
 
-	if (RenderSettings.TileSettings.GetType() == EDisplayClusterViewportTileType::Source)
+	// Tiles only evaluation
+	const EDisplayClusterViewportTileType TileType = RenderSettings.TileSettings.GetType();
+	if (TileType == EDisplayClusterViewportTileType::Source)
 	{
-		// The source viewport is not rendered by itself, because the rendering is done on tiles.
+		// SRC viewports should never get rendered
 		return false;
 	}
-
-	if (RenderSettings.HasAnyMediaStates(EDisplayClusterViewportMediaState::Inactive))
+	else if (TileType == EDisplayClusterViewportTileType::UnusedTile)
 	{
-		// Don't render inactive tiles
+		// UnusedTile is an internal state. We should never get it here.
 		return false;
 	}
+	else if (TileType == EDisplayClusterViewportTileType::Tile)
+	{
+		static const bool bIsCluster = (GDisplayCluster->GetOperationMode() == EDisplayClusterOperationMode::Cluster);
+		static const bool bIsOffscreen = FParse::Param(FCommandLine::Get(), TEXT("RenderOffscreen"));
+
+		const bool bHasInputAssigned  = RenderSettings.HasAnyMediaStates(EDisplayClusterViewportMediaState::Input);
+		const bool bHasOutputAssigned = RenderSettings.HasAnyMediaStates(EDisplayClusterViewportMediaState::Capture);
+
+		// In non-cluster modes, always render (i.e. preview in editor)
+		if (!bIsCluster)
+		{
+			return true;
+		}
+
+		// Tiles logic is a little more complicated
+		//
+		// In | Out | Render (offscreen) | Render (onscreen)
+		// ------------------------------------------------------------------
+		//  0 |  0  |        0           | bShouldRenderUnbound
+		//  0 |  1  |        1           | 1
+		//  1 |  0  |        0           | 0
+		//  1 |  1  |        0           | 0
+
+		if (bHasInputAssigned)
+		{
+			// This should have been validated previously in IsExternalRendering() but let's keep this condition
+			// active in case something is changed outside (i.e. validation order or condition changes).
+			return false;
+		}
+		else if (!bHasInputAssigned && !bHasOutputAssigned)
+		{
+			// Check if current node was explicitly allowed to render unbound tiles
+			const bool bShouldRenderUnbound = RenderSettings.TileSettings.HasAnyTileFlags(EDisplayClusterViewportTileFlags::AllowUnboundRender);
+			const bool bShouldRender = !bIsOffscreen && bShouldRenderUnbound;
+
+			return bShouldRender;
+		}
+	}
 
 	return true;
 }
diff --git a/Engine/Plugins/Runtime/nDisplay/Source/DisplayCluster/Public/Components/DisplayClusterICVFXCameraComponent.h b/Engine/Plugins/Runtime/nDisplay/Source/DisplayCluster/Public/Components/DisplayClusterICVFXCameraComponent.h
--- a/Engine/Plugins/Runtime/nDisplay/Source/DisplayCluster/Public/Components/DisplayClusterICVFXCameraComponent.h
+++ b/Engine/Plugins/Runtime/nDisplay/Source/DisplayCluster/Public/Components/DisplayClusterICVFXCameraComponent.h
@@ -177,6 +177,9 @@ private:
 	UPROPERTY(EditAnywhere, Transient, Category = "Tile Rendering", meta = (PropertyPath = "CameraSettings.CameraTile.TileOverscan", DisplayName = "Tile Overscan"))
 	FDisplayClusterEditorPropertyReference CameraTileOverscanRef;
 
+	UPROPERTY(EditAnywhere, Transient, Category = "Tile Rendering", meta = (PropertyPath = "CameraSettings.CameraTile.ClusterNodesToRenderUnboundTiles", DisplayName = "Nodes To Render Unbound Tiles", ToolTip = "Choose nodes that should render camera tiles that don't have any media assigned"))
+	FDisplayClusterEditorPropertyReference CameraTileNodesToRenderUnboundTilesRef;
+
 	UPROPERTY(EditAnywhere, Transient, Category = "Inner Frustum Color Grading", meta = (PropertyPath = "CameraSettings.EnableInnerFrustumColorGrading", DisplayName = "Enable Inner Frustum Color Grading"))
 	FDisplayClusterEditorPropertyReference EnableInnerFrustumColorGrading;
 
diff --git a/Engine/Plugins/Runtime/nDisplay/Source/DisplayCluster/Public/Render/Viewport/Containers/DisplayClusterViewport_Enums.h b/Engine/Plugins/Runtime/nDisplay/Source/DisplayCluster/Public/Render/Viewport/Containers/DisplayClusterViewport_Enums.h
--- a/Engine/Plugins/Runtime/nDisplay/Source/DisplayCluster/Public/Render/Viewport/Containers/DisplayClusterViewport_Enums.h
+++ b/Engine/Plugins/Runtime/nDisplay/Source/DisplayCluster/Public/Render/Viewport/Containers/DisplayClusterViewport_Enums.h
@@ -197,9 +197,6 @@ enum class EDisplayClusterViewportMediaState : uint8
 	Input = 1 << 4,
 	Input_ForceLateOCIOPass = 1 << 5,
 
-	// This viewport is inactive
-	Inactive = 1 << 6,
-
 	ForceLateOCIOPass = Capture_ForceLateOCIOPass | Input_ForceLateOCIOPass,
 };
 ENUM_CLASS_FLAGS(EDisplayClusterViewportMediaState);
diff --git a/Engine/Plugins/Runtime/nDisplay/Source/DisplayCluster/Public/Render/Viewport/Containers/DisplayClusterViewport_RenderSettings.h b/Engine/Plugins/Runtime/nDisplay/Source/DisplayCluster/Public/Render/Viewport/Containers/DisplayClusterViewport_RenderSettings.h
--- a/Engine/Plugins/Runtime/nDisplay/Source/DisplayCluster/Public/Render/Viewport/Containers/DisplayClusterViewport_RenderSettings.h
+++ b/Engine/Plugins/Runtime/nDisplay/Source/DisplayCluster/Public/Render/Viewport/Containers/DisplayClusterViewport_RenderSettings.h
@@ -180,6 +180,12 @@ public:
 		return EnumHasAnyFlags(MediaState, InMediaStates);
 	}
 
+	/** Returns current media states. */
+	inline EDisplayClusterViewportMediaState GetMediaStates() const
+	{
+		return MediaState;
+	}
+
 	/** Set media state for this viewport. */
 	inline void AssignMediaStates(const EDisplayClusterViewportMediaState InMediaStates)
 	{
diff --git a/Engine/Plugins/Runtime/nDisplay/Source/DisplayCluster/Public/Render/Viewport/Containers/DisplayClusterViewport_TileSettings.h b/Engine/Plugins/Runtime/nDisplay/Source/DisplayCluster/Public/Render/Viewport/Containers/DisplayClusterViewport_TileSettings.h
--- a/Engine/Plugins/Runtime/nDisplay/Source/DisplayCluster/Public/Render/Viewport/Containers/DisplayClusterViewport_TileSettings.h
+++ b/Engine/Plugins/Runtime/nDisplay/Source/DisplayCluster/Public/Render/Viewport/Containers/DisplayClusterViewport_TileSettings.h
@@ -3,7 +3,7 @@
 #pragma once
 
 #include "CoreMinimal.h"
-
+#include "Misc/EnumClassFlags.h"
 #include "Render/Viewport/Containers/DisplayClusterViewport_OverscanSettings.h"
 
 /**
@@ -25,6 +25,19 @@ enum class EDisplayClusterViewportTileType: uint8
 	UnusedTile
 };
 
+/**
+ * Additional tile settings represetned as a bitmask
+ */
+enum class EDisplayClusterViewportTileFlags : uint8
+{
+	None = 0,
+
+	// Allow this tile to render when unbound
+	AllowUnboundRender = 1 << 0,
+};
+ENUM_CLASS_FLAGS(EDisplayClusterViewportTileFlags);
+
+
 /**
  * nDisplay viewport tile settings.
  * These are runtime settings, updated every frame from the cluster configuration.
@@ -35,51 +48,70 @@ public:
 	FDisplayClusterViewport_TileSettings() = default;
 
 	/** Setup as source. */
-	FDisplayClusterViewport_TileSettings(const FIntPoint& InSize, const FDisplayClusterViewport_OverscanSettings& InOverscanSettings)
-		: Type(EDisplayClusterViewportTileType::Source), Size(InSize), OverscanSettings(InOverscanSettings)
+	FDisplayClusterViewport_TileSettings(const FIntPoint& InSize, const FDisplayClusterViewport_OverscanSettings& InOverscanSettings, const EDisplayClusterViewportTileFlags InTileFlags)
+		: Type(EDisplayClusterViewportTileType::Source)
+		, Size(InSize)
+		, OverscanSettings(InOverscanSettings)
+		, TileFlags(InTileFlags)
 	{ }
 
 	/** Setup as tile. */
-	FDisplayClusterViewport_TileSettings(const FString& InSourceViewportId, const FIntPoint& InPos, const FIntPoint& InSize)
-		: Type(EDisplayClusterViewportTileType::Tile), Size(InSize), Pos(InPos), SourceViewportId(InSourceViewportId)
+	FDisplayClusterViewport_TileSettings(const FString& InSourceViewportId, const FIntPoint& InPos, const FIntPoint& InSize, const EDisplayClusterViewportTileFlags InTileFlags)
+		: Type(EDisplayClusterViewportTileType::Tile)
+		, Size(InSize)
+		, Pos(InPos)
+		, SourceViewportId(InSourceViewportId)
+		, TileFlags(InTileFlags)
 	{ }
 
 	/** Returns the current viewport type for tile rendering. */
-	inline EDisplayClusterViewportTileType GetType() const
+	EDisplayClusterViewportTileType GetType() const
 	{
 		return Type;
 	}
 
 	/** Return true if this viewport is internal. */
-	inline bool IsInternalViewport() const
+	bool IsInternalViewport() const
 	{
 		return Type == EDisplayClusterViewportTileType::Tile;
 	}
 
 	/** Get Size value. */
-	inline const FIntPoint& GetSize() const
+	const FIntPoint& GetSize() const
 	{
 		return Size;
 	}
 
 	/** Get Pos value. */
-	inline const FIntPoint& GetPos() const
+	const FIntPoint& GetPos() const
 	{
 		return Pos;
 	}
 
 	/** Get SourceViewportId. */
-	inline const FString& GetSourceViewportId() const
+	const FString& GetSourceViewportId() const
 	{
 		return SourceViewportId;
 	}
 
 	/** Get overscan settings for tile rendering. */
-	inline const FDisplayClusterViewport_OverscanSettings& GetOverscanSettings() const
+	const FDisplayClusterViewport_OverscanSettings& GetOverscanSettings() const
 	{
 		return OverscanSettings;
 	}
 
+	/** Returns tile flags. */
+	EDisplayClusterViewportTileFlags GetTileFlags() const
+	{
+		return TileFlags;
+	}
+
+	/** Checks if any of specified tile flags are set. */
+	const bool HasAnyTileFlags(const EDisplayClusterViewportTileFlags RequestedTileFlags) const
+	{
+		return EnumHasAnyFlags(TileFlags, RequestedTileFlags);
+	}
+
 	/** Set tile state to be used. */
 	inline void SetTileStateToBeUsed(bool bUsed)
 	{
@@ -113,4 +145,7 @@ private:
 
 	// Overscan settings for tile rendering.
 	FDisplayClusterViewport_OverscanSettings OverscanSettings;
+
+	// Extra tile flags
+	EDisplayClusterViewportTileFlags TileFlags = EDisplayClusterViewportTileFlags::None;
 };
diff --git a/Engine/Plugins/Runtime/nDisplay/Source/DisplayClusterConfiguration/Public/DisplayClusterConfigurationTypes_ICVFX.h b/Engine/Plugins/Runtime/nDisplay/Source/DisplayClusterConfiguration/Public/DisplayClusterConfigurationTypes_ICVFX.h
--- a/Engine/Plugins/Runtime/nDisplay/Source/DisplayClusterConfiguration/Public/DisplayClusterConfigurationTypes_ICVFX.h
+++ b/Engine/Plugins/Runtime/nDisplay/Source/DisplayClusterConfiguration/Public/DisplayClusterConfigurationTypes_ICVFX.h
@@ -101,6 +101,10 @@ public:
 	/** Overscan settings for tile. */
 	UPROPERTY(EditAnywhere, BlueprintReadWrite, Category = "Tile Overscan")
 	FDisplayClusterConfigurationTile_Overscan TileOverscan;
+
+	/** Cluster nodes that should render unbound tiles. Unbound tiles are the tiles that don't have any media assigned */
+	UPROPERTY(EditAnywhere, BlueprintReadWrite, Category = Media, meta = (ClusterItemType = ClusterNodes, DisplayName = "Nodes To Render Unbound Tiles", ToolTip = "Choose nodes that should render camera tiles that don't have any media assigned"))
+	FDisplayClusterConfigurationClusterItemReferenceList ClusterNodesToRenderUnboundTiles;
 };
 
 USTRUCT(Blueprintable)
diff --git a/Engine/Plugins/Runtime/nDisplay/Source/DisplayClusterMedia/Private/Capture/DisplayClusterMediaCaptureViewport.cpp b/Engine/Plugins/Runtime/nDisplay/Source/DisplayClusterMedia/Private/Capture/DisplayClusterMediaCaptureViewport.cpp
--- a/Engine/Plugins/Runtime/nDisplay/Source/DisplayClusterMedia/Private/Capture/DisplayClusterMediaCaptureViewport.cpp
+++ b/Engine/Plugins/Runtime/nDisplay/Source/DisplayClusterMedia/Private/Capture/DisplayClusterMediaCaptureViewport.cpp
@@ -58,9 +58,6 @@ void FDisplayClusterMediaCaptureViewport::OnUpdateViewportMediaState(IDisplayClu
 	// In the future, after the media redesign, the DCRA name will also need to be checked here.
 	if (InViewport && InViewport->GetId().Equals(GetViewportId(), ESearchCase::IgnoreCase))
 	{
-		// Reset inactive state explicitly
-		InOutMediaState &= ~(EDisplayClusterViewportMediaState::Inactive);
-
 		// Raise flags that this viewport will be captured by media.
 		InOutMediaState |= EDisplayClusterViewportMediaState::Capture;
 
diff --git a/Engine/Plugins/Runtime/nDisplay/Source/DisplayClusterMedia/Private/Input/DisplayClusterMediaInputViewport.cpp b/Engine/Plugins/Runtime/nDisplay/Source/DisplayClusterMedia/Private/Input/DisplayClusterMediaInputViewport.cpp
--- a/Engine/Plugins/Runtime/nDisplay/Source/DisplayClusterMedia/Private/Input/DisplayClusterMediaInputViewport.cpp
+++ b/Engine/Plugins/Runtime/nDisplay/Source/DisplayClusterMedia/Private/Input/DisplayClusterMediaInputViewport.cpp
@@ -52,9 +52,6 @@ void FDisplayClusterMediaInputViewport::OnUpdateViewportMediaState(IDisplayClust
 	// In the future, after the media redesign, the DCRA name will also need to be checked here.
 	if (InViewport && InViewport->GetId().Equals(GetViewportId(), ESearchCase::IgnoreCase))
 	{
-		// Reset inactive state explicitly
-		InOutMediaState &= ~(EDisplayClusterViewportMediaState::Inactive);
-
 		// Raise flags that this viewport texture will be overridden by media.
 		InOutMediaState |= EDisplayClusterViewportMediaState::Input;
 
