Fix Incremental UnrealEditor Linux compilation.

diff --git a/Engine/Plugins/Runtime/nDisplay/Source/DisplayCluster/Private/Game/EngineClasses/Scene/DisplayClusterWidgetComponent.cpp b/Engine/Plugins/Runtime/nDisplay/Source/DisplayCluster/Private/Game/EngineClasses/Scene/DisplayClusterWidgetComponent.cpp
--- a/Engine/Plugins/Runtime/nDisplay/Source/DisplayCluster/Private/Game/EngineClasses/Scene/DisplayClusterWidgetComponent.cpp
+++ b/Engine/Plugins/Runtime/nDisplay/Source/DisplayCluster/Private/Game/EngineClasses/Scene/DisplayClusterWidgetComponent.cpp
@@ -36,10 +36,6 @@ public:
 	{
 		bWillEverBeLit = false;
 
-		FColor NewPropertyColor;
-		GEngine->GetPropertyColorationColor((UObject*)InComponent, NewPropertyColor);
-		SetPropertyColor(NewPropertyColor);
-
 		StaticMeshVertexBuffers.PositionVertexBuffer.Init(1);
 		StaticMeshVertexBuffers.StaticMeshVertexBuffer.Init(1, 1);
 		StaticMeshVertexBuffers.ColorVertexBuffer.Init(1);
