[nDisplay] Fixed incamera tiled capture initialization when running offscreen

- Fixed IsICVFXEnabled logic which is the culprit
- Changed media initialization trigger so it's always called from a single place regardless of any conditions
- Updated tiles configuration API in order to make it more clean

diff --git a/Engine/Plugins/Runtime/nDisplay/Source/DisplayCluster/Private/DisplayClusterCallbacks.h b/Engine/Plugins/Runtime/nDisplay/Source/DisplayCluster/Private/DisplayClusterCallbacks.h
--- a/Engine/Plugins/Runtime/nDisplay/Source/DisplayCluster/Private/DisplayClusterCallbacks.h
+++ b/Engine/Plugins/Runtime/nDisplay/Source/DisplayCluster/Private/DisplayClusterCallbacks.h
@@ -57,6 +57,11 @@ public:
 		return DisplayClusterEndSceneEvent;
 	}
 
+	virtual FDisplayClusterPreSubmitViewFamilies& OnDisplayClusterPreSubmitViewFamilies() override
+	{
+		return DisplayClusterPreSubmitViewFamiliesEvent;
+	}
+
 	virtual FDisplayClusterCustomPresentSetEvent& OnDisplayClusterCustomPresentSet() override
 	{
 		return DisplayClusterCustomPresentSetEvent;
@@ -152,6 +157,7 @@ private:
 	FDisplayClusterPostTickEvent             DisplayClusterPostTickEvent;
 	FDisplayClusterStartSceneEvent           DisplayClusterStartSceneEvent;
 	FDisplayClusterEndSceneEvent             DisplayClusterEndSceneEvent;
+	FDisplayClusterPreSubmitViewFamilies     DisplayClusterPreSubmitViewFamiliesEvent;
 	FDisplayClusterCustomPresentSetEvent     DisplayClusterCustomPresentSetEvent;
 	FDisplayClusterFailoverNodeDown          DisplayClusterFailoverNodeDown;
 
diff --git a/Engine/Plugins/Runtime/nDisplay/Source/DisplayCluster/Private/Game/EngineClasses/Basics/DisplayClusterViewportClient.cpp b/Engine/Plugins/Runtime/nDisplay/Source/DisplayCluster/Private/Game/EngineClasses/Basics/DisplayClusterViewportClient.cpp
--- a/Engine/Plugins/Runtime/nDisplay/Source/DisplayCluster/Private/Game/EngineClasses/Basics/DisplayClusterViewportClient.cpp
+++ b/Engine/Plugins/Runtime/nDisplay/Source/DisplayCluster/Private/Game/EngineClasses/Basics/DisplayClusterViewportClient.cpp
@@ -44,6 +44,7 @@
 #include "Misc/DisplayClusterGlobals.h"
 
 #include "IDisplayCluster.h"
+#include "IDisplayClusterCallbacks.h"
 #include "Game/IDisplayClusterGameManager.h"
 
 #include "Render/Viewport/IDisplayClusterViewportManager.h"
@@ -855,6 +856,9 @@ void UDisplayClusterViewportClient::Draw(FViewport* InViewport, FCanvas* SceneCa
 		}
 	}
 
+	// Trigger PreSubmitViewFamilies event before submitting to render
+	IDisplayCluster::Get().GetCallbacks().OnDisplayClusterPreSubmitViewFamilies().Broadcast(ViewFamilies);
+
 	// We gathered all the view families, now render them
 	if (!ViewFamilies.IsEmpty())
 	{
diff --git a/Engine/Plugins/Runtime/nDisplay/Source/DisplayCluster/Public/IDisplayClusterCallbacks.h b/Engine/Plugins/Runtime/nDisplay/Source/DisplayCluster/Public/IDisplayClusterCallbacks.h
--- a/Engine/Plugins/Runtime/nDisplay/Source/DisplayCluster/Public/IDisplayClusterCallbacks.h
+++ b/Engine/Plugins/Runtime/nDisplay/Source/DisplayCluster/Public/IDisplayClusterCallbacks.h
@@ -5,10 +5,14 @@
 #include "DisplayClusterEnums.h"
 #include "Render/Viewport/Containers/DisplayClusterViewport_Enums.h"
 
+#include "Delegates/DelegateCombinations.h"
+
+
 class FDisplayClusterShaderParameters_ICVFX;
 class FRDGBuilder;
 class FRHICommandListImmediate;
 class FSceneViewFamily;
+class FSceneViewFamilyContext;
 class FViewport;
 class IDisplayClusterViewport;
 class IDisplayClusterViewportProxy;
@@ -65,6 +69,10 @@ public:
 	DECLARE_EVENT(IDisplayClusterCallbacks, FDisplayClusterCustomPresentSetEvent);
 	virtual FDisplayClusterCustomPresentSetEvent& OnDisplayClusterCustomPresentSet() = 0;
 
+	/** Callback triggered before submitting view families to render **/
+	DECLARE_EVENT_OneParam(IDisplayClusterCallbacks, FDisplayClusterPreSubmitViewFamilies, TArray<FSceneViewFamilyContext*>&);
+	virtual FDisplayClusterPreSubmitViewFamilies& OnDisplayClusterPreSubmitViewFamilies() = 0;
+
 	/** Called before presentation synchronization is initiated **/
 	DECLARE_EVENT(IDisplayClusterCallbacks, FDisplayClusterPresentationPreSynchronization_RHIThread);
 	virtual FDisplayClusterPresentationPreSynchronization_RHIThread& OnDisplayClusterPresentationPreSynchronization_RHIThread() = 0;
diff --git a/Engine/Plugins/Runtime/nDisplay/Source/DisplayClusterConfiguration/Private/DisplayClusterConfigurationTypes_ICVFX.cpp b/Engine/Plugins/Runtime/nDisplay/Source/DisplayClusterConfiguration/Private/DisplayClusterConfigurationTypes_ICVFX.cpp
--- a/Engine/Plugins/Runtime/nDisplay/Source/DisplayClusterConfiguration/Private/DisplayClusterConfigurationTypes_ICVFX.cpp
+++ b/Engine/Plugins/Runtime/nDisplay/Source/DisplayClusterConfiguration/Private/DisplayClusterConfigurationTypes_ICVFX.cpp
@@ -209,7 +209,9 @@ bool FDisplayClusterConfigurationICVFX_CameraSettings::IsICVFXEnabled(const UDis
 	// If cluster mode + rendering offscreen, discover media output settings
 
 	// First condition to render offscreen: it has media output assigned
-	const bool bUsesMediaOutput = (RenderSettings.Media.bEnable && RenderSettings.Media.IsMediaOutputAssigned(InClusterNodeId));
+	const bool bUsesMediaOutput = RenderSettings.Media.bEnable && (
+			RenderSettings.Media.HasAnyMediaOutputAssigned(InClusterNodeId, EDisplayClusterConfigurationMediaSplitType::FullFrame) ||
+			RenderSettings.Media.HasAnyMediaOutputAssigned(InClusterNodeId, EDisplayClusterConfigurationMediaSplitType::UniformTiles));
 
 	// Get backbuffer media settings
 	const UDisplayClusterConfigurationClusterNode* const NodeCfg = InConfigurationData.Cluster->GetNode(InClusterNodeId);
diff --git a/Engine/Plugins/Runtime/nDisplay/Source/DisplayClusterConfiguration/Private/DisplayClusterConfigurationTypes_Media.cpp b/Engine/Plugins/Runtime/nDisplay/Source/DisplayClusterConfiguration/Private/DisplayClusterConfigurationTypes_Media.cpp
--- a/Engine/Plugins/Runtime/nDisplay/Source/DisplayClusterConfiguration/Private/DisplayClusterConfigurationTypes_Media.cpp
+++ b/Engine/Plugins/Runtime/nDisplay/Source/DisplayClusterConfiguration/Private/DisplayClusterConfigurationTypes_Media.cpp
@@ -41,99 +41,169 @@ bool FDisplayClusterConfigurationMediaViewport::IsMediaOutputAssigned() const
 ///////////////////////////////////////////////////
 // FDisplayClusterConfigurationMediaICVFX
 
-bool FDisplayClusterConfigurationMediaICVFX::IsMediaInputAssigned(const FString& NodeId) const
+bool FDisplayClusterConfigurationMediaICVFX::HasAnyMediaInputAssigned(const FString& NodeId, EDisplayClusterConfigurationMediaSplitType InSplitType) const
 {
-	return IsValid(GetMediaSource(NodeId));
+	// Nothing to do if a different split type requested
+	if (SplitType != InSplitType)
+	{
+		return false;
+	}
+
+	if (InSplitType == EDisplayClusterConfigurationMediaSplitType::FullFrame)
+	{
+		const bool bHasInput = IsValid(GetMediaSource(NodeId));
+		return bHasInput;
+	}
+	else if (InSplitType == EDisplayClusterConfigurationMediaSplitType::UniformTiles)
+	{
+		TArray<FDisplayClusterConfigurationMediaUniformTileInput> InputTiles;
+		GetMediaInputTiles(NodeId, InputTiles);
+
+		// Find any tile with a media source assigned
+		const bool bFoundInputTile = InputTiles.ContainsByPredicate([](const FDisplayClusterConfigurationMediaUniformTileInput& Tile)
+			{
+				return IsValid(Tile.MediaSource);
+			});
+
+		return bFoundInputTile;
+	}
+	else
+	{
+		checkNoEntry();
+	}
+
+	return false;
 }
 
-bool FDisplayClusterConfigurationMediaICVFX::IsMediaOutputAssigned(const FString& NodeId) const
+bool FDisplayClusterConfigurationMediaICVFX::HasAnyMediaOutputAssigned(const FString& NodeId, EDisplayClusterConfigurationMediaSplitType InSplitType) const
 {
-	const TArray<FDisplayClusterConfigurationMediaOutputGroup> NodeOutputGroups = GetMediaOutputGroups(NodeId);
+	if (SplitType != InSplitType)
+	{
+		return false;
+	}
 
-	for (const FDisplayClusterConfigurationMediaOutputGroup& NodeOutputGroup : NodeOutputGroups)
+	if (InSplitType == EDisplayClusterConfigurationMediaSplitType::FullFrame)
 	{
-		if (IsValid(NodeOutputGroup.MediaOutput))
-		{
-			return true;
-		}
+		const TArray<FDisplayClusterConfigurationMediaOutputGroup> NodeOutputGroups = GetMediaOutputGroups(NodeId);
+
+		const bool bFoundOutput = NodeOutputGroups.ContainsByPredicate([](const FDisplayClusterConfigurationMediaOutputGroup& OutputGroup)
+			{
+				return IsValid(OutputGroup.MediaOutput);
+			});
+
+		return bFoundOutput;
+	}
+	else if (InSplitType == EDisplayClusterConfigurationMediaSplitType::UniformTiles)
+	{
+		TArray<FDisplayClusterConfigurationMediaUniformTileOutput> OutputTiles;
+		GetMediaOutputTiles(NodeId, OutputTiles);
+
+		const bool bFoundOutputTile = OutputTiles.ContainsByPredicate([](const FDisplayClusterConfigurationMediaUniformTileOutput& Tile)
+			{
+				return IsValid(Tile.MediaOutput);
+			});
+
+		return bFoundOutputTile;
+	}
+	else
+	{
+		checkNoEntry();
 	}
 
 	return false;
 }
 
-
-//@note
-// The way how media is configured for ICVFX cameras technically allows to have multiple inputs assigned
-// to the same camera. Yes, this is something that contradicts to a single input concept. However, it provides
-// a very user-friendly GUI. So to follow the single input concept, we always return the first media input found.
+//@note: Full-frame input - we don't expect the same node ID to be used multiple times (in different groups)
 UMediaSource* FDisplayClusterConfigurationMediaICVFX::GetMediaSource(const FString& NodeId) const
 {
-	// Look up for a group that contains node ID specified
-	for (const FDisplayClusterConfigurationMediaInputGroup& MediaInputGroup : MediaInputGroups)
+	if (SplitType != EDisplayClusterConfigurationMediaSplitType::FullFrame)
 	{
-		const bool bNodeFound = MediaInputGroup.ClusterNodes.ItemNames.ContainsByPredicate([NodeId](const FString& Item)
-			{
-				return Item.Equals(NodeId, ESearchCase::IgnoreCase);
-			});
+		return nullptr;
+	}
 
-		if (bNodeFound && IsValid(MediaInputGroup.MediaSource))
+	// Look up for the first group that contains node ID specified
+	const FDisplayClusterConfigurationMediaInputGroup* const FoundGroup = MediaInputGroups.FindByPredicate([&NodeId](const FDisplayClusterConfigurationMediaInputGroup& MediaInputGroup)
 		{
-			return MediaInputGroup.MediaSource;
-		}
+			return MediaInputGroup.ClusterNodes.ItemNames.ContainsByPredicate([&NodeId](const FString& Item)
+				{
+					return Item.Equals(NodeId, ESearchCase::IgnoreCase);
+				});
+		});
+
+	if (FoundGroup)
+	{
+		return FoundGroup->MediaSource;
 	}
 
 	return nullptr;
 }
 
+//@note: Full-frame output, it's allowed to use the same node ID in different groups
 TArray<FDisplayClusterConfigurationMediaOutputGroup> FDisplayClusterConfigurationMediaICVFX::GetMediaOutputGroups(const FString& NodeId) const
 {
-	return MediaOutputGroups.FilterByPredicate([NodeId](const FDisplayClusterConfigurationMediaOutputGroup& Item)
+	if (SplitType != EDisplayClusterConfigurationMediaSplitType::FullFrame)
+	{
+		return TArray<FDisplayClusterConfigurationMediaOutputGroup>();
+	}
+
+	return MediaOutputGroups.FilterByPredicate([&NodeId](const FDisplayClusterConfigurationMediaOutputGroup& Item)
 		{
-			return Item.ClusterNodes.ItemNames.Contains(NodeId);
+			return Item.ClusterNodes.ItemNames.ContainsByPredicate([&NodeId](const FString& Item)
+				{
+					return Item.Equals(NodeId, ESearchCase::IgnoreCase);
+				});
 		});
 }
 
-/** Returns all tiles bound to a specific cluster node. */
-bool FDisplayClusterConfigurationMediaICVFX::GetMediaInputTiles(const FString& NodeId, TArray<FDisplayClusterConfigurationMediaUniformTileInput>& OutInputTiles) const
+//@note: Tiled input, it's allowed to use the same node ID in different groups
+bool FDisplayClusterConfigurationMediaICVFX::GetMediaInputTiles(const FString& NodeId, TArray<FDisplayClusterConfigurationMediaUniformTileInput>& OutTiles) const
 {
-	// Here we iterate through tiled media input groups, and look for a group that contains a cluster node ID specified.
-	for (const FDisplayClusterConfigurationMediaTiledInputGroup& TiledInputGroup : TiledMediaInputGroups)
+	if (SplitType != EDisplayClusterConfigurationMediaSplitType::UniformTiles)
 	{
-		// Look for NodeId in the group
-		const bool bFoundGroupWithNodeId = TiledInputGroup.ClusterNodes.ItemNames.ContainsByPredicate([NodeId](const FString& Item)
-			{
-				return Item.Equals(NodeId, ESearchCase::IgnoreCase);
-			});
+		return false;
+	}
+
+	OutTiles.Empty();
 
-		// If found
-		if (bFoundGroupWithNodeId)
+	// Find all input groups bound to the node ID specified
+	const TArray<FDisplayClusterConfigurationMediaTiledInputGroup> FoundInputGroups = TiledMediaInputGroups.FilterByPredicate([&NodeId](const FDisplayClusterConfigurationMediaTiledInputGroup& TiledInputGroup)
 		{
-			OutInputTiles = TiledInputGroup.Tiles;
-			return true;
-		}
+			return TiledInputGroup.ClusterNodes.ItemNames.ContainsByPredicate([&NodeId](const FString& Item)
+				{
+					return Item.Equals(NodeId, ESearchCase::IgnoreCase);
+				});
+		});
+
+	// Combine result out of the filtered groups
+	for (const FDisplayClusterConfigurationMediaTiledInputGroup& TiledInputGroup : FoundInputGroups)
+	{
+		OutTiles.Append(TiledInputGroup.Tiles);
 	}
 
 	return false;
 }
 
-/** Returns all tiles bound to a specific cluster node. */
-bool FDisplayClusterConfigurationMediaICVFX::GetMediaOutputTiles(const FString& NodeId, TArray<FDisplayClusterConfigurationMediaUniformTileOutput>& OutOutputTiles) const
+//@note: Tiled output, it's allowed to use the same node ID in different groups
+bool FDisplayClusterConfigurationMediaICVFX::GetMediaOutputTiles(const FString& NodeId, TArray<FDisplayClusterConfigurationMediaUniformTileOutput>& OutTiles) const
 {
-	// Here we iterate through tiled media output groups, and look for a group that contains a cluster node ID specified.
-	for (const FDisplayClusterConfigurationMediaTiledOutputGroup& TiledOutputGroup : TiledMediaOutputGroups)
+	if (SplitType != EDisplayClusterConfigurationMediaSplitType::UniformTiles)
 	{
-		// Look for NodeId in the group
-		const bool bFoundGroupWithNodeId = TiledOutputGroup.ClusterNodes.ItemNames.ContainsByPredicate([NodeId](const FString& Item)
-			{
-				return Item.Equals(NodeId, ESearchCase::IgnoreCase);
-			});
+		false;
+	}
 
-		// If found
-		if (bFoundGroupWithNodeId)
+	// Find all output groups bound to the node ID specified
+	const TArray<FDisplayClusterConfigurationMediaTiledOutputGroup> FoundOutputGroups = TiledMediaOutputGroups.FilterByPredicate([&NodeId](const FDisplayClusterConfigurationMediaTiledOutputGroup& TiledOutputGroup)
 		{
-			OutOutputTiles = TiledOutputGroup.Tiles;
-			return true;
-		}
+			return TiledOutputGroup.ClusterNodes.ItemNames.ContainsByPredicate([&NodeId](const FString& Item)
+				{
+					return Item.Equals(NodeId, ESearchCase::IgnoreCase);
+				});
+		});
+
+	// Combine result out of the filtered groups
+	for (const FDisplayClusterConfigurationMediaTiledOutputGroup& TiledOutputGroup : FoundOutputGroups)
+	{
+		OutTiles.Append(TiledOutputGroup.Tiles);
 	}
 
 	return false;
diff --git a/Engine/Plugins/Runtime/nDisplay/Source/DisplayClusterConfiguration/Public/DisplayClusterConfigurationTypes_Media.h b/Engine/Plugins/Runtime/nDisplay/Source/DisplayClusterConfiguration/Public/DisplayClusterConfigurationTypes_Media.h
--- a/Engine/Plugins/Runtime/nDisplay/Source/DisplayClusterConfiguration/Public/DisplayClusterConfigurationTypes_Media.h
+++ b/Engine/Plugins/Runtime/nDisplay/Source/DisplayClusterConfiguration/Public/DisplayClusterConfigurationTypes_Media.h
@@ -300,48 +300,61 @@ public:
 	bool bLateOCIOPass = false;
 
 public:
-	/** Returns true if a specific cluster node has media source assigned */
-	bool IsMediaInputAssigned(const FString& NodeId) const;
 
-	/** Returns true if a specific cluster node has media source assigned */
-	bool IsMediaOutputAssigned(const FString& NodeId) const;
+	/** Returns true if any media source of specified split-type is bound */
+	bool HasAnyMediaInputAssigned(const FString& NodeId, EDisplayClusterConfigurationMediaSplitType SplitType) const;
 
-	/** Returns media source bound to a specific cluster node */
+	/** Returns true if any media output of specified split-type is bound */
+	bool HasAnyMediaOutputAssigned(const FString& NodeId, EDisplayClusterConfigurationMediaSplitType SplitType) const;
+
+	/** [FullFrame] Returns media source bound to a specific cluster node if configured for full-frame media */
 	UMediaSource* GetMediaSource(const FString& NodeId) const;
 
-	/** Returns media outputs bound to a specific cluster node */
+	/** [FullFrame] Returns media outputs bound to a specific cluster node if configured for full-frame media */
 	TArray<FDisplayClusterConfigurationMediaOutputGroup> GetMediaOutputGroups(const FString& NodeId) const;
 
-	/** Returns all the tiles bound to a specific cluster node. */
-	bool GetMediaInputTiles(const FString& NodeId, TArray<FDisplayClusterConfigurationMediaUniformTileInput>& OutInputTiles) const;
+	/** [UniformTiles] Returns all the tiles bound to a specific cluster node. */
+	bool GetMediaInputTiles(const FString& NodeId, TArray<FDisplayClusterConfigurationMediaUniformTileInput>& OutTiles) const;
 
-	/** Returns all the tiles bound to a specific cluster node. */
-	bool GetMediaOutputTiles(const FString& NodeId, TArray<FDisplayClusterConfigurationMediaUniformTileOutput>& OutOutputTiles) const;
+	/** [UniformTiles] Returns all the tiles bound to a specific cluster node. */
+	bool GetMediaOutputTiles(const FString& NodeId, TArray<FDisplayClusterConfigurationMediaUniformTileOutput>& OutTiles) const;
 
 public:
-	UE_DEPRECATED(5.3, "This function has beend deprecated. Please use GetMediaOutputGroups.")
+	UE_DEPRECATED(5.3, "This function has been deprecated. Please use GetMediaOutputGroups.")
 	TArray<UMediaOutput*> GetMediaOutputs(const FString& NodeId) const
 	{
 		return {};
 	}
 
-	UE_DEPRECATED(5.3, "This function has beend deprecated. Please use GetMediaOutputGroups.")
+	UE_DEPRECATED(5.3, "This function has been deprecated. Please use GetMediaOutputGroups.")
 	UDisplayClusterMediaOutputSynchronizationPolicy* GetOutputSyncPolicy(const FString& NodeId) const
 	{
 		return nullptr;
 	}
 
-	UE_DEPRECATED(5.3, "This function has beend deprecated. Please use GetMediaOutputGroups.")
+	UE_DEPRECATED(5.3, "This function has been deprecated. Please use GetMediaOutputGroups.")
 	UMediaOutput* GetMediaOutput(const FString& NodeId) const
 	{
 		return nullptr;
 	}
 
-	UE_DEPRECATED(5.2, "This function has beend deprecated.")
+	UE_DEPRECATED(5.2, "This function has been deprecated.")
 	bool IsMediaSharingNode(const FString& InNodeId) const
 	{
 		return false;
 	}
+
+	UE_DEPRECATED(5.4, "This function has been deprecated. Please use `HasAnyMediaInputAssigned`.")
+	bool IsMediaInputAssigned(const FString& NodeId) const
+	{
+		return HasAnyMediaInputAssigned(NodeId, EDisplayClusterConfigurationMediaSplitType::FullFrame);
+	}
+
+	UE_DEPRECATED(5.4, "This function has been deprecated. Please use `HasAnyMediaOutputAssigned`.")
+	bool IsMediaOutputAssigned(const FString& NodeId) const
+	{
+		return HasAnyMediaOutputAssigned(NodeId, EDisplayClusterConfigurationMediaSplitType::FullFrame);
+	}
 };
 
 
diff --git a/Engine/Plugins/Runtime/nDisplay/Source/DisplayClusterMedia/Private/Capture/DisplayClusterMediaCaptureNode.cpp b/Engine/Plugins/Runtime/nDisplay/Source/DisplayClusterMedia/Private/Capture/DisplayClusterMediaCaptureNode.cpp
--- a/Engine/Plugins/Runtime/nDisplay/Source/DisplayClusterMedia/Private/Capture/DisplayClusterMediaCaptureNode.cpp
+++ b/Engine/Plugins/Runtime/nDisplay/Source/DisplayClusterMedia/Private/Capture/DisplayClusterMediaCaptureNode.cpp
@@ -4,6 +4,8 @@
 
 #include "Config/IDisplayClusterConfigManager.h"
 #include "DisplayClusterConfigurationTypes.h"
+#include "DisplayClusterMediaLog.h"
+
 #include "IDisplayCluster.h"
 #include "IDisplayClusterCallbacks.h"
 
@@ -48,9 +50,15 @@ FIntPoint FDisplayClusterMediaCaptureNode::GetCaptureSize() const
 	// Return backbuffer runtime size
 	if (GEngine && GEngine->GameViewport && GEngine->GameViewport->Viewport)
 	{
-		return GEngine->GameViewport->Viewport->GetSizeXY();
+		const FIntPoint Size = GEngine->GameViewport->Viewport->GetSizeXY();
+
+		UE_LOG(LogDisplayClusterMedia, Log, TEXT("'%s' capture size is [%d, %d]"), *GetMediaId(), Size.X, Size.Y);
+
+		return Size;
 	}
 
+	UE_LOG(LogDisplayClusterMedia, Warning, TEXT("'%s' couldn't get viewport size"), *GetMediaId());
+
 	return FIntPoint();
 }
 
diff --git a/Engine/Plugins/Runtime/nDisplay/Source/DisplayClusterMedia/Private/Capture/DisplayClusterMediaCaptureViewport.cpp b/Engine/Plugins/Runtime/nDisplay/Source/DisplayClusterMedia/Private/Capture/DisplayClusterMediaCaptureViewport.cpp
--- a/Engine/Plugins/Runtime/nDisplay/Source/DisplayClusterMedia/Private/Capture/DisplayClusterMediaCaptureViewport.cpp
+++ b/Engine/Plugins/Runtime/nDisplay/Source/DisplayClusterMedia/Private/Capture/DisplayClusterMediaCaptureViewport.cpp
@@ -6,6 +6,7 @@
 #include "IDisplayClusterCallbacks.h"
 
 #include "DisplayClusterConfigurationTypes_Viewport.h"
+#include "DisplayClusterMediaLog.h"
 
 #include "Config/IDisplayClusterConfigManager.h"
 
@@ -82,12 +83,15 @@ FIntPoint FDisplayClusterMediaCaptureViewport::GetCaptureSize() const
 				if (Contexts.Num() > 0)
 				{
 					const FIntPoint Size = Contexts[0].RenderTargetRect.Size();
+					UE_LOG(LogDisplayClusterMedia, Log, TEXT("'%s' capture size is [%d, %d]"), *GetMediaId(), Size.X, Size.Y);
 					return Size;
 				}
 			}
 		}
 	}
 
+	UE_LOG(LogDisplayClusterMedia, Warning, TEXT("'%s' couldn't get viewport size"), *GetMediaId());
+
 	return FIntPoint();
 }
 
diff --git a/Engine/Plugins/Runtime/nDisplay/Source/DisplayClusterMedia/Private/DisplayClusterMediaModule.cpp b/Engine/Plugins/Runtime/nDisplay/Source/DisplayClusterMedia/Private/DisplayClusterMediaModule.cpp
--- a/Engine/Plugins/Runtime/nDisplay/Source/DisplayClusterMedia/Private/DisplayClusterMediaModule.cpp
+++ b/Engine/Plugins/Runtime/nDisplay/Source/DisplayClusterMedia/Private/DisplayClusterMediaModule.cpp
@@ -30,7 +30,7 @@ void FDisplayClusterMediaModule::StartupModule()
 {
 	UE_LOG(LogDisplayClusterMedia, Log, TEXT("Starting module 'DisplayClusterMedia'..."));
 
-	IDisplayCluster::Get().GetCallbacks().OnDisplayClusterCustomPresentSet().AddRaw(this, &FDisplayClusterMediaModule::OnCustomPresentSet);
+	IDisplayCluster::Get().GetCallbacks().OnDisplayClusterPreSubmitViewFamilies().AddRaw(this, &FDisplayClusterMediaModule::OnPreSubmitViewFamilies);
 	FCoreDelegates::OnEnginePreExit.AddRaw(this, &FDisplayClusterMediaModule::OnEnginePreExit);
 }
 
@@ -38,17 +38,17 @@ void FDisplayClusterMediaModule::ShutdownModule()
 {
 	UE_LOG(LogDisplayClusterMedia, Log, TEXT("Shutting down module 'DisplayClusterMedia'..."));
 
-	IDisplayCluster::Get().GetCallbacks().OnDisplayClusterCustomPresentSet().RemoveAll(this);
+	// We should already be unsubscribed from it but do it in case the callback has never been called.
+	IDisplayCluster::Get().GetCallbacks().OnDisplayClusterPreSubmitViewFamilies().RemoveAll(this);
+
 	FCoreDelegates::OnEnginePreExit.RemoveAll(this);
 }
 
-void FDisplayClusterMediaModule::OnCustomPresentSet()
+void FDisplayClusterMediaModule::OnPreSubmitViewFamilies(TArray<FSceneViewFamilyContext*>&)
 {
-	// We initialize and start media when backbuffer is available. This CustomPresentSet event is a
-	// sign that the backbuffer is already available. This allows us to prevent any potential problems
-	// if any of the following is used under the hood, or will be used in the future:
-	// - UMediaCapture::CaptureActiveSceneViewport
-	// - UMediaCapture::CaptureSceneViewport
+	// Unsubscribe after first call. Currently, media initialization is a one time procedure. No need to receive any further callbacks.
+	IDisplayCluster::Get().GetCallbacks().OnDisplayClusterPreSubmitViewFamilies().RemoveAll(this);
+
 	InitializeMedia();
 	StartCapture();
 	PlayMedia();
diff --git a/Engine/Plugins/Runtime/nDisplay/Source/DisplayClusterMedia/Private/DisplayClusterMediaModule.h b/Engine/Plugins/Runtime/nDisplay/Source/DisplayClusterMedia/Private/DisplayClusterMediaModule.h
--- a/Engine/Plugins/Runtime/nDisplay/Source/DisplayClusterMedia/Private/DisplayClusterMediaModule.h
+++ b/Engine/Plugins/Runtime/nDisplay/Source/DisplayClusterMedia/Private/DisplayClusterMediaModule.h
@@ -10,11 +10,13 @@
 class FDisplayClusterMediaCaptureNode;
 class FDisplayClusterMediaCaptureViewport;
 class FDisplayClusterMediaInputViewport;
+class FSceneViewFamilyContext;
 class IMediaPlayerFactory;
 class UDisplayClusterConfigurationClusterNode;
 class UDisplayClusterConfigurationViewport;
 class UDisplayClusterICVFXCameraComponent;
 
+
 /**
  * Media module
  */
@@ -69,8 +71,8 @@ private:
 	void InitializeICVFXCameraUniformTilesOutput(const UDisplayClusterICVFXCameraComponent* ICVFXCameraComponent, const FString& RootActorName, const FString& ClusterNodeId);
 
 private:
-	/** CustomPresentSet event handler */
-	void OnCustomPresentSet();
+	/** PreSubmitViewFamilies event handler. It's used to initialize media on start. */
+	void OnPreSubmitViewFamilies(TArray<FSceneViewFamilyContext*>&);
 
 	/** EnginePreExit event handler */
 	void OnEnginePreExit();
