[nDisplay] Fixed CIS build issue

diff --git a/Engine/Plugins/Runtime/nDisplay/Source/DisplayClusterConfiguration/Private/DisplayClusterConfigurationTypes_Media.cpp b/Engine/Plugins/Runtime/nDisplay/Source/DisplayClusterConfiguration/Private/DisplayClusterConfigurationTypes_Media.cpp
--- a/Engine/Plugins/Runtime/nDisplay/Source/DisplayClusterConfiguration/Private/DisplayClusterConfigurationTypes_Media.cpp
+++ b/Engine/Plugins/Runtime/nDisplay/Source/DisplayClusterConfiguration/Private/DisplayClusterConfigurationTypes_Media.cpp
@@ -188,7 +188,7 @@ bool FDisplayClusterConfigurationMediaICVFX::GetMediaOutputTiles(const FString&
 {
 	if (SplitType != EDisplayClusterConfigurationMediaSplitType::UniformTiles)
 	{
-		false;
+		return false;
 	}
 
 	// Find all output groups bound to the node ID specified
