More bool to EAllowShrinking fixes.

diff --git a/Engine/Plugins/AI/MLAdapter/Source/MLAdapter/Private/MLAdapterSpace.cpp b/Engine/Plugins/AI/MLAdapter/Source/MLAdapter/Private/MLAdapterSpace.cpp
--- a/Engine/Plugins/AI/MLAdapter/Source/MLAdapter/Private/MLAdapterSpace.cpp
+++ b/Engine/Plugins/AI/MLAdapter/Source/MLAdapter/Private/MLAdapterSpace.cpp
@@ -135,7 +135,7 @@ namespace FMLAdapter
 			Contents += Space->ToJson();
 			Contents += TEXT(",");
 		}
-		Contents.RemoveAt(Contents.Len() - 1, 1, false);
+		Contents.RemoveAt(Contents.Len() - 1, 1, EAllowShrinking::No);
 		return FString::Printf(TEXT("{\"%s\":[%s]}"), *EnumToString(Type), *Contents);
 	}
 
@@ -205,7 +205,7 @@ FString FMLAdapterDescription::ToJson() const
 			RetString += TEXT(",");
 		}
 		// pop the last ,
-		RetString.RemoveAt(RetString.Len() - 1, 1, false);
+		RetString.RemoveAt(RetString.Len() - 1, 1, EAllowShrinking::No);
 		if (PrepData.Num() > 1)
 		{
 			RetString += TEXT("]");
diff --git a/Engine/Plugins/Editor/GuidedTutorials/Source/IntroTutorials/Private/STutorialEditableText.cpp b/Engine/Plugins/Editor/GuidedTutorials/Source/IntroTutorials/Private/STutorialEditableText.cpp
--- a/Engine/Plugins/Editor/GuidedTutorials/Source/IntroTutorials/Private/STutorialEditableText.cpp
+++ b/Engine/Plugins/Editor/GuidedTutorials/Source/IntroTutorials/Private/STutorialEditableText.cpp
@@ -440,7 +440,7 @@ void STutorialEditableText::HandleHyperlinkComboOpened()
 		{
 			if(FChar::IsLinebreak(SelectedText[SelectedTextIndex]))
 			{
-				SelectedText.LeftInline(SelectedTextIndex, false);
+				SelectedText.LeftInline(SelectedTextIndex, EAllowShrinking::No);
 				break;
 			}
 		}
diff --git a/Engine/Plugins/Experimental/AnimNext/Source/AnimNextEditor/Private/EditorUtils.cpp b/Engine/Plugins/Experimental/AnimNext/Source/AnimNextEditor/Private/EditorUtils.cpp
--- a/Engine/Plugins/Experimental/AnimNext/Source/AnimNextEditor/Private/EditorUtils.cpp
+++ b/Engine/Plugins/Experimental/AnimNext/Source/AnimNextEditor/Private/EditorUtils.cpp
@@ -93,7 +93,7 @@ FName FUtils::ValidateName(const UObject* InObject, const FString& InName)
 	FString Name = InName;
 	if (Name.StartsWith(TEXT("RigUnit_")))
 	{
-		Name.RightChopInline(8, false);
+		Name.RightChopInline(8, EAllowShrinking::No);
 	}
 
 	FNameValidator NameValidator(InObject);
@@ -302,7 +302,7 @@ FName FUtils::ValidateName(const UAnimNextParameterBlock_EditorData* InEditorDat
 	FString Name = InName;
 	if (Name.StartsWith(TEXT("RigUnit_")))
 	{
-		Name.RightChopInline(8, false);
+		Name.RightChopInline(8, EAllowShrinking::No);
 	}
 
 	FNameValidator NameValidator(InEditorData);
diff --git a/Engine/Plugins/Experimental/Avalanche/Source/AvalancheMediaEditor/Private/Playlist/Pages/Slate/SAvaPageList.cpp b/Engine/Plugins/Experimental/Avalanche/Source/AvalancheMediaEditor/Private/Playlist/Pages/Slate/SAvaPageList.cpp
--- a/Engine/Plugins/Experimental/Avalanche/Source/AvalancheMediaEditor/Private/Playlist/Pages/Slate/SAvaPageList.cpp
+++ b/Engine/Plugins/Experimental/Avalanche/Source/AvalancheMediaEditor/Private/Playlist/Pages/Slate/SAvaPageList.cpp
@@ -1068,7 +1068,7 @@ FText SAvaPageList::OnAssetSearchBoxSuggestionChosen(const FText& InSearchText,
 	UE::AvalanchePlaylist::Private::ExtractAssetSearchFilterTerms(InSearchText, nullptr, nullptr, &SuggestionInsertionIndex);
 
 	FString SearchString = InSearchText.ToString();
-	SearchString.RemoveAt(SuggestionInsertionIndex, SearchString.Len() - SuggestionInsertionIndex, false);
+	SearchString.RemoveAt(SuggestionInsertionIndex, SearchString.Len() - SuggestionInsertionIndex, EAllowShrinking::No);
 	SearchString.Append(InSuggestion);
 	return FText::FromString(SearchString);
 }
diff --git a/Engine/Plugins/Experimental/Avalanche/Source/AvalancheOutliner/Private/Slate/SAvaOutliner.cpp b/Engine/Plugins/Experimental/Avalanche/Source/AvalancheOutliner/Private/Slate/SAvaOutliner.cpp
--- a/Engine/Plugins/Experimental/Avalanche/Source/AvalancheOutliner/Private/Slate/SAvaOutliner.cpp
+++ b/Engine/Plugins/Experimental/Avalanche/Source/AvalancheOutliner/Private/Slate/SAvaOutliner.cpp
@@ -495,7 +495,7 @@ FText SAvaOutliner::OnAssetSearchSuggestionChosen(const FText& InSearchText, con
 	UE::AvaOutliner::Private::ExtractAssetSearchFilterTerms(InSearchText, nullptr, nullptr, &SuggestionInsertionIndex);
 
 	FString SearchString = InSearchText.ToString();
-	SearchString.RemoveAt(SuggestionInsertionIndex, SearchString.Len() - SuggestionInsertionIndex, false);
+	SearchString.RemoveAt(SuggestionInsertionIndex, SearchString.Len() - SuggestionInsertionIndex, EAllowShrinking::No);
 	SearchString.Append(InSuggestion);
 	return FText::FromString(SearchString);
 }
diff --git a/Engine/Plugins/Experimental/Avalanche/Source/AvalancheTransition/Private/Preview/AvaTransitionPreviewManager.cpp b/Engine/Plugins/Experimental/Avalanche/Source/AvalancheTransition/Private/Preview/AvaTransitionPreviewManager.cpp
--- a/Engine/Plugins/Experimental/Avalanche/Source/AvalancheTransition/Private/Preview/AvaTransitionPreviewManager.cpp
+++ b/Engine/Plugins/Experimental/Avalanche/Source/AvalancheTransition/Private/Preview/AvaTransitionPreviewManager.cpp
@@ -138,7 +138,7 @@ const FAvaTransitionPreviewLevelState* AAvaTransitionPreviewManager::FindExistin
 	FString ShortPackageName = FPackageName::GetShortName(LevelInstanceParams.LongPackageName);
 	if (ShortPackageName.StartsWith(LevelInstanceParams.World->StreamingLevelsPrefix))
 	{
-		ShortPackageName.RightChopInline(LevelInstanceParams.World->StreamingLevelsPrefix.Len(), false);
+		ShortPackageName.RightChopInline(LevelInstanceParams.World->StreamingLevelsPrefix.Len(), EAllowShrinking::No);
 	}
 
 	const FName PackageName = *(FPackageName::GetLongPackagePath(LevelInstanceParams.LongPackageName)
diff --git a/Engine/Plugins/Runtime/nDisplay/Source/DisplayCluster/Public/Misc/DisplayClusterHelpers.h b/Engine/Plugins/Runtime/nDisplay/Source/DisplayCluster/Public/Misc/DisplayClusterHelpers.h
--- a/Engine/Plugins/Runtime/nDisplay/Source/DisplayCluster/Public/Misc/DisplayClusterHelpers.h
+++ b/Engine/Plugins/Runtime/nDisplay/Source/DisplayCluster/Public/Misc/DisplayClusterHelpers.h
@@ -111,7 +111,7 @@ namespace DisplayClusterHelpers
 
 			if (InSeparator.Len() > 0 && InData.Num() > 0)
 			{
-				ResultStr.RemoveAt(ResultStr.Len() - InSeparator.Len(), InSeparator.Len(), !bAddQuotes);
+				ResultStr.RemoveAt(ResultStr.Len() - InSeparator.Len(), InSeparator.Len(), bAddQuotes ? EAllowShrinking::No : EAllowShrinking::Yes);
 			}
 
 			if (bAddQuotes)
diff --git a/Engine/Plugins/Runtime/nDisplay/Source/DisplayClusterConfigurator/Private/ClusterConfiguration/DisplayClusterConfiguratorClusterUtils.cpp b/Engine/Plugins/Runtime/nDisplay/Source/DisplayClusterConfigurator/Private/ClusterConfiguration/DisplayClusterConfiguratorClusterUtils.cpp
--- a/Engine/Plugins/Runtime/nDisplay/Source/DisplayClusterConfigurator/Private/ClusterConfiguration/DisplayClusterConfiguratorClusterUtils.cpp
+++ b/Engine/Plugins/Runtime/nDisplay/Source/DisplayClusterConfigurator/Private/ClusterConfiguration/DisplayClusterConfiguratorClusterUtils.cpp
@@ -51,7 +51,7 @@ namespace UE::DisplayClusterConfiguratorClusterUtils
 				FString NumericSuffix = BaseName.RightChop(Index);
 				Counter = FCString::Atoi(*NumericSuffix);
 				NumericSuffix = FString::FromInt(Counter); // Restringify the counter to account for leading 0s that we don't want to remove
-				BaseName.RemoveAt(BaseName.Len() - NumericSuffix.Len(), NumericSuffix.Len(), false);
+				BaseName.RemoveAt(BaseName.Len() - NumericSuffix.Len(), NumericSuffix.Len(), EAllowShrinking::No);
 			}
 			else
 			{
diff --git a/Engine/Plugins/VirtualProduction/DMX/DMXEngine/Source/DMXRuntime/Private/Library/DMXLibrary.cpp b/Engine/Plugins/VirtualProduction/DMX/DMXEngine/Source/DMXRuntime/Private/Library/DMXLibrary.cpp
--- a/Engine/Plugins/VirtualProduction/DMX/DMXEngine/Source/DMXRuntime/Private/Library/DMXLibrary.cpp
+++ b/Engine/Plugins/VirtualProduction/DMX/DMXEngine/Source/DMXRuntime/Private/Library/DMXLibrary.cpp
@@ -270,7 +270,7 @@ void UDMXLibrary::SetEntityIndex(UDMXEntity* InEntity, const int32 NewIndex)
 	{
 		if (NewIndex >= Entities.Num())
 		{
-			Entities.RemoveAt(OldIndex, 1, false);
+			Entities.RemoveAt(OldIndex, 1, EAllowShrinking::No);
 			Entities.Add(InEntity);
 			return;
 		}
diff --git a/Engine/Plugins/VirtualProduction/DMX/DMXProtocol/Source/DMXProtocolArtNet/Private/DMXProtocolArtNetReceiver.cpp b/Engine/Plugins/VirtualProduction/DMX/DMXProtocol/Source/DMXProtocolArtNet/Private/DMXProtocolArtNetReceiver.cpp
--- a/Engine/Plugins/VirtualProduction/DMX/DMXProtocol/Source/DMXProtocolArtNet/Private/DMXProtocolArtNetReceiver.cpp
+++ b/Engine/Plugins/VirtualProduction/DMX/DMXProtocol/Source/DMXProtocolArtNet/Private/DMXProtocolArtNetReceiver.cpp
@@ -163,7 +163,7 @@ void FDMXProtocolArtNetReceiver::Update(const FTimespan& SocketWaitTime)
 		
         if (Socket->RecvFrom(Reader->GetData(), Reader->Num(), Read, *ReceivedSenderInternetAddr))
 		{
-            Reader->RemoveAt(Read, Reader->Num() - Read, false);
+            Reader->RemoveAt(Read, Reader->Num() - Read, EAllowShrinking::No);
 			
 			DistributeReceivedData(Reader);
 		}
diff --git a/Engine/Plugins/VirtualProduction/DMX/DMXProtocol/Source/DMXProtocolSACN/Private/DMXProtocolSACNReceiver.cpp b/Engine/Plugins/VirtualProduction/DMX/DMXProtocol/Source/DMXProtocolSACN/Private/DMXProtocolSACNReceiver.cpp
--- a/Engine/Plugins/VirtualProduction/DMX/DMXProtocol/Source/DMXProtocolSACN/Private/DMXProtocolSACNReceiver.cpp
+++ b/Engine/Plugins/VirtualProduction/DMX/DMXProtocol/Source/DMXProtocolSACN/Private/DMXProtocolSACNReceiver.cpp
@@ -225,7 +225,7 @@ void FDMXProtocolSACNReceiver::Update(const FTimespan& SocketWaitTime)
 
 		if (Socket->RecvFromWithPktInfo(Reader->GetData(), Reader->Num(), NumBytesRead, *ReceivedSenderInternetAddr, *ReceivedDestinationInternetAddr))
 		{
-			Reader->RemoveAt(NumBytesRead, Reader->Num() - NumBytesRead, false);
+			Reader->RemoveAt(NumBytesRead, Reader->Num() - NumBytesRead, EAllowShrinking::No);
 			uint32 MulticastIp = 0;
 			ReceivedDestinationInternetAddr->GetIp(MulticastIp);
 
diff --git a/Engine/Plugins/VirtualProduction/RemoteControl/Source/WebRemoteControl/Private/WebRemoteControlInternalUtils.cpp b/Engine/Plugins/VirtualProduction/RemoteControl/Source/WebRemoteControl/Private/WebRemoteControlInternalUtils.cpp
--- a/Engine/Plugins/VirtualProduction/RemoteControl/Source/WebRemoteControl/Private/WebRemoteControlInternalUtils.cpp
+++ b/Engine/Plugins/VirtualProduction/RemoteControl/Source/WebRemoteControl/Private/WebRemoteControlInternalUtils.cpp
@@ -96,7 +96,7 @@ TSharedRef<FHttpServerRequest> UnwrapHttpRequest(const FRCRequestWrapper& Wrappe
 		FString PathWithoutParams = Wrapper.URL;
 
 		FString QueryParamsStr = PathWithoutParams.Mid(QueryParamsIndex + 1);
-		PathWithoutParams.MidInline(0, QueryParamsIndex, false);
+		PathWithoutParams.MidInline(0, QueryParamsIndex, EAllowShrinking::No);
 
 		// Split query params
 		TArray<FString> QueryParamPairs;
