[nDisplay] Changed/generalized the logic for ICVFX camera capture size

Instead of acquiring camera size from the camera settings, we now get it from
the RenderTargetSize. This allows to avoid any potential issues related to camera
buffer size computation (screen percentage, adopt resolution. overscan, etc.)
Also, dynamic camera resolution is supported now unless media device forbids that.

diff --git a/Engine/Plugins/Runtime/nDisplay/Source/DisplayClusterMedia/Private/Capture/DisplayClusterMediaCaptureCamera.cpp b/Engine/Plugins/Runtime/nDisplay/Source/DisplayClusterMedia/Private/Capture/DisplayClusterMediaCaptureCamera.cpp
--- a/Engine/Plugins/Runtime/nDisplay/Source/DisplayClusterMedia/Private/Capture/DisplayClusterMediaCaptureCamera.cpp
+++ b/Engine/Plugins/Runtime/nDisplay/Source/DisplayClusterMedia/Private/Capture/DisplayClusterMediaCaptureCamera.cpp
@@ -19,38 +19,4 @@ FDisplayClusterMediaCaptureCamera::FDisplayClusterMediaCaptureCamera(const FStri
 	: FDisplayClusterMediaCaptureViewport(InMediaId, InClusterNodeId, InViewportId, InMediaOutput, SyncPolicy)
 	, CameraId(InCameraId)
 {
-	if (const ADisplayClusterRootActor* const ActiveRootActor = IDisplayCluster::Get().GetGameMgr()->GetRootActor())
-	{
-		TArray<UActorComponent*> ICVFXCameraComponents;
-		IDisplayClusterViewportConfiguration* ViewportConfiguration = ActiveRootActor->GetViewportConfiguration();
-		if (const ADisplayClusterRootActor* const ConfigurationRootActor = ViewportConfiguration ? ViewportConfiguration->GetRootActor(EDisplayClusterRootActorType::Configuration) : nullptr)
-		{
-			ConfigurationRootActor->GetComponents(UDisplayClusterICVFXCameraComponent::StaticClass(), ICVFXCameraComponents);
-		}
-
-		for (const UActorComponent* const Component : ICVFXCameraComponents)
-		{
-			if (const UDisplayClusterICVFXCameraComponent* const ICVFXCamera = Cast<UDisplayClusterICVFXCameraComponent>(Component))
-			{
-				if (ICVFXCamera->GetName() == CameraId)
-				{
-					const FDisplayClusterConfigurationICVFX_CameraSettings& CameraSettings = ICVFXCamera->GetCameraSettingsICVFX();
-
-					if (CameraSettings.CustomFrustum.bEnable)
-					{
-						CameraResolution = CameraSettings.CustomFrustum.EstimatedOverscanResolution;
-					}
-					else
-					{
-						CameraResolution = CameraSettings.CustomFrustum.InnerFrustumResolution;
-					}
-				}
-			}
-		}
-	}
-}
-
-FIntPoint FDisplayClusterMediaCaptureCamera::GetCaptureSize() const
-{
-	return CameraResolution;
 }
diff --git a/Engine/Plugins/Runtime/nDisplay/Source/DisplayClusterMedia/Private/Capture/DisplayClusterMediaCaptureCamera.h b/Engine/Plugins/Runtime/nDisplay/Source/DisplayClusterMedia/Private/Capture/DisplayClusterMediaCaptureCamera.h
--- a/Engine/Plugins/Runtime/nDisplay/Source/DisplayClusterMedia/Private/Capture/DisplayClusterMediaCaptureCamera.h
+++ b/Engine/Plugins/Runtime/nDisplay/Source/DisplayClusterMedia/Private/Capture/DisplayClusterMediaCaptureCamera.h
@@ -18,11 +18,7 @@ class FDisplayClusterMediaCaptureCamera
 public:
 	FDisplayClusterMediaCaptureCamera(const FString& MediaId, const FString& ClusterNodeId, const FString& CameraId, const FString& ViewportId, UMediaOutput* MediaOutput, UDisplayClusterMediaOutputSynchronizationPolicy* SyncPolicy = nullptr);
 
-protected:
-	
-	virtual FIntPoint GetCaptureSize() const override;
-
 private:
+	/** ICVFX camera name */
 	const FString CameraId;
-	FIntPoint CameraResolution = FIntPoint::ZeroValue;
 };
