nDisplay:
- Fixed an optimization-related bug where the ICVFX camera was rendering in the background when using a full-frame chromakey. The expected behaviour was to disable camera rendering when a frame overlapped with a full-frame chromakey.
- Cleanup code in viewport rules related to media and tiles.

diff --git a/Engine/Plugins/Runtime/nDisplay/Source/DisplayCluster/Private/Render/Viewport/DisplayClusterViewport.h b/Engine/Plugins/Runtime/nDisplay/Source/DisplayCluster/Private/Render/Viewport/DisplayClusterViewport.h
--- a/Engine/Plugins/Runtime/nDisplay/Source/DisplayCluster/Private/Render/Viewport/DisplayClusterViewport.h
+++ b/Engine/Plugins/Runtime/nDisplay/Source/DisplayCluster/Private/Render/Viewport/DisplayClusterViewport.h
@@ -312,6 +312,12 @@ public:
 	/** Returns true if this viewport should be rendered. */
 	bool IsRenderEnabled() const;
 
+	/** Returns true if the rendering of this viewport is allowed by external media objects. */
+	bool IsRenderEnabledByMedia() const;
+
+	/** Returns true if this viewport used by external media objects. */
+	bool IsUsedByMedia() const;
+
 	/** Returns true if this viewport is to be used as a tile source. */
 	bool CanSplitIntoTiles() const;
 
diff --git a/Engine/Plugins/Runtime/nDisplay/Source/DisplayCluster/Private/Render/Viewport/DisplayClusterViewport_Rules.cpp b/Engine/Plugins/Runtime/nDisplay/Source/DisplayCluster/Private/Render/Viewport/DisplayClusterViewport_Rules.cpp
--- a/Engine/Plugins/Runtime/nDisplay/Source/DisplayCluster/Private/Render/Viewport/DisplayClusterViewport_Rules.cpp
+++ b/Engine/Plugins/Runtime/nDisplay/Source/DisplayCluster/Private/Render/Viewport/DisplayClusterViewport_Rules.cpp
@@ -110,14 +110,6 @@ bool FDisplayClusterViewport::IsExternalRendering() const
 		return true;
 	}
 
-	// Do not render the viewport if it is used as a media input.
-	if (RenderSettings.HasAnyMediaStates(EDisplayClusterViewportMediaState::Input))
-	{
-		// This viewport is not rendered but gets the image from the media.
-		// (media input replaces rendering.)
-		return true;
-	}
-
 	return false;
 }
 
@@ -125,42 +117,70 @@ bool FDisplayClusterViewport::IsRenderEnabled() const
 {
 	if (IsExternalRendering())
 	{
-		// The viewport is use external rendering
+		// The viewport uses an external rendering that overrides its RTT.
+		// Disable rendering for this viewport.
 		return false;
 	}
 
 	if (RenderSettings.bSkipRendering)
 	{
-		// Skip rendering
-		return true;
+		// Skip rendering.
+		// For example this feature is used when the ICVFX camera uses full-frame chromakey colour,
+		// thereby eliminating the ICVFX camera's viewport rendering for optimization.
+		return false;
 	}
 
-	// Tiles only evaluation
-	const EDisplayClusterViewportTileType TileType = RenderSettings.TileSettings.GetType();
-	if (TileType == EDisplayClusterViewportTileType::Source)
+	// Handle tile rendering rules.
+	switch(RenderSettings.TileSettings.GetType())
 	{
-		// SRC viewports should never get rendered
+	case EDisplayClusterViewportTileType::Tile:
+	case EDisplayClusterViewportTileType::None:
+		// When tile rendering is used, only tiles are rendered.
+		break;
+
+	default:
+		// When using tile rendering, other viewport types should never be rendered.
 		return false;
 	}
-	else if (TileType == EDisplayClusterViewportTileType::UnusedTile)
+
+	// Handle media rendering rules.
+	if (!IsRenderEnabledByMedia())
 	{
-		// UnusedTile is an internal state. We should never get it here.
+		// rendering of this viewport is not allowed by media.
 		return false;
 	}
-	else if (TileType == EDisplayClusterViewportTileType::Tile)
-	{
-		static const bool bIsCluster = (GDisplayCluster->GetOperationMode() == EDisplayClusterOperationMode::Cluster);
-		static const bool bIsOffscreen = FParse::Param(FCommandLine::Get(), TEXT("RenderOffscreen"));
 
-		const bool bHasInputAssigned  = RenderSettings.HasAnyMediaStates(EDisplayClusterViewportMediaState::Input);
-		const bool bHasOutputAssigned = RenderSettings.HasAnyMediaStates(EDisplayClusterViewportMediaState::Capture);
+	return true;
+}
 
-		// In non-cluster modes, always render (i.e. preview in editor)
-		if (!bIsCluster)
-		{
-			return true;
-		}
+bool FDisplayClusterViewport::IsUsedByMedia() const
+{
+	return RenderSettings.HasAnyMediaStates(EDisplayClusterViewportMediaState::Input | EDisplayClusterViewportMediaState::Capture);
+}
+
+bool FDisplayClusterViewport::IsRenderEnabledByMedia() const
+{
+	// Use this function to disable rendering of this viewport.
+	// All logic for handling media states should reside in this and IsUsedByMedia() functions only.
+
+	static const bool bIsCluster = (GDisplayCluster->GetOperationMode() == EDisplayClusterOperationMode::Cluster);
+	static const bool bIsOffscreen = FParse::Param(FCommandLine::Get(), TEXT("RenderOffscreen"));
 
+	const bool bHasInputAssigned = RenderSettings.HasAnyMediaStates(EDisplayClusterViewportMediaState::Input);
+	const bool bHasOutputAssigned = RenderSettings.HasAnyMediaStates(EDisplayClusterViewportMediaState::Capture);
+
+	// Do not render the viewport if it is used as a media input.
+	if (bHasInputAssigned)
+	{
+		// This viewport is not rendered but gets the image from the media.
+		// (media input replaces rendering.)
+		return false;
+	}
+
+	// Applying special media rules for tiles for cluster modes only.
+	if (bIsCluster && !bHasOutputAssigned
+		&& RenderSettings.TileSettings.GetType() == EDisplayClusterViewportTileType::Tile)
+	{
 		// Tiles logic is a little more complicated
 		//
 		// In | Out | Render (offscreen) | Render (onscreen)
@@ -169,21 +189,14 @@ bool FDisplayClusterViewport::IsRenderEnabled() const
 		//  0 |  1  |        1           | 1
 		//  1 |  0  |        0           | 0
 		//  1 |  1  |        0           | 0
-
-		if (bHasInputAssigned)
+		// Check if current node was explicitly allowed to render unbound tiles
+		const bool bShouldRenderUnbound = RenderSettings.TileSettings.HasAnyTileFlags(EDisplayClusterViewportTileFlags::AllowUnboundRender);
+		const bool bShouldRender = !bIsOffscreen && bShouldRenderUnbound;
+		if (!bShouldRender)
 		{
-			// This should have been validated previously in IsExternalRendering() but let's keep this condition
-			// active in case something is changed outside (i.e. validation order or condition changes).
+			// Disable tile viewport rendering.
 			return false;
 		}
-		else if (!bHasInputAssigned && !bHasOutputAssigned)
-		{
-			// Check if current node was explicitly allowed to render unbound tiles
-			const bool bShouldRenderUnbound = RenderSettings.TileSettings.HasAnyTileFlags(EDisplayClusterViewportTileFlags::AllowUnboundRender);
-			const bool bShouldRender = !bIsOffscreen && bShouldRenderUnbound;
-
-			return bShouldRender;
-		}
 	}
 
 	return true;
@@ -214,20 +227,27 @@ bool FDisplayClusterViewport::CanSplitIntoTiles() const
 		return false;
 	}
 
+	// Ignore viewports that used by media
+	if (IsUsedByMedia())
+	{
+		return false;
+	}
+
 	return true;
 }
 
 bool FDisplayClusterViewport::ShouldUseRenderTargetResource() const
 {
-	if (IsExternalRendering())
+	if (RenderSettings.bSkipRendering)
 	{
-		// Exceptions to the rules for a viewport that uses external rendering:
-		if (RenderSettings.HasAnyMediaStates(EDisplayClusterViewportMediaState::Input))
-		{
-			// Media input uses the resources of the viewport.
-			return true;
-		}
+		// When rendering is skipped, the RTT resources are not used.
+		// For example this feature is used when the ICVFX camera uses full-frame chromakey colour,
+		// thereby eliminating the ICVFX camera's viewport rendering resources for optimization.
+		return false;
+	}
 
+	if (IsExternalRendering())
+	{
 		return false;
 	}
 
@@ -239,6 +259,8 @@ bool FDisplayClusterViewport::ShouldUseInternalResources() const
 	if (RenderSettings.bSkipRendering)
 	{
 		// When rendering is skipped, internal resources are not used.
+		// For example this feature is used when the ICVFX camera uses full-frame chromakey colour,
+		// thereby eliminating the ICVFX camera's viewport rendering resources for optimization.
 		return false;
 	}
 
