ndispaly. Fixed bug: changed the return result of viewport gathering funcion for tiled media output.

diff --git a/Engine/Plugins/Runtime/nDisplay/Source/DisplayClusterConfiguration/Private/DisplayClusterConfigurationTypes_Media.cpp b/Engine/Plugins/Runtime/nDisplay/Source/DisplayClusterConfiguration/Private/DisplayClusterConfigurationTypes_Media.cpp
--- a/Engine/Plugins/Runtime/nDisplay/Source/DisplayClusterConfiguration/Private/DisplayClusterConfigurationTypes_Media.cpp
+++ b/Engine/Plugins/Runtime/nDisplay/Source/DisplayClusterConfiguration/Private/DisplayClusterConfigurationTypes_Media.cpp
@@ -180,7 +180,7 @@ bool FDisplayClusterConfigurationMediaICVFX::GetMediaInputTiles(const FString& N
 		OutTiles.Append(TiledInputGroup.Tiles);
 	}
 
-	return false;
+	return true;
 }
 
 //@note: Tiled output, it's allowed to use the same node ID in different groups
@@ -206,5 +206,5 @@ bool FDisplayClusterConfigurationMediaICVFX::GetMediaOutputTiles(const FString&
 		OutTiles.Append(TiledOutputGroup.Tiles);
 	}
 
-	return false;
+	return true;
 }
