[nDisplay] Another minor change to the tiled media configuration data gathering

- Don't clean the caller's buffer. Let the caller be able to collect multiple nodes to the same buffer my multiple GetMediaInputTiles() calls.

diff --git a/Engine/Plugins/Runtime/nDisplay/Source/DisplayClusterConfiguration/Private/DisplayClusterConfigurationTypes_Media.cpp b/Engine/Plugins/Runtime/nDisplay/Source/DisplayClusterConfiguration/Private/DisplayClusterConfigurationTypes_Media.cpp
--- a/Engine/Plugins/Runtime/nDisplay/Source/DisplayClusterConfiguration/Private/DisplayClusterConfigurationTypes_Media.cpp
+++ b/Engine/Plugins/Runtime/nDisplay/Source/DisplayClusterConfiguration/Private/DisplayClusterConfigurationTypes_Media.cpp
@@ -163,8 +163,6 @@ bool FDisplayClusterConfigurationMediaICVFX::GetMediaInputTiles(const FString& N
 		return false;
 	}
 
-	OutTiles.Empty();
-
 	// Find all input groups bound to the node ID specified
 	const TArray<FDisplayClusterConfigurationMediaTiledInputGroup> FoundInputGroups = TiledMediaInputGroups.FilterByPredicate([&NodeId](const FDisplayClusterConfigurationMediaTiledInputGroup& TiledInputGroup)
 		{
