nDisplay: Fixed an bug when post-processes of the main viewport were not copied to tiled viewports

diff --git a/Engine/Plugins/Runtime/nDisplay/Source/DisplayCluster/Private/Render/Viewport/Configuration/DisplayClusterViewportConfigurationHelpers_Tile.cpp b/Engine/Plugins/Runtime/nDisplay/Source/DisplayCluster/Private/Render/Viewport/Configuration/DisplayClusterViewportConfigurationHelpers_Tile.cpp
--- a/Engine/Plugins/Runtime/nDisplay/Source/DisplayCluster/Private/Render/Viewport/Configuration/DisplayClusterViewportConfigurationHelpers_Tile.cpp
+++ b/Engine/Plugins/Runtime/nDisplay/Source/DisplayCluster/Private/Render/Viewport/Configuration/DisplayClusterViewportConfigurationHelpers_Tile.cpp
@@ -210,6 +210,7 @@ FDisplayClusterViewport* FDisplayClusterViewportConfigurationHelpers_Tile::GetOr
 		InOutRenderSettings.TileSettings = FDisplayClusterViewport_TileSettings(InSourceViewport.GetId(), InTilePos, InTileSize, TileFlags);
 
 		// Copy internal render settings from the source:
+		TileViewport->GetCustomPostProcessSettings() = InSourceViewport.GetCustomPostProcessSettings();
 		TileViewport->GetVisibilitySettingsImpl() = InSourceViewport.GetVisibilitySettingsImpl();
 		TileViewport->GetCameraMotionBlurImpl() = InSourceViewport.GetCameraMotionBlurImpl();
 		TileViewport->GetCameraDepthOfFieldImpl() = InSourceViewport.GetCameraDepthOfFieldImpl();
diff --git a/Engine/Plugins/Runtime/nDisplay/Source/DisplayCluster/Private/Render/Viewport/DisplayClusterViewport.h b/Engine/Plugins/Runtime/nDisplay/Source/DisplayCluster/Private/Render/Viewport/DisplayClusterViewport.h
--- a/Engine/Plugins/Runtime/nDisplay/Source/DisplayCluster/Private/Render/Viewport/DisplayClusterViewport.h
+++ b/Engine/Plugins/Runtime/nDisplay/Source/DisplayCluster/Private/Render/Viewport/DisplayClusterViewport.h
@@ -297,6 +297,13 @@ public:
 		return CameraDepthOfField;
 	}
 
+	/** Gain direct access to internal PostProcess data of the viewport. */
+	FDisplayClusterViewport_CustomPostProcessSettings& GetCustomPostProcessSettings()
+	{
+		check(IsInGameThread());
+		return CustomPostProcessSettings;
+	}
+
 	/** Some viewports are used as internal and skip some logic steps.
 	* These viewports are handled separately from regular viewports.
 	* Context: icvfx, tile
