nDisplay: Created a new category in the ICVFX camera details panel called Camera which contains the camera's FIZ properties as well as motion blur and depth of field properties. Also set the default depth of field LUT to be the point RMS LUT instead of null, and renamed the In-Camera VFX category to Inner Frustum and shortened the color grading category name

diff --git a/Engine/Plugins/Runtime/nDisplay/Source/DisplayCluster/Private/Game/EngineClasses/Scene/DisplayClusterICVFXCameraComponent.cpp b/Engine/Plugins/Runtime/nDisplay/Source/DisplayCluster/Private/Game/EngineClasses/Scene/DisplayClusterICVFXCameraComponent.cpp
--- a/Engine/Plugins/Runtime/nDisplay/Source/DisplayCluster/Private/Game/EngineClasses/Scene/DisplayClusterICVFXCameraComponent.cpp
+++ b/Engine/Plugins/Runtime/nDisplay/Source/DisplayCluster/Private/Game/EngineClasses/Scene/DisplayClusterICVFXCameraComponent.cpp
@@ -59,6 +59,13 @@ void UDisplayClusterICVFXCameraComponent::PostLoad()
 	PRAGMA_ENABLE_DEPRECATION_WARNINGS
 }
 
+void UDisplayClusterICVFXCameraComponent::PostApplyToComponent()
+{
+	Super::PostApplyToComponent();
+
+	CameraSettings.CameraDepthOfField.UpdateDynamicCompensationLUT();
+}
+
 void UDisplayClusterICVFXCameraComponent::GetCameraView(float DeltaTime, FMinimalViewInfo& InOutViewInfo)
 {
 	const ADisplayClusterRootActor* RootActor = Cast<ADisplayClusterRootActor>(GetOwner());
@@ -226,7 +233,7 @@ FDisplayClusterViewport_CameraDepthOfField UDisplayClusterICVFXCameraComponent::
 	OutParameters.bEnableDepthOfFieldCompensation = CameraSettings.CameraDepthOfField.bEnableDepthOfFieldCompensation;
 	OutParameters.DistanceToWall = CameraSettings.CameraDepthOfField.DistanceToWall;
 	OutParameters.DistanceToWallOffset = CameraSettings.CameraDepthOfField.DistanceToWallOffset;
-	OutParameters.CompensationLUT = CameraSettings.CameraDepthOfField.DynamicCompensationLUT ? CameraSettings.CameraDepthOfField.DynamicCompensationLUT : CameraSettings.CameraDepthOfField.CompensationLUT;
+	OutParameters.CompensationLUT = CameraSettings.CameraDepthOfField.DynamicCompensationLUT ? CameraSettings.CameraDepthOfField.DynamicCompensationLUT : CameraSettings.CameraDepthOfField.CompensationLUT.Get();
 
 	return OutParameters;
 }
@@ -235,7 +242,14 @@ void UDisplayClusterICVFXCameraComponent::OnRegister()
 {
 	Super::OnRegister();
 
-	CameraSettings.CameraDepthOfField.UpdateDynamicCompensationLUT();
+	// If the blueprint is being reconstructed, we can't update the dynamic LUT here without causing issues
+	// when the reconstruction attempts to check if the component's properties are modified, as this call will
+	// load the compensation LUT soft pointer, resulting in a memory difference from the archetype.
+	// The PostApplyToComponent call handles rebuilding the dynamic LUT in such a case
+	if (!GIsReconstructingBlueprintInstances)
+	{
+		CameraSettings.CameraDepthOfField.UpdateDynamicCompensationLUT();
+	}
 
 #if WITH_EDITORONLY_DATA
 	// disable frustum for icvfx camera component
diff --git a/Engine/Plugins/Runtime/nDisplay/Source/DisplayCluster/Public/Components/DisplayClusterICVFXCameraComponent.h b/Engine/Plugins/Runtime/nDisplay/Source/DisplayCluster/Public/Components/DisplayClusterICVFXCameraComponent.h
--- a/Engine/Plugins/Runtime/nDisplay/Source/DisplayCluster/Public/Components/DisplayClusterICVFXCameraComponent.h
+++ b/Engine/Plugins/Runtime/nDisplay/Source/DisplayCluster/Public/Components/DisplayClusterICVFXCameraComponent.h
@@ -33,6 +33,7 @@ public:
 
 	virtual void Serialize(FArchive& Ar) override;
 	virtual void PostLoad() override;
+	virtual void PostApplyToComponent() override;
 
 public:
 	UPROPERTY(BlueprintReadWrite, EditAnywhere, Category = NDisplay)
@@ -104,44 +105,44 @@ public:
 private:
 	friend class FDisplayClusterICVFXCameraComponentDetailsCustomization;
 	
-	UPROPERTY(EditAnywhere, Transient, Category = "In Camera VFX", meta = (PropertyPath = "CameraSettings.bEnable"))
+	UPROPERTY(EditAnywhere, Transient, Category = "Inner Frustum", meta = (PropertyPath = "CameraSettings.bEnable"))
 	FDisplayClusterEditorPropertyReference IsEnabledRef;
 
-	UPROPERTY(EditAnywhere, Transient, Category = "In Camera VFX", meta = (PropertyPath = "CameraSettings.HiddenICVFXViewports"))
+	UPROPERTY(EditAnywhere, Transient, Category = "Inner Frustum", meta = (PropertyPath = "CameraSettings.HiddenICVFXViewports"))
 	FDisplayClusterEditorPropertyReference HiddenICVFXViewportsRef;
 
-	UPROPERTY(EditAnywhere, Transient, Category = "In Camera VFX", meta = (PropertyPath = "CameraSettings.ExternalCameraActor"))
-	FDisplayClusterEditorPropertyReference ExternalCameraActorRef;
-
-	UPROPERTY(EditAnywhere, Transient, Category = "In Camera VFX", meta = (PropertyPath = "CameraSettings.BufferRatio"))
+	UPROPERTY(EditAnywhere, Transient, Category = "Inner Frustum", meta = (PropertyPath = "CameraSettings.BufferRatio"))
 	FDisplayClusterEditorPropertyReference BufferRatioRef;
 
-	UPROPERTY(EditAnywhere, Transient, Category = "In Camera VFX", meta = (DisplayName = "Inner Frustum Overscan", PropertyPath = "CameraSettings.CustomFrustum"))
+	UPROPERTY(EditAnywhere, Transient, Category = "Inner Frustum", meta = (DisplayName = "Inner Frustum Overscan", PropertyPath = "CameraSettings.CustomFrustum"))
 	FDisplayClusterEditorPropertyReference CustomFrustumRef;
 
-	UPROPERTY(EditAnywhere, Transient, Category = "In Camera VFX", meta = (PropertyPath = "CameraSettings.SoftEdge"))
+	UPROPERTY(EditAnywhere, Transient, Category = "Inner Frustum", meta = (PropertyPath = "CameraSettings.SoftEdge"))
 	FDisplayClusterEditorPropertyReference SoftEdgeRef;
 
-	UPROPERTY(EditAnywhere, Transient, Category = "In Camera VFX", meta = (PropertyPath = "CameraSettings.Border"))
+	UPROPERTY(EditAnywhere, Transient, Category = "Inner Frustum", meta = (PropertyPath = "CameraSettings.Border"))
 	FDisplayClusterEditorPropertyReference BorderRef;
 
-	UPROPERTY(EditAnywhere, Transient, Category = "In Camera VFX", meta = (PropertyPath = "CameraSettings.FrustumRotation"))
+	UPROPERTY(EditAnywhere, Transient, Category = "Inner Frustum", meta = (PropertyPath = "CameraSettings.FrustumRotation"))
 	FDisplayClusterEditorPropertyReference FrustumRotationRef;
 
-	UPROPERTY(EditAnywhere, Transient, Category = "In Camera VFX", meta = (PropertyPath = "CameraSettings.FrustumOffset"))
+	UPROPERTY(EditAnywhere, Transient, Category = "Inner Frustum", meta = (PropertyPath = "CameraSettings.FrustumOffset"))
 	FDisplayClusterEditorPropertyReference FrustumOffsetRef;
 
-	UPROPERTY(EditAnywhere, Transient, Category = "In Camera VFX", meta = (PropertyPath = "CameraSettings.RenderSettings.GenerateMips"))
+	UPROPERTY(EditAnywhere, Transient, Category = "Inner Frustum", meta = (PropertyPath = "CameraSettings.RenderSettings.GenerateMips"))
 	FDisplayClusterEditorPropertyReference GenerateMipsRef;
 
-	UPROPERTY(EditAnywhere, Transient, Category = "In Camera VFX", meta = (PropertyPath = "CameraSettings.CameraMotionBlur"))
-	FDisplayClusterEditorPropertyReference CameraMotionBlurRef;
+	UPROPERTY(EditAnywhere, Transient, Category = "ICVFX Camera", meta = (PropertyPath = "CameraSettings.ExternalCameraActor"))
+	FDisplayClusterEditorPropertyReference ExternalCameraActorRef;
 
 	/** Exposed reference to the camera's inner depth of field settings */
-	UPROPERTY(EditAnywhere, Transient, Category = "In Camera VFX", meta = (PropertyPath = "CameraSettings.CameraDepthOfField"))
+	UPROPERTY(EditAnywhere, Transient, Category = "ICVFX Camera", meta = (PropertyPath = "CameraSettings.CameraDepthOfField", DisplayName = "ICVFX Depth of Field"))
 	FDisplayClusterEditorPropertyReference CameraDepthOfFieldRef;
 
-	UPROPERTY(EditAnywhere, Transient, Category = "In Camera VFX", meta = (PropertyPath = "CameraSettings.CameraHideList"))
+	UPROPERTY(EditAnywhere, Transient, Category = "ICVFX Camera", meta = (PropertyPath = "CameraSettings.CameraMotionBlur", DisplayName = "ICVFX Camera Motion Blur"))
+	FDisplayClusterEditorPropertyReference CameraMotionBlurRef;
+
+	UPROPERTY(EditAnywhere, Transient, Category = "Inner Frustum", meta = (PropertyPath = "CameraSettings.CameraHideList"))
 	FDisplayClusterEditorPropertyReference CameraHideListRef;
 
 	UPROPERTY(EditAnywhere, Transient, Category = Chromakey, meta = (PropertyPath = "CameraSettings.Chromakey.bEnable"))
@@ -180,13 +181,13 @@ private:
 	UPROPERTY(EditAnywhere, Transient, Category = "Tile Rendering", meta = (PropertyPath = "CameraSettings.CameraTile.ClusterNodesToRenderUnboundTiles", DisplayName = "Nodes To Render Unbound Tiles", ToolTip = "Choose nodes that should render camera tiles that don't have any media assigned"))
 	FDisplayClusterEditorPropertyReference CameraTileNodesToRenderUnboundTilesRef;
 
-	UPROPERTY(EditAnywhere, Transient, Category = "Inner Frustum Color Grading", meta = (PropertyPath = "CameraSettings.EnableInnerFrustumColorGrading", DisplayName = "Enable Inner Frustum Color Grading"))
+	UPROPERTY(EditAnywhere, Transient, Category = "Color Grading", meta = (PropertyPath = "CameraSettings.EnableInnerFrustumColorGrading", DisplayName = "Enable Inner Frustum Color Grading"))
 	FDisplayClusterEditorPropertyReference EnableInnerFrustumColorGrading;
 
-	UPROPERTY(EditAnywhere, Transient, Category = "Inner Frustum Color Grading", meta = (PropertyPath = "CameraSettings.AllNodesColorGrading", DisplayName = "All Nodes"))
+	UPROPERTY(EditAnywhere, Transient, Category = "Color Grading", meta = (PropertyPath = "CameraSettings.AllNodesColorGrading", DisplayName = "All Nodes"))
 	FDisplayClusterEditorPropertyReference AllNodesColorGradingRef;
 
-	UPROPERTY(EditAnywhere, Transient, Category = "Inner Frustum Color Grading", meta = (PropertyPath = "CameraSettings.PerNodeColorGrading"))
+	UPROPERTY(EditAnywhere, Transient, Category = "Color Grading", meta = (PropertyPath = "CameraSettings.PerNodeColorGrading"))
 	FDisplayClusterEditorPropertyReference PerNodeColorGradingRef;
 
 	UPROPERTY(EditAnywhere, Transient, Category = "Texture Replacement", meta = (PropertyPath = "CameraSettings.RenderSettings.Replace.bAllowReplace", DisplayName = "Enable Inner Frustum Texture Replacement", ToolTip = "Set to True to replace the entire inner frustum with the specified texture."))
diff --git a/Engine/Plugins/Runtime/nDisplay/Source/DisplayClusterConfiguration/Private/DisplayClusterConfigurationTypes_ICVFX.cpp b/Engine/Plugins/Runtime/nDisplay/Source/DisplayClusterConfiguration/Private/DisplayClusterConfigurationTypes_ICVFX.cpp
--- a/Engine/Plugins/Runtime/nDisplay/Source/DisplayClusterConfiguration/Private/DisplayClusterConfigurationTypes_ICVFX.cpp
+++ b/Engine/Plugins/Runtime/nDisplay/Source/DisplayClusterConfiguration/Private/DisplayClusterConfigurationTypes_ICVFX.cpp
@@ -585,6 +585,7 @@ bool FDisplayClusterConfigurationICVFX_ChromakeyRenderSettings::ShouldUseChromak
 
 void FDisplayClusterConfigurationICVFX_CameraDepthOfField::UpdateDynamicCompensationLUT()
 {
+	CompensationLUT.LoadSynchronous();
 	if (CompensationLUT)
 	{
 		FSharedImageConstRef CPUTextureRef = CompensationLUT->GetCPUCopy();
diff --git a/Engine/Plugins/Runtime/nDisplay/Source/DisplayClusterConfiguration/Public/DisplayClusterConfigurationStrings.h b/Engine/Plugins/Runtime/nDisplay/Source/DisplayClusterConfiguration/Public/DisplayClusterConfigurationStrings.h
--- a/Engine/Plugins/Runtime/nDisplay/Source/DisplayClusterConfiguration/Public/DisplayClusterConfigurationStrings.h
+++ b/Engine/Plugins/Runtime/nDisplay/Source/DisplayClusterConfiguration/Public/DisplayClusterConfigurationStrings.h
@@ -25,14 +25,15 @@ namespace DisplayClusterConfigurationStrings
 		static constexpr const TCHAR* ClusterConfigurationCategory = TEXT("NDisplay Cluster Configuration");
 		static constexpr const TCHAR* ClusterPostprocessCategory   = TEXT("Post Process");
 		static constexpr const TCHAR* ColorGradingCategory         = TEXT("Color Grading");
-		static constexpr const TCHAR* CameraColorGradingCategory   = TEXT("Inner Frustum Color Grading");
+		static constexpr const TCHAR* CameraColorGradingCategory   = TEXT("Color Grading");
 		static constexpr const TCHAR* ChromaKeyCategory            = TEXT("Chromakey");
 		static constexpr const TCHAR* LightcardCategory            = TEXT("Light Cards");
 		static constexpr const TCHAR* OCIOCategory                 = TEXT("OCIO");
 		static constexpr const TCHAR* TileCategory                 = TEXT("Tile Rendering");
 		static constexpr const TCHAR* OverrideCategory             = TEXT("Texture Replacement");
 		static constexpr const TCHAR* ViewportsCategory            = TEXT("Viewports");
-		static constexpr const TCHAR* ICVFXCategory                = TEXT("In Camera VFX");
+		static constexpr const TCHAR* ICVFXCategory                = TEXT("Inner Frustum");
+		static constexpr const TCHAR* ICVFXCameraCategory          = TEXT("ICVFX Camera");
 		static constexpr const TCHAR* ConfigurationCategory        = TEXT("Configuration");
 		static constexpr const TCHAR* PreviewCategory              = TEXT("Editor Preview");
 		static constexpr const TCHAR* PreviewInGameCategory        = TEXT("Preview In Game");
diff --git a/Engine/Plugins/Runtime/nDisplay/Source/DisplayClusterConfiguration/Public/DisplayClusterConfigurationTypes_ICVFX.h b/Engine/Plugins/Runtime/nDisplay/Source/DisplayClusterConfiguration/Public/DisplayClusterConfigurationTypes_ICVFX.h
--- a/Engine/Plugins/Runtime/nDisplay/Source/DisplayClusterConfiguration/Public/DisplayClusterConfigurationTypes_ICVFX.h
+++ b/Engine/Plugins/Runtime/nDisplay/Source/DisplayClusterConfiguration/Public/DisplayClusterConfigurationTypes_ICVFX.h
@@ -616,7 +616,7 @@ public:
 
 	/** Look-up texture that encodes the specific amount of compensation used for each combination of wall distance and object distance */
 	UPROPERTY(BlueprintReadWrite, EditAnywhere, Category = NDisplay)
-	UTexture2D* CompensationLUT = nullptr;
+	TSoftObjectPtr<UTexture2D> CompensationLUT = TSoftObjectPtr<UTexture2D>(FSoftObjectPath(TEXT("/nDisplay/DepthOfField/T_LUT_PointRMS.T_LUT_PointRMS")));
 
 	/** Actual LUT to use with the depth of field pipeline, copied and modified from CompensationLUT */
 	UPROPERTY(Transient)
diff --git a/Engine/Plugins/Runtime/nDisplay/Source/DisplayClusterConfigurator/Private/DisplayClusterConfiguratorModule.cpp b/Engine/Plugins/Runtime/nDisplay/Source/DisplayClusterConfigurator/Private/DisplayClusterConfiguratorModule.cpp
--- a/Engine/Plugins/Runtime/nDisplay/Source/DisplayClusterConfigurator/Private/DisplayClusterConfiguratorModule.cpp
+++ b/Engine/Plugins/Runtime/nDisplay/Source/DisplayClusterConfigurator/Private/DisplayClusterConfiguratorModule.cpp
@@ -286,12 +286,17 @@ void FDisplayClusterConfiguratorModule::RegisterSectionMappings()
 	// ICVFX Component
 	{
 		const TSharedRef<FPropertySection> Section = PropertyModule.FindOrCreateSection(UDisplayClusterICVFXCameraComponent::StaticClass()->GetFName(),
-		DisplayClusterConfigurationStrings::categories::ICVFXCategory, LOCTEXT("In-Camera VFX", "In-Camera VFX"));
+		DisplayClusterConfigurationStrings::categories::ICVFXCategory, LOCTEXT("InnerFrustumSectionLabel", "Inner Frustum"));
 		Section->AddCategory(DisplayClusterConfigurationStrings::categories::ICVFXCategory);
 	}
 	{
 		const TSharedRef<FPropertySection> Section = PropertyModule.FindOrCreateSection(UDisplayClusterICVFXCameraComponent::StaticClass()->GetFName(),
-		DisplayClusterConfigurationStrings::categories::CameraColorGradingCategory, LOCTEXT("Inner Frustum Color Grading", "Inner Frustum Color Grading"));
+			DisplayClusterConfigurationStrings::categories::ICVFXCameraCategory, LOCTEXT("InnerFrustumCameraSectionLabel", "Camera"));
+		Section->AddCategory(DisplayClusterConfigurationStrings::categories::ICVFXCameraCategory);
+	}
+	{
+		const TSharedRef<FPropertySection> Section = PropertyModule.FindOrCreateSection(UDisplayClusterICVFXCameraComponent::StaticClass()->GetFName(),
+		DisplayClusterConfigurationStrings::categories::CameraColorGradingCategory, LOCTEXT("InnerFrustumColorGradingLabel", "Color Grading"));
 		Section->AddCategory(DisplayClusterConfigurationStrings::categories::CameraColorGradingCategory);
 	}
 	{
diff --git a/Engine/Plugins/Runtime/nDisplay/Source/DisplayClusterConfigurator/Private/Views/Details/Components/DisplayClusterICVFXCameraComponentDetailsCustomization.cpp b/Engine/Plugins/Runtime/nDisplay/Source/DisplayClusterConfigurator/Private/Views/Details/Components/DisplayClusterICVFXCameraComponentDetailsCustomization.cpp
--- a/Engine/Plugins/Runtime/nDisplay/Source/DisplayClusterConfigurator/Private/Views/Details/Components/DisplayClusterICVFXCameraComponentDetailsCustomization.cpp
+++ b/Engine/Plugins/Runtime/nDisplay/Source/DisplayClusterConfigurator/Private/Views/Details/Components/DisplayClusterICVFXCameraComponentDetailsCustomization.cpp
@@ -6,6 +6,7 @@
 #include "Components/DisplayClusterICVFXCameraComponent.h"
 
 #include "CineCameraActor.h"
+#include "CineCameraComponent.h"
 #include "DisplayClusterRootActor.h"
 #include "UObject/SoftObjectPtr.h"
 
@@ -26,6 +27,7 @@ namespace DisplayClusterICVFXCameraComponentDetailsCustomizationUtils
 			TEXT("Variable"),
 			TEXT("TransformCommon"),
 			DisplayClusterConfigurationStrings::categories::ICVFXCategory,
+			DisplayClusterConfigurationStrings::categories::ICVFXCameraCategory,
 			DisplayClusterConfigurationStrings::categories::CameraColorGradingCategory,
 			DisplayClusterConfigurationStrings::categories::OCIOCategory,
 			DisplayClusterConfigurationStrings::categories::ChromaKeyCategory,
@@ -94,8 +96,24 @@ void FDisplayClusterICVFXCameraComponentDetailsCustomization::CustomizeDetails(I
 	// Sockets category must be hidden manually instead of through the HideCategories metadata specifier
 	InLayoutBuilder.HideCategory(TEXT("Sockets"));
 
-	// Manually label the ICVFX category to properly format it to have the dash in "In-Camera"
-	InLayoutBuilder.EditCategory(DisplayClusterConfigurationStrings::categories::ICVFXCategory, LOCTEXT("ICVFXCategoryLabel", "In-Camera VFX"));
+	IDetailCategoryBuilder& CameraCategory = InLayoutBuilder.EditCategory(DisplayClusterConfigurationStrings::categories::ICVFXCameraCategory, LOCTEXT("ICVFXCameraCategoryLabel", "Camera"));
+
+	// Re-add the external camera to the category to ensure it is always above the camera's fiz properties in the details panel
+	CameraCategory.AddProperty(GET_MEMBER_NAME_CHECKED(UDisplayClusterICVFXCameraComponent, ExternalCameraActorRef));
+
+	if (EditedObject.IsValid() && EditedObject->CameraSettings.ExternalCameraActor.IsValid())
+	{
+		TArray<UObject*> ExternalCameraComponents = { EditedObject->CameraSettings.ExternalCameraActor.Get()->GetCineCameraComponent() };
+		CameraCategory.AddExternalObjectProperty(ExternalCameraComponents, GET_MEMBER_NAME_CHECKED(UCineCameraComponent, FocusSettings));
+		CameraCategory.AddExternalObjectProperty(ExternalCameraComponents, GET_MEMBER_NAME_CHECKED(UCineCameraComponent, CurrentFocalLength));
+		CameraCategory.AddExternalObjectProperty(ExternalCameraComponents, GET_MEMBER_NAME_CHECKED(UCineCameraComponent, CurrentAperture));
+	}
+	else
+	{
+		CameraCategory.AddProperty(GET_MEMBER_NAME_CHECKED(UDisplayClusterICVFXCameraComponent, FocusSettings), UCineCameraComponent::StaticClass());
+		CameraCategory.AddProperty(GET_MEMBER_NAME_CHECKED(UDisplayClusterICVFXCameraComponent, CurrentFocalLength), UCineCameraComponent::StaticClass());
+		CameraCategory.AddProperty(GET_MEMBER_NAME_CHECKED(UDisplayClusterICVFXCameraComponent, CurrentAperture), UCineCameraComponent::StaticClass());
+	}
 
 	InLayoutBuilder.SortCategories(DisplayClusterICVFXCameraComponentDetailsCustomizationUtils::SortCategories);
 
