[nDisplay] Tiled media GUI redesign
 - Redesigned media settings details for ICVFX cameras
 - Added "Media" filter for ICVFX cameras

diff --git a/Engine/Plugins/Runtime/nDisplay/Source/DisplayCluster/Private/Render/Viewport/Configuration/DisplayClusterViewportConfigurationHelpers_ICVFX.cpp b/Engine/Plugins/Runtime/nDisplay/Source/DisplayCluster/Private/Render/Viewport/Configuration/DisplayClusterViewportConfigurationHelpers_ICVFX.cpp
--- a/Engine/Plugins/Runtime/nDisplay/Source/DisplayCluster/Private/Render/Viewport/Configuration/DisplayClusterViewportConfigurationHelpers_ICVFX.cpp
+++ b/Engine/Plugins/Runtime/nDisplay/Source/DisplayCluster/Private/Render/Viewport/Configuration/DisplayClusterViewportConfigurationHelpers_ICVFX.cpp
@@ -545,7 +545,7 @@ void FDisplayClusterViewportConfigurationHelpers_ICVFX::UpdateCameraViewportSett
 	DstViewport.SetViewportBufferRatio(InCameraSettings.GetCameraBufferRatio(*StageSettings));
 
 	// InCamera tile rendering.
-	FDisplayClusterViewportConfigurationHelpers_Tile::UpdateICVFXCameraViewportTileSettings(DstViewport, InCameraSettings.CameraTile);
+	FDisplayClusterViewportConfigurationHelpers_Tile::UpdateICVFXCameraViewportTileSettings(DstViewport, InCameraSettings.RenderSettings.Media);
 }
 
 void FDisplayClusterViewportConfigurationHelpers_ICVFX::UpdateChromakeyViewportSettings(FDisplayClusterViewport& DstViewport, FDisplayClusterViewport& InCameraViewport, const FDisplayClusterConfigurationICVFX_CameraSettings& InCameraSettings)
diff --git a/Engine/Plugins/Runtime/nDisplay/Source/DisplayCluster/Private/Render/Viewport/Configuration/DisplayClusterViewportConfigurationHelpers_Tile.cpp b/Engine/Plugins/Runtime/nDisplay/Source/DisplayCluster/Private/Render/Viewport/Configuration/DisplayClusterViewportConfigurationHelpers_Tile.cpp
--- a/Engine/Plugins/Runtime/nDisplay/Source/DisplayCluster/Private/Render/Viewport/Configuration/DisplayClusterViewportConfigurationHelpers_Tile.cpp
+++ b/Engine/Plugins/Runtime/nDisplay/Source/DisplayCluster/Private/Render/Viewport/Configuration/DisplayClusterViewportConfigurationHelpers_Tile.cpp
@@ -14,6 +14,7 @@
 
 #include "DisplayClusterEnums.h"
 #include "DisplayClusterConfigurationTypes_ICVFX.h"
+#include "DisplayClusterConfigurationTypes_Media.h"
 #include "DisplayClusterConfigurationTypes_Tile.h"
 #include "DisplayClusterProjectionStrings.h"
 
@@ -21,43 +22,55 @@
 // FDisplayClusterViewportConfigurationHelpers_Tile
 ////////////////////////////////////////////////////////////////////////
 
-void FDisplayClusterViewportConfigurationHelpers_Tile::UpdateICVFXCameraViewportTileSettings(FDisplayClusterViewport& InSourceViewport, const FDisplayClusterConfigurationICVFX_CameraTile& InCameraTile)
+void FDisplayClusterViewportConfigurationHelpers_Tile::UpdateICVFXCameraViewportTileSettings(FDisplayClusterViewport& InSourceViewport, const FDisplayClusterConfigurationMediaICVFX& InCameraMediaSettings)
 {
-	if (const FDisplayClusterConfigurationICVFX_StageSettings* StageSettings = InSourceViewport.Configuration->GetStageSettings())
+	// Nothing to do if media tiling is not configured
+	const bool bMediaEnabled = InCameraMediaSettings.bEnable;
+	const bool bMediaTiled = (InCameraMediaSettings.SplitType == EDisplayClusterConfigurationMediaSplitType::UniformTiles);
+	if (!bMediaEnabled || !bMediaTiled)
 	{
-		if (InCameraTile.TileSettings.IsEnabled(*StageSettings))
+		InSourceViewport.GetRenderSettingsImpl().TileSettings = { };
+		return;
+	}
+
+	// Perform additional validations
+	const FDisplayClusterConfigurationICVFX_StageSettings* const StageSettings = InSourceViewport.Configuration->GetStageSettings();
+	const bool bIsTilingAllowed = (StageSettings && FDisplayClusterConfigurationTile_Settings::IsEnabled(InCameraMediaSettings.TiledSplitLayout, *StageSettings));
+	if (!bIsTilingAllowed)
+	{
+		InSourceViewport.GetRenderSettingsImpl().TileSettings = { };
+		return;
+	}
+
+	// Check if viewport is able to split
+	const bool bCanSplitIntoTiles = InSourceViewport.CanSplitIntoTiles();
+	if (!bCanSplitIntoTiles)
+	{
+		InSourceViewport.GetRenderSettingsImpl().TileSettings = { };
+		return;
+	}
+
+	// Find if this cluster node is allowed to render unbound tiles
+	const FString ThisNodeId = InSourceViewport.GetClusterNodeId();
+	const bool bAllowRenderUnbound = InCameraMediaSettings.ClusterNodesToRenderUnboundTiles.ItemNames.ContainsByPredicate([&ThisNodeId](const FString& Item)
 		{
-			if (InSourceViewport.CanSplitIntoTiles())
-			{
-				const FDisplayClusterViewport_OverscanSettings& OverscanSettings = FDisplayClusterViewportConfigurationHelpers_Tile::GetTileOverscanSettings(InCameraTile.TileOverscan);
-				const FIntPoint Size(InCameraTile.TileSettings.TileX, InCameraTile.TileSettings.TileY);
-
-				// Find if this cluster node is allowed to render unbound tiles
-				const FString ThisNodeId = InSourceViewport.GetClusterNodeId();
-				const bool bAllowRenderUnbound = InCameraTile.ClusterNodesToRenderUnboundTiles.ItemNames.ContainsByPredicate([&ThisNodeId](const FString& Item)
-					{
-						return ThisNodeId.Equals(Item, ESearchCase::IgnoreCase);
-					});
-
-				// Prepare tile flags
-				EDisplayClusterViewportTileFlags TileFlags = EDisplayClusterViewportTileFlags::None;
-				TileFlags |= (bAllowRenderUnbound ? EDisplayClusterViewportTileFlags::AllowUnboundRender : EDisplayClusterViewportTileFlags::None);
-
-				// Set this viewport as the source for tile rendering.
-				FDisplayClusterViewport_TileSettings& OutTileSettings = InSourceViewport.GetRenderSettingsImpl().TileSettings;
-				OutTileSettings = FDisplayClusterViewport_TileSettings(Size, OverscanSettings, TileFlags);
-				OutTileSettings.bOptimizeTileOverscan = InCameraTile.TileOverscan.bOptimizeTileOverscan;
-
-				return;
-			}
+			return ThisNodeId.Equals(Item, ESearchCase::IgnoreCase);
+		});
 
-			// This viewport cannot be split because the current settings are conflicting.
-			UE_LOG(LogDisplayClusterViewport, Error, TEXT("Viewport '%s' cannot be tiled because the current settings are conflicting."), *InSourceViewport.GetId());
-		}
+	// Generate flags
+	EDisplayClusterViewportTileFlags TileFlags = EDisplayClusterViewportTileFlags::None;
+	if (bAllowRenderUnbound)
+	{
+		TileFlags |= EDisplayClusterViewportTileFlags::AllowUnboundRender;
 	}
 
-	// By default disable tile rendering.
-	InSourceViewport.GetRenderSettingsImpl().TileSettings = FDisplayClusterViewport_TileSettings();
+	// Generate overscan settings
+	const FDisplayClusterViewport_OverscanSettings& OverscanSettings = FDisplayClusterViewportConfigurationHelpers_Tile::GetTileOverscanSettings(InCameraMediaSettings.TileOverscan);
+
+	// Set this viewport as a source for tile rendering.
+	FDisplayClusterViewport_TileSettings& OutTileSettings = InSourceViewport.GetRenderSettingsImpl().TileSettings;
+	OutTileSettings = FDisplayClusterViewport_TileSettings(InCameraMediaSettings.TiledSplitLayout, OverscanSettings, TileFlags);
+	OutTileSettings.bOptimizeTileOverscan = InCameraMediaSettings.TileOverscan.bOptimizeTileOverscan;
 }
 
 FIntRect FDisplayClusterViewportConfigurationHelpers_Tile::GetDestRect(const FDisplayClusterViewport_TileSettings& InTileSettings, const FIntRect& InSourceRect)
diff --git a/Engine/Plugins/Runtime/nDisplay/Source/DisplayCluster/Private/Render/Viewport/Configuration/DisplayClusterViewportConfigurationHelpers_Tile.h b/Engine/Plugins/Runtime/nDisplay/Source/DisplayCluster/Private/Render/Viewport/Configuration/DisplayClusterViewportConfigurationHelpers_Tile.h
--- a/Engine/Plugins/Runtime/nDisplay/Source/DisplayCluster/Private/Render/Viewport/Configuration/DisplayClusterViewportConfigurationHelpers_Tile.h
+++ b/Engine/Plugins/Runtime/nDisplay/Source/DisplayCluster/Private/Render/Viewport/Configuration/DisplayClusterViewportConfigurationHelpers_Tile.h
@@ -7,9 +7,9 @@
 
 class FDisplayClusterViewport;
 class IDisplayClusterProjectionPolicy;
-struct FDisplayClusterViewport_TileSettings;
-struct FDisplayClusterConfigurationICVFX_CameraTile;
+struct FDisplayClusterConfigurationMediaICVFX;
 struct FDisplayClusterConfigurationTile_Overscan;
+struct FDisplayClusterViewport_TileSettings;
 
 /**
 * Tile configuration helper class.
@@ -28,7 +28,7 @@ public:
 
 public:
 	/** Get tile settings for the camera viewport. */
-	static void UpdateICVFXCameraViewportTileSettings(FDisplayClusterViewport& InSourceViewport, const FDisplayClusterConfigurationICVFX_CameraTile& InCameraTile);
+	static void UpdateICVFXCameraViewportTileSettings(FDisplayClusterViewport& InSourceViewport, const FDisplayClusterConfigurationMediaICVFX& InCameraMediaSettings);
 
 	/** Get a tile viewport rect inside the rect of the source viewport. */
 	static FIntRect GetDestRect(const FDisplayClusterViewport_TileSettings& InTileSettings, const FIntRect& InSourceRect);
diff --git a/Engine/Plugins/Runtime/nDisplay/Source/DisplayCluster/Public/Components/DisplayClusterICVFXCameraComponent.h b/Engine/Plugins/Runtime/nDisplay/Source/DisplayCluster/Public/Components/DisplayClusterICVFXCameraComponent.h
--- a/Engine/Plugins/Runtime/nDisplay/Source/DisplayCluster/Public/Components/DisplayClusterICVFXCameraComponent.h
+++ b/Engine/Plugins/Runtime/nDisplay/Source/DisplayCluster/Public/Components/DisplayClusterICVFXCameraComponent.h
@@ -172,15 +172,6 @@ private:
 	UPROPERTY(EditAnywhere, Transient, Category = OCIO, meta = (PropertyPath = "CameraSettings.CameraOCIO.PerNodeOCIOProfiles", DisplayName = "Per-Node OCIO Overrides"))
 	FDisplayClusterEditorPropertyReference PerNodeOCIOProfilesRef;
 
-	UPROPERTY(EditAnywhere, Transient, Category = "Tile Rendering", meta = (PropertyPath = "CameraSettings.CameraTile.TileSettings", DisplayName = "Tile Settings"))
-	FDisplayClusterEditorPropertyReference CameraTileSettingsRef;
-
-	UPROPERTY(EditAnywhere, Transient, Category = "Tile Rendering", meta = (PropertyPath = "CameraSettings.CameraTile.TileOverscan", DisplayName = "Tile Overscan"))
-	FDisplayClusterEditorPropertyReference CameraTileOverscanRef;
-
-	UPROPERTY(EditAnywhere, Transient, Category = "Tile Rendering", meta = (PropertyPath = "CameraSettings.CameraTile.ClusterNodesToRenderUnboundTiles", DisplayName = "Nodes To Render Unbound Tiles", ToolTip = "Choose nodes that should render camera tiles that don't have any media assigned"))
-	FDisplayClusterEditorPropertyReference CameraTileNodesToRenderUnboundTilesRef;
-
 	UPROPERTY(EditAnywhere, Transient, Category = "Color Grading", meta = (PropertyPath = "CameraSettings.EnableInnerFrustumColorGrading", DisplayName = "Enable Inner Frustum Color Grading"))
 	FDisplayClusterEditorPropertyReference EnableInnerFrustumColorGrading;
 
diff --git a/Engine/Plugins/Runtime/nDisplay/Source/DisplayClusterConfiguration/Private/DisplayClusterConfigurationTypes_Tile.cpp b/Engine/Plugins/Runtime/nDisplay/Source/DisplayClusterConfiguration/Private/DisplayClusterConfigurationTypes_Tile.cpp
--- a/Engine/Plugins/Runtime/nDisplay/Source/DisplayClusterConfiguration/Private/DisplayClusterConfigurationTypes_Tile.cpp
+++ b/Engine/Plugins/Runtime/nDisplay/Source/DisplayClusterConfiguration/Private/DisplayClusterConfigurationTypes_Tile.cpp
@@ -13,19 +13,24 @@ bool FDisplayClusterConfigurationTile_Settings::IsEnabled(const FDisplayClusterC
 
 	if (!bEnabled)
 	{
-		// This property should enable tile rendering.
+		// This property enables tile rendering.
 		return false;
 	}
 
-	if (TileX < 1 || TileY < 1)
+	return FDisplayClusterConfigurationTile_Settings::IsEnabled(FIntPoint{TileX, TileY}, InStageSettings);
+}
+
+bool FDisplayClusterConfigurationTile_Settings::IsEnabled(const FIntPoint& InTileLoc, const struct FDisplayClusterConfigurationICVFX_StageSettings& InStageSettings)
+{
+	// Ignore wrong values
+	if (InTileLoc.X < 1 || InTileLoc.Y < 1)
 	{
-		// Ignore wrong values
 		return false;
 	}
 
-	if (TileX == 1 && TileY == 1)
+	// Ignore 1x1 case
+	if (InTileLoc.X == 1 && InTileLoc.Y == 1)
 	{
-		// Ignore if only 1 tile is used.
 		return false;
 	}
 
diff --git a/Engine/Plugins/Runtime/nDisplay/Source/DisplayClusterConfiguration/Public/DisplayClusterConfigurationStrings.h b/Engine/Plugins/Runtime/nDisplay/Source/DisplayClusterConfiguration/Public/DisplayClusterConfigurationStrings.h
--- a/Engine/Plugins/Runtime/nDisplay/Source/DisplayClusterConfiguration/Public/DisplayClusterConfigurationStrings.h
+++ b/Engine/Plugins/Runtime/nDisplay/Source/DisplayClusterConfiguration/Public/DisplayClusterConfigurationStrings.h
@@ -29,6 +29,7 @@ namespace DisplayClusterConfigurationStrings
 		static constexpr const TCHAR* ChromaKeyCategory            = TEXT("Chromakey");
 		static constexpr const TCHAR* LightcardCategory            = TEXT("Light Cards");
 		static constexpr const TCHAR* OCIOCategory                 = TEXT("OCIO");
+		static constexpr const TCHAR* MediaCategory                = TEXT("Media");
 		static constexpr const TCHAR* TileCategory                 = TEXT("Tile Rendering");
 		static constexpr const TCHAR* OverrideCategory             = TEXT("Texture Replacement");
 		static constexpr const TCHAR* ViewportsCategory            = TEXT("Viewports");
diff --git a/Engine/Plugins/Runtime/nDisplay/Source/DisplayClusterConfiguration/Public/DisplayClusterConfigurationTypes_ICVFX.h b/Engine/Plugins/Runtime/nDisplay/Source/DisplayClusterConfiguration/Public/DisplayClusterConfigurationTypes_ICVFX.h
--- a/Engine/Plugins/Runtime/nDisplay/Source/DisplayClusterConfiguration/Public/DisplayClusterConfigurationTypes_ICVFX.h
+++ b/Engine/Plugins/Runtime/nDisplay/Source/DisplayClusterConfiguration/Public/DisplayClusterConfigurationTypes_ICVFX.h
@@ -16,7 +16,6 @@
 #include "DisplayClusterConfigurationTypes_PostRender.h"
 #include "DisplayClusterConfigurationTypes_Postprocess.h"
 #include "DisplayClusterConfigurationTypes_OCIO.h"
-#include "DisplayClusterConfigurationTypes_Tile.h"
 #include "DisplayClusterConfigurationTypes_ViewportOverscan.h"
 
 #include "Containers/DisplayClusterShader_Enums.h"
@@ -88,25 +87,6 @@ public:
 	TArray<FDisplayClusterConfigurationOCIOProfile> PerNodeOCIOProfiles;
 };
 
-USTRUCT(Blueprintable)
-struct DISPLAYCLUSTERCONFIGURATION_API FDisplayClusterConfigurationICVFX_CameraTile
-{
-	GENERATED_BODY()
-
-public:
-	/** Tile rendering settings. */
-	UPROPERTY(EditAnywhere, BlueprintReadWrite, Category = "Tile Rendering")
-	FDisplayClusterConfigurationTile_Settings TileSettings;
-
-	/** Overscan settings for tile. */
-	UPROPERTY(EditAnywhere, BlueprintReadWrite, Category = "Tile Overscan")
-	FDisplayClusterConfigurationTile_Overscan TileOverscan;
-
-	/** Cluster nodes that should render unbound tiles. Unbound tiles are the tiles that don't have any media assigned */
-	UPROPERTY(EditAnywhere, BlueprintReadWrite, Category = Media, meta = (ClusterItemType = ClusterNodes, DisplayName = "Nodes To Render Unbound Tiles", ToolTip = "Choose nodes that should render camera tiles that don't have any media assigned"))
-	FDisplayClusterConfigurationClusterItemReferenceList ClusterNodesToRenderUnboundTiles;
-};
-
 USTRUCT(Blueprintable)
 struct DISPLAYCLUSTERCONFIGURATION_API FDisplayClusterConfigurationICVFX_LightcardOCIO
 {
@@ -817,10 +797,6 @@ public:
 	UPROPERTY(EditAnywhere, BlueprintReadWrite, Category = OCIO)
 	FDisplayClusterConfigurationICVFX_CameraOCIO CameraOCIO;
 
-	/** Tile rendering settings for icvfx camera viewport. */
-	UPROPERTY(EditAnywhere, BlueprintReadWrite, Category = "Tile Rendering")
-	FDisplayClusterConfigurationICVFX_CameraTile CameraTile;
-
 	/** Entire Cluster Color Grading */
 	UPROPERTY(EditAnywhere, BlueprintReadWrite, Category = "Inner Frustum Color Grading", meta = (DisplayName = "Enable Inner Frustum Color Grading"))
 	bool EnableInnerFrustumColorGrading = true;
diff --git a/Engine/Plugins/Runtime/nDisplay/Source/DisplayClusterConfiguration/Public/DisplayClusterConfigurationTypes_Media.h b/Engine/Plugins/Runtime/nDisplay/Source/DisplayClusterConfiguration/Public/DisplayClusterConfigurationTypes_Media.h
--- a/Engine/Plugins/Runtime/nDisplay/Source/DisplayClusterConfiguration/Public/DisplayClusterConfigurationTypes_Media.h
+++ b/Engine/Plugins/Runtime/nDisplay/Source/DisplayClusterConfiguration/Public/DisplayClusterConfigurationTypes_Media.h
@@ -6,6 +6,7 @@
 
 #include "DisplayClusterConfigurationTypes_Base.h"
 #include "DisplayClusterConfigurationTypes_MediaSync.h"
+#include "DisplayClusterConfigurationTypes_Tile.h"
 
 #include "MediaPlayer.h"
 #include "MediaSource.h"
@@ -22,7 +23,7 @@ UENUM(BlueprintType)
 enum class EDisplayClusterConfigurationMediaSplitType : uint8
 {
 	FullFrame     UMETA(DisplayName = "Full Frame"),
-	UniformTiles  UMETA(DisplayName = "Uniform Tiles"),
+	UniformTiles  UMETA(DisplayName = "Tiled"),
 };
 
 
@@ -106,10 +107,6 @@ public:
 	UPROPERTY(EditAnywhere, BlueprintReadWrite, Category = "Media")
 	TArray<FDisplayClusterConfigurationMediaOutput> MediaOutputs;
 
-	/** Force late OCIO pass */
-	UPROPERTY(EditAnywhere, BlueprintReadWrite, Category = "Media", meta = (DisplayName = "Late OCIO Pass"))
-	bool bLateOCIOPass = false;
-
 public:
 	/** Returns true if a media source assigned */
 	bool IsMediaInputAssigned() const;
@@ -183,13 +180,9 @@ struct DISPLAYCLUSTERCONFIGURATION_API FDisplayClusterConfigurationMediaUniformT
 	GENERATED_BODY()
 
 public:
-	/** Tile X location */
-	UPROPERTY(EditAnywhere, BlueprintReadOnly, Category = "Media", meta = (UIMin = 0, ClampMin = 0, UIMax = 4, ClampMax = 4))
-	int32 TileX = 0;
-
-	/** Tile Y location */
-	UPROPERTY(EditAnywhere, BlueprintReadOnly, Category = "Media", meta = (UIMin = 0, ClampMin = 0, UIMax = 4, ClampMax = 4))
-	int32 TileY = 0;
+	/** Tile position */
+	UPROPERTY(EditAnywhere, BlueprintReadOnly, Category = "Media", meta = (ClampMin = 0, ClampMax = 3))
+	FIntPoint Position = FIntPoint::ZeroValue;
 
 	/** Media source to use */
 	UPROPERTY(Instanced, EditAnywhere, BlueprintReadWrite, Category = "Media")
@@ -206,13 +199,9 @@ struct DISPLAYCLUSTERCONFIGURATION_API FDisplayClusterConfigurationMediaUniformT
 	GENERATED_BODY()
 
 public:
-	/** Tile X location */
-	UPROPERTY(EditAnywhere, BlueprintReadOnly, Category = "Media", meta = (UIMin = 0, ClampMin = 0, UIMax = 4, ClampMax = 4))
-	int32 TileX = 0;
-
-	/** Tile Y location */
-	UPROPERTY(EditAnywhere, BlueprintReadOnly, Category = "Media", meta = (UIMin = 0, ClampMin = 0, UIMax = 4, ClampMax = 4))
-	int32 TileY = 0;
+	/** Tile position */
+	UPROPERTY(EditAnywhere, BlueprintReadOnly, Category = "Media", meta = (ClampMin = 0, ClampMax = 3))
+	FIntPoint Position = FIntPoint::ZeroValue;
 
 	/** Media output to use */
 	UPROPERTY(Instanced, EditAnywhere, BlueprintReadWrite, Category = "Media")
@@ -276,27 +265,43 @@ public:
 	bool bEnable = false;
 
 	/** Media frame split type */
-	UPROPERTY(EditAnywhere, BlueprintReadWrite, Category = "Media")
+	UPROPERTY(EditAnywhere, BlueprintReadWrite, Category = "Media", meta = (DisplayName = "Inner Frustum Type"))
 	EDisplayClusterConfigurationMediaSplitType SplitType = EDisplayClusterConfigurationMediaSplitType::FullFrame;
 
+	/// Full-frame
+
 	/** Media input mapping (Full frame) */
-	UPROPERTY(EditAnywhere, BlueprintReadWrite, Category = "Media", meta = (DisplayName = "Media Input Groups (Full Frame)", EditCondition = "SplitType == EDisplayClusterConfigurationMediaSplitType::FullFrame"))
+	UPROPERTY(EditAnywhere, BlueprintReadWrite, Category = "Media", meta = (DisplayName = "Media Input Groups"))
 	TArray<FDisplayClusterConfigurationMediaInputGroup> MediaInputGroups;
 
 	/** Media output mapping (Full frame) */
-	UPROPERTY(EditAnywhere, BlueprintReadWrite, Category = "Media", meta = (DisplayName = "Media Output Groups (Full Frame)", EditCondition = "SplitType == EDisplayClusterConfigurationMediaSplitType::FullFrame"))
+	UPROPERTY(EditAnywhere, BlueprintReadWrite, Category = "Media", meta = (DisplayName = "Media Output Groups"))
 	TArray<FDisplayClusterConfigurationMediaOutputGroup> MediaOutputGroups;
 
+	/// Uniform tiles
+
+	/** Split layout */
+	UPROPERTY(EditAnywhere, BlueprintReadWrite, Category = "Media", meta = (DisplayName = "Tiled Split Layout", ClampMin = 1, ClampMax = 4))
+	FIntPoint TiledSplitLayout = { 1, 1 };
+
+	/** Overscan settings for tile. */
+	UPROPERTY(EditAnywhere, BlueprintReadWrite, Category = "Media")
+	FDisplayClusterConfigurationTile_Overscan TileOverscan;
+
+	/** Cluster nodes that should render unbound tiles. Unbound tiles are the tiles that don't have any media assigned */
+	UPROPERTY(EditAnywhere, BlueprintReadWrite, Category = "Media", meta = (ClusterItemType = ClusterNodes, DisplayName = "Nodes To Render Unbound Tiles", ToolTip = "Choose nodes that should render camera tiles that don't have any media assigned"))
+	FDisplayClusterConfigurationClusterItemReferenceList ClusterNodesToRenderUnboundTiles;
+
 	/** Media input mapping (Tiled) */
-	UPROPERTY(EditAnywhere, BlueprintReadWrite, Category = "Media", meta = (DisplayName = "Media Input Groups (Tiled)", EditCondition = "SplitType == EDisplayClusterConfigurationMediaSplitType::UniformTiles"))
+	UPROPERTY(EditAnywhere, BlueprintReadWrite, Category = "Media", meta = (DisplayName = "Media Input Groups"))
 	TArray<FDisplayClusterConfigurationMediaTiledInputGroup> TiledMediaInputGroups;
 
 	/** Media output mapping (Tiled) */
-	UPROPERTY(EditAnywhere, BlueprintReadWrite, Category = "Media", meta = (DisplayName = "Media Output Groups (Tiled)", EditCondition = "SplitType == EDisplayClusterConfigurationMediaSplitType::UniformTiles"))
+	UPROPERTY(EditAnywhere, BlueprintReadWrite, Category = "Media", meta = (DisplayName = "Media Output Groups"))
 	TArray<FDisplayClusterConfigurationMediaTiledOutputGroup> TiledMediaOutputGroups;
 
 	/** Force late OCIO pass */
-	UPROPERTY(EditAnywhere, BlueprintReadWrite, Category = "Media", meta = (DisplayName = "Late OCIO Pass"))
+	UPROPERTY()
 	bool bLateOCIOPass = false;
 
 public:
diff --git a/Engine/Plugins/Runtime/nDisplay/Source/DisplayClusterConfiguration/Public/DisplayClusterConfigurationTypes_Tile.h b/Engine/Plugins/Runtime/nDisplay/Source/DisplayClusterConfiguration/Public/DisplayClusterConfigurationTypes_Tile.h
--- a/Engine/Plugins/Runtime/nDisplay/Source/DisplayClusterConfiguration/Public/DisplayClusterConfigurationTypes_Tile.h
+++ b/Engine/Plugins/Runtime/nDisplay/Source/DisplayClusterConfiguration/Public/DisplayClusterConfigurationTypes_Tile.h
@@ -21,6 +21,9 @@ public:
 	/** Return true if tile rendering can be used. */
 	bool IsEnabled(const struct FDisplayClusterConfigurationICVFX_StageSettings& InStageSettings) const;
 
+	/** Static validation for custom implementations. */
+	static bool IsEnabled(const FIntPoint& InTileLoc, const struct FDisplayClusterConfigurationICVFX_StageSettings& InStageSettings);
+
 public:
 	/** Enable tile rendering. */
 	UPROPERTY(EditAnywhere, BlueprintReadWrite, Category = "Tile Rendering")
@@ -42,23 +45,23 @@ struct DISPLAYCLUSTERCONFIGURATION_API FDisplayClusterConfigurationTile_Overscan
 
 public:
 	/** Enable/disable Viewport Overscan and specify units as percent or pixel values. */
-	UPROPERTY(EditAnywhere, BlueprintReadWrite, Category = "Tile Overscan", meta = (DisplayName = "Enable"))
+	UPROPERTY(EditAnywhere, BlueprintReadWrite, Category = "Tile Overscan", meta = (DisplayName = "Enable Tile Overscan"))
 	bool bEnabled = false;
 
 	/** Set to True to render at the overscan resolution, set to false to render at the resolution in the configuration and scale for overscan. */
-	UPROPERTY(EditAnywhere, BlueprintReadWrite, Category = "Tile Overscan", meta = (DisplayName = "Adapt Resolution", DisplayAfter = "Mode"))
+	UPROPERTY(EditAnywhere, BlueprintReadWrite, Category = "Tile Overscan", meta = (DisplayName = "Adapt Resolution", EditCondition = "bEnabled == true"))
 	bool bOversize = true;
 
 	/** Optimize overscan values on boundary tiles.
 	* When enabled, tile sides not in contact with other tiles will use zero overscan. */
-	UPROPERTY()
+	UPROPERTY(EditAnywhere, BlueprintReadWrite, Category = "Tile Overscan", meta = (DisplayName = "No Overscan For Edges", EditCondition = "bEnabled == true"))
 	bool bOptimizeTileOverscan = true;
 
 	/** Enable/disable Viewport Overscan and specify units as percent or pixel values. */
-	UPROPERTY(EditAnywhere, BlueprintReadWrite, Category = "Tile Overscan")
+	UPROPERTY(EditAnywhere, BlueprintReadWrite, Category = "Tile Overscan", meta = (DisplayName = "Overscan Units", EditCondition = "bEnabled == true"))
 	EDisplayClusterConfigurationViewportOverscanMode Mode = EDisplayClusterConfigurationViewportOverscanMode::Percent;
 
 	/** Overscan value for all sides. */
-	UPROPERTY(EditAnywhere, BlueprintReadWrite, Category = "Tile Overscan")
+	UPROPERTY(EditAnywhere, BlueprintReadWrite, Category = "Tile Overscan", meta = (DisplayName = "Overscan Value", EditCondition = "bEnabled == true"))
 	float AllSides = 10;
 };
diff --git a/Engine/Plugins/Runtime/nDisplay/Source/DisplayClusterConfigurator/Private/DisplayClusterConfiguratorModule.cpp b/Engine/Plugins/Runtime/nDisplay/Source/DisplayClusterConfigurator/Private/DisplayClusterConfiguratorModule.cpp
--- a/Engine/Plugins/Runtime/nDisplay/Source/DisplayClusterConfigurator/Private/DisplayClusterConfiguratorModule.cpp
+++ b/Engine/Plugins/Runtime/nDisplay/Source/DisplayClusterConfigurator/Private/DisplayClusterConfiguratorModule.cpp
@@ -26,6 +26,7 @@
 #include "Views/Details/Cluster/DisplayClusterConfiguratorRectangleCustomization.h"
 #include "Views/Details/Components/DisplayClusterConfiguratorScreenComponentDetailsCustomization.h"
 #include "Views/Details/Components/DisplayClusterICVFXCameraComponentDetailsCustomization.h"
+#include "Views/Details/Media/DisplayClusterConfiguratorICVFXMediaCustomization.h"
 #include "Views/Details/Policies/DisplayClusterConfiguratorPolicyDetailCustomization.h"
 
 #include "Blueprints/DisplayClusterBlueprint.h"
@@ -223,6 +224,7 @@ void FDisplayClusterConfiguratorModule::RegisterCustomLayouts()
 	REGISTER_PROPERTY_LAYOUT(FDisplayClusterConfigurationClusterItemReferenceList, FDisplayClusterConfiguratorClusterReferenceListCustomization);
 	REGISTER_PROPERTY_LAYOUT(FDisplayClusterConfigurationViewport_RemapData, FDisplayClusterConfiguratorViewportRemapCustomization);
 	REGISTER_PROPERTY_LAYOUT(FDisplayClusterConfigurationRectangle, FDisplayClusterConfiguratorRectangleCustomization);
+	REGISTER_PROPERTY_LAYOUT(FDisplayClusterConfigurationMediaICVFX, FDisplayClusterConfiguratorICVFXMediaCustomization);
 }
 
 void FDisplayClusterConfiguratorModule::UnregisterCustomLayouts()
@@ -304,6 +306,11 @@ void FDisplayClusterConfiguratorModule::RegisterSectionMappings()
 			DisplayClusterConfigurationStrings::categories::OCIOCategory, LOCTEXT("OCIO", "OCIO"));
 		Section->AddCategory(DisplayClusterConfigurationStrings::categories::OCIOCategory);
 	}
+	{
+		const TSharedRef<FPropertySection> Section = PropertyModule.FindOrCreateSection(UDisplayClusterICVFXCameraComponent::StaticClass()->GetFName(),
+			DisplayClusterConfigurationStrings::categories::MediaCategory, LOCTEXT("Media", "Media"));
+		Section->AddCategory(DisplayClusterConfigurationStrings::categories::MediaCategory);
+	}
 	{
 		const TSharedRef<FPropertySection> Section = PropertyModule.FindOrCreateSection(UDisplayClusterICVFXCameraComponent::StaticClass()->GetFName(),
 			DisplayClusterConfigurationStrings::categories::ChromaKeyCategory, LOCTEXT("Chromakey", "Chromakey"));
diff --git a/Engine/Plugins/Runtime/nDisplay/Source/DisplayClusterConfigurator/Private/Views/Details/Components/DisplayClusterICVFXCameraComponentDetailsCustomization.cpp b/Engine/Plugins/Runtime/nDisplay/Source/DisplayClusterConfigurator/Private/Views/Details/Components/DisplayClusterICVFXCameraComponentDetailsCustomization.cpp
--- a/Engine/Plugins/Runtime/nDisplay/Source/DisplayClusterConfigurator/Private/Views/Details/Components/DisplayClusterICVFXCameraComponentDetailsCustomization.cpp
+++ b/Engine/Plugins/Runtime/nDisplay/Source/DisplayClusterConfigurator/Private/Views/Details/Components/DisplayClusterICVFXCameraComponentDetailsCustomization.cpp
@@ -30,6 +30,7 @@ namespace DisplayClusterICVFXCameraComponentDetailsCustomizationUtils
 			DisplayClusterConfigurationStrings::categories::ICVFXCameraCategory,
 			DisplayClusterConfigurationStrings::categories::CameraColorGradingCategory,
 			DisplayClusterConfigurationStrings::categories::OCIOCategory,
+			DisplayClusterConfigurationStrings::categories::MediaCategory,
 			DisplayClusterConfigurationStrings::categories::ChromaKeyCategory,
 			DisplayClusterConfigurationStrings::categories::OverrideCategory,
 			DisplayClusterConfigurationStrings::categories::ConfigurationCategory
diff --git a/Engine/Plugins/Runtime/nDisplay/Source/DisplayClusterConfigurator/Private/Views/Details/Media/DisplayClusterConfiguratorICVFXMediaCustomization.cpp b/Engine/Plugins/Runtime/nDisplay/Source/DisplayClusterConfigurator/Private/Views/Details/Media/DisplayClusterConfiguratorICVFXMediaCustomization.cpp
new file mode 100644
--- /dev/null
+++ b/Engine/Plugins/Runtime/nDisplay/Source/DisplayClusterConfigurator/Private/Views/Details/Media/DisplayClusterConfiguratorICVFXMediaCustomization.cpp
@@ -0,0 +1,80 @@
+// Copyright Epic Games, Inc. All Rights Reserved.
+
+#include "Views/Details/Media/DisplayClusterConfiguratorICVFXMediaCustomization.h"
+
+#include "DisplayClusterConfigurationTypes_Media.h"
+
+#include "IDetailChildrenBuilder.h"
+#include "IPropertyUtilities.h"
+#include "PropertyHandle.h"
+
+
+#define LOCTEXT_NAMESPACE "FDisplayClusterConfiguratorICVFXMediaCustomization"
+
+void FDisplayClusterConfiguratorICVFXMediaCustomization::CustomizeChildren(TSharedRef<IPropertyHandle> InPropertyHandle, IDetailChildrenBuilder& InChildBuilder, IPropertyTypeCustomizationUtils& InCustomizationUtils)
+{
+	// SplitType property
+	TSharedPtr<IPropertyHandle> SplitTypeHandle = GET_CHILD_HANDLE(FDisplayClusterConfigurationMediaICVFX, SplitType);
+	check(SplitTypeHandle);
+
+	// Separate groups specific for every split type available
+	TArray<TSharedPtr<IPropertyHandle>> FullFramePropertyHandles;
+	TArray<TSharedPtr<IPropertyHandle>> UniformTilePropertyHandles;
+
+	// FullFrame properties
+	FullFramePropertyHandles.Add(GET_CHILD_HANDLE(FDisplayClusterConfigurationMediaICVFX, MediaInputGroups));
+	FullFramePropertyHandles.Add(GET_CHILD_HANDLE(FDisplayClusterConfigurationMediaICVFX, MediaOutputGroups));
+
+	// UniformTile properties
+	UniformTilePropertyHandles.Add(GET_CHILD_HANDLE(FDisplayClusterConfigurationMediaICVFX, TiledSplitLayout));
+	UniformTilePropertyHandles.Add(GET_CHILD_HANDLE(FDisplayClusterConfigurationMediaICVFX, TileOverscan));
+	UniformTilePropertyHandles.Add(GET_CHILD_HANDLE(FDisplayClusterConfigurationMediaICVFX, ClusterNodesToRenderUnboundTiles));
+	UniformTilePropertyHandles.Add(GET_CHILD_HANDLE(FDisplayClusterConfigurationMediaICVFX, TiledMediaInputGroups));
+	UniformTilePropertyHandles.Add(GET_CHILD_HANDLE(FDisplayClusterConfigurationMediaICVFX, TiledMediaOutputGroups));
+
+	// Get current split type
+	EDisplayClusterConfigurationMediaSplitType SplitTypeValue = EDisplayClusterConfigurationMediaSplitType::UniformTiles;
+	SplitTypeHandle->GetValue((uint8&)SplitTypeValue);
+
+	TSharedPtr<IPropertyUtilities> PropertyUtils = InCustomizationUtils.GetPropertyUtilities();
+	check(PropertyUtils);
+
+	// Setup details update on frustum type change
+	SplitTypeHandle->SetOnPropertyValueChanged(
+		FSimpleDelegate::CreateLambda([PropertyUtils]()
+		{
+			if (PropertyUtils)
+			{
+				PropertyUtils->RequestForceRefresh();
+			}
+		}));
+
+	// Filter properties to hide
+	TArray<TArray<TSharedPtr<IPropertyHandle>>*> HiddenPropertyHandles;
+	switch (SplitTypeValue)
+	{
+	case EDisplayClusterConfigurationMediaSplitType::FullFrame:
+		HiddenPropertyHandles.Add(&UniformTilePropertyHandles);
+		break;
+
+	case EDisplayClusterConfigurationMediaSplitType::UniformTiles:
+		HiddenPropertyHandles.Add(&FullFramePropertyHandles);
+		break;
+
+	default:
+		checkNoEntry();
+	}
+
+	// Hide unnecessary properties depending on the frustum (split) type currently selected
+	for (const TArray<TSharedPtr<IPropertyHandle>>* HiddenGroup : HiddenPropertyHandles)
+	{
+		for (const TSharedPtr<IPropertyHandle>& PropertyHandle : *HiddenGroup)
+		{
+			PropertyHandle->MarkHiddenByCustomization();
+		}
+	}
+
+	FDisplayClusterConfiguratorBaseTypeCustomization::CustomizeChildren(InPropertyHandle, InChildBuilder, InCustomizationUtils);
+}
+
+#undef LOCTEXT_NAMESPACE
diff --git a/Engine/Plugins/Runtime/nDisplay/Source/DisplayClusterConfigurator/Private/Views/Details/Media/DisplayClusterConfiguratorICVFXMediaCustomization.h b/Engine/Plugins/Runtime/nDisplay/Source/DisplayClusterConfigurator/Private/Views/Details/Media/DisplayClusterConfiguratorICVFXMediaCustomization.h
new file mode 100644
--- /dev/null
+++ b/Engine/Plugins/Runtime/nDisplay/Source/DisplayClusterConfigurator/Private/Views/Details/Media/DisplayClusterConfiguratorICVFXMediaCustomization.h
@@ -0,0 +1,27 @@
+// Copyright Epic Games, Inc. All Rights Reserved.
+
+#pragma once
+
+#include "CoreMinimal.h"
+#include "Views/Details/DisplayClusterConfiguratorBaseTypeCustomization.h"
+
+#include "IPropertyTypeCustomization.h"
+
+
+/**
+ * Details panel customization for the FDisplayClusterConfigurationMediaICVFX struct.
+ */
+class FDisplayClusterConfiguratorICVFXMediaCustomization
+	: public FDisplayClusterConfiguratorBaseTypeCustomization
+{
+public:
+	static TSharedRef<IPropertyTypeCustomization> MakeInstance()
+	{
+		return MakeShared<FDisplayClusterConfiguratorICVFXMediaCustomization>();
+	}
+
+protected:
+	//~ Begin IPropertyTypeCustomization
+	virtual void CustomizeChildren(TSharedRef<IPropertyHandle> PropertyHandle, IDetailChildrenBuilder& ChildBuilder, IPropertyTypeCustomizationUtils& CustomizationUtils) override;
+	//~ End IPropertyTypeCustomization
+};
diff --git a/Engine/Plugins/Runtime/nDisplay/Source/DisplayClusterMedia/Private/DisplayClusterMediaModule.cpp b/Engine/Plugins/Runtime/nDisplay/Source/DisplayClusterMedia/Private/DisplayClusterMediaModule.cpp
--- a/Engine/Plugins/Runtime/nDisplay/Source/DisplayClusterMedia/Private/DisplayClusterMediaModule.cpp
+++ b/Engine/Plugins/Runtime/nDisplay/Source/DisplayClusterMedia/Private/DisplayClusterMediaModule.cpp
@@ -391,18 +391,16 @@ void FDisplayClusterMediaModule::InitializeICVFXCameraUniformTilesInput(const UD
 				{
 					if (IsValid(MediaInputTile.MediaSource))
 					{
-						const FIntPoint TileLoc{ MediaInputTile.TileX, MediaInputTile.TileY };
-
-						const FString ICVFXViewportTileId = DisplayClusterMediaHelpers::GenerateTileViewportName(ICVFXViewportId, TileLoc);
+						const FString ICVFXViewportTileId = DisplayClusterMediaHelpers::GenerateTileViewportName(ICVFXViewportId, MediaInputTile.Position);
 
 						const FString MediaInputId = DisplayClusterMediaHelpers::MediaId::GenerateMediaId(
 							DisplayClusterMediaHelpers::MediaId::EMediaDeviceType::Input,
 							DisplayClusterMediaHelpers::MediaId::EMediaOwnerType::ICVFXCamera,
 							ClusterNodeId, RootActorName, ICVFXCameraName,
-							static_cast<uint8>(Index), &TileLoc);
+							static_cast<uint8>(Index), &MediaInputTile.Position);
 
 						UE_LOG(LogDisplayClusterMedia, Log, TEXT("Initializing ICVFX media input '%s' for camera '%s' tile '%d,%d'"),
-							*MediaInputId, *ICVFXCameraName, TileLoc.X, TileLoc.Y);
+							*MediaInputId, *ICVFXCameraName, MediaInputTile.Position.X, MediaInputTile.Position.Y);
 
 						TSharedPtr<FDisplayClusterMediaInputViewport> NewICVFXTileInput = MakeShared<FDisplayClusterMediaInputViewport>(
 							MediaInputId, ClusterNodeId,
@@ -445,18 +443,16 @@ void FDisplayClusterMediaModule::InitializeICVFXCameraUniformTilesOutput(const U
 				{
 					if (IsValid(MediaOutputTile.MediaOutput))
 					{
-						const FIntPoint TileLoc{ MediaOutputTile.TileX, MediaOutputTile.TileY };
-
-						const FString ICVFXViewportTileId = DisplayClusterMediaHelpers::GenerateTileViewportName(ICVFXViewportId, TileLoc);
+						const FString ICVFXViewportTileId = DisplayClusterMediaHelpers::GenerateTileViewportName(ICVFXViewportId, MediaOutputTile.Position);
 
 						const FString MediaOutputId = DisplayClusterMediaHelpers::MediaId::GenerateMediaId(
 							DisplayClusterMediaHelpers::MediaId::EMediaDeviceType::Output,
 							DisplayClusterMediaHelpers::MediaId::EMediaOwnerType::ICVFXCamera,
 							ClusterNodeId, RootActorName, ICVFXCameraName,
-							static_cast<uint8>(Index), &TileLoc);
+							static_cast<uint8>(Index), &MediaOutputTile.Position);
 
 						UE_LOG(LogDisplayClusterMedia, Log, TEXT("Initializing ICVFX media output '%s' for camera '%s' tile '%d,%d'"),
-							*MediaOutputId, *ICVFXCameraName, TileLoc.X, TileLoc.Y);
+							*MediaOutputId, *ICVFXCameraName, MediaOutputTile.Position.X, MediaOutputTile.Position.Y);
 
 						TSharedPtr<FDisplayClusterMediaCaptureViewport> NewICVFXTileOutput = MakeShared<FDisplayClusterMediaCaptureViewport>(
 							MediaOutputId, ClusterNodeId,
