nDisplay: Changed the logic that when a viewport references a viewpoint with an invalid name, the default viewpoint is now used.

diff --git a/Engine/Plugins/Runtime/nDisplay/Source/DisplayCluster/Private/Render/Viewport/DisplayClusterViewport_ViewPoint.cpp b/Engine/Plugins/Runtime/nDisplay/Source/DisplayCluster/Private/Render/Viewport/DisplayClusterViewport_ViewPoint.cpp
--- a/Engine/Plugins/Runtime/nDisplay/Source/DisplayCluster/Private/Render/Viewport/DisplayClusterViewport_ViewPoint.cpp
+++ b/Engine/Plugins/Runtime/nDisplay/Source/DisplayCluster/Private/Render/Viewport/DisplayClusterViewport_ViewPoint.cpp
@@ -160,10 +160,10 @@ UDisplayClusterCameraComponent* FDisplayClusterViewport::GetViewPointCameraCompo
 
 	if (CanShowLogMsgOnce(EDisplayClusterViewportShowLogMsgOnce::GetViewPointCameraComponent_NotFound))
 	{
-		UE_LOG(LogDisplayClusterViewport, Warning, TEXT("ViewPoint '%s' is not found for viewport '%s'"), *CameraId, *GetId());
+		UE_LOG(LogDisplayClusterViewport, Warning, TEXT("ViewPoint '%s' is not found for viewport '%s'. The default viewpoint will be used."), *CameraId, *GetId());
 	}
 
-	return nullptr;
+	return CameraId.IsEmpty() ? nullptr : RootActor->GetDefaultCamera();
 }
 
 bool FDisplayClusterViewport::SetupViewPoint(FMinimalViewInfo& InOutViewInfo)
