nDisplay: Fixed issue where the static mesh used by the stage geometry isosphere component did not have bAllowCPUAccess set to true, causing a crash in packaged builds.
Also added additional validation that checks this flag before attempting to access mesh data on the CPU, allowing a graceful error message instead of a crash

#virtualized

diff --git a/Engine/Plugins/Runtime/nDisplay/Source/DisplayCluster/Private/Game/EngineClasses/Scene/DisplayClusterStageIsosphereComponent.cpp b/Engine/Plugins/Runtime/nDisplay/Source/DisplayCluster/Private/Game/EngineClasses/Scene/DisplayClusterStageIsosphereComponent.cpp
--- a/Engine/Plugins/Runtime/nDisplay/Source/DisplayCluster/Private/Game/EngineClasses/Scene/DisplayClusterStageIsosphereComponent.cpp
+++ b/Engine/Plugins/Runtime/nDisplay/Source/DisplayCluster/Private/Game/EngineClasses/Scene/DisplayClusterStageIsosphereComponent.cpp
@@ -22,10 +22,8 @@ void UDisplayClusterStageIsosphereComponent::OnComponentCreated()
 	Super::OnComponentCreated();
 
 	// Load the isosphere's geometry into the procedural mesh
-	if (ensure(IsosphereMesh))
+	if (IsIsosphereMeshValid())
 	{
-		IsosphereMesh->bAllowCPUAccess = true;
-
 		const int32 NumSections = IsosphereMesh->GetNumSections(0);
 		for (int32 SectionIndex = 0; SectionIndex < NumSections; SectionIndex++)
 		{
@@ -54,7 +52,7 @@ void UDisplayClusterStageIsosphereComponent::OnComponentCreated()
 
 void UDisplayClusterStageIsosphereComponent::ResetIsosphere()
 {
-	if (ensure(IsosphereMesh))
+	if (IsIsosphereMeshValid())
 	{
 		const int32 NumSections = IsosphereMesh->GetNumSections(0);
 		for (int32 SectionIndex = 0; SectionIndex < NumSections; SectionIndex++)
@@ -81,3 +79,10 @@ void UDisplayClusterStageIsosphereComponent::ResetIsosphere()
 		}
 	}
 }
+
+bool UDisplayClusterStageIsosphereComponent::IsIsosphereMeshValid() const
+{
+	return
+		ensureMsgf(IsosphereMesh, TEXT("UDisplayClusterStageIsosphereComponent::IsosphereMesh was not properly loaded")) &&
+		ensureMsgf(IsosphereMesh->bAllowCPUAccess, TEXT("UDisplayClusterStageIsosphereComponent::IsosphereMesh does not have Allow CPUAccess enabled"));
+}
diff --git a/Engine/Plugins/Runtime/nDisplay/Source/DisplayCluster/Public/Components/DisplayClusterStageIsosphereComponent.h b/Engine/Plugins/Runtime/nDisplay/Source/DisplayCluster/Public/Components/DisplayClusterStageIsosphereComponent.h
--- a/Engine/Plugins/Runtime/nDisplay/Source/DisplayCluster/Public/Components/DisplayClusterStageIsosphereComponent.h
+++ b/Engine/Plugins/Runtime/nDisplay/Source/DisplayCluster/Public/Components/DisplayClusterStageIsosphereComponent.h
@@ -23,6 +23,10 @@ public:
 	 */
 	void ResetIsosphere();
 
+private:
+	/** Checks that the isosphere static mesh used when generating the morphed geometry is loaded and has bAllowCPUAccess set to true */
+	bool IsIsosphereMeshValid() const;
+
 private:
 	/** A static mesh containing an isosphere that can be morphed */
 	UPROPERTY(Transient)
