nDisplay: Fixed issue where stage components weren't being outlined when selected in the nDisplay config editor. The issue was caused by the mesh projection renderer used by the stage geometry component incorrectly overriding the components' SelectionOverrideDelegate, which prevented these components from being reported as selected to the scene renderer

diff --git a/Engine/Plugins/Runtime/nDisplay/Source/DisplayClusterLightCardEditorShaders/Private/DisplayClusterMeshProjectionRenderer.cpp b/Engine/Plugins/Runtime/nDisplay/Source/DisplayClusterLightCardEditorShaders/Private/DisplayClusterMeshProjectionRenderer.cpp
--- a/Engine/Plugins/Runtime/nDisplay/Source/DisplayClusterLightCardEditorShaders/Private/DisplayClusterMeshProjectionRenderer.cpp
+++ b/Engine/Plugins/Runtime/nDisplay/Source/DisplayClusterLightCardEditorShaders/Private/DisplayClusterMeshProjectionRenderer.cpp
@@ -801,7 +801,10 @@ FDisplayClusterMeshProjectionRenderer::~FDisplayClusterMeshProjectionRenderer()
 		}
 
 		// Unbind this since it will still have a raw pointer to this renderer
-		PrimitiveComponent->SelectionOverrideDelegate.Unbind();
+		if (PrimitiveComponent->SelectionOverrideDelegate.IsBoundToObject(this))
+		{
+			PrimitiveComponent->SelectionOverrideDelegate.Unbind();
+		}
 	}
 #endif
 }
@@ -857,7 +860,12 @@ void FDisplayClusterMeshProjectionRenderer::AddActor(AActor* Actor, const TFunct
 			PrimitiveComponents.Add(PrimitiveComponent);
 
 #if WITH_EDITOR
-			PrimitiveComponent->SelectionOverrideDelegate = UPrimitiveComponent::FSelectionOverride::CreateRaw(this, &FDisplayClusterMeshProjectionRenderer::IsPrimitiveComponentSelected);
+			// If the ActorSelectedDelegate is bound, that indicates that the renderer will handle rendering the selection outline itself, and the component's selection override delegate
+			// can be rebound; otherwise, don't touch the delegate as it may be being used by an editor viewport (such as the nDisplay config editor)
+			if (ActorSelectedDelegate.IsBound())
+			{
+				PrimitiveComponent->SelectionOverrideDelegate = UPrimitiveComponent::FSelectionOverride::CreateRaw(this, &FDisplayClusterMeshProjectionRenderer::IsPrimitiveComponentSelected);
+			}
 #endif
 		}
 	});
@@ -873,7 +881,7 @@ void FDisplayClusterMeshProjectionRenderer::RemoveActor(AActor* Actor)
 	for (const TWeakObjectPtr<UPrimitiveComponent>& PrimitiveComponent : ComponentsToRemove)
 	{
 #if WITH_EDITOR
-		if (PrimitiveComponent.IsValid() && PrimitiveComponent->SelectionOverrideDelegate.IsBound())
+		if (PrimitiveComponent.IsValid() && PrimitiveComponent->SelectionOverrideDelegate.IsBoundToObject(this))
 		{
 			PrimitiveComponent->SelectionOverrideDelegate.Unbind();
 		}
@@ -888,7 +896,7 @@ void FDisplayClusterMeshProjectionRenderer::ClearScene()
 	for (const TWeakObjectPtr<UPrimitiveComponent>& PrimitiveComponent : PrimitiveComponents)
 	{
 #if WITH_EDITOR
-		if (PrimitiveComponent.IsValid() && PrimitiveComponent->SelectionOverrideDelegate.IsBound())
+		if (PrimitiveComponent.IsValid() && PrimitiveComponent->SelectionOverrideDelegate.IsBoundToObject(this))
 		{
 			PrimitiveComponent->SelectionOverrideDelegate.Unbind();
 		}
