nDisplay: Fixed a bug that caused the application exit message to not appear in the log.

diff --git a/Engine/Plugins/Runtime/nDisplay/Source/DisplayCluster/Private/Misc/DisplayClusterAppExit.cpp b/Engine/Plugins/Runtime/nDisplay/Source/DisplayCluster/Private/Misc/DisplayClusterAppExit.cpp
--- a/Engine/Plugins/Runtime/nDisplay/Source/DisplayCluster/Private/Misc/DisplayClusterAppExit.cpp
+++ b/Engine/Plugins/Runtime/nDisplay/Source/DisplayCluster/Private/Misc/DisplayClusterAppExit.cpp
@@ -47,6 +47,8 @@ void FDisplayClusterAppExit::ExitApplication(const FString& Msg, EExitType ExitT
 		}
 		else if(ExitType == EExitType::KillImmediately)
 		{
+			UE_LOG(LogDisplayClusterModule, Error, TEXT("KillImmediately requested - %s"), *Msg);
+
 			FProcHandle hProc = FPlatformProcess::OpenProcess(FPlatformProcess::GetCurrentProcessId());
 			FPlatformProcess::TerminateProc(hProc, true);
 		}
