[nDisplay] Minor refactoring of tiled media

diff --git a/Engine/Plugins/Runtime/nDisplay/Source/DisplayCluster/Private/Render/Viewport/Configuration/DisplayClusterViewportConfigurationHelpers_Tile.cpp b/Engine/Plugins/Runtime/nDisplay/Source/DisplayCluster/Private/Render/Viewport/Configuration/DisplayClusterViewportConfigurationHelpers_Tile.cpp
--- a/Engine/Plugins/Runtime/nDisplay/Source/DisplayCluster/Private/Render/Viewport/Configuration/DisplayClusterViewportConfigurationHelpers_Tile.cpp
+++ b/Engine/Plugins/Runtime/nDisplay/Source/DisplayCluster/Private/Render/Viewport/Configuration/DisplayClusterViewportConfigurationHelpers_Tile.cpp
@@ -33,10 +33,9 @@ void FDisplayClusterViewportConfigurationHelpers_Tile::UpdateICVFXCameraViewport
 		return;
 	}
 
-	// Perform additional validations
-	const FDisplayClusterConfigurationICVFX_StageSettings* const StageSettings = InSourceViewport.Configuration->GetStageSettings();
-	const bool bIsTilingAllowed = (StageSettings && FDisplayClusterConfigurationTile_Settings::IsEnabled(InCameraMediaSettings.TiledSplitLayout, *StageSettings));
-	if (!bIsTilingAllowed)
+	// Validate the layout
+	const bool bIsLayoutValid = FDisplayClusterConfigurationTile_Settings::IsValid(InCameraMediaSettings.TiledSplitLayout);
+	if (!bIsLayoutValid)
 	{
 		InSourceViewport.GetRenderSettingsImpl().TileSettings = { };
 		return;
diff --git a/Engine/Plugins/Runtime/nDisplay/Source/DisplayClusterConfiguration/Private/DisplayClusterConfigurationTypes_Tile.cpp b/Engine/Plugins/Runtime/nDisplay/Source/DisplayClusterConfiguration/Private/DisplayClusterConfigurationTypes_Tile.cpp
--- a/Engine/Plugins/Runtime/nDisplay/Source/DisplayClusterConfiguration/Private/DisplayClusterConfigurationTypes_Tile.cpp
+++ b/Engine/Plugins/Runtime/nDisplay/Source/DisplayClusterConfiguration/Private/DisplayClusterConfigurationTypes_Tile.cpp
@@ -17,19 +17,25 @@ bool FDisplayClusterConfigurationTile_Settings::IsEnabled(const FDisplayClusterC
 		return false;
 	}
 
-	return FDisplayClusterConfigurationTile_Settings::IsEnabled(FIntPoint{TileX, TileY}, InStageSettings);
+	if (!FDisplayClusterConfigurationTile_Settings::IsValid(Layout))
+	{
+		// Invalid layout settings
+		return false;
+	}
+
+	return true;
 }
 
-bool FDisplayClusterConfigurationTile_Settings::IsEnabled(const FIntPoint& InTileLoc, const struct FDisplayClusterConfigurationICVFX_StageSettings& InStageSettings)
+bool FDisplayClusterConfigurationTile_Settings::IsValid(const FIntPoint& InTilesLayout)
 {
 	// Ignore wrong values
-	if (InTileLoc.X < 1 || InTileLoc.Y < 1)
+	if (InTilesLayout.X < 1 || InTilesLayout.Y < 1)
 	{
 		return false;
 	}
 
 	// Ignore 1x1 case
-	if (InTileLoc.X == 1 && InTileLoc.Y == 1)
+	if (InTilesLayout.X == 1 && InTilesLayout.Y == 1)
 	{
 		return false;
 	}
diff --git a/Engine/Plugins/Runtime/nDisplay/Source/DisplayClusterConfiguration/Public/DisplayClusterConfigurationTypes_Tile.h b/Engine/Plugins/Runtime/nDisplay/Source/DisplayClusterConfiguration/Public/DisplayClusterConfigurationTypes_Tile.h
--- a/Engine/Plugins/Runtime/nDisplay/Source/DisplayClusterConfiguration/Public/DisplayClusterConfigurationTypes_Tile.h
+++ b/Engine/Plugins/Runtime/nDisplay/Source/DisplayClusterConfiguration/Public/DisplayClusterConfigurationTypes_Tile.h
@@ -21,21 +21,17 @@ public:
 	/** Return true if tile rendering can be used. */
 	bool IsEnabled(const struct FDisplayClusterConfigurationICVFX_StageSettings& InStageSettings) const;
 
-	/** Static validation for custom implementations. */
-	static bool IsEnabled(const FIntPoint& InTileLoc, const struct FDisplayClusterConfigurationICVFX_StageSettings& InStageSettings);
+	/** Layout validation. */
+	static bool IsValid(const FIntPoint& InTilesLayout);
 
 public:
 	/** Enable tile rendering. */
 	UPROPERTY(EditAnywhere, BlueprintReadWrite, Category = "Tile Rendering")
 	bool bEnabled = false;
 
-	/** The viewport will be horizontally partitioned into the specified number of tiles.  */
-	UPROPERTY(EditAnywhere, BlueprintReadWrite, Category = "Tile Rendering", meta = (DisplayName = "Tile X", ClampMin = "1", UIMin = "1", ClampMax = "4", UIMax = "4"))
-	int32 TileX = 1;
-
-	/** The viewport will be vertically partitioned into the specified number of tiles.  */
-	UPROPERTY(EditAnywhere, BlueprintReadWrite, Category = "Tile Rendering", meta = (DisplayName = "Tile Y", ClampMin = "1", UIMin = "1", ClampMax = "4", UIMax = "4"))
-	int32 TileY = 1;
+	/** Tiling layout (X by Y tiles amount). */
+	UPROPERTY(EditAnywhere, BlueprintReadWrite, Category = "Tile Rendering", meta = (ClampMin = "1", ClampMax = "4"))
+	FIntPoint Layout = { 1, 1 };
 };
 
 USTRUCT(Blueprintable)
