[nDisplay] Fixed incamera capture initialization when camera is disabled on start
 - Media capture adapters always subscribe for API events regardless of capture start result
 - Allow capture restart if capture has never started

diff --git a/Engine/Plugins/Runtime/nDisplay/Source/DisplayClusterMedia/Private/Capture/DisplayClusterMediaCaptureBase.cpp b/Engine/Plugins/Runtime/nDisplay/Source/DisplayClusterMedia/Private/Capture/DisplayClusterMediaCaptureBase.cpp
--- a/Engine/Plugins/Runtime/nDisplay/Source/DisplayClusterMedia/Private/Capture/DisplayClusterMediaCaptureBase.cpp
+++ b/Engine/Plugins/Runtime/nDisplay/Source/DisplayClusterMedia/Private/Capture/DisplayClusterMediaCaptureBase.cpp
@@ -132,12 +132,12 @@ void FDisplayClusterMediaCaptureBase::ExportMediaData(FRDGBuilder& GraphBuilder,
 
 void FDisplayClusterMediaCaptureBase::OnPostClusterTick()
 {
-	if (MediaCapture && bWasCaptureStarted)
+	if (MediaCapture)
 	{
 		const EMediaCaptureState MediaCaptureState = MediaCapture->GetState();
 		const bool bMediaCaptureNeedsRestart = (MediaCaptureState == EMediaCaptureState::Error) || (MediaCaptureState == EMediaCaptureState::Stopped);
 
-		if (bMediaCaptureNeedsRestart)
+		if (!bWasCaptureStarted || bMediaCaptureNeedsRestart)
 		{
 			constexpr double Interval = 1.0;
 			const double CurrentTime = FPlatformTime::Seconds();
@@ -146,7 +146,7 @@ void FDisplayClusterMediaCaptureBase::OnPostClusterTick()
 			{
 				UE_LOG(LogDisplayClusterMedia, Log, TEXT("MediaCapture '%s' is in error or stopped, restarting it."), *GetMediaId());
 
-				StartMediaCapture();
+				bWasCaptureStarted = StartMediaCapture();
 				LastRestartTimestamp = CurrentTime;
 			}
 		}
@@ -166,6 +166,11 @@ bool FDisplayClusterMediaCaptureBase::StartMediaCapture()
 		Descriptor.ResourceSize = LastSrcRegionSize.load().ToIntPoint();
 	}
 	
+	if (Descriptor.ResourceSize == FIntPoint::ZeroValue)
+	{
+		return false;
+	}
+
 	FMediaCaptureOptions MediaCaptureOptions;
 	MediaCaptureOptions.NumberOfFramesToCapture = -1;
 	MediaCaptureOptions.bAutoRestartOnSourceSizeChange = true;
diff --git a/Engine/Plugins/Runtime/nDisplay/Source/DisplayClusterMedia/Private/Capture/DisplayClusterMediaCaptureViewport.cpp b/Engine/Plugins/Runtime/nDisplay/Source/DisplayClusterMedia/Private/Capture/DisplayClusterMediaCaptureViewport.cpp
--- a/Engine/Plugins/Runtime/nDisplay/Source/DisplayClusterMedia/Private/Capture/DisplayClusterMediaCaptureViewport.cpp
+++ b/Engine/Plugins/Runtime/nDisplay/Source/DisplayClusterMedia/Private/Capture/DisplayClusterMediaCaptureViewport.cpp
@@ -29,24 +29,20 @@ FDisplayClusterMediaCaptureViewport::FDisplayClusterMediaCaptureViewport(const F
 
 bool FDisplayClusterMediaCaptureViewport::StartCapture()
 {
-	// If capturing has started successfully, subscribe for the pipeline callbacks
-	if (FDisplayClusterMediaCaptureBase::StartCapture())
-	{
-		IDisplayCluster::Get().GetCallbacks().OnDisplayClusterPostRenderViewFamily_RenderThread().AddRaw(this, &FDisplayClusterMediaCaptureViewport::OnPostRenderViewFamily_RenderThread);
-		IDisplayCluster::Get().GetCallbacks().OnDisplayClusterUpdateViewportMediaState().AddRaw(this, &FDisplayClusterMediaCaptureViewport::OnUpdateViewportMediaState);
+	// Subscribe for events
+	IDisplayCluster::Get().GetCallbacks().OnDisplayClusterPostRenderViewFamily_RenderThread().AddRaw(this, &FDisplayClusterMediaCaptureViewport::OnPostRenderViewFamily_RenderThread);
+	IDisplayCluster::Get().GetCallbacks().OnDisplayClusterUpdateViewportMediaState().AddRaw(this, &FDisplayClusterMediaCaptureViewport::OnUpdateViewportMediaState);
 
-		return true;
-	}
+	// Start capture
+	const bool bStarted = FDisplayClusterMediaCaptureBase::StartCapture();
 
-	return false;
+	return bStarted;
 }
 
 void FDisplayClusterMediaCaptureViewport::StopCapture()
 {
-	// Stop rendering notifications
+	// Unsubscribe from events
 	IDisplayCluster::Get().GetCallbacks().OnDisplayClusterPostRenderViewFamily_RenderThread().RemoveAll(this);
-
-	// Stop raising media flags for the viewport.
 	IDisplayCluster::Get().GetCallbacks().OnDisplayClusterUpdateViewportMediaState().RemoveAll(this);
 
 	// Stop capturing
@@ -90,9 +86,7 @@ FIntPoint FDisplayClusterMediaCaptureViewport::GetCaptureSize() const
 		}
 	}
 
-	UE_LOG(LogDisplayClusterMedia, Warning, TEXT("'%s' couldn't get viewport size"), *GetMediaId());
-
-	return FIntPoint();
+	return FIntPoint::ZeroValue;
 }
 
 void FDisplayClusterMediaCaptureViewport::OnPostRenderViewFamily_RenderThread(FRDGBuilder& GraphBuilder, const FSceneViewFamily& ViewFamily, const IDisplayClusterViewportProxy* ViewportProxy)
