[nDisplay]: added Inner Frustum Projection Offset controls

diff --git a/Engine/Plugins/Runtime/nDisplay/Source/DisplayCluster/Private/Render/Viewport/Configuration/DisplayClusterViewportConfigurationHelpers_ICVFX.cpp b/Engine/Plugins/Runtime/nDisplay/Source/DisplayCluster/Private/Render/Viewport/Configuration/DisplayClusterViewportConfigurationHelpers_ICVFX.cpp
--- a/Engine/Plugins/Runtime/nDisplay/Source/DisplayCluster/Private/Render/Viewport/Configuration/DisplayClusterViewportConfigurationHelpers_ICVFX.cpp
+++ b/Engine/Plugins/Runtime/nDisplay/Source/DisplayCluster/Private/Render/Viewport/Configuration/DisplayClusterViewportConfigurationHelpers_ICVFX.cpp
@@ -101,6 +101,7 @@ bool FDisplayClusterViewportConfigurationHelpers_ICVFX::UpdateCameraProjectionSe
 	// Lens correction
 	PolicyCameraSettings.FrustumRotation = InCameraSettings.FrustumRotation;
 	PolicyCameraSettings.FrustumOffset = InCameraSettings.FrustumOffset;
+	PolicyCameraSettings.OffCenterProjectionOffset = InCameraSettings.OffCenterProjectionOffset;
 
 	static IDisplayClusterProjection& DisplayClusterProjectionAPI = IDisplayClusterProjection::Get();
 
diff --git a/Engine/Plugins/Runtime/nDisplay/Source/DisplayCluster/Public/Components/DisplayClusterICVFXCameraComponent.h b/Engine/Plugins/Runtime/nDisplay/Source/DisplayCluster/Public/Components/DisplayClusterICVFXCameraComponent.h
--- a/Engine/Plugins/Runtime/nDisplay/Source/DisplayCluster/Public/Components/DisplayClusterICVFXCameraComponent.h
+++ b/Engine/Plugins/Runtime/nDisplay/Source/DisplayCluster/Public/Components/DisplayClusterICVFXCameraComponent.h
@@ -129,6 +129,9 @@ private:
 	UPROPERTY(EditAnywhere, Transient, Category = "Inner Frustum", meta = (PropertyPath = "CameraSettings.FrustumOffset"))
 	FDisplayClusterEditorPropertyReference FrustumOffsetRef;
 
+	UPROPERTY(EditAnywhere, Transient, Category = "Inner Frustum", meta = (PropertyPath = "CameraSettings.OffCenterprojectionoffset"))
+	FDisplayClusterEditorPropertyReference OffCenterProjectionOffsetRef;
+
 	UPROPERTY(EditAnywhere, Transient, Category = "Inner Frustum", meta = (PropertyPath = "CameraSettings.RenderSettings.GenerateMips"))
 	FDisplayClusterEditorPropertyReference GenerateMipsRef;
 
diff --git a/Engine/Plugins/Runtime/nDisplay/Source/DisplayClusterConfiguration/Public/DisplayClusterConfigurationTypes_ICVFX.h b/Engine/Plugins/Runtime/nDisplay/Source/DisplayClusterConfiguration/Public/DisplayClusterConfigurationTypes_ICVFX.h
--- a/Engine/Plugins/Runtime/nDisplay/Source/DisplayClusterConfiguration/Public/DisplayClusterConfigurationTypes_ICVFX.h
+++ b/Engine/Plugins/Runtime/nDisplay/Source/DisplayClusterConfiguration/Public/DisplayClusterConfigurationTypes_ICVFX.h
@@ -775,6 +775,10 @@ public:
 	UPROPERTY(EditAnywhere, BlueprintReadWrite, Category = "In Camera VFX", meta = (DisplayName = "Inner Frustum Offset"))
 	FVector FrustumOffset = FVector::ZeroVector;
 
+	/** Off-axis / off-center projection offset as proportion of screen dimensions. */
+	UPROPERTY(EditAnywhere, BlueprintReadWrite, Category = "In Camera VFX", meta = (DisplayName = "Inner Frustum Projection Offset"))
+	FVector2D OffCenterProjectionOffset = FVector2D::ZeroVector;
+
 	/**Border for the inner frustum. */
 	UPROPERTY(EditAnywhere, BlueprintReadWrite, Category = "In Camera VFX", meta = (DisplayName = "Inner Frustum Border"))
 	FDisplayClusterConfigurationICVFX_CameraBorder Border;
diff --git a/Engine/Plugins/Runtime/nDisplay/Source/DisplayClusterProjection/Private/Policy/Camera/DisplayClusterProjectionCameraPolicy.cpp b/Engine/Plugins/Runtime/nDisplay/Source/DisplayClusterProjection/Private/Policy/Camera/DisplayClusterProjectionCameraPolicy.cpp
--- a/Engine/Plugins/Runtime/nDisplay/Source/DisplayClusterProjection/Private/Policy/Camera/DisplayClusterProjectionCameraPolicy.cpp
+++ b/Engine/Plugins/Runtime/nDisplay/Source/DisplayClusterProjection/Private/Policy/Camera/DisplayClusterProjectionCameraPolicy.cpp
@@ -130,6 +130,16 @@ bool FDisplayClusterProjectionCameraPolicy::GetProjectionMatrix(IDisplayClusterV
 		OutPrjMatrix = ComposureSettings.GetProjectionMatrix(ClampedCameraFOVDegrees, CameraAspectRatio);
 	}
 
+	if (!CameraSettings.OffCenterProjectionOffset.IsZero())
+	{
+		const float Left = -1.0f + CameraSettings.OffCenterProjectionOffset.X;
+		const float Right = Left + 2.0f;
+		const float Bottom = -1.0f + CameraSettings.OffCenterProjectionOffset.Y;
+		const float Top = Bottom + 2.0f;
+		OutPrjMatrix.M[2][0] = (Left + Right) / (Left - Right);
+		OutPrjMatrix.M[2][1] = (Bottom + Top) / (Bottom - Top);
+	}
+
 	return true;
 }
 
diff --git a/Engine/Plugins/Runtime/nDisplay/Source/DisplayClusterProjection/Public/Containers/DisplayClusterProjectionCameraPolicySettings.h b/Engine/Plugins/Runtime/nDisplay/Source/DisplayClusterProjection/Public/Containers/DisplayClusterProjectionCameraPolicySettings.h
--- a/Engine/Plugins/Runtime/nDisplay/Source/DisplayClusterProjection/Public/Containers/DisplayClusterProjectionCameraPolicySettings.h
+++ b/Engine/Plugins/Runtime/nDisplay/Source/DisplayClusterProjection/Public/Containers/DisplayClusterProjectionCameraPolicySettings.h
@@ -14,6 +14,9 @@ struct FDisplayClusterProjectionCameraPolicySettings
 	// Move incamera frustum on this value to fix broken lens on physic camera
 	FVector  FrustumOffset = FVector::ZeroVector;
 
+	// Off-axis / off-center projection offset as proportion of screen dimensions
+	FVector2D OffCenterProjectionOffset = FVector2D::ZeroVector;
+
 	// When enabled, ignore all next updates, until null cam ref not assigned (UE-137222)
 	bool bCameraOverrideDefaults = false;
 
