nDisplay: Fixed Alpha and Beta mapping shader.

diff --git a/Engine/Plugins/Runtime/nDisplay/Shaders/Private/MPCDIUtils.ush b/Engine/Plugins/Runtime/nDisplay/Shaders/Private/MPCDIUtils.ush
--- a/Engine/Plugins/Runtime/nDisplay/Shaders/Private/MPCDIUtils.ush
+++ b/Engine/Plugins/Runtime/nDisplay/Shaders/Private/MPCDIUtils.ush
@@ -12,9 +12,21 @@ float4 DrawRectangleUVScaleBias;
 // PS parameters
 Texture2D InputTexture;
 Texture2D WarpMapTexture;
+
+// AlphaMap and BetaMap textures.
 Texture2D AlphaMapTexture;
 Texture2D BetaMapTexture;
 
+// Indicates the number of components within the alpha or beta map. Valid values are 1 or 3.
+int AlphaMapComponentDepth;
+int BetaMapComponentDepth;
+
+// This parameter contains the 'gammaEmbedded' tag value, which is the gamma value for the AlphaMap texture.
+// Note: The MPCDI 2.0 standard does not define an 'EmbeddedGamma' tag value for the BetaMap tag.
+// However, it does require gamma correction for BetaMap.
+// Therefore, this parameter is also used for the BetaMap texture.
+float AlphaMapGammaEmbedded;
+
 Texture2D InnerCameraTexture;
 Texture2D ChromakeyCameraTexture;
 Texture2D ChromakeyMarkerTexture;
@@ -53,7 +65,6 @@ float4 OverlappedInnerCameraSoftEdges[MAX_INNER_CAMERAS_AMOUNT];
 #endif
 
 float LightCardGamma;
-float AlphaEmbeddedGamma;
 
 float4 InnerCameraSoftEdge;
 
@@ -137,24 +148,67 @@ void DrawRectangle(in float4 InPosition, in float2 InTexCoord, out float4 OutPos
 }
 
 
-float4 ApplyBlending(float3 color, OutputVS IN, float ColorAlphaValue)
+float4 ApplyBlending(float3 InColor_sRGB, OutputVS IN, float ColorAlphaValue)
 {
+	// 2.4 Apply Alpha and Beta mapping [MPCDI 2.0]
+	//
+	// Intensity blending and black level adjustments in projector overlap regions are expressed in alpha
+	// and beta maps respectively. These are expressed with the following equation.
+	//
+	//  Pf = (Pi * a * (1 - b)) + b
+	//
+	// where:
+	//   Pi - is the initial color
+	//   Pf - is final pixel color
+	//    a - is the gamma-corrected intensity mapping color
+	//    b - is the gamma-corrected black level offset color
+	
+	// BetaMap can be used by itself, without AlphaMap.
+#if ALPHAMAP_BLENDING || BETAMAP_BLENDING
+	float3 LinearColor = MPCDI_sRGBToLinear(InColor_sRGB);
+
+	// To convert a texture with its own gamma to a linear color space, just use this formula:
+	//     LinearTextureColor = pow(TextureColor, 1 / GammaEmbedded) function.
+	
+	// Note: AlphaMapGammaEmbedded is used for both AlphaMap and BetaMap.
+	// (The MPCDI 2.0 standard does not define an 'EmbeddedGamma' tag value for the BetaMap tag. However, it does require gamma correction for BetaMap.)
+	float InvGammaEmbedded = 1 / AlphaMapGammaEmbedded;
+	
 #if ALPHAMAP_BLENDING
-	color = MPCDI_sRGBToLinear(color);
-
-	// Load alpha
-	float alpha = pow(AlphaMapTexture.Sample(AlphaMapSampler, IN.UV.xy).r, AlphaEmbeddedGamma);
+	if(AlphaMapComponentDepth < 3)
+	{
+		// Grayscale AlphaMap: get a single-channel value and convert its color space to linear.
+		float LinearAlphaMap = pow(AlphaMapTexture.Sample(AlphaMapSampler, IN.UV.xy).r, InvGammaEmbedded);
+		LinearColor = LinearColor * LinearAlphaMap;
+	}
+	else
+	{
+		// RGB AlphaMap: get the values of the 3 texture channels and convert their color space to linear.
+		float3 LinearAlphaMap = pow(AlphaMapTexture.Sample(AlphaMapSampler, IN.UV.xy).rgb, InvGammaEmbedded);
+		LinearColor = LinearColor * LinearAlphaMap;
+	}
+#endif
 
 #if BETAMAP_BLENDING
-	float beta = BetaMapTexture.Sample(BetaMapSampler, IN.UV.xy).r;
-	color = (color*alpha*(1 - beta)) + beta;
-#else // NOT BETAMAP_BLENDING
-	color = color * alpha;
+	if(BetaMapComponentDepth < 3)
+	{
+		// Grayscale BetaMap: get a single-channel value and convert its color space to linear.
+		float LinearBetaMap = pow(BetaMapTexture.Sample(BetaMapSampler, IN.UV.xy).r, InvGammaEmbedded);
+		LinearColor = LinearColor * ( 1 - LinearBetaMap) + LinearBetaMap;
+	}
+	else
+	{
+		// RGB BetaMap: get the values of the 3 texture channels and convert their color space to linear.
+		float3 LinearBetaMap = pow(BetaMapTexture.Sample(BetaMapSampler, IN.UV.xy).rgb, InvGammaEmbedded);
+		LinearColor = LinearColor * ( 1 - LinearBetaMap) + LinearBetaMap;
+	}
 #endif // BETAMAP_BLENDING
-	// Load beta
 
-	return float4(MPCDI_LinearTosRGB(color), ColorAlphaValue);
-#else //NOT ALPHAMAP_BLENDING
-	return float4(color, ColorAlphaValue);
-#endif // ALPHAMAP_BLENDING
+	return float4(MPCDI_LinearTosRGB(LinearColor), ColorAlphaValue);
+
+#else // #if ALPHAMAP_BLENDING || BETAMAP_BLENDING
+	
+	return float4(InColor_sRGB, ColorAlphaValue);
+	
+#endif // #if ALPHAMAP_BLENDING || BETAMAP_BLENDING
 }
diff --git a/Engine/Plugins/Runtime/nDisplay/Source/DisplayCluster/Private/Render/Containers/DisplayClusterRender_Texture.h b/Engine/Plugins/Runtime/nDisplay/Source/DisplayCluster/Private/Render/Containers/DisplayClusterRender_Texture.h
--- a/Engine/Plugins/Runtime/nDisplay/Source/DisplayCluster/Private/Render/Containers/DisplayClusterRender_Texture.h
+++ b/Engine/Plugins/Runtime/nDisplay/Source/DisplayCluster/Private/Render/Containers/DisplayClusterRender_Texture.h
@@ -59,6 +59,12 @@ public:
 		return TextureResource.IsValid() ? TextureResource->GetPixelFormat() : PF_Unknown;
 	}
 
+	virtual uint32 GetComponentDepth() const override
+	{
+		const TSharedPtr<FDisplayClusterRender_TextureResource, ESPMode::ThreadSafe>& TextureResource = GetTextureResource();
+		return TextureResource.IsValid() ? TextureResource->GetComponentDepth() : 0;
+	}
+
 	virtual FRHITexture* GetRHITexture() const override
 	{
 		const TSharedPtr<FDisplayClusterRender_TextureResource, ESPMode::ThreadSafe>& TextureResource = GetTextureResource();
diff --git a/Engine/Plugins/Runtime/nDisplay/Source/DisplayCluster/Private/Render/Containers/DisplayClusterRender_TextureResource.cpp b/Engine/Plugins/Runtime/nDisplay/Source/DisplayCluster/Private/Render/Containers/DisplayClusterRender_TextureResource.cpp
--- a/Engine/Plugins/Runtime/nDisplay/Source/DisplayCluster/Private/Render/Containers/DisplayClusterRender_TextureResource.cpp
+++ b/Engine/Plugins/Runtime/nDisplay/Source/DisplayCluster/Private/Render/Containers/DisplayClusterRender_TextureResource.cpp
@@ -96,6 +96,7 @@ FDisplayClusterRender_TextureResource::FDisplayClusterRender_TextureResource(con
 	PixelFormat = DestPixelFormat;
 	Width = InWidth;
 	Height = InHeight;
+	ComponentDepth = InComponentDepth;
 }
 
 FDisplayClusterRender_TextureResource::~FDisplayClusterRender_TextureResource()
diff --git a/Engine/Plugins/Runtime/nDisplay/Source/DisplayCluster/Private/Render/Containers/DisplayClusterRender_TextureResource.h b/Engine/Plugins/Runtime/nDisplay/Source/DisplayCluster/Private/Render/Containers/DisplayClusterRender_TextureResource.h
--- a/Engine/Plugins/Runtime/nDisplay/Source/DisplayCluster/Private/Render/Containers/DisplayClusterRender_TextureResource.h
+++ b/Engine/Plugins/Runtime/nDisplay/Source/DisplayCluster/Private/Render/Containers/DisplayClusterRender_TextureResource.h
@@ -66,6 +66,12 @@ public:
 		return PixelFormat;
 	}
 
+	/** Get number of components. */
+	uint32 GetComponentDepth() const
+	{
+		return ComponentDepth;
+	}
+
 private:
 	// RAW data of the texture. Mips not supported now
 	void* TextureData = nullptr;
@@ -75,6 +81,9 @@ private:
 	
 	// Height of the texture.
 	uint32_t Height = 0;
+
+	// Number of components
+	uint32 ComponentDepth = 0;
 	
 	// Format in which mip data is stored.
 	EPixelFormat PixelFormat = EPixelFormat::PF_Unknown;
diff --git a/Engine/Plugins/Runtime/nDisplay/Source/DisplayCluster/Public/Render/Containers/IDisplayClusterRender_Texture.h b/Engine/Plugins/Runtime/nDisplay/Source/DisplayCluster/Public/Render/Containers/IDisplayClusterRender_Texture.h
--- a/Engine/Plugins/Runtime/nDisplay/Source/DisplayCluster/Public/Render/Containers/IDisplayClusterRender_Texture.h
+++ b/Engine/Plugins/Runtime/nDisplay/Source/DisplayCluster/Public/Render/Containers/IDisplayClusterRender_Texture.h
@@ -50,6 +50,9 @@ public:
 	/** Get texture pixel format. */
 	virtual EPixelFormat GetPixelFormat() const = 0;
 
+	/** Get number of components. */
+	virtual uint32 GetComponentDepth() const = 0;
+
 	/** Get the RHI resource. */
 	virtual FRHITexture* GetRHITexture() const = 0;
 };
diff --git a/Engine/Plugins/Runtime/nDisplay/Source/DisplayClusterShaders/Private/Shaders/DisplayClusterShadersWarpblend_ICVFX.cpp b/Engine/Plugins/Runtime/nDisplay/Source/DisplayClusterShaders/Private/Shaders/DisplayClusterShadersWarpblend_ICVFX.cpp
--- a/Engine/Plugins/Runtime/nDisplay/Source/DisplayClusterShaders/Private/Shaders/DisplayClusterShadersWarpblend_ICVFX.cpp
+++ b/Engine/Plugins/Runtime/nDisplay/Source/DisplayClusterShaders/Private/Shaders/DisplayClusterShadersWarpblend_ICVFX.cpp
@@ -16,6 +16,7 @@
 #include "ShaderPermutation.h"
 
 #include "Render/Containers/IDisplayClusterRender_MeshComponentProxy.h"
+#include "Render/Containers/IDisplayClusterRender_Texture.h"
 
 #include "ShaderParameters/DisplayClusterShaderParameters_WarpBlend.h"
 #include "ShaderParameters/DisplayClusterShaderParameters_ICVFX.h"
@@ -241,14 +242,9 @@ namespace IcvfxShaderPermutation
 		if (PermutationVector.Get<FIcvfxShaderChromakeyMarker>())
 		{
 			if (!PermutationVector.Get<FIcvfxShaderChromakey>() && !PermutationVector.Get<FIcvfxShaderChromakeyFrameColor>())
-		{
-			return false;
-		}
-		}
-
-		if (!PermutationVector.Get<FIcvfxShaderAlphaMapBlending>() && PermutationVector.Get<FIcvfxShaderBetaMapBlending>())
-		{
-			return false;
+			{
+				return false;
+			}
 		}
 
 		return true;
@@ -297,7 +293,11 @@ BEGIN_SHADER_PARAMETER_STRUCT(FIcvfxPixelShaderParameters, )
 	SHADER_PARAMETER(FMatrix44f, InnerCameraProjectionMatrix)
 
 	SHADER_PARAMETER(float, LightCardGamma)
-	SHADER_PARAMETER(float, AlphaEmbeddedGamma)
+
+	SHADER_PARAMETER(float, AlphaMapGammaEmbedded)
+
+	SHADER_PARAMETER(int, AlphaMapComponentDepth)
+	SHADER_PARAMETER(int, BetaMapComponentDepth)
 
 	SHADER_PARAMETER(FVector4f, InnerCameraSoftEdge)
 
@@ -562,21 +562,28 @@ public:
 	{
 		if (WarpBlendParameters.WarpInterface.IsValid())
 		{
-			FRHITexture* AlphaMap = WarpBlendParameters.WarpInterface->GetTexture(EDisplayClusterWarpBlendTextureType::AlphaMap);
-			if (AlphaMap)
+			TSharedPtr<IDisplayClusterRender_Texture, ESPMode::ThreadSafe> AlphaMap = WarpBlendParameters.WarpInterface->GetTextureInterface(EDisplayClusterWarpBlendTextureType::AlphaMap);
+			if (FRHITexture* AlphaMapTexture = AlphaMap.IsValid() ? AlphaMap->GetRHITexture() : nullptr)
 			{
-				RenderPassData.PSParameters.AlphaMapTexture = AlphaMap;
+				RenderPassData.PSParameters.AlphaMapTexture = AlphaMapTexture;
 				RenderPassData.PSParameters.AlphaMapSampler = TStaticSamplerState<SF_Trilinear, AM_Clamp, AM_Clamp, AM_Clamp>::GetRHI();
-				RenderPassData.PSParameters.AlphaEmbeddedGamma = WarpBlendParameters.WarpInterface->GetAlphaMapEmbeddedGamma();
+				RenderPassData.PSParameters.AlphaMapComponentDepth = AlphaMap->GetComponentDepth();
+				RenderPassData.PSParameters.AlphaMapGammaEmbedded = WarpBlendParameters.WarpInterface->GetAlphaMapEmbeddedGamma();
 
 				RenderPassData.PSPermutationVector.Set<IcvfxShaderPermutation::FIcvfxShaderAlphaMapBlending>(true);
 			}
 
-			FRHITexture* BetaMap = WarpBlendParameters.WarpInterface->GetTexture(EDisplayClusterWarpBlendTextureType::BetaMap);
-			if (BetaMap)
+			TSharedPtr<IDisplayClusterRender_Texture, ESPMode::ThreadSafe> BetaMap = WarpBlendParameters.WarpInterface->GetTextureInterface(EDisplayClusterWarpBlendTextureType::BetaMap);
+			if (FRHITexture* BetaMapTexture = BetaMap.IsValid() ? BetaMap->GetRHITexture() : nullptr)
 			{
-				RenderPassData.PSParameters.BetaMapTexture = BetaMap;
+				RenderPassData.PSParameters.BetaMapTexture = BetaMapTexture;
 				RenderPassData.PSParameters.BetaMapSampler = TStaticSamplerState<SF_Trilinear, AM_Clamp, AM_Clamp, AM_Clamp>::GetRHI();
+				RenderPassData.PSParameters.BetaMapComponentDepth = BetaMap->GetComponentDepth();
+
+				// Note: The MPCDI 2.0 standard does not define an 'EmbeddedGamma' tag value for the BetaMap tag.
+				// However, it does require gamma correction for BetaMap.
+				// Therefore, this parameter is also used for the BetaMap texture.
+				RenderPassData.PSParameters.AlphaMapGammaEmbedded = WarpBlendParameters.WarpInterface->GetAlphaMapEmbeddedGamma();
 
 				RenderPassData.PSPermutationVector.Set<IcvfxShaderPermutation::FIcvfxShaderBetaMapBlending>(true);
 			}
diff --git a/Engine/Plugins/Runtime/nDisplay/Source/DisplayClusterShaders/Private/Shaders/DisplayClusterShadersWarpblend_MPCDI.cpp b/Engine/Plugins/Runtime/nDisplay/Source/DisplayClusterShaders/Private/Shaders/DisplayClusterShadersWarpblend_MPCDI.cpp
--- a/Engine/Plugins/Runtime/nDisplay/Source/DisplayClusterShaders/Private/Shaders/DisplayClusterShadersWarpblend_MPCDI.cpp
+++ b/Engine/Plugins/Runtime/nDisplay/Source/DisplayClusterShaders/Private/Shaders/DisplayClusterShadersWarpblend_MPCDI.cpp
@@ -19,6 +19,7 @@
 #include "ShaderParameterStruct.h"
 
 #include "Render/Containers/IDisplayClusterRender_MeshComponentProxy.h"
+#include "Render/Containers/IDisplayClusterRender_Texture.h"
 #include "IDisplayClusterWarpBlend.h"
 
 #include "ShaderParameters/DisplayClusterShaderParameters_WarpBlend.h"
@@ -76,11 +77,6 @@ namespace MpcdiShaderPermutation
 
 	bool ShouldCompileCommonPSPermutation(const FGlobalShaderPermutationParameters& Parameters, const FCommonPSDomain& PermutationVector)
 	{
-		if (!PermutationVector.Get<FMpcdiShaderAlphaMapBlending>() && PermutationVector.Get<FMpcdiShaderBetaMapBlending>())
-		{
-			return false;
-		}
-
 		return true;
 	}
 
@@ -111,7 +107,10 @@ BEGIN_SHADER_PARAMETER_STRUCT(FMpcdiPixelShaderParameters, )
 
 	SHADER_PARAMETER(FMatrix44f, ViewportTextureProjectionMatrix)
 
-	SHADER_PARAMETER(float, AlphaEmbeddedGamma)
+	SHADER_PARAMETER(float, AlphaMapGammaEmbedded)
+
+	SHADER_PARAMETER(int, AlphaMapComponentDepth)
+	SHADER_PARAMETER(int, BetaMapComponentDepth)
 END_SHADER_PARAMETER_STRUCT()
 
 class FMpcdiWarpVS : public FGlobalShader
@@ -306,21 +305,28 @@ public:
 	{
 		if (WarpBlendParameters.WarpInterface.IsValid())
 		{
-			FRHITexture* AlphaMap = WarpBlendParameters.WarpInterface->GetTexture(EDisplayClusterWarpBlendTextureType::AlphaMap);
-			if (AlphaMap)
+			TSharedPtr<IDisplayClusterRender_Texture, ESPMode::ThreadSafe> AlphaMap = WarpBlendParameters.WarpInterface->GetTextureInterface(EDisplayClusterWarpBlendTextureType::AlphaMap);
+			if (FRHITexture* AlphaMapTexture = AlphaMap.IsValid() ? AlphaMap->GetRHITexture() : nullptr)
 			{
-				RenderPassData.PSParameters.AlphaMapTexture = AlphaMap;
+				RenderPassData.PSParameters.AlphaMapTexture = AlphaMapTexture;
 				RenderPassData.PSParameters.AlphaMapSampler = TStaticSamplerState<SF_Trilinear, AM_Clamp, AM_Clamp, AM_Clamp>::GetRHI();
-				RenderPassData.PSParameters.AlphaEmbeddedGamma = WarpBlendParameters.WarpInterface->GetAlphaMapEmbeddedGamma();
+				RenderPassData.PSParameters.AlphaMapComponentDepth = AlphaMap->GetComponentDepth();
+				RenderPassData.PSParameters.AlphaMapGammaEmbedded = WarpBlendParameters.WarpInterface->GetAlphaMapEmbeddedGamma();
 
 				RenderPassData.PSPermutationVector.Set<MpcdiShaderPermutation::FMpcdiShaderAlphaMapBlending>(true);
 			}
 
-			FRHITexture* BetaMap = WarpBlendParameters.WarpInterface->GetTexture(EDisplayClusterWarpBlendTextureType::BetaMap);
-			if (BetaMap)
+			TSharedPtr<IDisplayClusterRender_Texture, ESPMode::ThreadSafe> BetaMap = WarpBlendParameters.WarpInterface->GetTextureInterface(EDisplayClusterWarpBlendTextureType::BetaMap);
+			if (FRHITexture* BetaMapTexture = BetaMap.IsValid() ? BetaMap->GetRHITexture() : nullptr)
 			{
-				RenderPassData.PSParameters.BetaMapTexture = BetaMap;
+				RenderPassData.PSParameters.BetaMapTexture = BetaMapTexture;
 				RenderPassData.PSParameters.BetaMapSampler = TStaticSamplerState<SF_Trilinear, AM_Clamp, AM_Clamp, AM_Clamp>::GetRHI();
+				RenderPassData.PSParameters.BetaMapComponentDepth = BetaMap->GetComponentDepth();
+
+				// Note: The MPCDI 2.0 standard does not define an 'EmbeddedGamma' tag value for the BetaMap tag.
+				// However, it does require gamma correction for BetaMap.
+				// Therefore, this parameter is also used for the BetaMap texture.
+				RenderPassData.PSParameters.AlphaMapGammaEmbedded = WarpBlendParameters.WarpInterface->GetAlphaMapEmbeddedGamma();
 
 				RenderPassData.PSPermutationVector.Set<MpcdiShaderPermutation::FMpcdiShaderBetaMapBlending>(true);
 			}
diff --git a/Engine/Plugins/Runtime/nDisplay/Source/DisplayClusterWarp/Private/WarpBlend/DisplayClusterWarpBlend.h b/Engine/Plugins/Runtime/nDisplay/Source/DisplayClusterWarp/Private/WarpBlend/DisplayClusterWarpBlend.h
--- a/Engine/Plugins/Runtime/nDisplay/Source/DisplayClusterWarp/Private/WarpBlend/DisplayClusterWarpBlend.h
+++ b/Engine/Plugins/Runtime/nDisplay/Source/DisplayClusterWarp/Private/WarpBlend/DisplayClusterWarpBlend.h
@@ -47,6 +47,11 @@ public:
 		return pTexture.IsValid() ? pTexture->GetRHITexture() : nullptr;
 	}
 
+	virtual TSharedPtr<IDisplayClusterRender_Texture, ESPMode::ThreadSafe> GetTextureInterface(EDisplayClusterWarpBlendTextureType InTextureType) const override
+	{
+		return GeometryContext.GeometryProxy.ImplGetTexture(InTextureType);
+	}
+
 	virtual float GetAlphaMapEmbeddedGamma() const override
 	{ 
 		return GeometryContext.GeometryProxy.AlphaMapEmbeddedGamma; 
diff --git a/Engine/Plugins/Runtime/nDisplay/Source/DisplayClusterWarp/Public/IDisplayClusterWarpBlend.h b/Engine/Plugins/Runtime/nDisplay/Source/DisplayClusterWarp/Public/IDisplayClusterWarpBlend.h
--- a/Engine/Plugins/Runtime/nDisplay/Source/DisplayClusterWarp/Public/IDisplayClusterWarpBlend.h
+++ b/Engine/Plugins/Runtime/nDisplay/Source/DisplayClusterWarp/Public/IDisplayClusterWarpBlend.h
@@ -57,6 +57,11 @@ public:
 	 */
 	virtual class FRHITexture* GetTexture(EDisplayClusterWarpBlendTextureType InTextureType) const = 0;
 
+	/**
+	* Get texture interface by type
+	*/
+	virtual TSharedPtr<class IDisplayClusterRender_Texture, ESPMode::ThreadSafe> GetTextureInterface(EDisplayClusterWarpBlendTextureType InTextureType) const = 0;
+
 	/**
 	 * Return AlphaMap embedded gamma value
 	 */
