nDisplay: Fixed stereo rendering.

diff --git a/Engine/Plugins/Runtime/nDisplay/Source/DisplayCluster/Private/Render/Viewport/RenderTarget/DisplayClusterRenderTargetManager.cpp b/Engine/Plugins/Runtime/nDisplay/Source/DisplayCluster/Private/Render/Viewport/RenderTarget/DisplayClusterRenderTargetManager.cpp
--- a/Engine/Plugins/Runtime/nDisplay/Source/DisplayCluster/Private/Render/Viewport/RenderTarget/DisplayClusterRenderTargetManager.cpp
+++ b/Engine/Plugins/Runtime/nDisplay/Source/DisplayCluster/Private/Render/Viewport/RenderTarget/DisplayClusterRenderTargetManager.cpp
@@ -218,9 +218,15 @@ bool FDisplayClusterRenderTargetManager::AllocateFrameTargets(const FIntPoint& I
 	}
 
 	const uint32 FrameTargetsAmount = Configuration->GetRenderFrameSettings().GetViewPerViewportAmount();
-	const FIntPoint TargetOffset = Configuration->GetRenderFrameSettings().GetDesiredRTTSize(InViewportSize);
+	const FIntPoint DesiredRTTSize = Configuration->GetRenderFrameSettings().GetDesiredRTTSize(InViewportSize);
 	FIntPoint TargetLocation(ForceInitToZero);
 
+	// Offset for stereo rendering on a monoscopic display (side-by-side or top-bottom)
+	const FIntPoint TargetOffset(
+		DesiredRTTSize.X < InViewportSize.X ? DesiredRTTSize.X : 0,
+		DesiredRTTSize.Y < InViewportSize.Y ? DesiredRTTSize.Y : 0
+	);
+
 	// Reallocate frame target resources
 	TArray<TSharedPtr<FDisplayClusterViewportResource, ESPMode::ThreadSafe>> NewFrameTargetResources;
 	TArray<TSharedPtr<FDisplayClusterViewportResource, ESPMode::ThreadSafe>> NewAdditionalFrameTargetableResources;
