[nDisplay] added overridden IsRenderingSuspended() that needed to block Redraw calls while multiplayer session established on the client side, overriden function returns bForcedTickIdleMode flag switched to true when connection established, that ensures alignment on tick and prevents deadlock between GT/RT nDsiplay barriers. Fixed packet processing by passing correct variable, disabled tick for stage geometry component for server builds.

diff --git a/Engine/Plugins/Runtime/nDisplay/Source/DisplayCluster/Private/Game/EngineClasses/Basics/DisplayClusterGameEngine.cpp b/Engine/Plugins/Runtime/nDisplay/Source/DisplayCluster/Private/Game/EngineClasses/Basics/DisplayClusterGameEngine.cpp
--- a/Engine/Plugins/Runtime/nDisplay/Source/DisplayCluster/Private/Game/EngineClasses/Basics/DisplayClusterGameEngine.cpp
+++ b/Engine/Plugins/Runtime/nDisplay/Source/DisplayCluster/Private/Game/EngineClasses/Basics/DisplayClusterGameEngine.cpp
@@ -381,6 +381,11 @@ void UDisplayClusterGameEngine::Tick(float DeltaSeconds, bool bIdleMode)
 		// Perform EndFrame notification
 		GDisplayCluster->EndFrame(GFrameCounter);
 
+		if (bIsRenderingSuspended)
+		{
+			bIsRenderingSuspended = false;
+		}
+
 		UE_LOG(LogDisplayClusterEngine, Verbose, TEXT("Sync frame end"));
 	}
 	else
@@ -529,8 +534,15 @@ EBrowseReturnVal::Type UDisplayClusterGameEngine::BrowseLoadMap(FWorldContext& W
 				
 				if (bIsDisplayClusterNetDriver)
 				{
+					// multiplayer packes, including session handshake are processed on ticks thus we need to enforce engine to tick but preven from rendering until cluster is ready
 					// Force tick idle mode for multiplayer connections
 					bForcedTickIdleMode = true;
+
+					// Suspend rendering until the cluster is ready
+					// By default, engine allowed to render while session being established
+					// to prevent rendering from being invoked we overrided function IsRenderingSuspended() and only enabling this flag here in BrowseLoadMap which is called on in multiplayer
+					// once loading is finished and game tick called in the flag will be reset to false
+					bIsRenderingSuspended = true;
 				}
 			}
 
diff --git a/Engine/Plugins/Runtime/nDisplay/Source/DisplayCluster/Private/Game/EngineClasses/Scene/DisplayClusterStageGeometryComponent.cpp b/Engine/Plugins/Runtime/nDisplay/Source/DisplayCluster/Private/Game/EngineClasses/Scene/DisplayClusterStageGeometryComponent.cpp
--- a/Engine/Plugins/Runtime/nDisplay/Source/DisplayCluster/Private/Game/EngineClasses/Scene/DisplayClusterStageGeometryComponent.cpp
+++ b/Engine/Plugins/Runtime/nDisplay/Source/DisplayCluster/Private/Game/EngineClasses/Scene/DisplayClusterStageGeometryComponent.cpp
@@ -19,6 +19,7 @@ UDisplayClusterStageGeometryComponent::UDisplayClusterStageGeometryComponent()
 {
 	SetActive(true);
 	PrimaryComponentTick.bCanEverTick = true;
+	PrimaryComponentTick.bAllowTickOnDedicatedServer = false;
 
 	Renderer = MakeShared<FDisplayClusterMeshProjectionRenderer>();
 }
diff --git a/Engine/Plugins/Runtime/nDisplay/Source/DisplayCluster/Public/DisplayClusterGameEngine.h b/Engine/Plugins/Runtime/nDisplay/Source/DisplayCluster/Public/DisplayClusterGameEngine.h
--- a/Engine/Plugins/Runtime/nDisplay/Source/DisplayCluster/Public/DisplayClusterGameEngine.h
+++ b/Engine/Plugins/Runtime/nDisplay/Source/DisplayCluster/Public/DisplayClusterGameEngine.h
@@ -43,7 +43,7 @@ public:
 		return OperationMode;
 	}
 
-	/** Resets forced idle mode flag
+	/** Resets forced idle mode flag needed for multiplayer connections to tick the engine side to process handshake
 	 * Used by DisplayClusterNetDriver when all nodes are connected
 	 */
 	void ResetForcedIdleMode()
@@ -51,6 +51,14 @@ public:
 		bForcedTickIdleMode = false;
 	}
 
+	/** This override control either we want to trigger the rendering or not on the frame
+	 * By modifying variable state in BrowseLoadMap when multiplayer connections are being established to prevent Redraw on the client side
+	 */
+	virtual bool IsRenderingSuspended() const override
+	{
+		return bIsRenderingSuspended;
+	}
+
 protected:
 	virtual void UpdateTimeAndHandleMaxTickRate() override;
 
@@ -94,4 +102,7 @@ private:
 	// system have to wait for all nodes to be connected before proceeding with rendering otherwise cluster will deadlock itself between GT/RT barriers
 	// Enforced automatically into true on BrowseLoadMap for multiplayer connections
 	bool bForcedTickIdleMode = false;
+
+	// Flag used to suspend rendering and avoid Redraw method that leads to barrier lock among multiplayer client nodes, used in MP mode only
+	bool bIsRenderingSuspended = false;
 };
diff --git a/Engine/Plugins/Runtime/nDisplay/Source/DisplayClusterReplication/Private/DisplayClusterNetConnection.cpp b/Engine/Plugins/Runtime/nDisplay/Source/DisplayClusterReplication/Private/DisplayClusterNetConnection.cpp
--- a/Engine/Plugins/Runtime/nDisplay/Source/DisplayClusterReplication/Private/DisplayClusterNetConnection.cpp
+++ b/Engine/Plugins/Runtime/nDisplay/Source/DisplayClusterReplication/Private/DisplayClusterNetConnection.cpp
@@ -164,7 +164,7 @@ void UDisplayClusterNetConnection::ProcessPacket(int32 PacketId)
 
 		bool bSkipAck = false;
 		bool bHasBunchErrors = false;
-		DispatchPacket(InPackets[CurrentPacketId], PacketId, bSkipAck, bHasBunchErrors);
+		DispatchPacket(InPackets[CurrentPacketId], CurrentPacketId, bSkipAck, bHasBunchErrors);
 
 		InPackets.Remove(CurrentPacketId);
 	}
