nDisplay: Fixed issue where ICVFX camera would not use external cine camera when computing distance to stage geometry

diff --git a/Engine/Plugins/Runtime/nDisplay/Source/DisplayCluster/Private/Game/EngineClasses/Scene/DisplayClusterICVFXCameraComponent.cpp b/Engine/Plugins/Runtime/nDisplay/Source/DisplayCluster/Private/Game/EngineClasses/Scene/DisplayClusterICVFXCameraComponent.cpp
--- a/Engine/Plugins/Runtime/nDisplay/Source/DisplayCluster/Private/Game/EngineClasses/Scene/DisplayClusterICVFXCameraComponent.cpp
+++ b/Engine/Plugins/Runtime/nDisplay/Source/DisplayCluster/Private/Game/EngineClasses/Scene/DisplayClusterICVFXCameraComponent.cpp
@@ -133,8 +133,19 @@ void UDisplayClusterICVFXCameraComponent::TickComponent(float DeltaTime, ELevelT
 
 		if (ADisplayClusterRootActor* RootActor = Cast<ADisplayClusterRootActor>(GetOwner()))
 		{
-			const FVector CameraLocation = GetComponentLocation();
-			const FVector CameraDirection = GetComponentRotation().RotateVector(FVector::XAxisVector);
+			FVector CameraLocation = FVector::ZeroVector;
+			FVector CameraDirection = FVector::XAxisVector;
+			if (CameraSettings.ExternalCameraActor.IsValid())
+			{
+				CameraLocation = CameraSettings.ExternalCameraActor->GetActorLocation();
+				CameraDirection = CameraSettings.ExternalCameraActor->GetActorRotation().RotateVector(FVector::XAxisVector);
+			}
+			else
+			{
+				CameraLocation = GetComponentLocation();
+				CameraDirection = GetComponentRotation().RotateVector(FVector::XAxisVector);
+			}
+			
 			float DistanceToWall = 0.0;
 
 			// For now, do a single trace from the center of the camera to the stage geometry.
