nDisplay: Fixed issue where ICVFX camera color grading properties would be hidden when an external cine camera was set.

diff --git a/Engine/Plugins/Runtime/nDisplay/Source/DisplayCluster/Public/Components/DisplayClusterICVFXCameraComponent.h b/Engine/Plugins/Runtime/nDisplay/Source/DisplayCluster/Public/Components/DisplayClusterICVFXCameraComponent.h
--- a/Engine/Plugins/Runtime/nDisplay/Source/DisplayCluster/Public/Components/DisplayClusterICVFXCameraComponent.h
+++ b/Engine/Plugins/Runtime/nDisplay/Source/DisplayCluster/Public/Components/DisplayClusterICVFXCameraComponent.h
@@ -175,13 +175,13 @@ private:
 	UPROPERTY(EditAnywhere, Transient, Category = OCIO, meta = (PropertyPath = "CameraSettings.CameraOCIO.PerNodeOCIOProfiles", DisplayName = "Per-Node OCIO Overrides"))
 	FDisplayClusterEditorPropertyReference PerNodeOCIOProfilesRef;
 
-	UPROPERTY(EditAnywhere, Transient, Category = "Color Grading", meta = (PropertyPath = "CameraSettings.EnableInnerFrustumColorGrading", DisplayName = "Enable Inner Frustum Color Grading"))
+	UPROPERTY(EditAnywhere, Transient, Category = "Inner Frustum Color Grading", meta = (PropertyPath = "CameraSettings.EnableInnerFrustumColorGrading", DisplayName = "Enable Inner Frustum Color Grading"))
 	FDisplayClusterEditorPropertyReference EnableInnerFrustumColorGrading;
 
-	UPROPERTY(EditAnywhere, Transient, Category = "Color Grading", meta = (PropertyPath = "CameraSettings.AllNodesColorGrading", DisplayName = "All Nodes"))
+	UPROPERTY(EditAnywhere, Transient, Category = "Inner Frustum Color Grading", meta = (PropertyPath = "CameraSettings.AllNodesColorGrading", DisplayName = "All Nodes"))
 	FDisplayClusterEditorPropertyReference AllNodesColorGradingRef;
 
-	UPROPERTY(EditAnywhere, Transient, Category = "Color Grading", meta = (PropertyPath = "CameraSettings.PerNodeColorGrading"))
+	UPROPERTY(EditAnywhere, Transient, Category = "Inner Frustum Color Grading", meta = (PropertyPath = "CameraSettings.PerNodeColorGrading"))
 	FDisplayClusterEditorPropertyReference PerNodeColorGradingRef;
 
 	UPROPERTY(EditAnywhere, Transient, Category = "Texture Replacement", meta = (PropertyPath = "CameraSettings.RenderSettings.Replace.bAllowReplace", DisplayName = "Enable Inner Frustum Texture Replacement", ToolTip = "Set to True to replace the entire inner frustum with the specified texture."))
diff --git a/Engine/Plugins/Runtime/nDisplay/Source/DisplayClusterConfiguration/Public/DisplayClusterConfigurationStrings.h b/Engine/Plugins/Runtime/nDisplay/Source/DisplayClusterConfiguration/Public/DisplayClusterConfigurationStrings.h
--- a/Engine/Plugins/Runtime/nDisplay/Source/DisplayClusterConfiguration/Public/DisplayClusterConfigurationStrings.h
+++ b/Engine/Plugins/Runtime/nDisplay/Source/DisplayClusterConfiguration/Public/DisplayClusterConfigurationStrings.h
@@ -25,7 +25,7 @@ namespace DisplayClusterConfigurationStrings
 		static constexpr const TCHAR* ClusterConfigurationCategory = TEXT("NDisplay Cluster Configuration");
 		static constexpr const TCHAR* ClusterPostprocessCategory   = TEXT("Post Process");
 		static constexpr const TCHAR* ColorGradingCategory         = TEXT("Color Grading");
-		static constexpr const TCHAR* CameraColorGradingCategory   = TEXT("Color Grading");
+		static constexpr const TCHAR* CameraColorGradingCategory   = TEXT("Inner Frustum Color Grading");
 		static constexpr const TCHAR* ChromaKeyCategory            = TEXT("Chromakey");
 		static constexpr const TCHAR* LightcardCategory            = TEXT("Light Cards");
 		static constexpr const TCHAR* OCIOCategory                 = TEXT("OCIO");
diff --git a/Engine/Plugins/Runtime/nDisplay/Source/DisplayClusterConfigurator/Private/Views/Details/Components/DisplayClusterICVFXCameraComponentDetailsCustomization.cpp b/Engine/Plugins/Runtime/nDisplay/Source/DisplayClusterConfigurator/Private/Views/Details/Components/DisplayClusterICVFXCameraComponentDetailsCustomization.cpp
--- a/Engine/Plugins/Runtime/nDisplay/Source/DisplayClusterConfigurator/Private/Views/Details/Components/DisplayClusterICVFXCameraComponentDetailsCustomization.cpp
+++ b/Engine/Plugins/Runtime/nDisplay/Source/DisplayClusterConfigurator/Private/Views/Details/Components/DisplayClusterICVFXCameraComponentDetailsCustomization.cpp
@@ -97,6 +97,9 @@ void FDisplayClusterICVFXCameraComponentDetailsCustomization::CustomizeDetails(I
 	// Sockets category must be hidden manually instead of through the HideCategories metadata specifier
 	InLayoutBuilder.HideCategory(TEXT("Sockets"));
 
+	// Rename "Inner Frustum Color Grading" to "Color Grading" for brevity, as the category itself needs to remain distinct from the camera's "Color Grading" category.
+	InLayoutBuilder.EditCategory(DisplayClusterConfigurationStrings::categories::CameraColorGradingCategory, LOCTEXT("ICVFXColorGradingCategoryLabel", "Color Grading"));
+
 	IDetailCategoryBuilder& CameraCategory = InLayoutBuilder.EditCategory(DisplayClusterConfigurationStrings::categories::ICVFXCameraCategory, LOCTEXT("ICVFXCameraCategoryLabel", "Camera"));
 
 	// Re-add the external camera to the category to ensure it is always above the camera's fiz properties in the details panel
