nDisplay: added circle of confusion compensation for uniform and directional overscan.

diff --git a/Engine/Plugins/Runtime/nDisplay/Source/DisplayClusterConfiguration/Private/DisplayClusterConfigurationTypes_ICVFX.cpp b/Engine/Plugins/Runtime/nDisplay/Source/DisplayClusterConfiguration/Private/DisplayClusterConfigurationTypes_ICVFX.cpp
--- a/Engine/Plugins/Runtime/nDisplay/Source/DisplayClusterConfiguration/Private/DisplayClusterConfigurationTypes_ICVFX.cpp
+++ b/Engine/Plugins/Runtime/nDisplay/Source/DisplayClusterConfiguration/Private/DisplayClusterConfigurationTypes_ICVFX.cpp
@@ -14,6 +14,11 @@ namespace UE::DisplayClusterConfiguration::ICVFX
 
 		return FMath::Clamp(InValue, -MaxCustomFrustumValue, MaxCustomFrustumValue);
 	}
+
+	static float ClampCustomFrustum(float InValue, float InMax)
+	{
+		return FMath::Clamp(InValue, -InMax, InMax);
+	}
 };
 using namespace UE::DisplayClusterConfiguration::ICVFX;
 
@@ -272,7 +277,7 @@ bool FDisplayClusterConfigurationICVFX_CameraSettings::GetCameraBorder(const FDi
 void FDisplayClusterConfigurationICVFX_CameraSettings::SetupViewInfo(const FDisplayClusterConfigurationICVFX_StageSettings& InStageSettings, FMinimalViewInfo& InOutViewInfo)
 {
 	RenderSettings.SetupViewInfo(InStageSettings, InOutViewInfo);
-	CustomFrustum.SetupViewInfo(InStageSettings, InOutViewInfo);
+	CustomFrustum.SetupViewInfo(InStageSettings, *this, InOutViewInfo);
 	CameraMotionBlur.SetupViewInfo(InStageSettings, InOutViewInfo);
 }
 
@@ -618,12 +623,53 @@ void FDisplayClusterConfigurationICVFX_CameraDepthOfField::UpdateDynamicCompensa
 	DynamicCompensationLUT = nullptr;
 }
 
-void FDisplayClusterConfigurationICVFX_CameraCustomFrustum::SetupViewInfo(const FDisplayClusterConfigurationICVFX_StageSettings& InStageSettings, FMinimalViewInfo& InOutViewInfo) const
+void FDisplayClusterConfigurationICVFX_CameraCustomFrustum::SetupViewInfo(const FDisplayClusterConfigurationICVFX_StageSettings& InStageSettings, const FDisplayClusterConfigurationICVFX_CameraSettings& InCameraSettings, FMinimalViewInfo& InOutViewInfo) const
 {
 	// Since Circle of confusion is directly proportional to aperature, with wider FOV focal length needs to be shortened by the same amount as FOV.
-	const float FOVMultiplier = GetCameraFieldOfViewMultiplier(InStageSettings);
-	const float ClampedFieldOfViewMultiplier = (FOVMultiplier > 0.f) ? FOVMultiplier : 1.f;
-	InOutViewInfo.PostProcessSettings.DepthOfFieldFstop /= ClampedFieldOfViewMultiplier;
+	if (bEnable && InCameraSettings.ExternalCameraActor.IsValid())
+	{
+		// default - percents
+		const float ConvertToPercent = 0.01f;
+		const float MaxPercentOverscan = .5f;
+		const float MaxPixelOverscan = 5.f;
+
+		float LocLeft = Left;
+		float LocRight = Right;
+		float LocTop = Top;
+		float LocBottom = Bottom;
+
+		if (Mode == EDisplayClusterConfigurationViewportCustomFrustumMode::Pixels)
+		{
+			LocLeft = ClampCustomFrustum(Left * ConvertToPercent, MaxPixelOverscan);
+			LocRight = ClampCustomFrustum(Right * ConvertToPercent, MaxPixelOverscan);
+			LocTop = ClampCustomFrustum(Top * ConvertToPercent, MaxPixelOverscan);
+			LocBottom = ClampCustomFrustum(Bottom * ConvertToPercent, MaxPixelOverscan);
+
+			const float CameraBufferRatio = InCameraSettings.GetCameraBufferRatio(InStageSettings);
+			const FIntPoint FrameSize = InCameraSettings.GetCameraFrameSize(InStageSettings, *InCameraSettings.ExternalCameraActor->GetCineCameraComponent());
+			const float  FrameWidth = FrameSize.X * CameraBufferRatio;
+			const float FrameHeight = FrameSize.Y * CameraBufferRatio;
+
+			LocLeft = LocLeft / FrameWidth;
+			LocRight = LocRight / FrameWidth;
+			LocTop = LocTop / FrameHeight;
+			LocBottom = LocBottom / FrameHeight;
+		}
+
+		LocLeft = ClampCustomFrustum(LocLeft * ConvertToPercent, MaxPercentOverscan);
+		LocRight = ClampCustomFrustum(LocRight * ConvertToPercent, MaxPercentOverscan);
+		LocTop = ClampCustomFrustum(LocTop * ConvertToPercent, MaxPercentOverscan);
+		LocBottom = ClampCustomFrustum(LocBottom * ConvertToPercent, MaxPercentOverscan);
+
+		const float FOVMultiplier = GetCameraFieldOfViewMultiplier(InStageSettings);
+		const float ClampedFieldOfViewMultiplier = (FOVMultiplier > 0.f) ? FOVMultiplier : 1.f;
+		InOutViewInfo.PostProcessSettings.DepthOfFieldMinFstop /= ClampedFieldOfViewMultiplier;
+		InOutViewInfo.PostProcessSettings.DepthOfFieldFstop /= ClampedFieldOfViewMultiplier;
+
+		float Multiplier = (1. + (LocLeft + LocRight) / 2.) * (1. + (LocTop + LocBottom) / 2.);
+		InOutViewInfo.PostProcessSettings.DepthOfFieldMinFstop /= Multiplier;
+		InOutViewInfo.PostProcessSettings.DepthOfFieldFstop /= Multiplier;
+	}
 }
 
 float FDisplayClusterConfigurationICVFX_CameraCustomFrustum::GetCameraFieldOfViewMultiplier(const FDisplayClusterConfigurationICVFX_StageSettings& InStageSettings) const
diff --git a/Engine/Plugins/Runtime/nDisplay/Source/DisplayClusterConfiguration/Public/DisplayClusterConfigurationTypes_ICVFX.h b/Engine/Plugins/Runtime/nDisplay/Source/DisplayClusterConfiguration/Public/DisplayClusterConfigurationTypes_ICVFX.h
--- a/Engine/Plugins/Runtime/nDisplay/Source/DisplayClusterConfiguration/Public/DisplayClusterConfigurationTypes_ICVFX.h
+++ b/Engine/Plugins/Runtime/nDisplay/Source/DisplayClusterConfiguration/Public/DisplayClusterConfigurationTypes_ICVFX.h
@@ -691,7 +691,7 @@ struct DISPLAYCLUSTERCONFIGURATION_API FDisplayClusterConfigurationICVFX_CameraC
 
 public:
 	/** Propagates Custom Frustum related settings to the view info. */
-	void SetupViewInfo(const FDisplayClusterConfigurationICVFX_StageSettings& InStageSettings, FMinimalViewInfo& InOutViewInfo) const;
+	void SetupViewInfo(const FDisplayClusterConfigurationICVFX_StageSettings& InStageSettings, const FDisplayClusterConfigurationICVFX_CameraSettings& InCameraSettings, FMinimalViewInfo& InOutViewInfo) const;
 
 	/** Get camera FOV multiplier. */
 	float GetCameraFieldOfViewMultiplier(const FDisplayClusterConfigurationICVFX_StageSettings& InStageSettings) const;
