nDisplay: Fix Distance To Wall when camera is outside the stage geometry.

The distance from the camera to the wall is determined using a line trace to the stage geometry (which is always a gapless morphed isosphere). This works ok when the camera is inside the stage geometry, but when the camera is outside, it was only measuring the distance to the outside of the stage geo, which is not correct.

This was happening because procedural meshes default to being double sided, so to fix this, we force it to match that of the isosphere static mesh that it is morphing.

diff --git a/Engine/Plugins/Runtime/nDisplay/Source/DisplayCluster/Private/Game/EngineClasses/Scene/DisplayClusterStageIsosphereComponent.cpp b/Engine/Plugins/Runtime/nDisplay/Source/DisplayCluster/Private/Game/EngineClasses/Scene/DisplayClusterStageIsosphereComponent.cpp
--- a/Engine/Plugins/Runtime/nDisplay/Source/DisplayCluster/Private/Game/EngineClasses/Scene/DisplayClusterStageIsosphereComponent.cpp
+++ b/Engine/Plugins/Runtime/nDisplay/Source/DisplayCluster/Private/Game/EngineClasses/Scene/DisplayClusterStageIsosphereComponent.cpp
@@ -4,6 +4,7 @@
 
 #include "Engine/StaticMesh.h"
 #include "KismetProceduralMeshLibrary.h"
+#include "PhysicsEngine/BodySetup.h"
 #include "UObject/ConstructorHelpers.h"
 
 UDisplayClusterStageIsosphereComponent::UDisplayClusterStageIsosphereComponent(const FObjectInitializer& ObjectInitializer)
@@ -77,6 +78,13 @@ void UDisplayClusterStageIsosphereComponent::ResetIsosphere()
 				SetProcMeshSection(0, *Section);
 			}
 		}
+
+		// Match the sided-ness setting of the IsosphereMesh.
+		const UBodySetup* IsosphereMeshBodySetup = IsosphereMesh->GetBodySetup();
+		if (ProcMeshBodySetup && IsosphereMeshBodySetup)
+		{
+			ProcMeshBodySetup->bDoubleSidedGeometry = IsosphereMeshBodySetup->bDoubleSidedGeometry;
+		}
 	}
 }
 
