[nDisplay] Default parameters for ICVFX tiled media

diff --git a/Engine/Plugins/Runtime/nDisplay/Source/DisplayClusterConfiguration/Public/DisplayClusterConfigurationTypes_Media.h b/Engine/Plugins/Runtime/nDisplay/Source/DisplayClusterConfiguration/Public/DisplayClusterConfigurationTypes_Media.h
--- a/Engine/Plugins/Runtime/nDisplay/Source/DisplayClusterConfiguration/Public/DisplayClusterConfigurationTypes_Media.h
+++ b/Engine/Plugins/Runtime/nDisplay/Source/DisplayClusterConfiguration/Public/DisplayClusterConfigurationTypes_Media.h
@@ -223,7 +223,7 @@ public:
 	FDisplayClusterConfigurationClusterItemReferenceList ClusterNodes;
 
 	/** Tile mapping. Maps tiles to media sources. */
-	UPROPERTY(EditAnywhere, BlueprintReadWrite, Category = "Media", meta = (ClusterItemType = ClusterNodes))
+	UPROPERTY(EditAnywhere, BlueprintReadWrite, Category = "Media")
 	TArray<FDisplayClusterConfigurationMediaUniformTileInput> Tiles;
 };
 
@@ -242,7 +242,7 @@ public:
 	FDisplayClusterConfigurationClusterItemReferenceList ClusterNodes;
 
 	/** Tile mapping. Maps tiles to media outputs. */
-	UPROPERTY(EditAnywhere, BlueprintReadWrite, Category = "Media", meta = (ClusterItemType = ClusterNodes))
+	UPROPERTY(EditAnywhere, BlueprintReadWrite, Category = "Media")
 	TArray<FDisplayClusterConfigurationMediaUniformTileOutput> Tiles;
 
 	/** Media output synchronization policy */
diff --git a/Engine/Plugins/Runtime/nDisplay/Source/DisplayClusterConfigurator/DisplayClusterConfigurator.Build.cs b/Engine/Plugins/Runtime/nDisplay/Source/DisplayClusterConfigurator/DisplayClusterConfigurator.Build.cs
--- a/Engine/Plugins/Runtime/nDisplay/Source/DisplayClusterConfigurator/DisplayClusterConfigurator.Build.cs
+++ b/Engine/Plugins/Runtime/nDisplay/Source/DisplayClusterConfigurator/DisplayClusterConfigurator.Build.cs
@@ -12,6 +12,7 @@ public class DisplayClusterConfigurator : ModuleRules
 			{
 				"DisplayCluster",
 				"DisplayClusterConfiguration",
+				"DisplayClusterModularFeaturesEditor",
 				"DisplayClusterProjection",
 				"DisplayClusterWarp",
 
diff --git a/Engine/Plugins/Runtime/nDisplay/Source/DisplayClusterConfigurator/Private/DisplayClusterConfiguratorModule.cpp b/Engine/Plugins/Runtime/nDisplay/Source/DisplayClusterConfigurator/Private/DisplayClusterConfiguratorModule.cpp
--- a/Engine/Plugins/Runtime/nDisplay/Source/DisplayClusterConfigurator/Private/DisplayClusterConfiguratorModule.cpp
+++ b/Engine/Plugins/Runtime/nDisplay/Source/DisplayClusterConfigurator/Private/DisplayClusterConfiguratorModule.cpp
@@ -27,6 +27,7 @@
 #include "Views/Details/Components/DisplayClusterConfiguratorScreenComponentDetailsCustomization.h"
 #include "Views/Details/Components/DisplayClusterICVFXCameraComponentDetailsCustomization.h"
 #include "Views/Details/Media/DisplayClusterConfiguratorICVFXMediaCustomization.h"
+#include "Views/Details/Media/DisplayClusterConfiguratorICVFXMediaTileCustomization.h"
 #include "Views/Details/Policies/DisplayClusterConfiguratorPolicyDetailCustomization.h"
 
 #include "Blueprints/DisplayClusterBlueprint.h"
@@ -225,6 +226,8 @@ void FDisplayClusterConfiguratorModule::RegisterCustomLayouts()
 	REGISTER_PROPERTY_LAYOUT(FDisplayClusterConfigurationViewport_RemapData, FDisplayClusterConfiguratorViewportRemapCustomization);
 	REGISTER_PROPERTY_LAYOUT(FDisplayClusterConfigurationRectangle, FDisplayClusterConfiguratorRectangleCustomization);
 	REGISTER_PROPERTY_LAYOUT(FDisplayClusterConfigurationMediaICVFX, FDisplayClusterConfiguratorICVFXMediaCustomization);
+	REGISTER_PROPERTY_LAYOUT(FDisplayClusterConfigurationMediaUniformTileInput, FDisplayClusterConfiguratorICVFXMediaInputTileCustomization);
+	REGISTER_PROPERTY_LAYOUT(FDisplayClusterConfigurationMediaUniformTileOutput, FDisplayClusterConfiguratorICVFXMediaOutputTileCustomization);
 }
 
 void FDisplayClusterConfiguratorModule::UnregisterCustomLayouts()
diff --git a/Engine/Plugins/Runtime/nDisplay/Source/DisplayClusterConfigurator/Private/Views/Details/Media/DisplayClusterConfiguratorICVFXMediaCustomization.cpp b/Engine/Plugins/Runtime/nDisplay/Source/DisplayClusterConfigurator/Private/Views/Details/Media/DisplayClusterConfiguratorICVFXMediaCustomization.cpp
--- a/Engine/Plugins/Runtime/nDisplay/Source/DisplayClusterConfigurator/Private/Views/Details/Media/DisplayClusterConfiguratorICVFXMediaCustomization.cpp
+++ b/Engine/Plugins/Runtime/nDisplay/Source/DisplayClusterConfigurator/Private/Views/Details/Media/DisplayClusterConfiguratorICVFXMediaCustomization.cpp
@@ -1,13 +1,18 @@
 // Copyright Epic Games, Inc. All Rights Reserved.
 
 #include "Views/Details/Media/DisplayClusterConfiguratorICVFXMediaCustomization.h"
+#include "Views/Details/Media/DisplayClusterConfiguratorMediaUtils.h"
 
 #include "DisplayClusterConfigurationTypes_Media.h"
 
+#include "DetailWidgetRow.h"
 #include "IDetailChildrenBuilder.h"
 #include "IPropertyUtilities.h"
 #include "PropertyHandle.h"
 
+#include "Widgets/Input/SButton.h"
+#include "Widgets/Layout/SBox.h"
+#include "Widgets/Text/STextBlock.h"
 
 #define LOCTEXT_NAMESPACE "FDisplayClusterConfiguratorICVFXMediaCustomization"
 
@@ -15,7 +20,11 @@ void FDisplayClusterConfiguratorICVFXMediaCustomization::CustomizeChildren(TShar
 {
 	// SplitType property
 	TSharedPtr<IPropertyHandle> SplitTypeHandle = GET_CHILD_HANDLE(FDisplayClusterConfigurationMediaICVFX, SplitType);
-	check(SplitTypeHandle);
+	check(SplitTypeHandle->IsValidHandle());
+
+	// Layout property
+	TilesLayoutHandle = GET_CHILD_HANDLE(FDisplayClusterConfigurationMediaICVFX, TiledSplitLayout);
+	check(TilesLayoutHandle->IsValidHandle());
 
 	// Separate groups specific for every split type available
 	TArray<TSharedPtr<IPropertyHandle>> FullFramePropertyHandles;
@@ -74,7 +83,42 @@ void FDisplayClusterConfiguratorICVFXMediaCustomization::CustomizeChildren(TShar
 		}
 	}
 
+	// Create all the property widgets
 	FDisplayClusterConfiguratorBaseTypeCustomization::CustomizeChildren(InPropertyHandle, InChildBuilder, InCustomizationUtils);
+
+	// Create auto-configure button in the bottom
+	if (SplitTypeValue == EDisplayClusterConfigurationMediaSplitType::UniformTiles)
+	{
+		AddAutoConfigurationButton(InChildBuilder);
+	}
+}
+
+void FDisplayClusterConfiguratorICVFXMediaCustomization::AddAutoConfigurationButton(IDetailChildrenBuilder& InChildBuilder)
+{
+	InChildBuilder.AddCustomRow(FText::GetEmpty())
+		.WholeRowContent()
+		[
+			SNew(SBox)
+				.Padding(5.f)
+				[
+					SNew(SButton)
+						.HAlign(HAlign_Center)
+						.OnClicked(this, &FDisplayClusterConfiguratorICVFXMediaCustomization::OnAutoConfigureButtonClicked)
+						[
+							SNew(STextBlock)
+								.Font(FAppStyle::GetFontStyle(TEXT("PropertyWindow.NormalFont")))
+								.Text(LOCTEXT("AutoConfigureTiledInputButtonTitle", "Auto-Configure"))
+						]
+				]
+		];
+}
+
+FReply FDisplayClusterConfiguratorICVFXMediaCustomization::OnAutoConfigureButtonClicked()
+{
+	// Notify tile customizators to re-initialize their media subjects
+	FDisplayClusterConfiguratorMediaUtils::Get().OnTiledMediaAutoConfiguration().Broadcast(EditingObject.Get());
+
+	return FReply::Handled();
 }
 
 #undef LOCTEXT_NAMESPACE
diff --git a/Engine/Plugins/Runtime/nDisplay/Source/DisplayClusterConfigurator/Private/Views/Details/Media/DisplayClusterConfiguratorICVFXMediaCustomization.h b/Engine/Plugins/Runtime/nDisplay/Source/DisplayClusterConfigurator/Private/Views/Details/Media/DisplayClusterConfiguratorICVFXMediaCustomization.h
--- a/Engine/Plugins/Runtime/nDisplay/Source/DisplayClusterConfigurator/Private/Views/Details/Media/DisplayClusterConfiguratorICVFXMediaCustomization.h
+++ b/Engine/Plugins/Runtime/nDisplay/Source/DisplayClusterConfigurator/Private/Views/Details/Media/DisplayClusterConfiguratorICVFXMediaCustomization.h
@@ -5,6 +5,8 @@
 #include "CoreMinimal.h"
 #include "Views/Details/DisplayClusterConfiguratorBaseTypeCustomization.h"
 
+#include "Input/Reply.h"
+
 #include "IPropertyTypeCustomization.h"
 
 
@@ -24,4 +26,17 @@ protected:
 	//~ Begin IPropertyTypeCustomization
 	virtual void CustomizeChildren(TSharedRef<IPropertyHandle> PropertyHandle, IDetailChildrenBuilder& ChildBuilder, IPropertyTypeCustomizationUtils& CustomizationUtils) override;
 	//~ End IPropertyTypeCustomization
+
+private:
+
+	/** Builds auto-conf button widget. */
+	void AddAutoConfigurationButton(IDetailChildrenBuilder& ChildBuilder);
+
+	/** Handles auto-configuration button clicks. */
+	FReply OnAutoConfigureButtonClicked();
+
+private:
+
+	/** Tiles layout property. */
+	TSharedPtr<IPropertyHandle> TilesLayoutHandle;
 };
diff --git a/Engine/Plugins/Runtime/nDisplay/Source/DisplayClusterConfigurator/Private/Views/Details/Media/DisplayClusterConfiguratorICVFXMediaTileCustomization.cpp b/Engine/Plugins/Runtime/nDisplay/Source/DisplayClusterConfigurator/Private/Views/Details/Media/DisplayClusterConfiguratorICVFXMediaTileCustomization.cpp
new file mode 100644
--- /dev/null
+++ b/Engine/Plugins/Runtime/nDisplay/Source/DisplayClusterConfigurator/Private/Views/Details/Media/DisplayClusterConfiguratorICVFXMediaTileCustomization.cpp
@@ -0,0 +1,266 @@
+// Copyright Epic Games, Inc. All Rights Reserved.
+
+#include "Views/Details/Media/DisplayClusterConfiguratorICVFXMediaTileCustomization.h"
+#include "Views/Details/Media/DisplayClusterConfiguratorMediaUtils.h"
+
+#include "Components/DisplayClusterICVFXCameraComponent.h"
+#include "DisplayClusterConfigurationTypes_Media.h"
+#include "DisplayClusterConfiguratorBlueprintEditor.h"
+#include "DisplayClusterConfiguratorUtils.h"
+#include "DisplayClusterRootActor.h"
+
+#include "IDisplayClusterModularFeatureMediaInitializer.h"
+
+#include "IDetailChildrenBuilder.h"
+#include "IPropertyUtilities.h"
+#include "PropertyHandle.h"
+
+#include "Features/IModularFeatures.h"
+
+#define LOCTEXT_NAMESPACE "FDisplayClusterConfiguratorICVFXMediaCustomization"
+
+
+// Enable/disable live auto-configuration. Does not affect "auto-configure" button.
+TAutoConsoleVariable<bool> CVarTiledMediaAutoInitializationEnabled(
+	TEXT("nDisplay.Media.TiledMediaAutoInitialization"),
+	true,
+	TEXT("nDisplay tiled media auto-configuration\n")
+	TEXT("0 : Disabled\n")
+	TEXT("1 : Enabled\n"),
+	ECVF_Default
+);
+
+
+//
+// Input tiles
+//
+void FDisplayClusterConfiguratorICVFXMediaInputTileCustomization::CustomizeChildren(TSharedRef<IPropertyHandle> InPropertyHandle, IDetailChildrenBuilder& InChildBuilder, IPropertyTypeCustomizationUtils& InCustomizationUtils)
+{
+	// MediaSource property
+	MediaSubjectHandle = GET_CHILD_HANDLE(FDisplayClusterConfigurationMediaUniformTileInput, MediaSource);
+	check(MediaSubjectHandle->IsValidHandle());
+
+	// Tile position property
+	TilePosHandle = GET_CHILD_HANDLE(FDisplayClusterConfigurationMediaUniformTileInput, Position);
+	check(TilePosHandle->IsValidHandle());
+
+	FDisplayClusterConfiguratorICVFXMediaTileCustomizationBase::CustomizeChildren(InPropertyHandle, InChildBuilder, InCustomizationUtils);
+}
+
+//
+// Output tiles
+//
+void FDisplayClusterConfiguratorICVFXMediaOutputTileCustomization::CustomizeChildren(TSharedRef<IPropertyHandle> InPropertyHandle, IDetailChildrenBuilder& InChildBuilder, IPropertyTypeCustomizationUtils& InCustomizationUtils)
+{
+	// MediaOutput property
+	MediaSubjectHandle = GET_CHILD_HANDLE(FDisplayClusterConfigurationMediaUniformTileOutput, MediaOutput);
+	check(MediaSubjectHandle->IsValidHandle());
+
+	// Tile position property
+	TilePosHandle = GET_CHILD_HANDLE(FDisplayClusterConfigurationMediaUniformTileOutput, Position);
+	check(TilePosHandle->IsValidHandle());
+
+	FDisplayClusterConfiguratorICVFXMediaTileCustomizationBase::CustomizeChildren(InPropertyHandle, InChildBuilder, InCustomizationUtils);
+}
+
+//
+// Common customization
+//
+
+// Static members
+bool FDisplayClusterConfiguratorICVFXMediaTileCustomizationBase::bMediaInitializersObtained = false;
+TArray<IDisplayClusterModularFeatureMediaInitializer*> FDisplayClusterConfiguratorICVFXMediaTileCustomizationBase::MediaInitializers;
+
+FDisplayClusterConfiguratorICVFXMediaTileCustomizationBase::FDisplayClusterConfiguratorICVFXMediaTileCustomizationBase()
+{
+	// Obtain media initializer modular features once. We don't expect any of them to be loaded dynamically in runtime.
+	if (!bMediaInitializersObtained)
+	{
+		IModularFeatures& ModularFeatures = IModularFeatures::Get();
+		ModularFeatures.LockModularFeatureList();
+		MediaInitializers = ModularFeatures.GetModularFeatureImplementations<IDisplayClusterModularFeatureMediaInitializer>(IDisplayClusterModularFeatureMediaInitializer::ModularFeatureName);
+		ModularFeatures.UnlockModularFeatureList();
+
+		bMediaInitializersObtained = true;
+	}
+
+	// Subscribe for auto-configure event
+	FDisplayClusterConfiguratorMediaUtils::Get().OnTiledMediaAutoConfiguration().AddRaw(this, &FDisplayClusterConfiguratorICVFXMediaInputTileCustomization::OnAutoConfigureRequested);
+}
+
+FDisplayClusterConfiguratorICVFXMediaTileCustomizationBase::~FDisplayClusterConfiguratorICVFXMediaTileCustomizationBase()
+{
+	// Unsubscribe from auto-configure event
+	FDisplayClusterConfiguratorMediaUtils::Get().OnTiledMediaAutoConfiguration().RemoveAll(this);
+}
+
+void FDisplayClusterConfiguratorICVFXMediaTileCustomizationBase::CustomizeChildren(TSharedRef<IPropertyHandle> InPropertyHandle, IDetailChildrenBuilder& InChildBuilder, IPropertyTypeCustomizationUtils& InCustomizationUtils)
+{
+	if (CVarTiledMediaAutoInitializationEnabled.GetValueOnGameThread())
+	{
+		// Subscribe for change callbacks
+		if (MediaSubjectHandle->IsValidHandle() && TilePosHandle->IsValidHandle())
+		{
+			MediaSubjectHandle->SetOnPropertyValueChanged(FSimpleDelegate::CreateRaw(this, &FDisplayClusterConfiguratorICVFXMediaTileCustomizationBase::OnMediaSubjectChanged));
+
+			TilePosHandle->SetOnPropertyValueChanged(FSimpleDelegate::CreateRaw(this, &FDisplayClusterConfiguratorICVFXMediaTileCustomizationBase::OnTilePositionChanged));
+			TilePosHandle->SetOnChildPropertyValueChanged(FSimpleDelegate::CreateRaw(this, &FDisplayClusterConfiguratorICVFXMediaTileCustomizationBase::OnTilePositionChanged));
+		}
+	}
+
+	FDisplayClusterConfiguratorBaseTypeCustomization::CustomizeChildren(InPropertyHandle, InChildBuilder, InCustomizationUtils);
+}
+
+void FDisplayClusterConfiguratorICVFXMediaTileCustomizationBase::ModifyMediaSubjectParameters()
+{
+	if (!EditingObject.IsValid())
+	{
+		return;
+	}
+
+	UDisplayClusterICVFXCameraComponent* ICVFXCameraComponent = Cast<UDisplayClusterICVFXCameraComponent>(EditingObject.Get());
+	if (!MediaSubjectHandle->IsValidHandle() || !ICVFXCameraComponent)
+	{
+		return;
+	}
+
+	UObject* NewMediaSubject = nullptr;
+	if (MediaSubjectHandle->GetValue(NewMediaSubject) != FPropertyAccess::Success)
+	{
+		return;
+	}
+
+	if (!NewMediaSubject)
+	{
+		return;
+	}
+
+	// Get input tile position
+	const FIntPoint TilePos = GetEditedTilePos();
+	if (TilePos.X == FIntPoint::NoneValue.X || TilePos.Y == FIntPoint::NoneValue.Y)
+	{
+		return;
+	}
+
+	// Find a suitable initializer, and let it process new media data
+	for (IDisplayClusterModularFeatureMediaInitializer* Initializer : MediaInitializers)
+	{
+		if (Initializer && Initializer->IsMediaSubjectSupported(NewMediaSubject))
+		{
+			const FString OwnerName = GetOwnerName(ICVFXCameraComponent);
+			const uint8 UniqueOwnerIndex = GenerateUniqueOwnerIndex(OwnerName);
+
+			if (UniqueOwnerIndex < TNumericLimits<uint8>::Max())
+			{
+				Initializer->InitializeMediaSubjectForTile(NewMediaSubject, OwnerName, UniqueOwnerIndex, TilePos);
+			}
+
+			break;
+		}
+	}
+}
+
+AActor* FDisplayClusterConfiguratorICVFXMediaTileCustomizationBase::GetOwningActor() const
+{
+	if (UDisplayClusterICVFXCameraComponent* ICVFXCameraComponent = Cast<UDisplayClusterICVFXCameraComponent>(EditingObject))
+	{
+		if (AActor* Actor1 = ICVFXCameraComponent->GetOwner())
+		{
+			return Actor1;
+		}
+		else if (AActor* Actor2 = Cast<AActor>(FindRootActor()))
+		{
+			return Actor2;
+		}
+		else if(FDisplayClusterConfiguratorBlueprintEditor* BlueprintEditor = FDisplayClusterConfiguratorUtils::GetBlueprintEditorForObject(ICVFXCameraComponent))
+		{
+			return BlueprintEditor->GetPreviewActor();
+		}
+	}
+
+	return nullptr;
+}
+
+FString FDisplayClusterConfiguratorICVFXMediaTileCustomizationBase::GetOwnerName(const UDisplayClusterICVFXCameraComponent* ICVFXCameraComponent) const
+{
+	// In blueprint editor, the components are transient with the following suffix. We don't need it. The instances in a level don't have such suffix.
+	FString CameraName = ICVFXCameraComponent->GetName();
+	CameraName.RemoveFromEnd(TEXT("_GEN_VARIABLE"));
+
+	return CameraName;
+}
+
+FIntPoint FDisplayClusterConfiguratorICVFXMediaTileCustomizationBase::GetEditedTilePos() const
+{
+	FIntPoint TilePos = FIntPoint::NoneValue;
+
+	// Children 0 and 1 correspond to X and Y subproperties of FIntPoint
+	TilePosHandle->GetChildHandle(0)->GetValue(TilePos.X);
+	TilePosHandle->GetChildHandle(1)->GetValue(TilePos.Y);
+
+	return TilePos;
+}
+
+uint8 FDisplayClusterConfiguratorICVFXMediaTileCustomizationBase::GenerateUniqueOwnerIndex(const FString& OwnerName) const
+{
+	if (const AActor* OwningActor = GetOwningActor())
+	{
+		// Get all camera components
+		TArray<UDisplayClusterICVFXCameraComponent*> ICVFXCameras;
+		OwningActor->GetComponents(ICVFXCameras);
+
+		if (ICVFXCameras.Num() > 0)
+		{
+			// Camera sort predicate
+			struct FCameraSortPredicate
+			{
+				bool operator()(UDisplayClusterICVFXCameraComponent& LHS, UDisplayClusterICVFXCameraComponent& RHS) const
+				{
+					return LHS.GetName().Compare(RHS.GetName(), ESearchCase::IgnoreCase) < 0;
+				}
+			};
+
+			// Just in case, sort by name to always keep the same order
+			ICVFXCameras.Sort(FCameraSortPredicate());
+
+			const int32 CamerasAmount = ICVFXCameras.Num();
+			checkSlow(CamerasAmount > 0 && CamerasAmount < TNumericLimits<uint8>::Max());
+
+			// Find camera index in the list
+			for (uint8 CameraIdx = 0; CameraIdx < CamerasAmount; ++CameraIdx)
+			{
+				if (ICVFXCameras[CameraIdx]->GetName().Equals(OwnerName, ESearchCase::IgnoreCase))
+				{
+					return CameraIdx;
+				}
+			}
+		}
+		else
+		{
+			// Should not normally happen because the camera being edited should be there at least.
+			checkSlow(false);
+		}
+	}
+
+	return TNumericLimits<uint8>::Max();
+}
+
+void FDisplayClusterConfiguratorICVFXMediaTileCustomizationBase::OnMediaSubjectChanged()
+{
+	ModifyMediaSubjectParameters();
+}
+
+void FDisplayClusterConfiguratorICVFXMediaTileCustomizationBase::OnTilePositionChanged()
+{
+	ModifyMediaSubjectParameters();
+}
+
+void FDisplayClusterConfiguratorICVFXMediaTileCustomizationBase::OnAutoConfigureRequested(UObject* InEditingObject)
+{
+	if (EditingObject == InEditingObject)
+	{
+		ModifyMediaSubjectParameters();
+	}
+}
+
+#undef LOCTEXT_NAMESPACE
diff --git a/Engine/Plugins/Runtime/nDisplay/Source/DisplayClusterConfigurator/Private/Views/Details/Media/DisplayClusterConfiguratorICVFXMediaTileCustomization.h b/Engine/Plugins/Runtime/nDisplay/Source/DisplayClusterConfigurator/Private/Views/Details/Media/DisplayClusterConfiguratorICVFXMediaTileCustomization.h
new file mode 100644
--- /dev/null
+++ b/Engine/Plugins/Runtime/nDisplay/Source/DisplayClusterConfigurator/Private/Views/Details/Media/DisplayClusterConfiguratorICVFXMediaTileCustomization.h
@@ -0,0 +1,109 @@
+// Copyright Epic Games, Inc. All Rights Reserved.
+
+#pragma once
+
+#include "CoreMinimal.h"
+#include "Views/Details/DisplayClusterConfiguratorBaseTypeCustomization.h"
+
+class IDisplayClusterModularFeatureMediaInitializer;
+class UDisplayClusterICVFXCameraComponent;
+
+
+/**
+ * Helper class to hanldle common functionality of input and output tiles customization
+ */
+class FDisplayClusterConfiguratorICVFXMediaTileCustomizationBase
+	: public FDisplayClusterConfiguratorBaseTypeCustomization
+{
+public:
+	FDisplayClusterConfiguratorICVFXMediaTileCustomizationBase();
+	~FDisplayClusterConfiguratorICVFXMediaTileCustomizationBase();
+
+protected:
+
+	//~ Begin IPropertyTypeCustomization
+	virtual void CustomizeChildren(TSharedRef<IPropertyHandle> PropertyHandle, IDetailChildrenBuilder& ChildBuilder, IPropertyTypeCustomizationUtils& CustomizationUtils) override;
+	//~ End IPropertyTypeCustomization
+
+private:
+
+	/** Entry point to modify media subjects. */
+	void ModifyMediaSubjectParameters();
+
+	/** Returns the actor that owns the object being edited. */
+	AActor* GetOwningActor() const;
+
+	/** Returns the name of media subject's owner. */
+	FString GetOwnerName(const UDisplayClusterICVFXCameraComponent* ICVFXCameraComponent) const;
+
+	/** Returns tile position currently set. */
+	FIntPoint GetEditedTilePos() const;
+
+	/** Generate unique index for the media subject's owner. */
+	uint8 GenerateUniqueOwnerIndex(const FString& OwnerName) const;
+
+private:
+
+	/** Handles media source/output change callbacks */
+	void OnMediaSubjectChanged();
+
+	/** Handles tile position change callbacks */
+	void OnTilePositionChanged();
+
+	/** Auto-configuration requests handler */
+	void OnAutoConfigureRequested(UObject* EditingObject);
+
+protected:
+
+	/** MediaSource or MediaOutput property handle, depending on the child implementation. */
+	TSharedPtr<IPropertyHandle> MediaSubjectHandle;
+
+	/** Tile position property handle. */
+	TSharedPtr<IPropertyHandle> TilePosHandle;
+
+private:
+
+	/** Media initializers (modular features) available. */
+	static TArray<IDisplayClusterModularFeatureMediaInitializer*> MediaInitializers;
+
+	/** Optimization flag to prevent retrieving modular features multiple times.*/
+	static bool bMediaInitializersObtained;
+};
+
+
+/**
+ * Details panel customization for FDisplayClusterConfigurationMediaUniformTileInput struct (input tiles).
+ */
+class FDisplayClusterConfiguratorICVFXMediaInputTileCustomization
+	: public FDisplayClusterConfiguratorICVFXMediaTileCustomizationBase
+{
+public:
+	static TSharedRef<IPropertyTypeCustomization> MakeInstance()
+	{
+		return MakeShared<FDisplayClusterConfiguratorICVFXMediaInputTileCustomization>();
+	}
+
+protected:
+	//~ Begin IPropertyTypeCustomization
+	virtual void CustomizeChildren(TSharedRef<IPropertyHandle> PropertyHandle, IDetailChildrenBuilder& ChildBuilder, IPropertyTypeCustomizationUtils& CustomizationUtils) override;
+	//~ End IPropertyTypeCustomization
+};
+
+
+/**
+ * Details panel customization for FDisplayClusterConfigurationMediaUniformTileOutput struct (output tiles).
+ */
+class FDisplayClusterConfiguratorICVFXMediaOutputTileCustomization
+	: public FDisplayClusterConfiguratorICVFXMediaTileCustomizationBase
+{
+public:
+	static TSharedRef<IPropertyTypeCustomization> MakeInstance()
+	{
+		return MakeShared<FDisplayClusterConfiguratorICVFXMediaOutputTileCustomization>();
+	}
+
+protected:
+	//~ Begin IPropertyTypeCustomization
+	virtual void CustomizeChildren(TSharedRef<IPropertyHandle> PropertyHandle, IDetailChildrenBuilder& ChildBuilder, IPropertyTypeCustomizationUtils& CustomizationUtils) override;
+	//~ End IPropertyTypeCustomization
+};
diff --git a/Engine/Plugins/Runtime/nDisplay/Source/DisplayClusterConfigurator/Private/Views/Details/Media/DisplayClusterConfiguratorMediaUtils.cpp b/Engine/Plugins/Runtime/nDisplay/Source/DisplayClusterConfigurator/Private/Views/Details/Media/DisplayClusterConfiguratorMediaUtils.cpp
new file mode 100644
--- /dev/null
+++ b/Engine/Plugins/Runtime/nDisplay/Source/DisplayClusterConfigurator/Private/Views/Details/Media/DisplayClusterConfiguratorMediaUtils.cpp
@@ -0,0 +1,10 @@
+// Copyright Epic Games, Inc. All Rights Reserved.
+
+#include "Views/Details/Media/DisplayClusterConfiguratorMediaUtils.h"
+
+
+FDisplayClusterConfiguratorMediaUtils& FDisplayClusterConfiguratorMediaUtils::Get()
+{
+	static FDisplayClusterConfiguratorMediaUtils Instance;
+	return Instance;
+}
diff --git a/Engine/Plugins/Runtime/nDisplay/Source/DisplayClusterConfigurator/Private/Views/Details/Media/DisplayClusterConfiguratorMediaUtils.h b/Engine/Plugins/Runtime/nDisplay/Source/DisplayClusterConfigurator/Private/Views/Details/Media/DisplayClusterConfiguratorMediaUtils.h
new file mode 100644
--- /dev/null
+++ b/Engine/Plugins/Runtime/nDisplay/Source/DisplayClusterConfigurator/Private/Views/Details/Media/DisplayClusterConfiguratorMediaUtils.h
@@ -0,0 +1,31 @@
+// Copyright Epic Games, Inc. All Rights Reserved.
+
+#pragma once
+
+#include "CoreMinimal.h"
+
+
+/**
+ * Media customization utilities
+ */
+class FDisplayClusterConfiguratorMediaUtils
+{
+public:
+
+	/** Singleton access */
+	static FDisplayClusterConfiguratorMediaUtils& Get();
+
+public:
+
+	/** Tiled media auto-configuration event */
+	DECLARE_EVENT_OneParam(FDisplayClusterConfiguratorMediaUtils, FDisplayClusterConfiguratorTiledMediaAutoConfigurationRequestedEvent, UObject*);
+	FDisplayClusterConfiguratorTiledMediaAutoConfigurationRequestedEvent& OnTiledMediaAutoConfiguration()
+	{
+		return TiledMediaAutoConfigurationEvent;
+	}
+
+private:
+
+	/** Tiled media auto-configuration event */
+	FDisplayClusterConfiguratorTiledMediaAutoConfigurationRequestedEvent TiledMediaAutoConfigurationEvent;
+};
diff --git a/Engine/Plugins/Runtime/nDisplay/Source/SharedMemoryMedia/SharedMemoryMedia.Build.cs b/Engine/Plugins/Runtime/nDisplay/Source/SharedMemoryMedia/SharedMemoryMedia.Build.cs
--- a/Engine/Plugins/Runtime/nDisplay/Source/SharedMemoryMedia/SharedMemoryMedia.Build.cs
+++ b/Engine/Plugins/Runtime/nDisplay/Source/SharedMemoryMedia/SharedMemoryMedia.Build.cs
@@ -8,6 +8,10 @@ namespace UnrealBuildTool.Rules
 	{
 		public SharedMemoryMedia(ReadOnlyTargetRules Target) : base(Target)
 		{
+			PublicDependencyModuleNames.AddRange(
+				new string[] {
+					"MediaIOCore",
+				});
 
 			PrivateDependencyModuleNames.AddRange(
 				new string[] {
@@ -16,7 +20,6 @@ namespace UnrealBuildTool.Rules
 					"Engine",
 					"Media",
 					"MediaAssets",
-					"MediaIOCore",
 					"RenderCore",
 					"Renderer",
 					"RHI",
diff --git a/Engine/Plugins/Runtime/nDisplay/Source/SharedMemoryMediaEditor/Private/ModularFeatures/SharedMemoryMediaInitializerFeature.cpp b/Engine/Plugins/Runtime/nDisplay/Source/SharedMemoryMediaEditor/Private/ModularFeatures/SharedMemoryMediaInitializerFeature.cpp
new file mode 100644
--- /dev/null
+++ b/Engine/Plugins/Runtime/nDisplay/Source/SharedMemoryMediaEditor/Private/ModularFeatures/SharedMemoryMediaInitializerFeature.cpp
@@ -0,0 +1,43 @@
+// Copyright Epic Games, Inc. All Rights Reserved.
+
+#pragma once
+
+#include "ModularFeatures/SharedMemoryMediaInitializerFeature.h"
+
+#include "SharedMemoryMediaOutput.h"
+#include "SharedMemoryMediaSource.h"
+
+
+bool FSharedMemoryMediaInitializerFeature::IsMediaSubjectSupported(const UObject* MediaSubject)
+{
+	if (MediaSubject)
+	{
+		return MediaSubject->IsA<USharedMemoryMediaSource>() || MediaSubject->IsA<USharedMemoryMediaOutput>();
+	}
+
+	return false;
+}
+
+void FSharedMemoryMediaInitializerFeature::InitializeMediaSubjectForTile(UObject* MediaSubject, const FString& OwnerName, uint8 OwnerUniqueIdx, const FIntPoint& TilePos)
+{
+	checkSlow(MediaSubject);
+
+	if (USharedMemoryMediaSource* SMMediaSource = Cast<USharedMemoryMediaSource>(MediaSubject))
+	{
+		SMMediaSource->bZeroLatency = true;
+		SMMediaSource->Mode         = ESharedMemoryMediaSourceMode::Framelocked;
+		SMMediaSource->UniqueName   = GenerateUniqueName(OwnerName, TilePos);
+	}
+	else if (USharedMemoryMediaOutput* SMMediaOutput = Cast<USharedMemoryMediaOutput>(MediaSubject))
+	{
+		SMMediaOutput->UniqueName   = GenerateUniqueName(OwnerName, TilePos);
+		SMMediaOutput->bInvertAlpha = true;
+		SMMediaOutput->bCrossGpu    = true;
+		SMMediaOutput->NumberOfTextureBuffers = 4;
+	}
+}
+
+FString FSharedMemoryMediaInitializerFeature::GenerateUniqueName(const FString& OwnerName, const FIntPoint& TilePos)
+{
+	return FString::Printf(TEXT("%s_tile_%d:%d"), *OwnerName, TilePos.X, TilePos.Y);
+}
diff --git a/Engine/Plugins/Runtime/nDisplay/Source/SharedMemoryMediaEditor/Private/ModularFeatures/SharedMemoryMediaInitializerFeature.h b/Engine/Plugins/Runtime/nDisplay/Source/SharedMemoryMediaEditor/Private/ModularFeatures/SharedMemoryMediaInitializerFeature.h
new file mode 100644
--- /dev/null
+++ b/Engine/Plugins/Runtime/nDisplay/Source/SharedMemoryMediaEditor/Private/ModularFeatures/SharedMemoryMediaInitializerFeature.h
@@ -0,0 +1,26 @@
+// Copyright Epic Games, Inc. All Rights Reserved.
+
+#pragma once
+
+#include "CoreMinimal.h"
+#include "IDisplayClusterModularFeatureMediaInitializer.h"
+
+
+/**
+ * ShareMemory media source/output initializer for nDisplay
+ */
+class FSharedMemoryMediaInitializerFeature
+	: public IDisplayClusterModularFeatureMediaInitializer
+{
+public:
+
+	//~ Begin IDisplayClusterModularFeatureMediaInitializer
+	virtual bool IsMediaSubjectSupported(const UObject* MediaSubject) override;
+	virtual void InitializeMediaSubjectForTile(UObject* MediaSubject, const FString& OwnerName, uint8 OwnerUniqueIdx, const FIntPoint& TilePos) override;
+	//~ End IDisplayClusterModularFeatureMediaInitializer
+
+private:
+
+	/** Generates unique name for SMM subjects */
+	FString GenerateUniqueName(const FString& OwnerName, const FIntPoint& TilePos);
+};
diff --git a/Engine/Plugins/Runtime/nDisplay/Source/SharedMemoryMediaEditor/Private/SharedMemoryMediaEditorModule.cpp b/Engine/Plugins/Runtime/nDisplay/Source/SharedMemoryMediaEditor/Private/SharedMemoryMediaEditorModule.cpp
new file mode 100644
--- /dev/null
+++ b/Engine/Plugins/Runtime/nDisplay/Source/SharedMemoryMediaEditor/Private/SharedMemoryMediaEditorModule.cpp
@@ -0,0 +1,39 @@
+// Copyright Epic Games, Inc. All Rights Reserved.
+
+#include "SharedMemoryMediaEditorModule.h"
+
+#include "Features/IModularFeatures.h"
+#include "Modules/ModuleManager.h"
+
+#include "ModularFeatures/SharedMemoryMediaInitializerFeature.h"
+
+
+void FSharedMemoryMediaEditorModule::StartupModule()
+{
+	RegisterModularFeatures();
+}
+
+void FSharedMemoryMediaEditorModule::ShutdownModule()
+{
+	UnregisterModularFeatures();
+}
+
+void FSharedMemoryMediaEditorModule::RegisterModularFeatures()
+{
+	// Instantiate modular features
+	MediaInitializer = MakeUnique<FSharedMemoryMediaInitializerFeature>();
+
+	// Register modular features
+	IModularFeatures& ModularFeatures = IModularFeatures::Get();
+	ModularFeatures.RegisterModularFeature(FSharedMemoryMediaInitializerFeature::ModularFeatureName, MediaInitializer.Get());
+}
+
+void FSharedMemoryMediaEditorModule::UnregisterModularFeatures()
+{
+	// Unregister modular features
+	IModularFeatures& ModularFeatures = IModularFeatures::Get();
+	ModularFeatures.UnregisterModularFeature(FSharedMemoryMediaInitializerFeature::ModularFeatureName, MediaInitializer.Get());
+}
+
+
+IMPLEMENT_MODULE(FSharedMemoryMediaEditorModule, SharedMemoryMediaEditor)
diff --git a/Engine/Plugins/Runtime/nDisplay/Source/SharedMemoryMediaEditor/Private/SharedMemoryMediaEditorModule.h b/Engine/Plugins/Runtime/nDisplay/Source/SharedMemoryMediaEditor/Private/SharedMemoryMediaEditorModule.h
new file mode 100644
--- /dev/null
+++ b/Engine/Plugins/Runtime/nDisplay/Source/SharedMemoryMediaEditor/Private/SharedMemoryMediaEditorModule.h
@@ -0,0 +1,33 @@
+// Copyright Epic Games, Inc. All Rights Reserved.
+
+#pragma once
+
+#include "Modules/ModuleInterface.h"
+#include "Templates/UniquePtr.h"
+
+class FSharedMemoryMediaInitializerFeature;
+
+
+class FSharedMemoryMediaEditorModule
+	: public IModuleInterface
+{
+public:
+	
+	//~ Begin IModuleInterface
+	virtual void StartupModule() override;
+	virtual void ShutdownModule() override;
+	//~ End IModuleInterface
+
+private:
+
+	/** Registers modular features for external modules */
+	void RegisterModularFeatures();
+
+	/** Unregisters modular features */
+	void UnregisterModularFeatures();
+
+private:
+
+	/** MediaInitializer modular feature instance */
+	TUniquePtr<FSharedMemoryMediaInitializerFeature> MediaInitializer;
+};
diff --git a/Engine/Plugins/Runtime/nDisplay/Source/SharedMemoryMediaEditor/SharedMemoryMediaEditor.Build.cs b/Engine/Plugins/Runtime/nDisplay/Source/SharedMemoryMediaEditor/SharedMemoryMediaEditor.Build.cs
new file mode 100644
--- /dev/null
+++ b/Engine/Plugins/Runtime/nDisplay/Source/SharedMemoryMediaEditor/SharedMemoryMediaEditor.Build.cs
@@ -0,0 +1,22 @@
+// Copyright Epic Games, Inc. All Rights Reserved.
+
+using System.IO;
+
+namespace UnrealBuildTool.Rules
+{
+	public class SharedMemoryMediaEditor : ModuleRules
+	{
+		public SharedMemoryMediaEditor(ReadOnlyTargetRules Target) : base(Target)
+		{
+			PrivateDependencyModuleNames.AddRange(
+				new string[] {
+					"Core",
+					"CoreUObject",
+					"DisplayClusterModularFeaturesEditor",
+					"Engine",
+					"SharedMemoryMedia",
+					"UnrealEd",
+				});
+		}
+	}
+}
diff --git a/Engine/Plugins/Runtime/nDisplay/nDisplay.uplugin b/Engine/Plugins/Runtime/nDisplay/nDisplay.uplugin
--- a/Engine/Plugins/Runtime/nDisplay/nDisplay.uplugin
+++ b/Engine/Plugins/Runtime/nDisplay/nDisplay.uplugin
@@ -223,6 +223,15 @@
 				"Linux"
 			]
 		},
+		{
+			"Name": "SharedMemoryMediaEditor",
+			"Type": "Editor",
+			"LoadingPhase": "Default",
+			"PlatformAllowList": [
+				"Win64",
+				"Linux"
+			]
+		},
 		{
 			"Name": "DisplayClusterReplication",
 			"Type": "Runtime",
diff --git a/Engine/Plugins/Runtime/nDisplayModularFeatures/Source/DisplayClusterModularFeaturesEditor/DisplayClusterModularFeaturesEditor.build.cs b/Engine/Plugins/Runtime/nDisplayModularFeatures/Source/DisplayClusterModularFeaturesEditor/DisplayClusterModularFeaturesEditor.build.cs
new file mode 100644
--- /dev/null
+++ b/Engine/Plugins/Runtime/nDisplayModularFeatures/Source/DisplayClusterModularFeaturesEditor/DisplayClusterModularFeaturesEditor.build.cs
@@ -0,0 +1,23 @@
+// Copyright Epic Games, Inc. All Rights Reserved.
+
+using UnrealBuildTool;
+using System.IO;
+
+public class DisplayClusterModularFeaturesEditor : ModuleRules
+{
+	public DisplayClusterModularFeaturesEditor(ReadOnlyTargetRules ROTargetRules) : base(ROTargetRules)
+	{
+		PublicDependencyModuleNames.AddRange(
+			new string[] {
+				"Core",
+				"CoreUObject",
+				"Engine"
+			});
+
+		PrivateDependencyModuleNames.AddRange(
+			new string[]
+			{
+				"UnrealEd"
+			});
+	}
+}
diff --git a/Engine/Plugins/Runtime/nDisplayModularFeatures/Source/DisplayClusterModularFeaturesEditor/Private/DisplayClusterModularFeaturesEditorModule.cpp b/Engine/Plugins/Runtime/nDisplayModularFeatures/Source/DisplayClusterModularFeaturesEditor/Private/DisplayClusterModularFeaturesEditorModule.cpp
new file mode 100644
--- /dev/null
+++ b/Engine/Plugins/Runtime/nDisplayModularFeatures/Source/DisplayClusterModularFeaturesEditor/Private/DisplayClusterModularFeaturesEditorModule.cpp
@@ -0,0 +1,5 @@
+// Copyright Epic Games, Inc. All Rights Reserved.
+
+#include "Modules/ModuleManager.h"
+
+IMPLEMENT_MODULE(FDefaultModuleImpl, DisplayClusterModularFeaturesEditor);
diff --git a/Engine/Plugins/Runtime/nDisplayModularFeatures/Source/DisplayClusterModularFeaturesEditor/Private/DisplayClusterModularFeaturesEditorNames.cpp b/Engine/Plugins/Runtime/nDisplayModularFeatures/Source/DisplayClusterModularFeaturesEditor/Private/DisplayClusterModularFeaturesEditorNames.cpp
new file mode 100644
--- /dev/null
+++ b/Engine/Plugins/Runtime/nDisplayModularFeatures/Source/DisplayClusterModularFeaturesEditor/Private/DisplayClusterModularFeaturesEditorNames.cpp
@@ -0,0 +1,5 @@
+// Copyright Epic Games, Inc. All Rights Reserved.
+
+#include "IDisplayClusterModularFeatureMediaInitializer.h"
+
+const FName IDisplayClusterModularFeatureMediaInitializer::ModularFeatureName = TEXT("DisplayClusterModularFeatureMediaInitializer");
diff --git a/Engine/Plugins/Runtime/nDisplayModularFeatures/Source/DisplayClusterModularFeaturesEditor/Public/IDisplayClusterModularFeatureMediaInitializer.h b/Engine/Plugins/Runtime/nDisplayModularFeatures/Source/DisplayClusterModularFeaturesEditor/Public/IDisplayClusterModularFeatureMediaInitializer.h
new file mode 100644
--- /dev/null
+++ b/Engine/Plugins/Runtime/nDisplayModularFeatures/Source/DisplayClusterModularFeaturesEditor/Public/IDisplayClusterModularFeatureMediaInitializer.h
@@ -0,0 +1,43 @@
+// Copyright Epic Games, Inc. All Rights Reserved.
+
+#pragma once
+
+#include "CoreMinimal.h"
+#include "Features/IModularFeature.h"
+
+
+/**
+ * Base class for nDisplay media initializer implementations.
+ */
+class DISPLAYCLUSTERMODULARFEATURESEDITOR_API IDisplayClusterModularFeatureMediaInitializer
+	: public IModularFeature
+{
+public:
+	/** Public feature name */
+	static const FName ModularFeatureName;
+
+public:
+	virtual ~IDisplayClusterModularFeatureMediaInitializer() = default;
+
+public:
+
+	/**
+	 * Checks if media subject is supported by the initializer
+	 *
+	 * @param MediaSubject - UMediaSource or UMediaOutput instance
+	 * @return true if media subject supported
+	 */
+	virtual bool IsMediaSubjectSupported(const UObject* MediaSubject) = 0;
+
+	/**
+	 * Performs initialization of a media subject
+	 *
+	 * @param MediaSubject   - UMediaSource or UMediaOutput instance
+	 * @param OwnerName      - Owner name (e.g. ICVFX camera component)
+	 * @param OwnerUniqueIdx - Unique owner index
+	 * @param TilePos        - TIle XY-position
+	 * 
+	 * @return true if media subject supported
+	 */
+	virtual void InitializeMediaSubjectForTile(UObject* MediaSubject, const FString& OwnerName, uint8 OwnerUniqueIdx, const FIntPoint& TilePos) = 0;
+};
diff --git a/Engine/Plugins/Runtime/nDisplayModularFeatures/nDisplayModularFeatures.uplugin b/Engine/Plugins/Runtime/nDisplayModularFeatures/nDisplayModularFeatures.uplugin
--- a/Engine/Plugins/Runtime/nDisplayModularFeatures/nDisplayModularFeatures.uplugin
+++ b/Engine/Plugins/Runtime/nDisplayModularFeatures/nDisplayModularFeatures.uplugin
@@ -19,6 +19,11 @@
 			"Name": "DisplayClusterLightCardExtender",
 			"Type": "Runtime",
 			"LoadingPhase": "Default"
+		},
+		{
+			"Name": "DisplayClusterModularFeaturesEditor",
+			"Type": "Editor",
+			"LoadingPhase": "Default"
 		}
 	]
 }
diff --git a/Engine/Plugins/VirtualProduction/Rivermax/RivermaxMedia/RivermaxMedia.uplugin b/Engine/Plugins/VirtualProduction/Rivermax/RivermaxMedia/RivermaxMedia.uplugin
--- a/Engine/Plugins/VirtualProduction/Rivermax/RivermaxMedia/RivermaxMedia.uplugin
+++ b/Engine/Plugins/VirtualProduction/Rivermax/RivermaxMedia/RivermaxMedia.uplugin
@@ -20,6 +20,10 @@
 			"Name": "MediaIOFramework",
 			"Enabled": true
 		},
+		{
+			"Name": "nDisplayModularFeatures",
+			"Enabled": true
+		},
 		{
 			"Name": "RivermaxCore",
 			"Enabled": true
diff --git a/Engine/Plugins/VirtualProduction/Rivermax/RivermaxMedia/Source/RivermaxMediaEditor/Private/ModularFeatures/RivermaxMediaInitializerFeature.cpp b/Engine/Plugins/VirtualProduction/Rivermax/RivermaxMedia/Source/RivermaxMediaEditor/Private/ModularFeatures/RivermaxMediaInitializerFeature.cpp
new file mode 100644
--- /dev/null
+++ b/Engine/Plugins/VirtualProduction/Rivermax/RivermaxMedia/Source/RivermaxMediaEditor/Private/ModularFeatures/RivermaxMediaInitializerFeature.cpp
@@ -0,0 +1,98 @@
+// Copyright Epic Games, Inc. All Rights Reserved.
+
+#pragma once
+
+#include "ModularFeatures/RivermaxMediaInitializerFeature.h"
+
+#include "IRivermaxCoreModule.h"
+#include "IRivermaxManager.h"
+
+#include "RivermaxMediaSource.h"
+#include "RivermaxMediaOutput.h"
+
+
+bool FRivermaxMediaInitializerFeature::IsMediaSubjectSupported(const UObject* MediaSubject)
+{
+	if (MediaSubject)
+	{
+		return MediaSubject->IsA<URivermaxMediaSource>() || MediaSubject->IsA<URivermaxMediaOutput>();
+	}
+
+	return false;
+}
+
+void FRivermaxMediaInitializerFeature::InitializeMediaSubjectForTile(UObject* MediaSubject, const FString& OwnerName, uint8 OwnerUniqueIdx, const FIntPoint& TilePos)
+{
+	checkSlow(MediaSubject);
+
+	if (URivermaxMediaSource* RivermaxMediaSource = Cast<URivermaxMediaSource>(MediaSubject))
+	{
+		RivermaxMediaSource->PlayerMode          = ERivermaxPlayerMode::Framelock;
+		RivermaxMediaSource->bUseZeroLatency     = true;
+		RivermaxMediaSource->bOverrideResolution = false;
+		//RivermaxMediaSource->Resolution          = default value
+		RivermaxMediaSource->FrameRate           = { 60,1 };
+		RivermaxMediaSource->PixelFormat         = ERivermaxMediaSourcePixelFormat::RGB_10bit;
+		RivermaxMediaSource->InterfaceAddress    = GetRivermaxInterfaceAddress();
+		RivermaxMediaSource->StreamAddress       = GenerateStreamAddress(OwnerUniqueIdx, TilePos);
+		RivermaxMediaSource->Port                = 50000;
+		RivermaxMediaSource->bIsSRGBInput        = false;
+		RivermaxMediaSource->bUseGPUDirect       = true;
+	}
+	else if (URivermaxMediaOutput* RivermaxMediaOutput = Cast<URivermaxMediaOutput>(MediaSubject))
+	{
+		RivermaxMediaOutput->AlignmentMode       = ERivermaxMediaAlignmentMode::FrameCreation;
+		RivermaxMediaOutput->bDoContinuousOutput = false;
+		RivermaxMediaOutput->FrameLockingMode    = ERivermaxFrameLockingMode::BlockOnReservation;
+		RivermaxMediaOutput->PresentationQueueSize = 2;
+		RivermaxMediaOutput->bDoFrameCounterTimestamping = true;
+		RivermaxMediaOutput->bOverrideResolution = false;
+		//RivermaxMediaOutput->Resolution          = default value
+		RivermaxMediaOutput->FrameRate           = { 60,1 };
+		RivermaxMediaOutput->PixelFormat         = ERivermaxMediaOutputPixelFormat::PF_10BIT_RGB;
+		RivermaxMediaOutput->InterfaceAddress    = GetRivermaxInterfaceAddress();
+		RivermaxMediaOutput->StreamAddress       = GenerateStreamAddress(OwnerUniqueIdx, TilePos);
+		RivermaxMediaOutput->Port                = 50000;
+		RivermaxMediaOutput->bUseGPUDirect       = true;
+	}
+}
+
+FString FRivermaxMediaInitializerFeature::GetRivermaxInterfaceAddress() const
+{
+	FString ResultAddress{ TEXT("*.*.*.*") };
+
+	// Now let's see if we have any interfaces available
+	IRivermaxCoreModule& RivermaxModule = FModuleManager::LoadModuleChecked<IRivermaxCoreModule>(TEXT("RivermaxCore"));
+	const TConstArrayView<UE::RivermaxCore::FRivermaxDeviceInfo> Devices = RivermaxModule.GetRivermaxManager()->GetDevices();
+	if (Devices.Num() > 0)
+	{
+		// Split address into octets
+		TArray<FString> Octets;
+		Devices[0].InterfaceAddress.ParseIntoArray(Octets, TEXT("."));
+
+		// IPv4?
+		if (Octets.Num() == 4)
+		{
+			ResultAddress = FString::Printf(TEXT("%s.%s.%s.*"), *Octets[0], *Octets[1], *Octets[2]);
+		}
+		// IPv6?
+		else
+		{
+			Devices[0].InterfaceAddress.ParseIntoArray(Octets, TEXT(":"));
+			if (Octets.Num() == 6)
+			{
+				ResultAddress = FString::Printf(TEXT("%s:%s:%s:%s:%s:*"), *Octets[0], *Octets[1], *Octets[2], *Octets[3], *Octets[4]);
+			}
+		}
+	}
+
+	return ResultAddress;
+}
+
+FString FRivermaxMediaInitializerFeature::GenerateStreamAddress(uint8 OwnerUniqueIdx, const FIntPoint& TilePos) const
+{
+	constexpr uint8 MaxVal = TNumericLimits<uint8>::Max();
+	checkSlow(OwnerUniqueIdx < MaxVal && TilePos.X < MaxVal && TilePos.Y < MaxVal);
+
+	return FString::Printf(TEXT("228.%u.%u.%u"), OwnerUniqueIdx + 1, TilePos.X + 1, TilePos.Y + 1);
+}
diff --git a/Engine/Plugins/VirtualProduction/Rivermax/RivermaxMedia/Source/RivermaxMediaEditor/Private/ModularFeatures/RivermaxMediaInitializerFeature.h b/Engine/Plugins/VirtualProduction/Rivermax/RivermaxMedia/Source/RivermaxMediaEditor/Private/ModularFeatures/RivermaxMediaInitializerFeature.h
new file mode 100644
--- /dev/null
+++ b/Engine/Plugins/VirtualProduction/Rivermax/RivermaxMedia/Source/RivermaxMediaEditor/Private/ModularFeatures/RivermaxMediaInitializerFeature.h
@@ -0,0 +1,29 @@
+// Copyright Epic Games, Inc. All Rights Reserved.
+
+#pragma once
+
+#include "CoreMinimal.h"
+#include "IDisplayClusterModularFeatureMediaInitializer.h"
+
+
+/**
+ * Rivermax media source/output initializer for nDisplay
+ */
+class FRivermaxMediaInitializerFeature
+	: public IDisplayClusterModularFeatureMediaInitializer
+{
+public:
+
+	//~ Begin IDisplayClusterModularFeatureMediaInitializer
+	virtual bool IsMediaSubjectSupported(const UObject* MediaSubject) override;
+	virtual void InitializeMediaSubjectForTile(UObject* MediaSubject, const FString& OwnerName, uint8 OwnerUniqueIdx, const FIntPoint& TilePos) override;
+	//~ End IDisplayClusterModularFeatureMediaInitializer
+
+private:
+
+	/** Returns first interface address if available */
+	FString GetRivermaxInterfaceAddress() const;
+
+	/** Generates stream address based on the function parameters */
+	FString GenerateStreamAddress(uint8 OwnerUniqueIdx, const FIntPoint& TilePos) const;
+};
diff --git a/Engine/Plugins/VirtualProduction/Rivermax/RivermaxMedia/Source/RivermaxMediaEditor/Private/RivermaxMediaEditorModule.cpp b/Engine/Plugins/VirtualProduction/Rivermax/RivermaxMedia/Source/RivermaxMediaEditor/Private/RivermaxMediaEditorModule.cpp
--- a/Engine/Plugins/VirtualProduction/Rivermax/RivermaxMedia/Source/RivermaxMediaEditor/Private/RivermaxMediaEditorModule.cpp
+++ b/Engine/Plugins/VirtualProduction/Rivermax/RivermaxMedia/Source/RivermaxMediaEditor/Private/RivermaxMediaEditorModule.cpp
@@ -6,6 +6,8 @@
 #include "RivermaxMediaOutput.h"
 #include "RivermaxMediaSource.h"
 #include "Customization/RivermaxMediaDetailsCustomization.h"
+#include "Features/IModularFeatures.h"
+#include "ModularFeatures/RivermaxMediaInitializerFeature.h"
 
 class FRivermaxMediaEditorModule : public IModuleInterface
 {
@@ -15,11 +17,13 @@ public:
 	virtual void StartupModule() override
 	{
 		RegisterCustomizations();
+		RegisterModularFeatures();
 	}
 	
 	virtual void ShutdownModule() override
 	{
 		UnregisterCustomizations();
+		UnregisterModularFeatures();
 	}
 	//~End IModuleInterface
 
@@ -52,8 +56,30 @@ private:
 		}
 	}
 
+	/** Registers modular features for external modules */
+	void RegisterModularFeatures()
+	{
+		// Instantiate modular features
+		MediaInitializer = MakeUnique<FRivermaxMediaInitializerFeature>();
+
+		// Register modular features
+		IModularFeatures& ModularFeatures = IModularFeatures::Get();
+		ModularFeatures.RegisterModularFeature(FRivermaxMediaInitializerFeature::ModularFeatureName, MediaInitializer.Get());
+	}
+
+	/** Unregisters modular features */
+	void UnregisterModularFeatures()
+	{
+		// Unregister modular features
+		IModularFeatures& ModularFeatures = IModularFeatures::Get();
+		ModularFeatures.UnregisterModularFeature(FRivermaxMediaInitializerFeature::ModularFeatureName, MediaInitializer.Get());
+	}
+
 private:
 	TArray<FName> CustomizedClasses;
+
+	/** MediaInitializer modular feature instance */
+	TUniquePtr<FRivermaxMediaInitializerFeature> MediaInitializer;
 };
 
 
diff --git a/Engine/Plugins/VirtualProduction/Rivermax/RivermaxMedia/Source/RivermaxMediaEditor/RivermaxMediaEditor.build.cs b/Engine/Plugins/VirtualProduction/Rivermax/RivermaxMedia/Source/RivermaxMediaEditor/RivermaxMediaEditor.build.cs
--- a/Engine/Plugins/VirtualProduction/Rivermax/RivermaxMedia/Source/RivermaxMediaEditor/RivermaxMediaEditor.build.cs
+++ b/Engine/Plugins/VirtualProduction/Rivermax/RivermaxMedia/Source/RivermaxMediaEditor/RivermaxMediaEditor.build.cs
@@ -17,6 +17,7 @@ public class RivermaxMediaEditor : ModuleRules
 				"AssetTools",
 				"Core",
 				"CoreUObject",
+				"DisplayClusterModularFeaturesEditor",
 				"EditorStyle",
 				"Engine",
 				"InputCore",
