nDisplay: Fixed AlphaMap and BetaMap color correction math of MPCDI shader.

diff --git a/Engine/Plugins/Runtime/nDisplay/Shaders/Private/MPCDIUtils.ush b/Engine/Plugins/Runtime/nDisplay/Shaders/Private/MPCDIUtils.ush
--- a/Engine/Plugins/Runtime/nDisplay/Shaders/Private/MPCDIUtils.ush
+++ b/Engine/Plugins/Runtime/nDisplay/Shaders/Private/MPCDIUtils.ush
@@ -167,24 +167,24 @@ float4 ApplyBlending(float3 InColor_sRGB, OutputVS IN, float ColorAlphaValue)
 #if ALPHAMAP_BLENDING || BETAMAP_BLENDING
 	float3 LinearColor = MPCDI_sRGBToLinear(InColor_sRGB);
 
-	// To convert a texture with its own gamma to a linear color space, just use this formula:
-	//     LinearTextureColor = pow(TextureColor, 1 / GammaEmbedded) function.
+	// To convert a texture between gamma and linear color space, simply use these formulas:
+	// LinearColor = pow(sRGBColor, Gamma)
+	//   sRGBColor = pow(LinearColor, 1.0 / Gamma)
 	
 	// Note: AlphaMapGammaEmbedded is used for both AlphaMap and BetaMap.
 	// (The MPCDI 2.0 standard does not define an 'EmbeddedGamma' tag value for the BetaMap tag. However, it does require gamma correction for BetaMap.)
-	float InvGammaEmbedded = 1 / AlphaMapGammaEmbedded;
 	
 #if ALPHAMAP_BLENDING
 	if(AlphaMapComponentDepth < 3)
 	{
 		// Grayscale AlphaMap: get a single-channel value and convert its color space to linear.
-		float LinearAlphaMap = pow(AlphaMapTexture.Sample(AlphaMapSampler, IN.UV.xy).r, InvGammaEmbedded);
+		float LinearAlphaMap = pow(AlphaMapTexture.Sample(AlphaMapSampler, IN.UV.xy).r, AlphaMapGammaEmbedded);
 		LinearColor = LinearColor * LinearAlphaMap;
 	}
 	else
 	{
 		// RGB AlphaMap: get the values of the 3 texture channels and convert their color space to linear.
-		float3 LinearAlphaMap = pow(AlphaMapTexture.Sample(AlphaMapSampler, IN.UV.xy).rgb, InvGammaEmbedded);
+		float3 LinearAlphaMap = pow(AlphaMapTexture.Sample(AlphaMapSampler, IN.UV.xy).rgb, AlphaMapGammaEmbedded);
 		LinearColor = LinearColor * LinearAlphaMap;
 	}
 #endif
@@ -193,13 +193,13 @@ float4 ApplyBlending(float3 InColor_sRGB, OutputVS IN, float ColorAlphaValue)
 	if(BetaMapComponentDepth < 3)
 	{
 		// Grayscale BetaMap: get a single-channel value and convert its color space to linear.
-		float LinearBetaMap = pow(BetaMapTexture.Sample(BetaMapSampler, IN.UV.xy).r, InvGammaEmbedded);
+		float LinearBetaMap = pow(BetaMapTexture.Sample(BetaMapSampler, IN.UV.xy).r, AlphaMapGammaEmbedded);
 		LinearColor = LinearColor * ( 1 - LinearBetaMap) + LinearBetaMap;
 	}
 	else
 	{
 		// RGB BetaMap: get the values of the 3 texture channels and convert their color space to linear.
-		float3 LinearBetaMap = pow(BetaMapTexture.Sample(BetaMapSampler, IN.UV.xy).rgb, InvGammaEmbedded);
+		float3 LinearBetaMap = pow(BetaMapTexture.Sample(BetaMapSampler, IN.UV.xy).rgb, AlphaMapGammaEmbedded);
 		LinearColor = LinearColor * ( 1 - LinearBetaMap) + LinearBetaMap;
 	}
 #endif // BETAMAP_BLENDING
