nDisplay: fixed regression with overscan estimation calculation and icvfx rtt size

diff --git a/Engine/Plugins/Runtime/nDisplay/Source/DisplayCluster/Private/Game/EngineClasses/Scene/DisplayClusterICVFXCameraComponent.cpp b/Engine/Plugins/Runtime/nDisplay/Source/DisplayCluster/Private/Game/EngineClasses/Scene/DisplayClusterICVFXCameraComponent.cpp
--- a/Engine/Plugins/Runtime/nDisplay/Source/DisplayCluster/Private/Game/EngineClasses/Scene/DisplayClusterICVFXCameraComponent.cpp
+++ b/Engine/Plugins/Runtime/nDisplay/Source/DisplayCluster/Private/Game/EngineClasses/Scene/DisplayClusterICVFXCameraComponent.cpp
@@ -171,32 +171,53 @@ void UDisplayClusterICVFXCameraComponent::UpdateOverscanEstimatedFrameSize()
 	UCineCameraComponent* ActualCineCameraComponent = GetActualCineCameraComponent();
 	check(ActualCineCameraComponent);
 
+	// additional multipliers from FDisplayClusterConfigurationRenderFrame are not used in following calculations
 	const float CameraBufferRatio = CameraSettings.GetCameraBufferRatio(StageSettings);
 	const FIntPoint CameraFrameSize = CameraSettings.GetCameraFrameSize(StageSettings, *ActualCineCameraComponent);
-
-	const FIntPoint InnerFrustumResolution(
-		CameraFrameSize.X * CameraBufferRatio,
-		CameraFrameSize.Y * CameraBufferRatio
-	);
-
-	const FIntPoint UpscaledInnerFrustumResolution(
-		InnerFrustumResolution.X * CameraSettings.CustomFrustum.FieldOfViewMultiplier,
-		InnerFrustumResolution.Y * CameraSettings.CustomFrustum.FieldOfViewMultiplier
+	const FIntPoint InnerFrustumResolution( CameraFrameSize.X * CameraBufferRatio, CameraFrameSize.Y * CameraBufferRatio);
+	
+	{
+		// calculate estimations
+		FDisplayClusterConfigurationICVFX_CameraCustomFrustum EstimatedCustomFrustum = CameraSettings.CustomFrustum;
+		EstimatedCustomFrustum.bEnable = true;
+		EstimatedCustomFrustum.bAdaptResolution = true;
+
+		const float EstimatedCameraAdaptResolutionRatio = EstimatedCustomFrustum.GetCameraAdaptResolutionRatio(StageSettings);
+		const FIntPoint EstimatedInnerFrustumResolution(
+			InnerFrustumResolution.X * EstimatedCameraAdaptResolutionRatio,
+			InnerFrustumResolution.Y * EstimatedCameraAdaptResolutionRatio
 	);
 
-	// Read configuration data
-	FDisplayClusterViewport_CustomFrustumSettings CustomFrustumSettings;
-	FDisplayClusterViewportConfigurationHelpers_ICVFX::UpdateCameraCustomFrustum(CameraSettings.CustomFrustum, CustomFrustumSettings);
+		FIntRect EstimatedViewportRect(FIntPoint(0, 0), EstimatedInnerFrustumResolution);
+		FDisplayClusterViewport_CustomFrustumSettings EstimatedFrustumSettings;
+		FDisplayClusterViewport_CustomFrustumRuntimeSettings EstimatedFrustumRuntimeSettings;
 
-	// Calculate overscan to temp var
-	FIntRect FinalViewportRect(FIntPoint(0,0), UpscaledInnerFrustumResolution);
-	FDisplayClusterViewport_CustomFrustumRuntimeSettings CustomFrustumRuntimeSettings;
-	FDisplayClusterViewport_CustomFrustumRuntimeSettings::UpdateCustomFrustumSettings(GetName(), CustomFrustumSettings, CustomFrustumRuntimeSettings, FinalViewportRect);
+		FDisplayClusterViewportConfigurationHelpers_ICVFX::UpdateCameraCustomFrustum(EstimatedCustomFrustum, EstimatedFrustumSettings);
+		FDisplayClusterViewport_CustomFrustumRuntimeSettings::UpdateCustomFrustumSettings(GetName(), EstimatedFrustumSettings, EstimatedFrustumRuntimeSettings, EstimatedViewportRect);
 
-	// Show calculated values
-	CameraSettings.CustomFrustum.InnerFrustumResolution = InnerFrustumResolution;
-	CameraSettings.CustomFrustum.EstimatedOverscanResolution = FinalViewportRect.Size();
+		// Assign estimated calculated values
+		CameraSettings.CustomFrustum.EstimatedOverscanResolution = EstimatedViewportRect.Size();
+	}
 
+	{
+		// calculate real
+		const float RealCameraAdaptResolutionRatio = CameraSettings.CustomFrustum.GetCameraAdaptResolutionRatio(StageSettings);
+		const FIntPoint RealInnerFrustumResolution(
+			InnerFrustumResolution.X * RealCameraAdaptResolutionRatio,
+			InnerFrustumResolution.Y * RealCameraAdaptResolutionRatio
+		);
+
+		FIntRect RealViewportRect(FIntPoint(0, 0), RealInnerFrustumResolution);
+		FDisplayClusterViewport_CustomFrustumSettings RealFrustumSettings;
+		FDisplayClusterViewport_CustomFrustumRuntimeSettings RealFrustumRuntimeSettings;
+
+		FDisplayClusterViewportConfigurationHelpers_ICVFX::UpdateCameraCustomFrustum(CameraSettings.CustomFrustum, RealFrustumSettings);
+		FDisplayClusterViewport_CustomFrustumRuntimeSettings::UpdateCustomFrustumSettings(GetName(), RealFrustumSettings, RealFrustumRuntimeSettings, RealViewportRect);
+
+		// Assign real calculated values
+		CameraSettings.CustomFrustum.InnerFrustumResolution = RealViewportRect.Size();
+	}
+	
 	const int32 EstimatedPixel = CameraSettings.CustomFrustum.EstimatedOverscanResolution.X * CameraSettings.CustomFrustum.EstimatedOverscanResolution.Y;
 	const int32 BasePixels = CameraSettings.CustomFrustum.InnerFrustumResolution.X * CameraSettings.CustomFrustum.InnerFrustumResolution.Y;
 
