[nDisplay] Media configuration improvements

diff --git a/Engine/Plugins/Runtime/nDisplay/Source/DisplayClusterConfigurator/Private/DisplayClusterConfiguratorModule.cpp b/Engine/Plugins/Runtime/nDisplay/Source/DisplayClusterConfigurator/Private/DisplayClusterConfiguratorModule.cpp
--- a/Engine/Plugins/Runtime/nDisplay/Source/DisplayClusterConfigurator/Private/DisplayClusterConfiguratorModule.cpp
+++ b/Engine/Plugins/Runtime/nDisplay/Source/DisplayClusterConfigurator/Private/DisplayClusterConfiguratorModule.cpp
@@ -27,7 +27,8 @@
 #include "Views/Details/Components/DisplayClusterConfiguratorScreenComponentDetailsCustomization.h"
 #include "Views/Details/Components/DisplayClusterICVFXCameraComponentDetailsCustomization.h"
 #include "Views/Details/Media/DisplayClusterConfiguratorICVFXMediaCustomization.h"
-#include "Views/Details/Media/DisplayClusterConfiguratorICVFXMediaTileCustomization.h"
+#include "Views/Details/Media/DisplayClusterConfiguratorMediaFullFrameCustomization.h"
+#include "Views/Details/Media/DisplayClusterConfiguratorMediaTileCustomization.h"
 #include "Views/Details/Policies/DisplayClusterConfiguratorPolicyDetailCustomization.h"
 
 #include "Blueprints/DisplayClusterBlueprint.h"
@@ -226,8 +227,12 @@ void FDisplayClusterConfiguratorModule::RegisterCustomLayouts()
 	REGISTER_PROPERTY_LAYOUT(FDisplayClusterConfigurationViewport_RemapData, FDisplayClusterConfiguratorViewportRemapCustomization);
 	REGISTER_PROPERTY_LAYOUT(FDisplayClusterConfigurationRectangle, FDisplayClusterConfiguratorRectangleCustomization);
 	REGISTER_PROPERTY_LAYOUT(FDisplayClusterConfigurationMediaICVFX, FDisplayClusterConfiguratorICVFXMediaCustomization);
-	REGISTER_PROPERTY_LAYOUT(FDisplayClusterConfigurationMediaUniformTileInput, FDisplayClusterConfiguratorICVFXMediaInputTileCustomization);
-	REGISTER_PROPERTY_LAYOUT(FDisplayClusterConfigurationMediaUniformTileOutput, FDisplayClusterConfiguratorICVFXMediaOutputTileCustomization);
+	REGISTER_PROPERTY_LAYOUT(FDisplayClusterConfigurationMediaUniformTileInput,  FDisplayClusterConfiguratorMediaInputTileCustomization);
+	REGISTER_PROPERTY_LAYOUT(FDisplayClusterConfigurationMediaUniformTileOutput, FDisplayClusterConfiguratorMediaOutputTileCustomization);
+	REGISTER_PROPERTY_LAYOUT(FDisplayClusterConfigurationMediaInput,  FDisplayClusterConfiguratorMediaFullFrameInputCustomization);
+	REGISTER_PROPERTY_LAYOUT(FDisplayClusterConfigurationMediaOutput, FDisplayClusterConfiguratorMediaFullFrameOutputCustomization);
+	REGISTER_PROPERTY_LAYOUT(FDisplayClusterConfigurationMediaInputGroup,  FDisplayClusterConfiguratorMediaFullFrameInputCustomization);
+	REGISTER_PROPERTY_LAYOUT(FDisplayClusterConfigurationMediaOutputGroup, FDisplayClusterConfiguratorMediaFullFrameOutputCustomization);
 }
 
 void FDisplayClusterConfiguratorModule::UnregisterCustomLayouts()
diff --git a/Engine/Plugins/Runtime/nDisplay/Source/DisplayClusterConfigurator/Private/Views/Details/Media/DisplayClusterConfiguratorICVFXMediaCustomization.cpp b/Engine/Plugins/Runtime/nDisplay/Source/DisplayClusterConfigurator/Private/Views/Details/Media/DisplayClusterConfiguratorICVFXMediaCustomization.cpp
--- a/Engine/Plugins/Runtime/nDisplay/Source/DisplayClusterConfigurator/Private/Views/Details/Media/DisplayClusterConfiguratorICVFXMediaCustomization.cpp
+++ b/Engine/Plugins/Runtime/nDisplay/Source/DisplayClusterConfigurator/Private/Views/Details/Media/DisplayClusterConfiguratorICVFXMediaCustomization.cpp
@@ -87,10 +87,7 @@ void FDisplayClusterConfiguratorICVFXMediaCustomization::CustomizeChildren(TShar
 	FDisplayClusterConfiguratorBaseTypeCustomization::CustomizeChildren(InPropertyHandle, InChildBuilder, InCustomizationUtils);
 
 	// Create auto-configure button in the bottom
-	if (SplitTypeValue == EDisplayClusterConfigurationMediaSplitType::UniformTiles)
-	{
-		AddAutoConfigurationButton(InChildBuilder);
-	}
+	AddAutoConfigurationButton(InChildBuilder);
 }
 
 void FDisplayClusterConfiguratorICVFXMediaCustomization::AddAutoConfigurationButton(IDetailChildrenBuilder& InChildBuilder)
@@ -107,7 +104,7 @@ void FDisplayClusterConfiguratorICVFXMediaCustomization::AddAutoConfigurationBut
 						[
 							SNew(STextBlock)
 								.Font(FAppStyle::GetFontStyle(TEXT("PropertyWindow.NormalFont")))
-								.Text(LOCTEXT("AutoConfigureTiledInputButtonTitle", "Auto-Configure"))
+								.Text(LOCTEXT("AutoConfigureInputButtonTitle", "Auto-Configure"))
 						]
 				]
 		];
@@ -116,7 +113,7 @@ void FDisplayClusterConfiguratorICVFXMediaCustomization::AddAutoConfigurationBut
 FReply FDisplayClusterConfiguratorICVFXMediaCustomization::OnAutoConfigureButtonClicked()
 {
 	// Notify tile customizators to re-initialize their media subjects
-	FDisplayClusterConfiguratorMediaUtils::Get().OnTiledMediaAutoConfiguration().Broadcast(EditingObject.Get());
+	FDisplayClusterConfiguratorMediaUtils::Get().OnMediaAutoConfiguration().Broadcast(EditingObject.Get());
 
 	return FReply::Handled();
 }
diff --git a/Engine/Plugins/Runtime/nDisplay/Source/DisplayClusterConfigurator/Private/Views/Details/Media/DisplayClusterConfiguratorICVFXMediaTileCustomization.cpp b/Engine/Plugins/Runtime/nDisplay/Source/DisplayClusterConfigurator/Private/Views/Details/Media/DisplayClusterConfiguratorICVFXMediaTileCustomization.cpp
deleted file mode 100644
diff --git a/Engine/Plugins/Runtime/nDisplay/Source/DisplayClusterConfigurator/Private/Views/Details/Media/DisplayClusterConfiguratorICVFXMediaTileCustomization.h b/Engine/Plugins/Runtime/nDisplay/Source/DisplayClusterConfigurator/Private/Views/Details/Media/DisplayClusterConfiguratorICVFXMediaTileCustomization.h
deleted file mode 100644
diff --git a/Engine/Plugins/Runtime/nDisplay/Source/DisplayClusterConfigurator/Private/Views/Details/Media/DisplayClusterConfiguratorMediaCustomizationBase.cpp b/Engine/Plugins/Runtime/nDisplay/Source/DisplayClusterConfigurator/Private/Views/Details/Media/DisplayClusterConfiguratorMediaCustomizationBase.cpp
new file mode 100644
--- /dev/null
+++ b/Engine/Plugins/Runtime/nDisplay/Source/DisplayClusterConfigurator/Private/Views/Details/Media/DisplayClusterConfiguratorMediaCustomizationBase.cpp
@@ -0,0 +1,346 @@
+// Copyright Epic Games, Inc. All Rights Reserved.
+
+#include "Views/Details/Media/DisplayClusterConfiguratorMediaCustomizationBase.h"
+#include "Views/Details/Media/DisplayClusterConfiguratorMediaCustomizationCVars.h"
+#include "Views/Details/Media/DisplayClusterConfiguratorMediaUtils.h"
+
+#include "Components/DisplayClusterICVFXCameraComponent.h"
+
+#include "DisplayClusterConfigurationTypes.h"
+#include "DisplayClusterConfigurationTypes_Media.h"
+#include "DisplayClusterConfiguratorBlueprintEditor.h"
+#include "DisplayClusterConfiguratorUtils.h"
+#include "DisplayClusterRootActor.h"
+
+#include "IDetailChildrenBuilder.h"
+#include "IPropertyUtilities.h"
+#include "PropertyHandle.h"
+
+
+FDisplayClusterConfiguratorMediaFullFrameCustomizationBase::FDisplayClusterConfiguratorMediaFullFrameCustomizationBase()
+{
+	// Subscribe for auto-configure event
+	FDisplayClusterConfiguratorMediaUtils::Get().OnMediaAutoConfiguration().AddRaw(this, &FDisplayClusterConfiguratorMediaFullFrameCustomizationBase::OnAutoConfigureRequested);
+}
+
+FDisplayClusterConfiguratorMediaFullFrameCustomizationBase::~FDisplayClusterConfiguratorMediaFullFrameCustomizationBase()
+{
+	// Unsubscribe from auto-configure event
+	FDisplayClusterConfiguratorMediaUtils::Get().OnMediaAutoConfiguration().RemoveAll(this);
+}
+
+void FDisplayClusterConfiguratorMediaFullFrameCustomizationBase::CustomizeChildren(TSharedRef<IPropertyHandle> InPropertyHandle, IDetailChildrenBuilder& InChildBuilder, IPropertyTypeCustomizationUtils& InCustomizationUtils)
+{
+	if (CVarMediaAutoInitializationEnabled.GetValueOnGameThread())
+	{
+		// Subscribe for change callbacks
+		if (MediaSubjectHandle->IsValidHandle())
+		{
+			MediaSubjectHandle->SetOnPropertyValueChanged(FSimpleDelegate::CreateSP(this, &FDisplayClusterConfiguratorMediaFullFrameCustomizationBase::OnMediaSubjectChanged));
+		}
+	}
+
+	FDisplayClusterConfiguratorBaseTypeCustomization::CustomizeChildren(InPropertyHandle, InChildBuilder, InCustomizationUtils);
+}
+
+void FDisplayClusterConfiguratorMediaFullFrameCustomizationBase::ModifyMediaSubjectParameters()
+{
+	// Nothing to do if disabled by a CVar
+	if (!CVarMediaAutoInitializationEnabled.GetValueOnGameThread())
+	{
+		return;
+	}
+
+	// Validate object being edited
+	if (!EditingObject.IsValid())
+	{
+		return;
+	}
+
+	UObject* Owner = EditingObject.Get();
+	if (!MediaSubjectHandle->IsValidHandle() || !Owner)
+	{
+		return;
+	}
+
+	// Get media subject
+	UObject* NewMediaSubject = nullptr;
+	if (MediaSubjectHandle->GetValue(NewMediaSubject) != FPropertyAccess::Success)
+	{
+		return;
+	}
+
+	if (!NewMediaSubject)
+	{
+		return;
+	}
+
+	// Find a suitable initializer, and let it process new media data
+	const TArray<IDisplayClusterModularFeatureMediaInitializer*> MediaInitializers = FDisplayClusterConfiguratorMediaUtils::Get().GetMediaInitializers();
+	for (IDisplayClusterModularFeatureMediaInitializer* Initializer : MediaInitializers)
+	{
+		if (Initializer && Initializer->IsMediaSubjectSupported(NewMediaSubject))
+		{
+			if (PerformMediaInitialization(Owner, NewMediaSubject, Initializer))
+			{
+				ModifyBlueprint();
+			}
+
+			break;
+		}
+	}
+}
+
+bool FDisplayClusterConfiguratorMediaFullFrameCustomizationBase::PerformMediaInitialization(UObject* Owner, UObject* MediaSubject, IDisplayClusterModularFeatureMediaInitializer* Initializer)
+{
+	IDisplayClusterModularFeatureMediaInitializer::FMediaSubjectOwnerInfo OwnerInfo;
+	const bool bGotOwnerData = GetOwnerData(Owner, OwnerInfo);
+
+	if (bGotOwnerData && Initializer)
+	{
+		Initializer->InitializeMediaSubjectForFullFrame(MediaSubject, OwnerInfo);
+		return true;
+	}
+
+	return false;
+}
+
+bool FDisplayClusterConfiguratorMediaFullFrameCustomizationBase::GetOwnerData(const UObject* Owner, IDisplayClusterModularFeatureMediaInitializer::FMediaSubjectOwnerInfo& OutOwnerInfo) const
+{
+	if (const UDisplayClusterICVFXCameraComponent* ICVFXCameraComponent = Cast<UDisplayClusterICVFXCameraComponent>(Owner))
+	{
+		return GetOwnerData(ICVFXCameraComponent, OutOwnerInfo);
+	}
+	else if (const UDisplayClusterConfigurationViewport* ViewportCfg = Cast<UDisplayClusterConfigurationViewport>(Owner))
+	{
+		return GetOwnerData(ViewportCfg, OutOwnerInfo);
+	}
+	else if (const UDisplayClusterConfigurationClusterNode* NodeCfg = Cast<UDisplayClusterConfigurationClusterNode>(Owner))
+	{
+		return GetOwnerData(NodeCfg, OutOwnerInfo);
+	}
+	else
+	{
+		checkNoEntry();
+	}
+
+	return false;
+}
+
+bool FDisplayClusterConfiguratorMediaFullFrameCustomizationBase::GetOwnerData(const UDisplayClusterICVFXCameraComponent* ICVFXCameraComponent, IDisplayClusterModularFeatureMediaInitializer::FMediaSubjectOwnerInfo& OutOwnerInfo) const
+{
+	if (const AActor* const OwningActor = GetOwningActor())
+	{
+		// Get all camera components
+		TArray<UDisplayClusterICVFXCameraComponent*> ICVFXCameras;
+		OwningActor->GetComponents(ICVFXCameras);
+
+		const int32 CamerasAmount = ICVFXCameras.Num();
+		if (CamerasAmount > 0)
+		{
+			// Camera sort predicate
+			struct FCameraSortPredicate
+			{
+				bool operator()(UDisplayClusterICVFXCameraComponent& LHS, UDisplayClusterICVFXCameraComponent& RHS) const
+				{
+					return LHS.GetName().Compare(RHS.GetName(), ESearchCase::IgnoreCase) < 0;
+				}
+			};
+
+			// Sort by name to always keep the same ABC-order
+			ICVFXCameras.Sort(FCameraSortPredicate());
+
+			FString OrigCameraName = ICVFXCameraComponent->GetName();
+			OrigCameraName.RemoveFromEnd(TEXT("_GEN_VARIABLE"));
+
+			// Find camera index in the list
+			for (uint8 CameraIdx = 0; CameraIdx < CamerasAmount; ++CameraIdx)
+			{
+				if (ICVFXCameras[CameraIdx]->GetName().Equals(OrigCameraName, ESearchCase::IgnoreCase))
+				{
+					OutOwnerInfo.OwnerType = IDisplayClusterModularFeatureMediaInitializer::FMediaSubjectOwnerInfo::EMediaSubjectOwnerType::ICVFXCamera;
+					OutOwnerInfo.OwnerName = OrigCameraName;
+					OutOwnerInfo.OwnerUniqueIdx = CameraIdx;
+
+					return true;
+				}
+			}
+		}
+	}
+
+	return false;
+}
+
+bool FDisplayClusterConfiguratorMediaFullFrameCustomizationBase::GetOwnerData(const UDisplayClusterConfigurationViewport* ViewportCfg, IDisplayClusterModularFeatureMediaInitializer::FMediaSubjectOwnerInfo& OutOwnerInfo) const
+{
+	// Get cluster node that holds the viewport requested
+	if (const UDisplayClusterConfigurationClusterNode* const ClusterNodeCfg = ViewportCfg->GetTypedOuter<UDisplayClusterConfigurationClusterNode>())
+	{
+		// Find the pair that holds the viewport requested
+		for (const TPair<FString, TObjectPtr<UDisplayClusterConfigurationViewport>>& ViewportIt : ClusterNodeCfg->Viewports)
+		{
+			if (ViewportIt.Value == ViewportCfg)
+			{
+				// Get ABC-sorted list of the viewport IDs
+				TArray<FString> ViewportIds;
+				ClusterNodeCfg->Viewports.GenerateKeyArray(ViewportIds);
+				ViewportIds.Sort();
+
+				// Find viewport index
+				for (int32 ViewportIdx = 0; ViewportIdx < ViewportIds.Num(); ++ViewportIdx)
+				{
+					if (ViewportIds[ViewportIdx].Equals(ViewportIt.Key, ESearchCase::IgnoreCase))
+					{
+						checkSlow(ViewportIdx < TNumericLimits<uint8>::Max());
+
+						OutOwnerInfo.OwnerType = IDisplayClusterModularFeatureMediaInitializer::FMediaSubjectOwnerInfo::EMediaSubjectOwnerType::Viewport;
+						OutOwnerInfo.OwnerName = ViewportIt.Key;
+						OutOwnerInfo.OwnerUniqueIdx = static_cast<uint8>(ViewportIdx);
+
+						// Here we leverage cluster node data provider to get unique cluster node index
+						{
+							IDisplayClusterModularFeatureMediaInitializer::FMediaSubjectOwnerInfo TempNodeInfo;
+							if (GetOwnerData(ClusterNodeCfg, TempNodeInfo))
+							{
+								OutOwnerInfo.ClusterNodeUniqueIdx = TempNodeInfo.ClusterNodeUniqueIdx;
+							}
+						}
+
+						return true;
+					}
+				}
+			}
+		}
+	}
+
+	return false;
+}
+
+bool FDisplayClusterConfiguratorMediaFullFrameCustomizationBase::GetOwnerData(const UDisplayClusterConfigurationClusterNode* NodeCfg, IDisplayClusterModularFeatureMediaInitializer::FMediaSubjectOwnerInfo& OutOwnerInfo) const
+{
+	// Get cluster object holding the node requested
+	if (const UDisplayClusterConfigurationCluster* const ClusterCfg = NodeCfg->GetTypedOuter<UDisplayClusterConfigurationCluster>())
+	{
+		// Find the pair that holds the node requested
+		for (const TPair<FString, TObjectPtr<UDisplayClusterConfigurationClusterNode>>& NodeIt : ClusterCfg->Nodes)
+		{
+			if (NodeIt.Value == NodeCfg)
+			{
+				// Get ABC-sorted list of the cluster nodes
+				TArray<FString> NodeIds;
+				ClusterCfg->Nodes.GenerateKeyArray(NodeIds);
+				NodeIds.Sort();
+
+				// Find node index
+				for (int32 NodeIdx = 0; NodeIdx < ClusterCfg->Nodes.Num(); ++NodeIdx)
+				{
+					if (NodeIds[NodeIdx].Equals(NodeIt.Key, ESearchCase::IgnoreCase))
+					{
+						checkSlow(NodeIdx < TNumericLimits<uint8>::Max());
+
+						OutOwnerInfo.OwnerType = IDisplayClusterModularFeatureMediaInitializer::FMediaSubjectOwnerInfo::EMediaSubjectOwnerType::Backbuffer;
+						OutOwnerInfo.OwnerName = NodeIt.Key;
+						OutOwnerInfo.OwnerUniqueIdx = static_cast<uint8>(NodeIdx);
+						OutOwnerInfo.ClusterNodeUniqueIdx = static_cast<uint8>(NodeIdx);
+
+						return true;
+					}
+				}
+			}
+		}
+	}
+
+	return false;
+}
+
+AActor* FDisplayClusterConfiguratorMediaFullFrameCustomizationBase::GetOwningActor() const
+{
+	if (UDisplayClusterICVFXCameraComponent* ICVFXCameraComponent = Cast<UDisplayClusterICVFXCameraComponent>(EditingObject))
+	{
+		if (AActor* Actor1 = ICVFXCameraComponent->GetOwner())
+		{
+			return Actor1;
+		}
+		else if (AActor* Actor2 = Cast<AActor>(FindRootActor()))
+		{
+			return Actor2;
+		}
+		else if (FDisplayClusterConfiguratorBlueprintEditor* BlueprintEditor = FDisplayClusterConfiguratorUtils::GetBlueprintEditorForObject(ICVFXCameraComponent))
+		{
+			return BlueprintEditor->GetPreviewActor();
+		}
+	}
+
+	return nullptr;
+}
+
+void FDisplayClusterConfiguratorMediaFullFrameCustomizationBase::OnMediaSubjectChanged()
+{
+	ModifyMediaSubjectParameters();
+}
+
+void FDisplayClusterConfiguratorMediaFullFrameCustomizationBase::OnAutoConfigureRequested(UObject* InEditingObject)
+{
+	if (EditingObject == InEditingObject)
+	{
+		ModifyMediaSubjectParameters();
+	}
+}
+
+
+//
+// Common customization
+//
+
+void FDisplayClusterConfiguratorMediaTileCustomizationBase::CustomizeChildren(TSharedRef<IPropertyHandle> InPropertyHandle, IDetailChildrenBuilder& InChildBuilder, IPropertyTypeCustomizationUtils& InCustomizationUtils)
+{
+	if (CVarMediaAutoInitializationEnabled.GetValueOnGameThread())
+	{
+		// Subscribe for change callbacks
+		if (TilePosHandle->IsValidHandle())
+		{
+			TilePosHandle->SetOnPropertyValueChanged(FSimpleDelegate::CreateSP(this, &FDisplayClusterConfiguratorMediaTileCustomizationBase::OnTilePositionChanged));
+			TilePosHandle->SetOnChildPropertyValueChanged(FSimpleDelegate::CreateSP(this, &FDisplayClusterConfiguratorMediaTileCustomizationBase::OnTilePositionChanged));
+		}
+	}
+
+	FDisplayClusterConfiguratorMediaFullFrameCustomizationBase::CustomizeChildren(InPropertyHandle, InChildBuilder, InCustomizationUtils);
+}
+
+bool FDisplayClusterConfiguratorMediaTileCustomizationBase::PerformMediaInitialization(UObject* Owner, UObject* MediaSubject, IDisplayClusterModularFeatureMediaInitializer* Initializer)
+{
+	IDisplayClusterModularFeatureMediaInitializer::FMediaSubjectOwnerInfo OwnerInfo;
+	const bool bGotOwnerData = GetOwnerData(Owner, OwnerInfo);
+
+	if (bGotOwnerData)
+	{
+		// Get input tile position
+		const FIntPoint TilePos = GetEditedTilePos();
+		if (TilePos.X != FIntPoint::NoneValue.X && TilePos.Y != FIntPoint::NoneValue.Y)
+		{
+			if (Initializer)
+			{
+				Initializer->InitializeMediaSubjectForTile(MediaSubject, OwnerInfo, TilePos);
+				return true;
+			}
+		}
+	}
+
+	return false;
+}
+
+FIntPoint FDisplayClusterConfiguratorMediaTileCustomizationBase::GetEditedTilePos() const
+{
+	FIntPoint TilePos = FIntPoint::NoneValue;
+
+	// Children 0 and 1 correspond to X and Y subproperties of FIntPoint
+	TilePosHandle->GetChildHandle(0)->GetValue(TilePos.X);
+	TilePosHandle->GetChildHandle(1)->GetValue(TilePos.Y);
+
+	return TilePos;
+}
+
+void FDisplayClusterConfiguratorMediaTileCustomizationBase::OnTilePositionChanged()
+{
+	ModifyMediaSubjectParameters();
+}
diff --git a/Engine/Plugins/Runtime/nDisplay/Source/DisplayClusterConfigurator/Private/Views/Details/Media/DisplayClusterConfiguratorMediaCustomizationBase.h b/Engine/Plugins/Runtime/nDisplay/Source/DisplayClusterConfigurator/Private/Views/Details/Media/DisplayClusterConfiguratorMediaCustomizationBase.h
new file mode 100644
--- /dev/null
+++ b/Engine/Plugins/Runtime/nDisplay/Source/DisplayClusterConfigurator/Private/Views/Details/Media/DisplayClusterConfiguratorMediaCustomizationBase.h
@@ -0,0 +1,107 @@
+// Copyright Epic Games, Inc. All Rights Reserved.
+
+#pragma once
+
+#include "CoreMinimal.h"
+#include "Views/Details/DisplayClusterConfiguratorBaseTypeCustomization.h"
+
+#include "IDisplayClusterModularFeatureMediaInitializer.h"
+
+class UDisplayClusterConfigurationClusterNode;
+class UDisplayClusterConfigurationViewport;
+class UDisplayClusterICVFXCameraComponent;
+
+
+/**
+ * Base class for full frame media input & output customization
+ */
+class FDisplayClusterConfiguratorMediaFullFrameCustomizationBase
+	: public FDisplayClusterConfiguratorBaseTypeCustomization
+{
+public:
+	FDisplayClusterConfiguratorMediaFullFrameCustomizationBase();
+	~FDisplayClusterConfiguratorMediaFullFrameCustomizationBase();
+
+protected:
+
+	//~ Begin IPropertyTypeCustomization
+	virtual void CustomizeChildren(TSharedRef<IPropertyHandle> PropertyHandle, IDetailChildrenBuilder& ChildBuilder, IPropertyTypeCustomizationUtils& CustomizationUtils) override;
+	//~ End IPropertyTypeCustomization
+
+protected:
+
+	/** Entry point to modify media subjects. */
+	void ModifyMediaSubjectParameters();
+
+	/** Virtual media initialization depending on media type (full-frame, split, etc.) */
+	virtual bool PerformMediaInitialization(UObject* Owner, UObject* MediaSubject, IDisplayClusterModularFeatureMediaInitializer* Initializer);
+
+	/** Returns the actor that owns the object being edited. */
+	AActor* GetOwningActor() const;
+
+	/** Returns the name of media subject's owner. */
+	bool GetOwnerData(const UObject* Owner, IDisplayClusterModularFeatureMediaInitializer::FMediaSubjectOwnerInfo& OutOwnerInfo) const;
+
+private:
+
+	/** Helper data provider for ICVFX cameras. */
+	bool GetOwnerData(const UDisplayClusterICVFXCameraComponent* ICVFXCameraComponent, IDisplayClusterModularFeatureMediaInitializer::FMediaSubjectOwnerInfo& OutOwnerInfo) const;
+
+	/** Helper data provider for viewports. */
+	bool GetOwnerData(const UDisplayClusterConfigurationViewport* ViewportCfg, IDisplayClusterModularFeatureMediaInitializer::FMediaSubjectOwnerInfo& OutOwnerInfo) const;
+
+	/** Helper data provider for cluster nodes. */
+	bool GetOwnerData(const UDisplayClusterConfigurationClusterNode* NodeCfg, IDisplayClusterModularFeatureMediaInitializer::FMediaSubjectOwnerInfo& OutOwnerInfo) const;
+
+private:
+
+	/** Handles media source/output change callbacks */
+	void OnMediaSubjectChanged();
+
+	/** Auto-configuration requests handler */
+	void OnAutoConfigureRequested(UObject* EditingObject);
+
+protected:
+
+	/** MediaSource or MediaOutput property handle, depending on the child implementation. */
+	TSharedPtr<IPropertyHandle> MediaSubjectHandle;
+};
+
+
+
+/**
+ * Base class for tiled media input & output customization
+ */
+class FDisplayClusterConfiguratorMediaTileCustomizationBase
+	: public FDisplayClusterConfiguratorMediaFullFrameCustomizationBase
+{
+protected:
+
+	//~ Begin IPropertyTypeCustomization
+	virtual void CustomizeChildren(TSharedRef<IPropertyHandle> PropertyHandle, IDetailChildrenBuilder& ChildBuilder, IPropertyTypeCustomizationUtils& CustomizationUtils) override;
+	//~ End IPropertyTypeCustomization
+
+protected:
+
+	//~ Begin FDisplayClusterConfiguratorMediaFullFrameCustomizationBase
+
+	/** Virtual media initialization depending on media type (full-frame, split, etc.) */
+	virtual bool PerformMediaInitialization(UObject* Owner, UObject* MediaSubject, IDisplayClusterModularFeatureMediaInitializer* Initializer) override;
+
+	//~ End FDisplayClusterConfiguratorMediaFullFrameCustomizationBase
+
+private:
+
+	/** Returns tile position currently set. */
+	FIntPoint GetEditedTilePos() const;
+
+private:
+
+	/** Handles tile position change callbacks */
+	void OnTilePositionChanged();
+
+protected:
+
+	/** Tile position property handle. */
+	TSharedPtr<IPropertyHandle> TilePosHandle;
+};
diff --git a/Engine/Plugins/Runtime/nDisplay/Source/DisplayClusterConfigurator/Private/Views/Details/Media/DisplayClusterConfiguratorMediaCustomizationCVars.cpp b/Engine/Plugins/Runtime/nDisplay/Source/DisplayClusterConfigurator/Private/Views/Details/Media/DisplayClusterConfiguratorMediaCustomizationCVars.cpp
new file mode 100644
--- /dev/null
+++ b/Engine/Plugins/Runtime/nDisplay/Source/DisplayClusterConfigurator/Private/Views/Details/Media/DisplayClusterConfiguratorMediaCustomizationCVars.cpp
@@ -0,0 +1,13 @@
+// Copyright Epic Games, Inc. All Rights Reserved.
+
+#include "Views\Details\Media\DisplayClusterConfiguratorMediaCustomizationCVars.h"
+
+
+TAutoConsoleVariable<bool> CVarMediaAutoInitializationEnabled(
+	TEXT("nDisplay.Media.MediaAutoInitialization"),
+	true,
+	TEXT("nDisplay tiled media auto-configuration\n")
+	TEXT("0 : Disabled\n")
+	TEXT("1 : Enabled\n"),
+	ECVF_Default
+);
diff --git a/Engine/Plugins/Runtime/nDisplay/Source/DisplayClusterConfigurator/Private/Views/Details/Media/DisplayClusterConfiguratorMediaCustomizationCVars.h b/Engine/Plugins/Runtime/nDisplay/Source/DisplayClusterConfigurator/Private/Views/Details/Media/DisplayClusterConfiguratorMediaCustomizationCVars.h
new file mode 100644
--- /dev/null
+++ b/Engine/Plugins/Runtime/nDisplay/Source/DisplayClusterConfigurator/Private/Views/Details/Media/DisplayClusterConfiguratorMediaCustomizationCVars.h
@@ -0,0 +1,9 @@
+// Copyright Epic Games, Inc. All Rights Reserved.
+
+#pragma once
+
+#include "CoreMinimal.h"
+#include "HAL/IConsoleManager.h"
+
+// Enable/disable live auto-configuration. Does not affect "auto-configure" button.
+extern TAutoConsoleVariable<bool> CVarMediaAutoInitializationEnabled;
diff --git a/Engine/Plugins/Runtime/nDisplay/Source/DisplayClusterConfigurator/Private/Views/Details/Media/DisplayClusterConfiguratorMediaFullFrameCustomization.cpp b/Engine/Plugins/Runtime/nDisplay/Source/DisplayClusterConfigurator/Private/Views/Details/Media/DisplayClusterConfiguratorMediaFullFrameCustomization.cpp
new file mode 100644
--- /dev/null
+++ b/Engine/Plugins/Runtime/nDisplay/Source/DisplayClusterConfigurator/Private/Views/Details/Media/DisplayClusterConfiguratorMediaFullFrameCustomization.cpp
@@ -0,0 +1,34 @@
+// Copyright Epic Games, Inc. All Rights Reserved.
+
+#include "Views/Details/Media/DisplayClusterConfiguratorMediaFullFrameCustomization.h"
+
+#include "DisplayClusterConfigurationTypes_Media.h"
+
+#include "IDetailChildrenBuilder.h"
+#include "IPropertyUtilities.h"
+#include "PropertyHandle.h"
+
+
+//
+// Input tiles
+//
+void FDisplayClusterConfiguratorMediaFullFrameInputCustomization::CustomizeChildren(TSharedRef<IPropertyHandle> InPropertyHandle, IDetailChildrenBuilder& InChildBuilder, IPropertyTypeCustomizationUtils& InCustomizationUtils)
+{
+	// MediaSource property
+	MediaSubjectHandle = GET_CHILD_HANDLE(FDisplayClusterConfigurationMediaInput, MediaSource);
+	check(MediaSubjectHandle->IsValidHandle());
+
+	FDisplayClusterConfiguratorMediaFullFrameCustomizationBase::CustomizeChildren(InPropertyHandle, InChildBuilder, InCustomizationUtils);
+}
+
+//
+// Output tiles
+//
+void FDisplayClusterConfiguratorMediaFullFrameOutputCustomization::CustomizeChildren(TSharedRef<IPropertyHandle> InPropertyHandle, IDetailChildrenBuilder& InChildBuilder, IPropertyTypeCustomizationUtils& InCustomizationUtils)
+{
+	// MediaOutput property
+	MediaSubjectHandle = GET_CHILD_HANDLE(FDisplayClusterConfigurationMediaOutput, MediaOutput);
+	check(MediaSubjectHandle->IsValidHandle());
+
+	FDisplayClusterConfiguratorMediaFullFrameCustomizationBase::CustomizeChildren(InPropertyHandle, InChildBuilder, InCustomizationUtils);
+}
diff --git a/Engine/Plugins/Runtime/nDisplay/Source/DisplayClusterConfigurator/Private/Views/Details/Media/DisplayClusterConfiguratorMediaFullFrameCustomization.h b/Engine/Plugins/Runtime/nDisplay/Source/DisplayClusterConfigurator/Private/Views/Details/Media/DisplayClusterConfiguratorMediaFullFrameCustomization.h
new file mode 100644
--- /dev/null
+++ b/Engine/Plugins/Runtime/nDisplay/Source/DisplayClusterConfigurator/Private/Views/Details/Media/DisplayClusterConfiguratorMediaFullFrameCustomization.h
@@ -0,0 +1,44 @@
+// Copyright Epic Games, Inc. All Rights Reserved.
+
+#pragma once
+
+#include "CoreMinimal.h"
+#include "Views/Details/Media/DisplayClusterConfiguratorMediaCustomizationBase.h"
+
+
+/**
+ * Details panel customization for FDisplayClusterConfigurationMediaUniformTileInput struct (input tiles).
+ */
+class FDisplayClusterConfiguratorMediaFullFrameInputCustomization
+	: public FDisplayClusterConfiguratorMediaFullFrameCustomizationBase
+{
+public:
+	static TSharedRef<IPropertyTypeCustomization> MakeInstance()
+	{
+		return MakeShared<FDisplayClusterConfiguratorMediaFullFrameInputCustomization>();
+	}
+
+protected:
+	//~ Begin IPropertyTypeCustomization
+	virtual void CustomizeChildren(TSharedRef<IPropertyHandle> PropertyHandle, IDetailChildrenBuilder& ChildBuilder, IPropertyTypeCustomizationUtils& CustomizationUtils) override;
+	//~ End IPropertyTypeCustomization
+};
+
+
+/**
+ * Details panel customization for FDisplayClusterConfigurationMediaUniformTileOutput struct (output tiles).
+ */
+class FDisplayClusterConfiguratorMediaFullFrameOutputCustomization
+	: public FDisplayClusterConfiguratorMediaFullFrameCustomizationBase
+{
+public:
+	static TSharedRef<IPropertyTypeCustomization> MakeInstance()
+	{
+		return MakeShared<FDisplayClusterConfiguratorMediaFullFrameOutputCustomization>();
+	}
+
+protected:
+	//~ Begin IPropertyTypeCustomization
+	virtual void CustomizeChildren(TSharedRef<IPropertyHandle> PropertyHandle, IDetailChildrenBuilder& ChildBuilder, IPropertyTypeCustomizationUtils& CustomizationUtils) override;
+	//~ End IPropertyTypeCustomization
+};
diff --git a/Engine/Plugins/Runtime/nDisplay/Source/DisplayClusterConfigurator/Private/Views/Details/Media/DisplayClusterConfiguratorMediaTileCustomization.cpp b/Engine/Plugins/Runtime/nDisplay/Source/DisplayClusterConfigurator/Private/Views/Details/Media/DisplayClusterConfiguratorMediaTileCustomization.cpp
new file mode 100644
--- /dev/null
+++ b/Engine/Plugins/Runtime/nDisplay/Source/DisplayClusterConfigurator/Private/Views/Details/Media/DisplayClusterConfiguratorMediaTileCustomization.cpp
@@ -0,0 +1,43 @@
+// Copyright Epic Games, Inc. All Rights Reserved.
+
+#include "Views/Details/Media/DisplayClusterConfiguratorMediaTileCustomization.h"
+
+#include "DisplayClusterConfigurationTypes_Media.h"
+
+#include "IDetailChildrenBuilder.h"
+#include "IPropertyUtilities.h"
+#include "PropertyHandle.h"
+
+
+//
+// Input tiles
+//
+void FDisplayClusterConfiguratorMediaInputTileCustomization::CustomizeChildren(TSharedRef<IPropertyHandle> InPropertyHandle, IDetailChildrenBuilder& InChildBuilder, IPropertyTypeCustomizationUtils& InCustomizationUtils)
+{
+	// MediaSource property
+	MediaSubjectHandle = GET_CHILD_HANDLE(FDisplayClusterConfigurationMediaUniformTileInput, MediaSource);
+	check(MediaSubjectHandle->IsValidHandle());
+
+	// Tile position property
+	TilePosHandle = GET_CHILD_HANDLE(FDisplayClusterConfigurationMediaUniformTileInput, Position);
+	check(TilePosHandle->IsValidHandle());
+
+	FDisplayClusterConfiguratorMediaTileCustomizationBase::CustomizeChildren(InPropertyHandle, InChildBuilder, InCustomizationUtils);
+}
+
+
+//
+// Output tiles
+//
+void FDisplayClusterConfiguratorMediaOutputTileCustomization::CustomizeChildren(TSharedRef<IPropertyHandle> InPropertyHandle, IDetailChildrenBuilder& InChildBuilder, IPropertyTypeCustomizationUtils& InCustomizationUtils)
+{
+	// MediaOutput property
+	MediaSubjectHandle = GET_CHILD_HANDLE(FDisplayClusterConfigurationMediaUniformTileOutput, MediaOutput);
+	check(MediaSubjectHandle->IsValidHandle());
+
+	// Tile position property
+	TilePosHandle = GET_CHILD_HANDLE(FDisplayClusterConfigurationMediaUniformTileOutput, Position);
+	check(TilePosHandle->IsValidHandle());
+
+	FDisplayClusterConfiguratorMediaTileCustomizationBase::CustomizeChildren(InPropertyHandle, InChildBuilder, InCustomizationUtils);
+}
diff --git a/Engine/Plugins/Runtime/nDisplay/Source/DisplayClusterConfigurator/Private/Views/Details/Media/DisplayClusterConfiguratorMediaTileCustomization.h b/Engine/Plugins/Runtime/nDisplay/Source/DisplayClusterConfigurator/Private/Views/Details/Media/DisplayClusterConfiguratorMediaTileCustomization.h
new file mode 100644
--- /dev/null
+++ b/Engine/Plugins/Runtime/nDisplay/Source/DisplayClusterConfigurator/Private/Views/Details/Media/DisplayClusterConfiguratorMediaTileCustomization.h
@@ -0,0 +1,44 @@
+// Copyright Epic Games, Inc. All Rights Reserved.
+
+#pragma once
+
+#include "CoreMinimal.h"
+#include "Views/Details/Media/DisplayClusterConfiguratorMediaCustomizationBase.h"
+
+
+/**
+ * Details panel customization for FDisplayClusterConfigurationMediaUniformTileInput struct (input tiles).
+ */
+class FDisplayClusterConfiguratorMediaInputTileCustomization
+	: public FDisplayClusterConfiguratorMediaTileCustomizationBase
+{
+public:
+	static TSharedRef<IPropertyTypeCustomization> MakeInstance()
+	{
+		return MakeShared<FDisplayClusterConfiguratorMediaInputTileCustomization>();
+	}
+
+protected:
+	//~ Begin IPropertyTypeCustomization
+	virtual void CustomizeChildren(TSharedRef<IPropertyHandle> PropertyHandle, IDetailChildrenBuilder& ChildBuilder, IPropertyTypeCustomizationUtils& CustomizationUtils) override;
+	//~ End IPropertyTypeCustomization
+};
+
+
+/**
+ * Details panel customization for FDisplayClusterConfigurationMediaUniformTileOutput struct (output tiles).
+ */
+class FDisplayClusterConfiguratorMediaOutputTileCustomization
+	: public FDisplayClusterConfiguratorMediaTileCustomizationBase
+{
+public:
+	static TSharedRef<IPropertyTypeCustomization> MakeInstance()
+	{
+		return MakeShared<FDisplayClusterConfiguratorMediaOutputTileCustomization>();
+	}
+
+protected:
+	//~ Begin IPropertyTypeCustomization
+	virtual void CustomizeChildren(TSharedRef<IPropertyHandle> PropertyHandle, IDetailChildrenBuilder& ChildBuilder, IPropertyTypeCustomizationUtils& CustomizationUtils) override;
+	//~ End IPropertyTypeCustomization
+};
diff --git a/Engine/Plugins/Runtime/nDisplay/Source/DisplayClusterConfigurator/Private/Views/Details/Media/DisplayClusterConfiguratorMediaUtils.cpp b/Engine/Plugins/Runtime/nDisplay/Source/DisplayClusterConfigurator/Private/Views/Details/Media/DisplayClusterConfiguratorMediaUtils.cpp
--- a/Engine/Plugins/Runtime/nDisplay/Source/DisplayClusterConfigurator/Private/Views/Details/Media/DisplayClusterConfiguratorMediaUtils.cpp
+++ b/Engine/Plugins/Runtime/nDisplay/Source/DisplayClusterConfigurator/Private/Views/Details/Media/DisplayClusterConfiguratorMediaUtils.cpp
@@ -2,9 +2,21 @@
 
 #include "Views/Details/Media/DisplayClusterConfiguratorMediaUtils.h"
 
+#include "IDisplayClusterModularFeatureMediaInitializer.h"
+#include "Features/IModularFeatures.h"
+
 
 FDisplayClusterConfiguratorMediaUtils& FDisplayClusterConfiguratorMediaUtils::Get()
 {
 	static FDisplayClusterConfiguratorMediaUtils Instance;
 	return Instance;
 }
+
+FDisplayClusterConfiguratorMediaUtils::FDisplayClusterConfiguratorMediaUtils()
+{
+	// Obtain media initializer modular features once. We don't expect any of them to be loaded dynamically in runtime.
+	IModularFeatures& ModularFeatures = IModularFeatures::Get();
+	ModularFeatures.LockModularFeatureList();
+	MediaInitializers = ModularFeatures.GetModularFeatureImplementations<IDisplayClusterModularFeatureMediaInitializer>(IDisplayClusterModularFeatureMediaInitializer::ModularFeatureName);
+	ModularFeatures.UnlockModularFeatureList();
+}
diff --git a/Engine/Plugins/Runtime/nDisplay/Source/DisplayClusterConfigurator/Private/Views/Details/Media/DisplayClusterConfiguratorMediaUtils.h b/Engine/Plugins/Runtime/nDisplay/Source/DisplayClusterConfigurator/Private/Views/Details/Media/DisplayClusterConfiguratorMediaUtils.h
--- a/Engine/Plugins/Runtime/nDisplay/Source/DisplayClusterConfigurator/Private/Views/Details/Media/DisplayClusterConfiguratorMediaUtils.h
+++ b/Engine/Plugins/Runtime/nDisplay/Source/DisplayClusterConfigurator/Private/Views/Details/Media/DisplayClusterConfiguratorMediaUtils.h
@@ -4,12 +4,16 @@
 
 #include "CoreMinimal.h"
 
+class IDisplayClusterModularFeatureMediaInitializer;
 
 /**
  * Media customization utilities
  */
 class FDisplayClusterConfiguratorMediaUtils
 {
+protected:
+	FDisplayClusterConfiguratorMediaUtils();
+
 public:
 
 	/** Singleton access */
@@ -17,15 +21,24 @@ public:
 
 public:
 
+	/** Returns media initializers */
+	const TArray<IDisplayClusterModularFeatureMediaInitializer*>& GetMediaInitializers() const
+	{
+		return MediaInitializers;
+	}
+
 	/** Tiled media auto-configuration event */
-	DECLARE_EVENT_OneParam(FDisplayClusterConfiguratorMediaUtils, FDisplayClusterConfiguratorTiledMediaAutoConfigurationRequestedEvent, UObject*);
-	FDisplayClusterConfiguratorTiledMediaAutoConfigurationRequestedEvent& OnTiledMediaAutoConfiguration()
+	DECLARE_EVENT_OneParam(FDisplayClusterConfiguratorMediaUtils, FDisplayClusterConfiguratorMediaAutoConfigurationRequestedEvent, UObject*);
+	FDisplayClusterConfiguratorMediaAutoConfigurationRequestedEvent& OnMediaAutoConfiguration()
 	{
-		return TiledMediaAutoConfigurationEvent;
+		return MediaAutoConfigurationEvent;
 	}
 
 private:
 
+	/** Media initializers (modular features) available. */
+	TArray<IDisplayClusterModularFeatureMediaInitializer*> MediaInitializers;
+
 	/** Tiled media auto-configuration event */
-	FDisplayClusterConfiguratorTiledMediaAutoConfigurationRequestedEvent TiledMediaAutoConfigurationEvent;
+	FDisplayClusterConfiguratorMediaAutoConfigurationRequestedEvent MediaAutoConfigurationEvent;
 };
diff --git a/Engine/Plugins/Runtime/nDisplay/Source/SharedMemoryMediaEditor/Private/ModularFeatures/SharedMemoryMediaInitializerFeature.cpp b/Engine/Plugins/Runtime/nDisplay/Source/SharedMemoryMediaEditor/Private/ModularFeatures/SharedMemoryMediaInitializerFeature.cpp
--- a/Engine/Plugins/Runtime/nDisplay/Source/SharedMemoryMediaEditor/Private/ModularFeatures/SharedMemoryMediaInitializerFeature.cpp
+++ b/Engine/Plugins/Runtime/nDisplay/Source/SharedMemoryMediaEditor/Private/ModularFeatures/SharedMemoryMediaInitializerFeature.cpp
@@ -18,26 +18,58 @@ bool FSharedMemoryMediaInitializerFeature::IsMediaSubjectSupported(const UObject
 	return false;
 }
 
-void FSharedMemoryMediaInitializerFeature::InitializeMediaSubjectForTile(UObject* MediaSubject, const FString& OwnerName, uint8 OwnerUniqueIdx, const FIntPoint& TilePos)
+static constexpr const TCHAR* GetMediaPrefix(const IDisplayClusterModularFeatureMediaInitializer::FMediaSubjectOwnerInfo::EMediaSubjectOwnerType OwnerType)
 {
-	checkSlow(MediaSubject);
+	switch (OwnerType)
+	{
+	case IDisplayClusterModularFeatureMediaInitializer::FMediaSubjectOwnerInfo::EMediaSubjectOwnerType::ICVFXCamera:
+		return TEXT("icam");
+
+	case IDisplayClusterModularFeatureMediaInitializer::FMediaSubjectOwnerInfo::EMediaSubjectOwnerType::Viewport:
+		return TEXT("vp");
+
+	case IDisplayClusterModularFeatureMediaInitializer::FMediaSubjectOwnerInfo::EMediaSubjectOwnerType::Backbuffer:
+		return TEXT("node");
+
+	default:
+		return TEXT("unknown");
+	}
+}
+
+void FSharedMemoryMediaInitializerFeature::InitializeMediaSubjectForTile(UObject* MediaSubject, const FMediaSubjectOwnerInfo& OnwerInfo, const FIntPoint& TilePos)
+{
+	const FString UniqueName = FString::Printf(TEXT("%s@%s_tile_%d:%d"), GetMediaPrefix(OnwerInfo.OwnerType), *OnwerInfo.OwnerName, TilePos.X, TilePos.Y);
 
 	if (USharedMemoryMediaSource* SMMediaSource = Cast<USharedMemoryMediaSource>(MediaSubject))
 	{
+		SMMediaSource->UniqueName   = UniqueName;
 		SMMediaSource->bZeroLatency = true;
 		SMMediaSource->Mode         = ESharedMemoryMediaSourceMode::Framelocked;
-		SMMediaSource->UniqueName   = GenerateUniqueName(OwnerName, TilePos);
 	}
 	else if (USharedMemoryMediaOutput* SMMediaOutput = Cast<USharedMemoryMediaOutput>(MediaSubject))
 	{
-		SMMediaOutput->UniqueName   = GenerateUniqueName(OwnerName, TilePos);
+		SMMediaOutput->UniqueName   = UniqueName;
 		SMMediaOutput->bInvertAlpha = true;
 		SMMediaOutput->bCrossGpu    = true;
 		SMMediaOutput->NumberOfTextureBuffers = 4;
 	}
 }
 
-FString FSharedMemoryMediaInitializerFeature::GenerateUniqueName(const FString& OwnerName, const FIntPoint& TilePos)
+void FSharedMemoryMediaInitializerFeature::InitializeMediaSubjectForFullFrame(UObject* MediaSubject, const FMediaSubjectOwnerInfo& OnwerInfo)
 {
-	return FString::Printf(TEXT("%s_tile_%d:%d"), *OwnerName, TilePos.X, TilePos.Y);
+	const FString UniqueName = FString::Printf(TEXT("%s@%s"), GetMediaPrefix(OnwerInfo.OwnerType), *OnwerInfo.OwnerName);
+
+	if (USharedMemoryMediaSource* SMMediaSource = Cast<USharedMemoryMediaSource>(MediaSubject))
+	{
+		SMMediaSource->UniqueName   = UniqueName;
+		SMMediaSource->bZeroLatency = true;
+		SMMediaSource->Mode         = ESharedMemoryMediaSourceMode::Framelocked;
+	}
+	else if (USharedMemoryMediaOutput* SMMediaOutput = Cast<USharedMemoryMediaOutput>(MediaSubject))
+	{
+		SMMediaOutput->UniqueName   = UniqueName;
+		SMMediaOutput->bInvertAlpha = true;
+		SMMediaOutput->bCrossGpu    = true;
+		SMMediaOutput->NumberOfTextureBuffers = 4;
+	}
 }
diff --git a/Engine/Plugins/Runtime/nDisplay/Source/SharedMemoryMediaEditor/Private/ModularFeatures/SharedMemoryMediaInitializerFeature.h b/Engine/Plugins/Runtime/nDisplay/Source/SharedMemoryMediaEditor/Private/ModularFeatures/SharedMemoryMediaInitializerFeature.h
--- a/Engine/Plugins/Runtime/nDisplay/Source/SharedMemoryMediaEditor/Private/ModularFeatures/SharedMemoryMediaInitializerFeature.h
+++ b/Engine/Plugins/Runtime/nDisplay/Source/SharedMemoryMediaEditor/Private/ModularFeatures/SharedMemoryMediaInitializerFeature.h
@@ -16,11 +16,7 @@ public:
 
 	//~ Begin IDisplayClusterModularFeatureMediaInitializer
 	virtual bool IsMediaSubjectSupported(const UObject* MediaSubject) override;
-	virtual void InitializeMediaSubjectForTile(UObject* MediaSubject, const FString& OwnerName, uint8 OwnerUniqueIdx, const FIntPoint& TilePos) override;
+	virtual void InitializeMediaSubjectForTile(UObject* MediaSubject, const FMediaSubjectOwnerInfo& OnwerInfo, const FIntPoint& TilePos) override;
+	virtual void InitializeMediaSubjectForFullFrame(UObject* MediaSubject, const FMediaSubjectOwnerInfo& OnwerInfo) override;
 	//~ End IDisplayClusterModularFeatureMediaInitializer
-
-private:
-
-	/** Generates unique name for SMM subjects */
-	FString GenerateUniqueName(const FString& OwnerName, const FIntPoint& TilePos);
 };
diff --git a/Engine/Plugins/Runtime/nDisplayModularFeatures/Source/DisplayClusterModularFeaturesEditor/Public/IDisplayClusterModularFeatureMediaInitializer.h b/Engine/Plugins/Runtime/nDisplayModularFeatures/Source/DisplayClusterModularFeaturesEditor/Public/IDisplayClusterModularFeatureMediaInitializer.h
--- a/Engine/Plugins/Runtime/nDisplayModularFeatures/Source/DisplayClusterModularFeaturesEditor/Public/IDisplayClusterModularFeatureMediaInitializer.h
+++ b/Engine/Plugins/Runtime/nDisplayModularFeatures/Source/DisplayClusterModularFeaturesEditor/Public/IDisplayClusterModularFeatureMediaInitializer.h
@@ -17,6 +17,43 @@ public:
 	static const FName ModularFeatureName;
 
 public:
+
+	/**
+	 * Container to carrier the info about media subject's owner
+	 */
+	struct FMediaSubjectOwnerInfo
+	{
+		/**
+		 * Type of the media subject's owner
+		 */
+		enum class EMediaSubjectOwnerType : uint8
+		{
+			ICVFXCamera = 0,
+			Viewport,
+			Backbuffer
+		};
+
+
+		/** Owner name (ICVFX camera component name, viewport or node name) */
+		FString OwnerName;
+
+		/** Owner type (ICVFX camera component name, viewport or node name) */
+		EMediaSubjectOwnerType OwnerType;
+
+		/** Optional unique index of the cluster node holding the owner object */
+		TOptional<uint8> ClusterNodeUniqueIdx = 0;
+
+		/** 
+		 * Unique index of the owner
+		 *   Camera     - within a config
+		 *   Viewport   - within a cluster node
+		 *   Backbuffer - within a config
+		 */
+		uint8 OwnerUniqueIdx = 0;
+	};
+
+public:
+
 	virtual ~IDisplayClusterModularFeatureMediaInitializer() = default;
 
 public:
@@ -30,14 +67,19 @@ public:
 	virtual bool IsMediaSubjectSupported(const UObject* MediaSubject) = 0;
 
 	/**
-	 * Performs initialization of a media subject
+	 * Performs initialization of a media subject for tiled input/output
 	 *
-	 * @param MediaSubject   - UMediaSource or UMediaOutput instance
-	 * @param OwnerName      - Owner name (e.g. ICVFX camera component)
-	 * @param OwnerUniqueIdx - Unique owner index
-	 * @param TilePos        - TIle XY-position
-	 * 
-	 * @return true if media subject supported
+	 * @param MediaSubject - UMediaSource or UMediaOutput instance
+	 * @param OwnerInfo    - Additional information about the object holding the media subject
+	 * @param TilePos      - TIle XY-position
+	 */
+	virtual void InitializeMediaSubjectForTile(UObject* MediaSubject, const FMediaSubjectOwnerInfo& OnwerInfo, const FIntPoint& TilePos) = 0;
+
+	/**
+	 * Performs initialization of a media subject for full frame input/output
+	 *
+	 * @param MediaSubject - UMediaSource or UMediaOutput instance
+	 * @param OwnerInfo    - Additional information about the object holding the media subject
 	 */
-	virtual void InitializeMediaSubjectForTile(UObject* MediaSubject, const FString& OwnerName, uint8 OwnerUniqueIdx, const FIntPoint& TilePos) = 0;
+	virtual void InitializeMediaSubjectForFullFrame(UObject* MediaSubject, const FMediaSubjectOwnerInfo& OnwerInfo) = 0;
 };
diff --git a/Engine/Plugins/VirtualProduction/Rivermax/RivermaxMedia/Source/RivermaxMediaEditor/Private/ModularFeatures/RivermaxMediaInitializerFeature.cpp b/Engine/Plugins/VirtualProduction/Rivermax/RivermaxMedia/Source/RivermaxMediaEditor/Private/ModularFeatures/RivermaxMediaInitializerFeature.cpp
--- a/Engine/Plugins/VirtualProduction/Rivermax/RivermaxMedia/Source/RivermaxMediaEditor/Private/ModularFeatures/RivermaxMediaInitializerFeature.cpp
+++ b/Engine/Plugins/VirtualProduction/Rivermax/RivermaxMedia/Source/RivermaxMediaEditor/Private/ModularFeatures/RivermaxMediaInitializerFeature.cpp
@@ -21,10 +21,8 @@ bool FRivermaxMediaInitializerFeature::IsMediaSubjectSupported(const UObject* Me
 	return false;
 }
 
-void FRivermaxMediaInitializerFeature::InitializeMediaSubjectForTile(UObject* MediaSubject, const FString& OwnerName, uint8 OwnerUniqueIdx, const FIntPoint& TilePos)
+void FRivermaxMediaInitializerFeature::InitializeMediaSubjectForTile(UObject* MediaSubject, const FMediaSubjectOwnerInfo& OnwerInfo, const FIntPoint& TilePos)
 {
-	checkSlow(MediaSubject);
-
 	if (URivermaxMediaSource* RivermaxMediaSource = Cast<URivermaxMediaSource>(MediaSubject))
 	{
 		RivermaxMediaSource->PlayerMode          = ERivermaxPlayerMode::Framelock;
@@ -34,7 +32,7 @@ void FRivermaxMediaInitializerFeature::InitializeMediaSubjectForTile(UObject* Me
 		RivermaxMediaSource->FrameRate           = { 60,1 };
 		RivermaxMediaSource->PixelFormat         = ERivermaxMediaSourcePixelFormat::RGB_10bit;
 		RivermaxMediaSource->InterfaceAddress    = GetRivermaxInterfaceAddress();
-		RivermaxMediaSource->StreamAddress       = GenerateStreamAddress(OwnerUniqueIdx, TilePos);
+		RivermaxMediaSource->StreamAddress       = GenerateStreamAddress(OnwerInfo.OwnerUniqueIdx, TilePos);
 		RivermaxMediaSource->Port                = 50000;
 		RivermaxMediaSource->bIsSRGBInput        = false;
 		RivermaxMediaSource->bUseGPUDirect       = true;
@@ -51,12 +49,64 @@ void FRivermaxMediaInitializerFeature::InitializeMediaSubjectForTile(UObject* Me
 		RivermaxMediaOutput->FrameRate           = { 60,1 };
 		RivermaxMediaOutput->PixelFormat         = ERivermaxMediaOutputPixelFormat::PF_10BIT_RGB;
 		RivermaxMediaOutput->InterfaceAddress    = GetRivermaxInterfaceAddress();
-		RivermaxMediaOutput->StreamAddress       = GenerateStreamAddress(OwnerUniqueIdx, TilePos);
+		RivermaxMediaOutput->StreamAddress       = GenerateStreamAddress(OnwerInfo.OwnerUniqueIdx, TilePos);
 		RivermaxMediaOutput->Port                = 50000;
 		RivermaxMediaOutput->bUseGPUDirect       = true;
 	}
 }
 
+void FRivermaxMediaInitializerFeature::InitializeMediaSubjectForFullFrame(UObject* MediaSubject, const FMediaSubjectOwnerInfo& OnwerInfo)
+{
+	if (URivermaxMediaSource* RivermaxMediaSource = Cast<URivermaxMediaSource>(MediaSubject))
+	{
+		RivermaxMediaSource->PlayerMode          = ERivermaxPlayerMode::Framelock;
+		RivermaxMediaSource->bUseZeroLatency     = true;
+		RivermaxMediaSource->bOverrideResolution = false;
+		//RivermaxMediaSource->Resolution          = default value
+		RivermaxMediaSource->FrameRate           = { 60,1 };
+		RivermaxMediaSource->PixelFormat         = ERivermaxMediaSourcePixelFormat::RGB_10bit;
+		RivermaxMediaSource->InterfaceAddress    = GetRivermaxInterfaceAddress();
+		RivermaxMediaSource->StreamAddress       = GenerateStreamAddress(OnwerInfo.ClusterNodeUniqueIdx.Get(0), OnwerInfo.OwnerUniqueIdx, OnwerInfo.OwnerType);
+		RivermaxMediaSource->Port                = 50000;
+		RivermaxMediaSource->bIsSRGBInput        = false;
+		RivermaxMediaSource->bUseGPUDirect       = true;
+	}
+	else if (URivermaxMediaOutput* RivermaxMediaOutput = Cast<URivermaxMediaOutput>(MediaSubject))
+	{
+		RivermaxMediaOutput->FrameLockingMode      = ERivermaxFrameLockingMode::BlockOnReservation;
+		RivermaxMediaOutput->PresentationQueueSize = 2;
+		RivermaxMediaOutput->bOverrideResolution   = false;
+		//RivermaxMediaOutput->Resolution            = default value
+		RivermaxMediaOutput->PixelFormat           = ERivermaxMediaOutputPixelFormat::PF_10BIT_RGB;
+		RivermaxMediaOutput->InterfaceAddress      = GetRivermaxInterfaceAddress();
+		RivermaxMediaOutput->StreamAddress         = GenerateStreamAddress(OnwerInfo.ClusterNodeUniqueIdx.Get(0), OnwerInfo.OwnerUniqueIdx, OnwerInfo.OwnerType);
+		RivermaxMediaOutput->Port                  = 50000;
+
+		switch (OnwerInfo.OwnerType)
+		{
+		case FMediaSubjectOwnerInfo::EMediaSubjectOwnerType::ICVFXCamera:
+		case FMediaSubjectOwnerInfo::EMediaSubjectOwnerType::Viewport:
+			RivermaxMediaOutput->AlignmentMode       = ERivermaxMediaAlignmentMode::FrameCreation;
+			RivermaxMediaOutput->bDoContinuousOutput = false;
+			RivermaxMediaOutput->bDoFrameCounterTimestamping = true;
+			RivermaxMediaOutput->FrameRate           = { 60,1 };
+			RivermaxMediaOutput->bUseGPUDirect       = true;
+			break;
+
+		case FMediaSubjectOwnerInfo::EMediaSubjectOwnerType::Backbuffer:
+			RivermaxMediaOutput->AlignmentMode       = ERivermaxMediaAlignmentMode::AlignmentPoint;
+			RivermaxMediaOutput->bDoContinuousOutput = true;
+			RivermaxMediaOutput->bDoFrameCounterTimestamping = false;
+			RivermaxMediaOutput->FrameRate           = { 24,1 };
+			RivermaxMediaOutput->bUseGPUDirect       = false;
+			break;
+
+		default:
+			checkNoEntry();
+		}
+	}
+}
+
 FString FRivermaxMediaInitializerFeature::GetRivermaxInterfaceAddress() const
 {
 	FString ResultAddress{ TEXT("*.*.*.*") };
@@ -70,20 +120,11 @@ FString FRivermaxMediaInitializerFeature::GetRivermaxInterfaceAddress() const
 		TArray<FString> Octets;
 		Devices[0].InterfaceAddress.ParseIntoArray(Octets, TEXT("."));
 
-		// IPv4?
+		// IPv4 always has 4 octets
 		if (Octets.Num() == 4)
 		{
 			ResultAddress = FString::Printf(TEXT("%s.%s.%s.*"), *Octets[0], *Octets[1], *Octets[2]);
 		}
-		// IPv6?
-		else
-		{
-			Devices[0].InterfaceAddress.ParseIntoArray(Octets, TEXT(":"));
-			if (Octets.Num() == 6)
-			{
-				ResultAddress = FString::Printf(TEXT("%s:%s:%s:%s:%s:*"), *Octets[0], *Octets[1], *Octets[2], *Octets[3], *Octets[4]);
-			}
-		}
 	}
 
 	return ResultAddress;
@@ -91,8 +132,18 @@ FString FRivermaxMediaInitializerFeature::GetRivermaxInterfaceAddress() const
 
 FString FRivermaxMediaInitializerFeature::GenerateStreamAddress(uint8 OwnerUniqueIdx, const FIntPoint& TilePos) const
 {
-	constexpr uint8 MaxVal = TNumericLimits<uint8>::Max();
-	checkSlow(OwnerUniqueIdx < MaxVal && TilePos.X < MaxVal && TilePos.Y < MaxVal);
+	static const constexpr uint8 MaxVal = TNumericLimits<uint8>::Max();
+	checkSlow(TilePos.X < MaxVal && TilePos.Y < MaxVal);
 
-	return FString::Printf(TEXT("228.%u.%u.%u"), OwnerUniqueIdx + 1, TilePos.X + 1, TilePos.Y + 1);
+	static constexpr uint8 AddressOffsetForTiles = 200;
+	checkSlow((AddressOffsetForTiles + OwnerUniqueIdx) < MaxVal);
+
+	// 228.200.*.* - 228.255.*.* - range for tiled media (max 56 objects)
+	return FString::Printf(TEXT("228.%u.%u.%u"), AddressOffsetForTiles + OwnerUniqueIdx, TilePos.X, TilePos.Y);
+}
+
+FString FRivermaxMediaInitializerFeature::GenerateStreamAddress(uint8 ClusterNodeUniqueIdx, uint8 OwnerUniqueIdx, const FMediaSubjectOwnerInfo::EMediaSubjectOwnerType OwnerType) const
+{
+	// 228.0.*.* - 228.199.*.* - range for full-frame media (max 200 objects). But could be extended up to the limit if no tiles used.
+	return FString::Printf(TEXT("228.%u.%u.%u"), ClusterNodeUniqueIdx, static_cast<uint8>(OwnerType), OwnerUniqueIdx);
 }
diff --git a/Engine/Plugins/VirtualProduction/Rivermax/RivermaxMedia/Source/RivermaxMediaEditor/Private/ModularFeatures/RivermaxMediaInitializerFeature.h b/Engine/Plugins/VirtualProduction/Rivermax/RivermaxMedia/Source/RivermaxMediaEditor/Private/ModularFeatures/RivermaxMediaInitializerFeature.h
--- a/Engine/Plugins/VirtualProduction/Rivermax/RivermaxMedia/Source/RivermaxMediaEditor/Private/ModularFeatures/RivermaxMediaInitializerFeature.h
+++ b/Engine/Plugins/VirtualProduction/Rivermax/RivermaxMedia/Source/RivermaxMediaEditor/Private/ModularFeatures/RivermaxMediaInitializerFeature.h
@@ -16,7 +16,8 @@ public:
 
 	//~ Begin IDisplayClusterModularFeatureMediaInitializer
 	virtual bool IsMediaSubjectSupported(const UObject* MediaSubject) override;
-	virtual void InitializeMediaSubjectForTile(UObject* MediaSubject, const FString& OwnerName, uint8 OwnerUniqueIdx, const FIntPoint& TilePos) override;
+	virtual void InitializeMediaSubjectForTile(UObject* MediaSubject, const FMediaSubjectOwnerInfo& OnwerInfo, const FIntPoint& TilePos) override;
+	virtual void InitializeMediaSubjectForFullFrame(UObject* MediaSubject, const FMediaSubjectOwnerInfo& OnwerInfo) override;
 	//~ End IDisplayClusterModularFeatureMediaInitializer
 
 private:
@@ -26,4 +27,7 @@ private:
 
 	/** Generates stream address based on the function parameters */
 	FString GenerateStreamAddress(uint8 OwnerUniqueIdx, const FIntPoint& TilePos) const;
+
+	/** Generates stream address based on the function parameters */
+	FString GenerateStreamAddress(uint8 ClusterNodeUniqueIdx, uint8 OwnerUniqueIdx, const FMediaSubjectOwnerInfo::EMediaSubjectOwnerType OwnerType) const;
 };
