nDisplay: Replaced 2D LUT with 1D F16 LUT for depth of field compensation

#virtualized

diff --git a/Engine/Plugins/Runtime/nDisplay/Source/DisplayClusterConfiguration/Private/DisplayClusterConfigurationTypes_ICVFX.cpp b/Engine/Plugins/Runtime/nDisplay/Source/DisplayClusterConfiguration/Private/DisplayClusterConfigurationTypes_ICVFX.cpp
--- a/Engine/Plugins/Runtime/nDisplay/Source/DisplayClusterConfiguration/Private/DisplayClusterConfigurationTypes_ICVFX.cpp
+++ b/Engine/Plugins/Runtime/nDisplay/Source/DisplayClusterConfiguration/Private/DisplayClusterConfigurationTypes_ICVFX.cpp
@@ -594,24 +594,30 @@ void FDisplayClusterConfigurationICVFX_CameraDepthOfField::UpdateDynamicCompensa
 	if (CompensationLUT)
 	{
 		FSharedImageConstRef CPUTextureRef = CompensationLUT->GetCPUCopy();
-		if (CPUTextureRef.IsValid())
+		if (CPUTextureRef.IsValid() && CPUTextureRef->Format == ERawImageFormat::R32F)
 		{
-			TArray64<uint8> Pixels(CPUTextureRef->RawData);
-			for (int32 Index = 0; Index < Pixels.Num(); ++Index)
+			TArrayView64<const float> SrcPixels = CPUTextureRef->AsR32F();
+			TArray64<FFloat16> DestPixels;
+			DestPixels.AddZeroed(SrcPixels.Num());
+
+			for (int32 Index = 0; Index < SrcPixels.Num(); ++Index)
 			{
 				// Scale the offset encoded in the LUT so that the final CoC when computed in the DoF pipeline is scaled by the gain.
 				// The actual new offset needed to accomplish this comes from the following equation:
 				// c * (CoC_obj + CoC_off) = CoC_obj + newOffset =>
 				// newOffset = (1 - c) * CoC_obj + c * CoC_off
-				const int32 ObjectCoC = Index % CPUTextureRef->GetHeight();
-				const int32 Offset = Pixels[Index] - 128;
-				const int32 ScaledOffset = (1 - DepthOfFieldGain) * ObjectCoC + DepthOfFieldGain * Offset;
-				Pixels[Index] = FMath::Clamp(ScaledOffset + 128, 0, 255);
+				
+				const float ObjectCoC = Index / 32.0f + 1;
+				const float Offset = SrcPixels[Index];
+				const float ScaledOffset = (1 - DepthOfFieldGain) * ObjectCoC + DepthOfFieldGain * Offset;
+				DestPixels[Index] = ScaledOffset;
 			}
 
+			TArrayView<uint8> PixelsView((uint8*)DestPixels.GetData(), int64(DestPixels.Num() * sizeof(FFloat16)));
+
 			// Texture format is assumed to be greyscale (PF_G8), and we must disable sRGB on the texture to ensure the raw byte value, which encodes
 			// the offset in pixels, is passed unmodified to the depth of field shader
-			if (UTexture2D* NewTexture = UTexture2D::CreateTransient(CPUTextureRef->GetWidth(), CPUTextureRef->GetHeight(), PF_G8, NAME_None, Pixels))
+			if (UTexture2D* NewTexture = UTexture2D::CreateTransient(CPUTextureRef->GetWidth(), CPUTextureRef->GetHeight(), PF_R16F, NAME_None, PixelsView))
 			{
 				DynamicCompensationLUT = NewTexture;
 				DynamicCompensationLUT->SRGB = 0;
diff --git a/Engine/Shaders/Private/DiaphragmDOF/DOFCommon.ush b/Engine/Shaders/Private/DiaphragmDOF/DOFCommon.ush
--- a/Engine/Shaders/Private/DiaphragmDOF/DOFCommon.ush
+++ b/Engine/Shaders/Private/DiaphragmDOF/DOFCommon.ush
@@ -161,7 +161,6 @@ float DepthBlurExponent;
 
 Texture2D		DynamicRadiusOffsetLUT;
 SamplerState	DynamicRadiusOffsetLUTSampler;
-float2			DynamicRadiusOffsetLUTInvSize;
 
 //------------------------------------------------------- DEBUG OUTPUT
 
@@ -178,19 +177,14 @@ float GetCocOffset(float CocRadius)
 	float DynamicOffset = 0.0;
 	BRANCH if (bCocEnableDynamicRadiusOffset)
 	{
-		// By convention, the LUT is a 1-to-1 map from CoC size in pixels to CoC offset, so scale the CoC radii by the texture
-		// size to get the texel coordinate. If this assumption ever changes, we will need to scale this value
-		float2 TexCoords = float2(CocRadius, CocInFocusRadius) * DynamicRadiusOffsetLUTInvSize;
-
-		// Assume that the LUT is homogeneous, which means that the value of any pixel at s*uv is equal to s times the value at uv
-		// This allows us to compute what a pixel should be outside the bounds of the LUT by finding the largest integer
-		// scalar s such that uv is between 0 and 1, and then scaling the sampled pixel by that much
-		float WrapCoefficient = ceil(max(TexCoords.x, TexCoords.y));
-
-		// Also by convention, a value of 128 is used as the zero offset value, allowing support for positive and negative offsets, so scale by 255
-		// to put the pixel into byte space, then offset by 128 to get the final value. Also, enforce a maximum offset to be the initial CoC radius;
-		// the offset should only ever make the final CoC between the initial CoC and 0
-		DynamicOffset = -min((DynamicRadiusOffsetLUT.SampleLevel(DynamicRadiusOffsetLUTSampler, TexCoords / WrapCoefficient, 0).r * 255 - 128) * WrapCoefficient, CocRadius);
+		// The 1D LUT is indexed by ratio of object CoC to in-focus CoC, starting at ratio of 1 (since we only care about objects at or beyond the in-focus plane)
+		// and going to a maximum ratio calculated by the number of pixels in the LUT and the starting object and wall CoCs. By convention, the LUT has been
+		// computed with starting CoC of 32
+		float MaxRatio = (256 + 32) / 32;
+		float Ratio = CocInFocusRadius > 0 ? CocRadius / CocInFocusRadius : MaxRatio;
+		float2 TexCoords = float2((Ratio - 1) / (MaxRatio - 1), 0);
+		float Offset = DynamicRadiusOffsetLUT.SampleLevel(DynamicRadiusOffsetLUTSampler, TexCoords, 0).r * CocInFocusRadius;
+		DynamicOffset = -min(Offset, CocRadius);
 	}
 
 	return DynamicOffset;
diff --git a/Engine/Source/Runtime/Renderer/Private/PostProcess/DiaphragmDOF.cpp b/Engine/Source/Runtime/Renderer/Private/PostProcess/DiaphragmDOF.cpp
--- a/Engine/Source/Runtime/Renderer/Private/PostProcess/DiaphragmDOF.cpp
+++ b/Engine/Source/Runtime/Renderer/Private/PostProcess/DiaphragmDOF.cpp
@@ -640,7 +640,6 @@ BEGIN_SHADER_PARAMETER_STRUCT(FDOFCocModelShaderParameters, )
 	SHADER_PARAMETER(float, CocInFocusRadius)
 	SHADER_PARAMETER_RDG_TEXTURE(Texture2D, DynamicRadiusOffsetLUT)
 	SHADER_PARAMETER_SAMPLER(SamplerState, DynamicRadiusOffsetLUTSampler)
-	SHADER_PARAMETER(FVector2f, DynamicRadiusOffsetLUTInvSize)
 	SHADER_PARAMETER(uint32, bCocEnableDynamicRadiusOffset)
 	SHADER_PARAMETER(float, CocMinRadius)
 	SHADER_PARAMETER(float, CocMaxRadius)
@@ -669,19 +668,15 @@ void SetCocModelParameters(
 	if (CocModel.DynamicRadiusOffsetLUT)
 	{
 		OutParameters->DynamicRadiusOffsetLUT = GraphBuilder.RegisterExternalTexture(CreateRenderTarget(CocModel.DynamicRadiusOffsetLUT, TEXT("DynamicRadiusOffsetLUT")));
-		OutParameters->DynamicRadiusOffsetLUTInvSize.X = 1.0f / CocModel.DynamicRadiusOffsetLUT->GetDesc().Extent.X;
-		OutParameters->DynamicRadiusOffsetLUTInvSize.Y = 1.0f / CocModel.DynamicRadiusOffsetLUT->GetDesc().Extent.Y;
 	}
 	else
 	{
 		const FRDGSystemTextures& SystemTextures = FRDGSystemTextures::Get(GraphBuilder);
 
 		OutParameters->DynamicRadiusOffsetLUT = SystemTextures.Black;
-		OutParameters->DynamicRadiusOffsetLUTInvSize.X = 1.0f;
-		OutParameters->DynamicRadiusOffsetLUTInvSize.Y = 1.0f;
 	}
 
-	OutParameters->DynamicRadiusOffsetLUTSampler = TStaticSamplerState<SF_Bilinear>::GetRHI();
+	OutParameters->DynamicRadiusOffsetLUTSampler = TStaticSamplerState<SF_Bilinear, AM_Clamp, AM_Wrap>::GetRHI();
 }
 
 
