[nDisplay] Fixed media capture initialization to prevent log spam
 - Media capture always start, even if camera is disabled
 - No more log spam trying to re-start capture every second
 - Special initialization case for tiles when the owner is disabled on start

diff --git a/Engine/Plugins/Runtime/nDisplay/Source/DisplayClusterMedia/Private/Capture/DisplayClusterMediaCaptureCamera.cpp b/Engine/Plugins/Runtime/nDisplay/Source/DisplayClusterMedia/Private/Capture/DisplayClusterMediaCaptureCamera.cpp
--- a/Engine/Plugins/Runtime/nDisplay/Source/DisplayClusterMedia/Private/Capture/DisplayClusterMediaCaptureCamera.cpp
+++ b/Engine/Plugins/Runtime/nDisplay/Source/DisplayClusterMedia/Private/Capture/DisplayClusterMediaCaptureCamera.cpp
@@ -3,20 +3,52 @@
 #include "Capture/DisplayClusterMediaCaptureCamera.h"
 
 #include "Components/DisplayClusterICVFXCameraComponent.h"
-#include "Config/IDisplayClusterConfigManager.h"
-#include "DisplayClusterConfigurationTypes_Viewport.h"
+#include "DisplayClusterConfigurationTypes.h"
 #include "DisplayClusterRootActor.h"
 #include "IDisplayCluster.h"
-#include "IDisplayClusterCallbacks.h"
 #include "Game/IDisplayClusterGameManager.h"
 
-#include "Render/Viewport/IDisplayClusterViewportManagerProxy.h"
-#include "RHICommandList.h"
-#include "RHIResources.h"
-
 
 FDisplayClusterMediaCaptureCamera::FDisplayClusterMediaCaptureCamera(const FString& InMediaId, const FString& InClusterNodeId, const FString& InCameraId, const FString& InViewportId, UMediaOutput* InMediaOutput, UDisplayClusterMediaOutputSynchronizationPolicy* SyncPolicy)
 	: FDisplayClusterMediaCaptureViewport(InMediaId, InClusterNodeId, InViewportId, InMediaOutput, SyncPolicy)
 	, CameraId(InCameraId)
 {
 }
+
+bool FDisplayClusterMediaCaptureCamera::GetCaptureSizeFromConfig(FIntPoint& OutSize) const
+{
+	if (const ADisplayClusterRootActor* const ActiveRootActor = IDisplayCluster::Get().GetGameMgr()->GetRootActor())
+	{
+		TArray<UActorComponent*> ICVFXCameraComponents;
+		IDisplayClusterViewportConfiguration* ViewportConfiguration = ActiveRootActor->GetViewportConfiguration();
+		if (const ADisplayClusterRootActor* const ConfigurationRootActor = ViewportConfiguration ? ViewportConfiguration->GetRootActor(EDisplayClusterRootActorType::Configuration) : nullptr)
+		{
+			ConfigurationRootActor->GetComponents(UDisplayClusterICVFXCameraComponent::StaticClass(), ICVFXCameraComponents);
+		}
+
+		// Get capture size from camera config
+		for (const UActorComponent* const Component : ICVFXCameraComponents)
+		{
+			if (const UDisplayClusterICVFXCameraComponent* const ICVFXCamera = Cast<UDisplayClusterICVFXCameraComponent>(Component))
+			{
+				if (ICVFXCamera->GetName() == CameraId)
+				{
+					const FDisplayClusterConfigurationICVFX_CameraSettings& CameraSettings = ICVFXCamera->GetCameraSettingsICVFX();
+
+					if (CameraSettings.CustomFrustum.bEnable)
+					{
+						OutSize = CameraSettings.CustomFrustum.EstimatedOverscanResolution;
+					}
+					else
+					{
+						OutSize = CameraSettings.CustomFrustum.InnerFrustumResolution;
+					}
+
+					return true;
+				}
+			}
+		}
+	}
+
+	return false;
+}
diff --git a/Engine/Plugins/Runtime/nDisplay/Source/DisplayClusterMedia/Private/Capture/DisplayClusterMediaCaptureCamera.h b/Engine/Plugins/Runtime/nDisplay/Source/DisplayClusterMedia/Private/Capture/DisplayClusterMediaCaptureCamera.h
--- a/Engine/Plugins/Runtime/nDisplay/Source/DisplayClusterMedia/Private/Capture/DisplayClusterMediaCaptureCamera.h
+++ b/Engine/Plugins/Runtime/nDisplay/Source/DisplayClusterMedia/Private/Capture/DisplayClusterMediaCaptureCamera.h
@@ -4,10 +4,6 @@
 
 #include "Capture/DisplayClusterMediaCaptureViewport.h"
 
-class FRDGBuilder;
-class FSceneViewFamily;
-class IDisplayClusterViewportProxy;
-
 
 /**
  * Viewport media capture
@@ -18,6 +14,11 @@ class FDisplayClusterMediaCaptureCamera
 public:
 	FDisplayClusterMediaCaptureCamera(const FString& MediaId, const FString& ClusterNodeId, const FString& CameraId, const FString& ViewportId, UMediaOutput* MediaOutput, UDisplayClusterMediaOutputSynchronizationPolicy* SyncPolicy = nullptr);
 
+protected:
+	//~ Begin FDisplayClusterMediaCaptureViewport
+	virtual bool GetCaptureSizeFromConfig(FIntPoint& OutSize) const override;
+	//~ End FDisplayClusterMediaCaptureViewport
+
 private:
 	/** ICVFX camera name */
 	const FString CameraId;
diff --git a/Engine/Plugins/Runtime/nDisplay/Source/DisplayClusterMedia/Private/Capture/DisplayClusterMediaCaptureTile.cpp b/Engine/Plugins/Runtime/nDisplay/Source/DisplayClusterMedia/Private/Capture/DisplayClusterMediaCaptureTile.cpp
new file mode 100644
--- /dev/null
+++ b/Engine/Plugins/Runtime/nDisplay/Source/DisplayClusterMedia/Private/Capture/DisplayClusterMediaCaptureTile.cpp
@@ -0,0 +1,36 @@
+// Copyright Epic Games, Inc. All Rights Reserved.
+
+#include "Capture/DisplayClusterMediaCaptureTile.h"
+
+
+FDisplayClusterMediaCaptureTile::FDisplayClusterMediaCaptureTile(const FString& InMediaId, const FString& InClusterNodeId, const FString& InViewportId, UMediaOutput* InMediaOutput, UDisplayClusterMediaOutputSynchronizationPolicy* SyncPolicy)
+	: FDisplayClusterMediaCaptureViewport(InMediaId, InClusterNodeId, InViewportId, InMediaOutput, SyncPolicy)
+{
+}
+
+bool FDisplayClusterMediaCaptureTile::GetCaptureSizeFromConfig(FIntPoint& OutSize) const
+{
+	// The upper level logic always tries to get actual capture size from the game proxy. If the game proxy
+	// is not availalbe (camera/viewport is disabled or deactivated on start and therefore not being rendered),
+	// it tries to acquire capture size from config. Being here means exactly this case.
+	//
+	// The problem is the tiles can't know their size until they start to render. There is a bunch of intermediate
+	// logic that affects the final size of the owning camera/viewport texture such as "adopt resolution",
+	// "screen percentage", "overscan", etc.
+	//
+	// But we need to provide something to start capture successfully. As a workaround, we return capture size
+	// of 64x64. It looks valid and small enough to initialize media capture pipeline with a little
+	// resource allocation. When tile rendering is started (e.g. camera/viewport is activated), its texture
+	// will have final size available in the game proxy. Most likely it will be different so the media capture
+	// pipeline will reconfigure for the new size.
+	//
+	// If game proxy is available, the capture size will be acquired by the higher level in
+	// FDisplayClusterMediaCaptureViewport::GetCaptureSizeFromGameProxy()
+	//
+	// This trick looks kind of legal as we allow to change let's say ICVFX camera texture size in runtime.
+	// The media pipeline is able to recognize the size changes and reconfigure.
+
+	OutSize = { 64, 64 };
+
+	return true;
+}
diff --git a/Engine/Plugins/Runtime/nDisplay/Source/DisplayClusterMedia/Private/Capture/DisplayClusterMediaCaptureTile.h b/Engine/Plugins/Runtime/nDisplay/Source/DisplayClusterMedia/Private/Capture/DisplayClusterMediaCaptureTile.h
new file mode 100644
--- /dev/null
+++ b/Engine/Plugins/Runtime/nDisplay/Source/DisplayClusterMedia/Private/Capture/DisplayClusterMediaCaptureTile.h
@@ -0,0 +1,22 @@
+// Copyright Epic Games, Inc. All Rights Reserved.
+
+#pragma once
+
+#include "Capture/DisplayClusterMediaCaptureViewport.h"
+
+
+/**
+ * Tile media capture
+ */
+class FDisplayClusterMediaCaptureTile
+	: public FDisplayClusterMediaCaptureViewport
+{
+public:
+	FDisplayClusterMediaCaptureTile(const FString& MediaId, const FString& ClusterNodeId, const FString& ViewportId, UMediaOutput* MediaOutput, UDisplayClusterMediaOutputSynchronizationPolicy* SyncPolicy = nullptr);
+
+protected:
+
+	//~ Begin FDisplayClusterMediaCaptureViewport
+	virtual bool GetCaptureSizeFromConfig(FIntPoint& OutSize) const override;
+	//~ End FDisplayClusterMediaCaptureViewport
+};
diff --git a/Engine/Plugins/Runtime/nDisplay/Source/DisplayClusterMedia/Private/Capture/DisplayClusterMediaCaptureViewport.cpp b/Engine/Plugins/Runtime/nDisplay/Source/DisplayClusterMedia/Private/Capture/DisplayClusterMediaCaptureViewport.cpp
--- a/Engine/Plugins/Runtime/nDisplay/Source/DisplayClusterMedia/Private/Capture/DisplayClusterMediaCaptureViewport.cpp
+++ b/Engine/Plugins/Runtime/nDisplay/Source/DisplayClusterMedia/Private/Capture/DisplayClusterMediaCaptureViewport.cpp
@@ -7,8 +7,10 @@
 
 #include "DisplayClusterConfigurationTypes_Viewport.h"
 #include "DisplayClusterMediaLog.h"
+#include "DisplayClusterRootActor.h"
 
 #include "Config/IDisplayClusterConfigManager.h"
+#include "Game/IDisplayClusterGameManager.h"
 
 #include "Render/IDisplayClusterRenderManager.h"
 #include "Render/Viewport/IDisplayClusterViewport.h"
@@ -67,6 +69,46 @@ void FDisplayClusterMediaCaptureViewport::OnUpdateViewportMediaState(IDisplayClu
 }
 
 FIntPoint FDisplayClusterMediaCaptureViewport::GetCaptureSize() const
+{
+	FIntPoint CaptureSize{ FIntPoint::ZeroValue };
+
+	if (GetCaptureSizeFromGameProxy(CaptureSize))
+	{
+		UE_LOG(LogDisplayClusterMedia, Verbose, TEXT("'%s' acquired capture size from game proxy [%d, %d]"), *GetMediaId(), CaptureSize.X, CaptureSize.Y);
+	}
+	else if (GetCaptureSizeFromConfig(CaptureSize))
+	{
+		UE_LOG(LogDisplayClusterMedia, Verbose, TEXT("'%s' acquired capture size from config [%d, %d]"), *GetMediaId(), CaptureSize.X, CaptureSize.Y);
+	}
+	else
+	{
+		UE_LOG(LogDisplayClusterMedia, Verbose, TEXT("'%s' couldn't acquire capture"), *GetMediaId());
+	}
+
+	return CaptureSize;
+}
+
+bool FDisplayClusterMediaCaptureViewport::GetCaptureSizeFromConfig(FIntPoint& OutSize) const
+{
+	if (const ADisplayClusterRootActor* const ActiveRootActor = IDisplayCluster::Get().GetGameMgr()->GetRootActor())
+	{
+		if (const UDisplayClusterConfigurationData* const ConfigData = ActiveRootActor->GetConfigData())
+		{
+			const FString& NodeId = GetClusterNodeId();
+			if (const UDisplayClusterConfigurationViewport* const ViewportCfg = ConfigData->GetViewport(NodeId, ViewportId))
+			{
+				const FIntRect ViewportRect = ViewportCfg->Region.ToRect();
+				OutSize = FIntPoint(ViewportRect.Width(), ViewportRect.Height());
+
+				return true;
+			}
+		}
+	}
+
+	return false;
+}
+
+bool FDisplayClusterMediaCaptureViewport::GetCaptureSizeFromGameProxy(FIntPoint& OutSize) const
 {
 	// We need to get actual texture size for the viewport
 	if (const IDisplayClusterRenderManager* const RenderMgr = IDisplayCluster::Get().GetRenderMgr())
@@ -78,15 +120,14 @@ FIntPoint FDisplayClusterMediaCaptureViewport::GetCaptureSize() const
 				const TArray<FDisplayClusterViewport_Context>& Contexts = Viewport->GetContexts();
 				if (Contexts.Num() > 0)
 				{
-					const FIntPoint Size = Contexts[0].RenderTargetRect.Size();
-					UE_LOG(LogDisplayClusterMedia, Log, TEXT("'%s' capture size is [%d, %d]"), *GetMediaId(), Size.X, Size.Y);
-					return Size;
+					OutSize = Contexts[0].RenderTargetRect.Size();
+					return true;
 				}
 			}
 		}
 	}
 
-	return FIntPoint::ZeroValue;
+	return false;
 }
 
 void FDisplayClusterMediaCaptureViewport::OnPostRenderViewFamily_RenderThread(FRDGBuilder& GraphBuilder, const FSceneViewFamily& ViewFamily, const IDisplayClusterViewportProxy* ViewportProxy)
diff --git a/Engine/Plugins/Runtime/nDisplay/Source/DisplayClusterMedia/Private/Capture/DisplayClusterMediaCaptureViewport.h b/Engine/Plugins/Runtime/nDisplay/Source/DisplayClusterMedia/Private/Capture/DisplayClusterMediaCaptureViewport.h
--- a/Engine/Plugins/Runtime/nDisplay/Source/DisplayClusterMedia/Private/Capture/DisplayClusterMediaCaptureViewport.h
+++ b/Engine/Plugins/Runtime/nDisplay/Source/DisplayClusterMedia/Private/Capture/DisplayClusterMediaCaptureViewport.h
@@ -38,6 +38,12 @@ protected:
 	/** Returns texture size of a viewport assigned to capture */
 	virtual FIntPoint GetCaptureSize() const override;
 
+	/** Provides default texture size from config */
+	virtual bool GetCaptureSizeFromConfig(FIntPoint& OutSize) const;
+
+	/** Provides texture size from a game proxy (if available) */
+	bool GetCaptureSizeFromGameProxy(FIntPoint& OutSize) const;
+
 private:
 	/** PostRenderViewFamily callback handler where data is captured */
 	void OnPostRenderViewFamily_RenderThread(FRDGBuilder& GraphBuilder, const FSceneViewFamily& ViewFamily, const IDisplayClusterViewportProxy* ViewportProxy);
diff --git a/Engine/Plugins/Runtime/nDisplay/Source/DisplayClusterMedia/Private/DisplayClusterMediaModule.cpp b/Engine/Plugins/Runtime/nDisplay/Source/DisplayClusterMedia/Private/DisplayClusterMediaModule.cpp
--- a/Engine/Plugins/Runtime/nDisplay/Source/DisplayClusterMedia/Private/DisplayClusterMediaModule.cpp
+++ b/Engine/Plugins/Runtime/nDisplay/Source/DisplayClusterMedia/Private/DisplayClusterMediaModule.cpp
@@ -19,6 +19,7 @@
 
 #include "Capture/DisplayClusterMediaCaptureCamera.h"
 #include "Capture/DisplayClusterMediaCaptureNode.h"
+#include "Capture/DisplayClusterMediaCaptureTile.h"
 #include "Capture/DisplayClusterMediaCaptureViewport.h"
 #include "Input/DisplayClusterMediaInputViewport.h"
 
@@ -454,7 +455,7 @@ void FDisplayClusterMediaModule::InitializeICVFXCameraUniformTilesOutput(const U
 						UE_LOG(LogDisplayClusterMedia, Log, TEXT("Initializing ICVFX media output '%s' for camera '%s' tile '%d,%d'"),
 							*MediaOutputId, *ICVFXCameraName, MediaOutputTile.Position.X, MediaOutputTile.Position.Y);
 
-						TSharedPtr<FDisplayClusterMediaCaptureViewport> NewICVFXTileOutput = MakeShared<FDisplayClusterMediaCaptureViewport>(
+						TSharedPtr<FDisplayClusterMediaCaptureTile> NewICVFXTileOutput = MakeShared<FDisplayClusterMediaCaptureTile>(
 							MediaOutputId, ClusterNodeId,
 							ICVFXViewportTileId,
 							MediaOutputTile.MediaOutput);
