nDisplay: added ability to specify custom render target size for custom mappings

diff --git a/Engine/Plugins/Runtime/nDisplay/Source/DisplayCluster/Private/Render/Viewport/DisplayClusterViewport.cpp b/Engine/Plugins/Runtime/nDisplay/Source/DisplayCluster/Private/Render/Viewport/DisplayClusterViewport.cpp
--- a/Engine/Plugins/Runtime/nDisplay/Source/DisplayCluster/Private/Render/Viewport/DisplayClusterViewport.cpp
+++ b/Engine/Plugins/Runtime/nDisplay/Source/DisplayCluster/Private/Render/Viewport/DisplayClusterViewport.cpp
@@ -47,6 +47,17 @@ static FAutoConsoleVariableRef CVarDisplayClusterMultiGPUEnable(
 	ECVF_Default
 );
 
+namespace UE::DisplayCluster::Viewport
+{
+	static inline void AdjustRect(FIntRect& InOutRect, const float multX, const float multY)
+	{
+		InOutRect.Min.X *= multX;
+		InOutRect.Max.X *= multX;
+		InOutRect.Min.Y *= multY;
+		InOutRect.Max.Y *= multY;
+	}
+};
+
 ///////////////////////////////////////////////////////////////////////////////////////
 //          FDisplayClusterViewport
 ///////////////////////////////////////////////////////////////////////////////////////
@@ -357,14 +368,6 @@ void FDisplayClusterViewport::SetupSceneView(uint32 ContextNum, class UWorld* Wo
 	}
 }
 
-inline void AdjustRect(FIntRect& InOutRect, const float multX, const float multY)
-{
-	InOutRect.Min.X *= multX;
-	InOutRect.Max.X *= multX;
-	InOutRect.Min.Y *= multY;
-	InOutRect.Max.Y *= multY;
-}
-
 float FDisplayClusterViewport::GetClusterRenderTargetRatioMult(const FDisplayClusterRenderFrameSettings& InFrameSettings) const
 {
 	float ClusterRenderTargetRatioMult = InFrameSettings.ClusterRenderTargetRatioMult;
@@ -383,8 +386,13 @@ float FDisplayClusterViewport::GetClusterRenderTargetRatioMult(const FDisplayClu
 	return FMath::Clamp(ClusterRenderTargetRatioMult, 0.f, 1.f);
 }
 
-FIntPoint FDisplayClusterViewport::GetDesiredContextSize(const FIntPoint& InSize, const FDisplayClusterRenderFrameSettings& InFrameSettings) const
+FIntPoint FDisplayClusterViewport::GetDesiredContextSize(const FIntPoint& InContextSize, const FDisplayClusterRenderFrameSettings& InFrameSettings) const
 {
+	// Overrides the base size of the RenderTarget texture for all viewport contexts.
+	// The rest of the RTT size modifiers are applied after this.
+	FIntPoint CustomCustomRenderTargetSize;
+	const FIntPoint InSize = (ProjectionPolicy.IsValid() && ProjectionPolicy->GetCustomRenderTargetSize(this, CustomCustomRenderTargetSize)) ? CustomCustomRenderTargetSize : InContextSize;
+
 	const float ClusterRenderTargetRatioMult = GetClusterRenderTargetRatioMult(InFrameSettings);
 
 	// Check size multipliers in order bellow:
@@ -392,7 +400,12 @@ FIntPoint FDisplayClusterViewport::GetDesiredContextSize(const FIntPoint& InSize
 	const float RenderTargetRatio = FDisplayClusterViewportHelpers::GetValidSizeMultiplier(InSize, RenderSettings.RenderTargetRatio, ClusterRenderTargetRatioMult * RenderTargetAdaptRatio);
 	const float ClusterMult = FDisplayClusterViewportHelpers::GetValidSizeMultiplier(InSize, ClusterRenderTargetRatioMult, RenderTargetRatio * RenderTargetAdaptRatio);
 
-	FIntPoint DesiredContextSize = FDisplayClusterViewportHelpers::ScaleTextureSize(InSize, FMath::Max(RenderTargetAdaptRatio * RenderTargetRatio * ClusterMult, 0.f));
+	const float FinalRenderTargetMult = FMath::Max(RenderTargetAdaptRatio * RenderTargetRatio * ClusterMult, 0.f);
+
+	// Scale RTT size
+	FIntPoint DesiredContextSize = (ProjectionPolicy.IsValid() && !ProjectionPolicy->ShouldUseAnySizeScaleForRenderTarget(this))
+		? InSize // Use original RenderTarget size.
+		: FDisplayClusterViewportHelpers::ScaleTextureSize(InSize, FinalRenderTargetMult);
 
 	const int32 MaxTextureSize = FDisplayClusterViewportHelpers::GetMaxTextureDimension();
 	DesiredContextSize.X = FMath::Min(DesiredContextSize.X, MaxTextureSize);
@@ -444,7 +457,7 @@ bool FDisplayClusterViewport::UpdateFrameContexts(const uint32 InStereoViewIndex
 
 	// Apply desired frame mult
 	const FVector2D DesiredFrameMult = InFrameSettings.GetDesiredFrameMult();
-	AdjustRect(DesiredFrameTargetRect, DesiredFrameMult.X, DesiredFrameMult.Y);
+	UE::DisplayCluster::Viewport::AdjustRect(DesiredFrameTargetRect, DesiredFrameMult.X, DesiredFrameMult.Y);
 
 	// Support preview in scene rendering
 	if (InFrameSettings.IsPreviewRendering())
diff --git a/Engine/Plugins/Runtime/nDisplay/Source/DisplayCluster/Public/Render/Projection/IDisplayClusterProjectionPolicy.h b/Engine/Plugins/Runtime/nDisplay/Source/DisplayCluster/Public/Render/Projection/IDisplayClusterProjectionPolicy.h
--- a/Engine/Plugins/Runtime/nDisplay/Source/DisplayCluster/Public/Render/Projection/IDisplayClusterProjectionPolicy.h
+++ b/Engine/Plugins/Runtime/nDisplay/Source/DisplayCluster/Public/Render/Projection/IDisplayClusterProjectionPolicy.h
@@ -217,6 +217,33 @@ public:
 		return false;
 	}
 
+	/**
+	* This function can override the size of the RenderTarget texture for the viewport.
+	*
+	* @param InViewport          - viewport to override RTT size.
+	* @param OutRenderTargetSize - (out) the desired RTT size.
+	*
+	* @return is true if the RTT size should be overridden.
+	*/
+	virtual bool GetCustomRenderTargetSize(const IDisplayClusterViewport* InViewport, FIntPoint& OutRenderTargetSize) const
+	{
+		return false;
+	}
+
+	/**
+	* This function controls the RTT size multipliers - if false is returned, all modifiers should be ignored.
+	* Here is the list of ignored modifiers:
+	*    RenderTargetAdaptRatio, RenderTargetRatio,
+	*    ClusterRenderTargetRatioMult,
+	*    ClusterICVFXOuterViewportRenderTargetRatioMult, ClusterICVFXInnerViewportRenderTargetRatioMult
+	* 
+	* @param InViewport - the DC viewport.
+	*/
+	virtual bool ShouldUseAnySizeScaleForRenderTarget(const IDisplayClusterViewport* InViewport) const
+	{
+		return true;
+	}
+
 	/**
 	* Initializing the projection policy logic for the current frame before applying warp blending. Called if IsWarpBlendSupported() returns true
 	*
