Fix narrowing conversion with VS 17.9 msvc

#rnx

diff --git a/Engine/Plugins/Runtime/nDisplay/Source/SharedMemoryMedia/Private/SharedMemoryMediaTypes.h b/Engine/Plugins/Runtime/nDisplay/Source/SharedMemoryMedia/Private/SharedMemoryMediaTypes.h
--- a/Engine/Plugins/Runtime/nDisplay/Source/SharedMemoryMedia/Private/SharedMemoryMediaTypes.h
+++ b/Engine/Plugins/Runtime/nDisplay/Source/SharedMemoryMedia/Private/SharedMemoryMediaTypes.h
@@ -88,7 +88,7 @@ struct FSharedMemoryMediaFrameMetadata
 	struct FReceiver
 	{
 		// Frame number that the receiver is done reading. This signals the sender that it can re-use the texture.
-		std::atomic<uint32> FrameNumberAcked = ~0;
+		std::atomic<uint32> FrameNumberAcked = ~0u;
 
 		/** 
 		 * The receiver sets this register. This signals the sender that it should wait for the related frame number ack before re-using the texture.
