nDisplay: Fixed overscan math

diff --git a/Engine/Plugins/Runtime/nDisplay/Source/DisplayCluster/Private/Render/Viewport/Containers/DisplayClusterViewport_CustomFrustumRuntimeSettings.cpp b/Engine/Plugins/Runtime/nDisplay/Source/DisplayCluster/Private/Render/Viewport/Containers/DisplayClusterViewport_CustomFrustumRuntimeSettings.cpp
--- a/Engine/Plugins/Runtime/nDisplay/Source/DisplayCluster/Private/Render/Viewport/Containers/DisplayClusterViewport_CustomFrustumRuntimeSettings.cpp
+++ b/Engine/Plugins/Runtime/nDisplay/Source/DisplayCluster/Private/Render/Viewport/Containers/DisplayClusterViewport_CustomFrustumRuntimeSettings.cpp
@@ -39,6 +39,7 @@ namespace UE::DisplayCluster::Viewport::CustomFrustumHelpers
 ///////////////////////////////////////////////////////////////////////////////////////////
 bool FDisplayClusterViewport_CustomFrustumRuntimeSettings::UpdateProjectionAngles(
 	const FDisplayClusterViewport_CustomFrustumRuntimeSettings& InRuntimeSettings,
+	const FIntPoint& InRenderTargetSize,
 	double& InOutLeft,
 	double& InOutRight,
 	double& InOutTop,
@@ -49,10 +50,13 @@ bool FDisplayClusterViewport_CustomFrustumRuntimeSettings::UpdateProjectionAngle
 		const double Horizontal = InOutRight - InOutLeft;
 		const double Vertical = InOutTop - InOutBottom;
 
-		InOutLeft   -= Horizontal * InRuntimeSettings.CustomFrustumPercent.Left;
-		InOutRight  += Horizontal * InRuntimeSettings.CustomFrustumPercent.Right;
-		InOutBottom -= Vertical * InRuntimeSettings.CustomFrustumPercent.Bottom;
-		InOutTop    += Vertical * InRuntimeSettings.CustomFrustumPercent.Top;
+		// Use the inner region of the texture as the base of the frustum.
+		const FIntPoint InnerSize = InRenderTargetSize - InRuntimeSettings.CustomFrustumPixels.Size();
+
+		InOutLeft -= Horizontal * InRuntimeSettings.CustomFrustumPixels.Left / InnerSize.X;
+		InOutRight += Horizontal * InRuntimeSettings.CustomFrustumPixels.Right / InnerSize.X;
+		InOutBottom -= Vertical * InRuntimeSettings.CustomFrustumPixels.Bottom / InnerSize.Y;
+		InOutTop += Vertical * InRuntimeSettings.CustomFrustumPixels.Top / InnerSize.Y;
 
 		return true;
 	}
diff --git a/Engine/Plugins/Runtime/nDisplay/Source/DisplayCluster/Private/Render/Viewport/Containers/DisplayClusterViewport_CustomFrustumRuntimeSettings.h b/Engine/Plugins/Runtime/nDisplay/Source/DisplayCluster/Private/Render/Viewport/Containers/DisplayClusterViewport_CustomFrustumRuntimeSettings.h
--- a/Engine/Plugins/Runtime/nDisplay/Source/DisplayCluster/Private/Render/Viewport/Containers/DisplayClusterViewport_CustomFrustumRuntimeSettings.h
+++ b/Engine/Plugins/Runtime/nDisplay/Source/DisplayCluster/Private/Render/Viewport/Containers/DisplayClusterViewport_CustomFrustumRuntimeSettings.h
@@ -21,12 +21,13 @@ struct FDisplayClusterViewport_CustomFrustumRuntimeSettings
 	/** Update projection angles by custom frustum
 	*
 	* @param InOverscanRuntimeSettings - the CustomFrustum runtime settings.
+	* @param InRenderTargetSize        - RTT size that used for custom frustum
 	* @param InOutLeft                 - the value of the left projection plane that you want to change
 	* @param InOutRight                - the value of the right projection plane that you want to change
 	* @param InOutTop                  - the value of the top projection plane that you want to change
 	* @parma InOutBottom               - the value of the bottom projection plane that you want to change
 	*/
-	static bool UpdateProjectionAngles(const FDisplayClusterViewport_CustomFrustumRuntimeSettings& InRuntimeSettings, double& InOutLeft, double& InOutRight, double& InOutTop, double& InOutBottom);
+	static bool UpdateProjectionAngles(const FDisplayClusterViewport_CustomFrustumRuntimeSettings& InRuntimeSettings, const FIntPoint& InRenderTargetSize, double& InOutLeft, double& InOutRight, double& InOutTop, double& InOutBottom);
 
 	/**
 	* Values in percent
diff --git a/Engine/Plugins/Runtime/nDisplay/Source/DisplayCluster/Private/Render/Viewport/Containers/DisplayClusterViewport_OverscanRuntimeSettings.cpp b/Engine/Plugins/Runtime/nDisplay/Source/DisplayCluster/Private/Render/Viewport/Containers/DisplayClusterViewport_OverscanRuntimeSettings.cpp
--- a/Engine/Plugins/Runtime/nDisplay/Source/DisplayCluster/Private/Render/Viewport/Containers/DisplayClusterViewport_OverscanRuntimeSettings.cpp
+++ b/Engine/Plugins/Runtime/nDisplay/Source/DisplayCluster/Private/Render/Viewport/Containers/DisplayClusterViewport_OverscanRuntimeSettings.cpp
@@ -38,6 +38,7 @@ namespace UE::DisplayCluster::Viewport::OverscanHelpers
 ///////////////////////////////////////////////////////////////////////////////////////////
 bool FDisplayClusterViewport_OverscanRuntimeSettings::UpdateProjectionAngles(
 	const FDisplayClusterViewport_OverscanRuntimeSettings& InOverscanRuntimeSettings,
+	const FIntPoint& InRenderTargetSize,
 	double& InOutLeft,
 	double& InOutRight,
 	double& InOutTop,
@@ -48,10 +49,13 @@ bool FDisplayClusterViewport_OverscanRuntimeSettings::UpdateProjectionAngles(
 		double Horizontal = InOutRight - InOutLeft;
 		double Vertical = InOutTop - InOutBottom;
 
-		InOutLeft   -= Horizontal * InOverscanRuntimeSettings.OverscanPercent.Left;
-		InOutRight  += Horizontal * InOverscanRuntimeSettings.OverscanPercent.Right;
-		InOutBottom -= Vertical * InOverscanRuntimeSettings.OverscanPercent.Bottom;
-		InOutTop    += Vertical * InOverscanRuntimeSettings.OverscanPercent.Top;
+		// Use the inner region of the texture as the base of the frustum.
+		const FIntPoint InnerSize = InRenderTargetSize - InOverscanRuntimeSettings.OverscanPixels.Size();
+
+		InOutLeft   -= Horizontal * InOverscanRuntimeSettings.OverscanPixels.Left / InnerSize.X;
+		InOutRight  += Horizontal * InOverscanRuntimeSettings.OverscanPixels.Right / InnerSize.X;
+		InOutBottom -= Vertical * InOverscanRuntimeSettings.OverscanPixels.Bottom / InnerSize.Y;
+		InOutTop    += Vertical * InOverscanRuntimeSettings.OverscanPixels.Top / InnerSize.Y;
 
 		return true;
 	}
diff --git a/Engine/Plugins/Runtime/nDisplay/Source/DisplayCluster/Private/Render/Viewport/Containers/DisplayClusterViewport_OverscanRuntimeSettings.h b/Engine/Plugins/Runtime/nDisplay/Source/DisplayCluster/Private/Render/Viewport/Containers/DisplayClusterViewport_OverscanRuntimeSettings.h
--- a/Engine/Plugins/Runtime/nDisplay/Source/DisplayCluster/Private/Render/Viewport/Containers/DisplayClusterViewport_OverscanRuntimeSettings.h
+++ b/Engine/Plugins/Runtime/nDisplay/Source/DisplayCluster/Private/Render/Viewport/Containers/DisplayClusterViewport_OverscanRuntimeSettings.h
@@ -21,12 +21,13 @@ struct FDisplayClusterViewport_OverscanRuntimeSettings
 	/** Update projection angles by overscan
 	* 
 	* @param InOverscanRuntimeSettings - the overscan runtime settings.
+	* @param InRenderTargetSize         - RenderTarget texture size
 	* @param InOutLeft                 - the value of the left projection plane that you want to change
 	* @param InOutRight                - the value of the right projection plane that you want to change
 	* @param InOutTop                  - the value of the top projection plane that you want to change
 	* @parma InOutBottom               - the value of the bottom projection plane that you want to change
 	*/
-	static bool UpdateProjectionAngles(const FDisplayClusterViewport_OverscanRuntimeSettings& InOverscanRuntimeSettings, double& InOutLeft, double& InOutRight, double& InOutTop, double& InOutBottom);
+	static bool UpdateProjectionAngles(const FDisplayClusterViewport_OverscanRuntimeSettings& InOverscanRuntimeSettings, const FIntPoint& InRenderTargetSize, double& InOutLeft, double& InOutRight, double& InOutTop, double& InOutBottom);
 
 	/**
 	* Overscan values in percent
diff --git a/Engine/Plugins/Runtime/nDisplay/Source/DisplayCluster/Private/Render/Viewport/DisplayClusterViewport_Math.cpp b/Engine/Plugins/Runtime/nDisplay/Source/DisplayCluster/Private/Render/Viewport/DisplayClusterViewport_Math.cpp
--- a/Engine/Plugins/Runtime/nDisplay/Source/DisplayCluster/Private/Render/Viewport/DisplayClusterViewport_Math.cpp
+++ b/Engine/Plugins/Runtime/nDisplay/Source/DisplayCluster/Private/Render/Viewport/DisplayClusterViewport_Math.cpp
@@ -235,7 +235,7 @@ void FDisplayClusterViewport::CalculateProjectionMatrix(const uint32 InContextNu
 
 	// Support custom frustum rendering
 	const double OrigValues[] = {l, r, t, b};
-	if (FDisplayClusterViewport_CustomFrustumRuntimeSettings::UpdateProjectionAngles(CustomFrustumRuntimeSettings, l, r, t, b))
+	if (FDisplayClusterViewport_CustomFrustumRuntimeSettings::UpdateProjectionAngles(CustomFrustumRuntimeSettings, Contexts[InContextNum].RenderTargetRect.Size(), l, r, t, b))
 	{
 		const bool bIsValidLimits =  FMath::IsWithin(l, -MaxValue, MaxValue)
 							&& FMath::IsWithin(r, -MaxValue, MaxValue)
@@ -267,7 +267,7 @@ void FDisplayClusterViewport::CalculateProjectionMatrix(const uint32 InContextNu
 	CachedProjectionData.ZFar = f;
 	CachedProjectionData.bValid = true;
 
-	if (FDisplayClusterViewport_OverscanRuntimeSettings::UpdateProjectionAngles(OverscanRuntimeSettings, l, r, t, b))
+	if (FDisplayClusterViewport_OverscanRuntimeSettings::UpdateProjectionAngles(OverscanRuntimeSettings, Contexts[InContextNum].RenderTargetRect.Size(), l, r, t, b))
 	{
 		if (FMath::IsWithin(l, -MaxValue, MaxValue) &&
 			FMath::IsWithin(r, -MaxValue, MaxValue) &&
