nDisplay: Fixed 2110 Split Frustum Tearing

Currently, the incoming texture may have extra black columns on the right side due to required padding. This will fixed at the Rivermax level in the future, but this change is a workaround that trims them based on the destination texture in nDisplay. For the cases of resolutions that are different, which would not be the viewport-viewport case, a cvar was provided to control how many rightmost columns to trim.

diff --git a/Engine/Plugins/Runtime/nDisplay/Source/DisplayClusterMedia/Private/Input/DisplayClusterMediaInputBase.cpp b/Engine/Plugins/Runtime/nDisplay/Source/DisplayClusterMedia/Private/Input/DisplayClusterMediaInputBase.cpp
--- a/Engine/Plugins/Runtime/nDisplay/Source/DisplayClusterMedia/Private/Input/DisplayClusterMediaInputBase.cpp
+++ b/Engine/Plugins/Runtime/nDisplay/Source/DisplayClusterMedia/Private/Input/DisplayClusterMediaInputBase.cpp
@@ -15,6 +15,31 @@
 #include "UObject/Package.h"
 
 
+TAutoConsoleVariable<bool> CVarTempRivermaxCropWorkaround(
+	TEXT("nDisplay.Media.Rivermax.CropWorkaround"),
+	true,
+	TEXT("nDisplay workaround for Rivermax input\n")
+	TEXT("0 : Disabled\n")
+	TEXT("1 : Enabled\n"),
+	ECVF_RenderThreadSafe
+);
+
+// Based on the discussion, it looks like the problem is the incoming 2110 textures
+// may have up to 3 ExtraPixelsThreshold extra pixels.
+TAutoConsoleVariable<int32> CVarTempRivermaxExtraPixelsThreshold(
+	TEXT("nDisplay.Media.Rivermax.ExtraPixelsThreshold"),
+	3,
+	TEXT("nDisplay workaround for Rivermax input\n"),
+	ECVF_RenderThreadSafe
+);
+
+TAutoConsoleVariable<int32> CVarTempRivermaxExtraPixelsRemove(
+	TEXT("nDisplay.Media.Rivermax.ExtraPixelsRemove"),
+	0,
+	TEXT("nDisplay workaround for Rivermax input\n"),
+	ECVF_RenderThreadSafe
+);
+
 FDisplayClusterMediaInputBase::FDisplayClusterMediaInputBase(const FString& InMediaId, const FString& InClusterNodeId, UMediaSource* InMediaSource)
 	: FDisplayClusterMediaBase(InMediaId, InClusterNodeId)
 {
@@ -57,7 +82,10 @@ bool FDisplayClusterMediaInputBase::Play()
 		MediaPlayer->OnMediaEvent().AddRaw(this, &FDisplayClusterMediaInputBase::OnMediaEvent);
 
 		bWasPlayerStarted = MediaPlayer->OpenSource(MediaSource);
-		
+
+		static FName RiverMaxPlayerName(TEXT("RivermaxMedia"));
+		bRunningRivermaxMedia = MediaPlayer->GetPlayerName() == RiverMaxPlayerName;
+
 		return bWasPlayerStarted;
 	}
 
@@ -72,6 +100,43 @@ void FDisplayClusterMediaInputBase::Stop()
 		MediaPlayer->Close();
 		MediaPlayer->OnMediaEvent().RemoveAll(this);
 	}
+
+	bRunningRivermaxMedia = false;
+}
+
+void FDisplayClusterMediaInputBase::OverrideTextureRegions_RenderThread(FRHITexture* const SrcTexture, FRHITexture* const DstTexture, FIntRect& InOutSrcRect, FIntRect& InOutDstRect) const
+{
+	const FIntPoint SrcSize = InOutSrcRect.Size();
+	const FIntPoint DstSize = InOutDstRect.Size();
+	if (SrcSize == DstSize)
+	{
+		return;
+	}
+
+	// [Workaround]
+	// Based on the discussion, it looks like the problem is the incoming 2110 textures
+	// may have up to 3 ExtraPixelsThreshold extra pixels.
+	// If this is the only difference, we just copy the required subregion.
+	if (bRunningRivermaxMedia && CVarTempRivermaxCropWorkaround.GetValueOnRenderThread())
+	{
+		const int32 ExtraPixelsThreshold = CVarTempRivermaxExtraPixelsThreshold.GetValueOnRenderThread();
+
+		// Crop if required
+		if (SrcSize.Y == DstSize.Y
+			&& SrcSize.X >= DstSize.X
+			&& (SrcSize.X - DstSize.X) <= ExtraPixelsThreshold)
+		{
+			// Use Dest size
+			InOutSrcRect.Max.X = InOutSrcRect.Min.X + DstSize.X;
+
+			return;
+		}
+
+		// By default we always remove extra pixels from the right side.
+		const int32 ExtraPixelsRemove = CVarTempRivermaxExtraPixelsRemove.GetValueOnRenderThread();
+
+		InOutSrcRect.Max.X -= ExtraPixelsRemove;
+	}
 }
 
 void FDisplayClusterMediaInputBase::ImportMediaData(FRHICommandListImmediate& RHICmdList, const FMediaTextureInfo& TextureInfo)
@@ -85,29 +150,27 @@ void FDisplayClusterMediaInputBase::ImportMediaData(FRHICommandListImmediate& RH
 
 	if (SrcTexture && DstTexture)
 	{
-		const FIntPoint DstRegionSize = TextureInfo.Region.Size();
+		FIntRect SrcRect(FIntPoint(0, 0), SrcTexture->GetDesc().Extent);
+		FIntRect DstRect(TextureInfo.Region);
+		OverrideTextureRegions_RenderThread(SrcTexture, DstTexture, SrcRect, DstRect);
 
 		const bool bSrcSrgb = EnumHasAnyFlags(SrcTexture->GetFlags(), TexCreate_SRGB);
 		const bool bDstSrgb = EnumHasAnyFlags(DstTexture->GetFlags(), TexCreate_SRGB);
 
 		if ((SrcTexture->GetDesc().Format == DstTexture->GetDesc().Format)
-			&& (SrcTexture->GetDesc().Extent == DstRegionSize)
+			&& (SrcRect.Size() == DstRect.Size())
 			&& (bSrcSrgb == bDstSrgb)
-		)
+			)
 		{
 			FRHICopyTextureInfo CopyInfo;
-
-			CopyInfo.DestPosition = FIntVector(TextureInfo.Region.Min.X, TextureInfo.Region.Min.Y, 0);
-			CopyInfo.Size = FIntVector(DstRegionSize.X, DstRegionSize.Y, 0);
-
+			CopyInfo.SourcePosition = FIntVector(SrcRect.Min.X, SrcRect.Min.Y, 0);
+			CopyInfo.DestPosition = FIntVector(DstRect.Min.X, DstRect.Min.Y, 0);
+			CopyInfo.Size = FIntVector(DstRect.Size().X, DstRect.Size().Y, 0);
 			TransitionAndCopyTexture(RHICmdList, SrcTexture, DstTexture, CopyInfo);
 		}
 		else
 		{
-			DisplayClusterMediaHelpers::ResampleTexture_RenderThread(
-				RHICmdList, SrcTexture, DstTexture,
-				FIntRect(FIntPoint(0, 0), SrcTexture->GetDesc().Extent),
-				TextureInfo.Region);
+			DisplayClusterMediaHelpers::ResampleTexture_RenderThread(RHICmdList, SrcTexture, DstTexture, SrcRect, DstRect);
 		}
 	}
 }
diff --git a/Engine/Plugins/Runtime/nDisplay/Source/DisplayClusterMedia/Private/Input/DisplayClusterMediaInputBase.h b/Engine/Plugins/Runtime/nDisplay/Source/DisplayClusterMedia/Private/Input/DisplayClusterMediaInputBase.h
--- a/Engine/Plugins/Runtime/nDisplay/Source/DisplayClusterMedia/Private/Input/DisplayClusterMediaInputBase.h
+++ b/Engine/Plugins/Runtime/nDisplay/Source/DisplayClusterMedia/Private/Input/DisplayClusterMediaInputBase.h
@@ -61,6 +61,8 @@ private:
 	bool StartPlayer();
 	void OnPlayerClosed();
 
+	void OverrideTextureRegions_RenderThread(FRHITexture* const SrcTexture, FRHITexture* const DstTexture, FIntRect& InOutSrcRect, FIntRect& InOutDstRect) const;
+
 private:
 	//~ Begin GC by AddReferencedObjects
 	TObjectPtr<UMediaSource>  MediaSource = nullptr;
@@ -73,4 +75,7 @@ private:
 
 	// Used to control the rate at which we try to restart the player
 	double LastRestartTimestamp = 0;
+
+	// [Temp workaround] Whether current media is Rivermax
+	bool bRunningRivermaxMedia = false;
 };
