nDisplay: Make Media Settings EditDefaultOnly.

This avoids the problem of the instance media settings being inadvertently dirty and never propagated when the user edits the media settings in the nDisplay blueprint. This limits this property to the CDO but this is where it is almost always edited and also avoids the confusion that is happening relatively often.

diff --git a/Engine/Plugins/Runtime/nDisplay/Source/DisplayCluster/Public/Components/DisplayClusterICVFXCameraComponent.h b/Engine/Plugins/Runtime/nDisplay/Source/DisplayCluster/Public/Components/DisplayClusterICVFXCameraComponent.h
--- a/Engine/Plugins/Runtime/nDisplay/Source/DisplayCluster/Public/Components/DisplayClusterICVFXCameraComponent.h
+++ b/Engine/Plugins/Runtime/nDisplay/Source/DisplayCluster/Public/Components/DisplayClusterICVFXCameraComponent.h
@@ -196,7 +196,7 @@ private:
 	UPROPERTY(EditAnywhere, Transient, Category = "Texture Replacement", meta = (PropertyPath = "CameraSettings.RenderSettings.Replace.TextureRegion"))
 	FDisplayClusterEditorPropertyReference TextureRegionRef;
 
-	UPROPERTY(EditAnywhere, Transient, Category = "Media", meta = (PropertyPath = "CameraSettings.RenderSettings.Media"))
+	UPROPERTY(EditDefaultsOnly, Transient, Category = "Media", meta = (PropertyPath = "CameraSettings.RenderSettings.Media"))
 	FDisplayClusterEditorPropertyReference MediaRef;
 
 	UPROPERTY(EditAnywhere, Transient, Category = Configuration, meta = (PropertyPath = "CameraSettings.RenderSettings.RenderOrder"))
diff --git a/Engine/Plugins/Runtime/nDisplay/Source/DisplayCluster/Public/DisplayClusterRootActor.h b/Engine/Plugins/Runtime/nDisplay/Source/DisplayCluster/Public/DisplayClusterRootActor.h
--- a/Engine/Plugins/Runtime/nDisplay/Source/DisplayCluster/Public/DisplayClusterRootActor.h
+++ b/Engine/Plugins/Runtime/nDisplay/Source/DisplayCluster/Public/DisplayClusterRootActor.h
@@ -339,7 +339,7 @@ private:
 	FDisplayClusterEditorPropertyReference LightcardPerViewportOCIOProfilesRef;
 
 	// Media
-	UPROPERTY(EditAnywhere, Transient, Category = "Media", meta = (PropertyPath = "CurrentConfigData.MediaSettings"))
+	UPROPERTY(EditDefaultsOnly, Transient, Category = "Media", meta = (PropertyPath = "CurrentConfigData.MediaSettings"))
 	FDisplayClusterEditorPropertyReference MediaSettingsRef;
 
 #endif // WITH_EDITORONLY_DATA
diff --git a/Engine/Plugins/Runtime/nDisplay/Source/DisplayClusterConfiguration/Public/DisplayClusterConfigurationTypes.h b/Engine/Plugins/Runtime/nDisplay/Source/DisplayClusterConfiguration/Public/DisplayClusterConfigurationTypes.h
--- a/Engine/Plugins/Runtime/nDisplay/Source/DisplayClusterConfiguration/Public/DisplayClusterConfigurationTypes.h
+++ b/Engine/Plugins/Runtime/nDisplay/Source/DisplayClusterConfiguration/Public/DisplayClusterConfigurationTypes.h
@@ -358,7 +358,7 @@ public:
 	TMap<FString, FDisplayClusterConfigurationPostprocess> Postprocess;
 
 	// Media settings
-	UPROPERTY(EditAnywhere, BlueprintReadWrite, Category = "Media", meta = (DisplayName = "Media"))
+	UPROPERTY(EditDefaultsOnly, BlueprintReadWrite, Category = "Media", meta = (DisplayName = "Media"))
 	FDisplayClusterConfigurationMediaNodeBackbuffer MediaSettings;
 
 #if WITH_EDITORONLY_DATA
diff --git a/Engine/Plugins/Runtime/nDisplay/Source/DisplayClusterConfiguration/Public/DisplayClusterConfigurationTypes_ICVFX.h b/Engine/Plugins/Runtime/nDisplay/Source/DisplayClusterConfiguration/Public/DisplayClusterConfigurationTypes_ICVFX.h
--- a/Engine/Plugins/Runtime/nDisplay/Source/DisplayClusterConfiguration/Public/DisplayClusterConfigurationTypes_ICVFX.h
+++ b/Engine/Plugins/Runtime/nDisplay/Source/DisplayClusterConfiguration/Public/DisplayClusterConfigurationTypes_ICVFX.h
@@ -503,7 +503,7 @@ public:
 	FDisplayClusterConfigurationPostRender_Override Replace;
 
 	// Media settings
-	UPROPERTY(EditAnywhere, BlueprintReadWrite, Category = NDisplay, meta = (DisplayName = "Media"))
+	UPROPERTY(EditDefaultsOnly, BlueprintReadWrite, Category = NDisplay, meta = (DisplayName = "Media"))
 	FDisplayClusterConfigurationMediaICVFX Media;
 
 	UPROPERTY()
