Removed #pragma once from cpp file

diff --git a/Engine/Plugins/Runtime/nDisplay/Source/SharedMemoryMediaEditor/Private/ModularFeatures/SharedMemoryMediaInitializerFeature.cpp b/Engine/Plugins/Runtime/nDisplay/Source/SharedMemoryMediaEditor/Private/ModularFeatures/SharedMemoryMediaInitializerFeature.cpp
--- a/Engine/Plugins/Runtime/nDisplay/Source/SharedMemoryMediaEditor/Private/ModularFeatures/SharedMemoryMediaInitializerFeature.cpp
+++ b/Engine/Plugins/Runtime/nDisplay/Source/SharedMemoryMediaEditor/Private/ModularFeatures/SharedMemoryMediaInitializerFeature.cpp
@@ -1,7 +1,5 @@
 // Copyright Epic Games, Inc. All Rights Reserved.
 
-#pragma once
-
 #include "ModularFeatures/SharedMemoryMediaInitializerFeature.h"
 
 #include "SharedMemoryMediaOutput.h"
