Initialized some FVector variables that were potentially uninitialized, deleted some that weren't used.

diff --git a/Engine/Plugins/Experimental/Avalanche/Source/AvalancheEffectorsEditor/Private/Cloner/AvaClonerActorVis.cpp b/Engine/Plugins/Experimental/Avalanche/Source/AvalancheEffectorsEditor/Private/Cloner/AvaClonerActorVis.cpp
--- a/Engine/Plugins/Experimental/Avalanche/Source/AvalancheEffectorsEditor/Private/Cloner/AvaClonerActorVis.cpp
+++ b/Engine/Plugins/Experimental/Avalanche/Source/AvalancheEffectorsEditor/Private/Cloner/AvaClonerActorVis.cpp
@@ -55,7 +55,7 @@ void FAvaClonerActorVisualizer::StoreInitialValues()
 	{
 		const ECEClonerPlane Plane = HoneycombLayout->GetPlane();
 
-		FVector Spacing;
+		FVector Spacing = FVector::ZeroVector;
 		if (Plane == ECEClonerPlane::XY)
 		{
 			Spacing = FVector(HoneycombLayout->GetWidthSpacing(), HoneycombLayout->GetHeightSpacing(), 0);
diff --git a/Engine/Plugins/Experimental/Avalanche/Source/AvalancheLevelViewport/Private/Customizations/AvaLevelViewportComponentTransformDetails.cpp b/Engine/Plugins/Experimental/Avalanche/Source/AvalancheLevelViewport/Private/Customizations/AvaLevelViewportComponentTransformDetails.cpp
--- a/Engine/Plugins/Experimental/Avalanche/Source/AvalancheLevelViewport/Private/Customizations/AvaLevelViewportComponentTransformDetails.cpp
+++ b/Engine/Plugins/Experimental/Avalanche/Source/AvalancheLevelViewport/Private/Customizations/AvaLevelViewportComponentTransformDetails.cpp
@@ -1056,9 +1056,9 @@ void FAvaLevelViewportComponentTransformDetails::CacheTransform()
 		CachedScale.Z = SelectionScale.Z;
 	}
 
-	FVector CurLoc;
+	FVector CurLoc = FVector::ZeroVector;
 	FRotator CurRot;
-	FVector CurScale;
+	FVector CurScale = FVector::ZeroVector;
 	bool bFirstObject = true;
 	bool bSkipFirst = false;
 
diff --git a/Engine/Plugins/Experimental/Avalanche/Source/AvalancheTextEditor/Private/Visualizer/AvaTextVisualizer.cpp b/Engine/Plugins/Experimental/Avalanche/Source/AvalancheTextEditor/Private/Visualizer/AvaTextVisualizer.cpp
--- a/Engine/Plugins/Experimental/Avalanche/Source/AvalancheTextEditor/Private/Visualizer/AvaTextVisualizer.cpp
+++ b/Engine/Plugins/Experimental/Avalanche/Source/AvalancheTextEditor/Private/Visualizer/AvaTextVisualizer.cpp
@@ -535,7 +535,7 @@ bool FAvaTextVisualizer::HandleInputDeltaInternal(FEditorViewportClient* InViewp
 					float GradientRot = InitialGradientRotation;
 					float RotIncrZ = InAccumulatedTranslation.Z * FTextVisualizerStatics::GradientRotHandleScale;
 
-					FVector HandleToCenter;
+					FVector HandleToCenter = FVector::ZeroVector;
 
 					if (bEditingGradientRotation_EndHandle)
 					{
diff --git a/Engine/Plugins/Experimental/Avalanche/Source/AvalancheViewport/Private/Interaction/AvaSnapOperation.cpp b/Engine/Plugins/Experimental/Avalanche/Source/AvalancheViewport/Private/Interaction/AvaSnapOperation.cpp
--- a/Engine/Plugins/Experimental/Avalanche/Source/AvalancheViewport/Private/Interaction/AvaSnapOperation.cpp
+++ b/Engine/Plugins/Experimental/Avalanche/Source/AvalancheViewport/Private/Interaction/AvaSnapOperation.cpp
@@ -1248,7 +1248,7 @@ bool FAvaSnapOperation::SnapDragLocation(const TArray<FAvaSnapPoint>& InDraggedA
 			// Single axis movement
 			if (DraggedAxisCount == 1)
 			{
-				FVector DragDirection;
+				FVector DragDirection = FVector::ZeroVector;
 
 				if (bIsDraggingX)
 				{
diff --git a/Engine/Plugins/Experimental/GizmoEdMode/Source/LightGizmos/Private/SpotLightGizmo.cpp b/Engine/Plugins/Experimental/GizmoEdMode/Source/LightGizmos/Private/SpotLightGizmo.cpp
--- a/Engine/Plugins/Experimental/GizmoEdMode/Source/LightGizmos/Private/SpotLightGizmo.cpp
+++ b/Engine/Plugins/Experimental/GizmoEdMode/Source/LightGizmos/Private/SpotLightGizmo.cpp
@@ -113,7 +113,6 @@ void USpotLightGizmo::Shutdown()
 FInputRayHit USpotLightGizmo::BeginHoverSequenceHitTest(const FInputDeviceRay& PressPos)
 {
 	FHitResult HitResult;
-	FVector HitAxis;
 	FTransform DragTransform;
 
 	if (HitTest(PressPos.WorldRay, HitResult, DragTransform))
@@ -300,7 +299,6 @@ void USpotLightGizmo::OnBeginDrag(const FInputDeviceRay& Ray)
 
 	FRay HitCheckRay(Start, End - Start);
 	FHitResult HitResult;
-	FVector HitAxis;
 	FTransform DragTransform;
 
 	// Check if any of the components were hit
diff --git a/Engine/Plugins/NetcodeUnitTest/NetcodeUnitTest/Source/NetcodeUnitTest/Private/UnitTests/Engine/PackedVectorTest.cpp b/Engine/Plugins/NetcodeUnitTest/NetcodeUnitTest/Source/NetcodeUnitTest/Private/UnitTests/Engine/PackedVectorTest.cpp
--- a/Engine/Plugins/NetcodeUnitTest/NetcodeUnitTest/Source/NetcodeUnitTest/Private/UnitTests/Engine/PackedVectorTest.cpp
+++ b/Engine/Plugins/NetcodeUnitTest/NetcodeUnitTest/Source/NetcodeUnitTest/Private/UnitTests/Engine/PackedVectorTest.cpp
@@ -321,7 +321,7 @@ bool UPackedVectorTest::ExecuteWriteDoubleReadFloatTest(TMap<FString, bool>& Tes
 
 			const double ScalarValue = Test.TestValues[ValueIt];
 			const FVector3d WriteValue(ScalarValue);
-			FVector3f ReadValue;
+			FVector3f ReadValue = FVector3f::ZeroVector;
 
 			const bool bOverflowOrNan = !Test.Writer(WriteValue, Writer);
 			bool LocalSuccess = !Writer.GetError();
diff --git a/Engine/Plugins/Runtime/Soundscape/Source/SoundScape/Private/SoundscapeColor.cpp b/Engine/Plugins/Runtime/Soundscape/Source/SoundScape/Private/SoundscapeColor.cpp
--- a/Engine/Plugins/Runtime/Soundscape/Source/SoundScape/Private/SoundscapeColor.cpp
+++ b/Engine/Plugins/Runtime/Soundscape/Source/SoundScape/Private/SoundscapeColor.cpp
@@ -432,9 +432,9 @@ void UActiveSoundscapeColor::CalculateSpawnParams()
 			// Get the new sound's start time
 			NewSoundParams.StartTime = PlaybackBehavior.bRandomizeStartingSeekTime ? FMath::FRandRange(0.0f, Sound->Duration) : 0.0f;
 
-			FVector ListenerLocation;
-			FVector ListenerForward;
-			FVector ListenerUp;
+			FVector ListenerLocation = FVector::ZeroVector;
+			FVector ListenerForward = FVector::ZeroVector;
+			FVector ListenerUp = FVector::ZeroVector;
 
 			if (FAudioDevice* AudioDevice = World->GetAudioDeviceRaw())
 			{
diff --git a/Engine/Plugins/Runtime/nDisplay/Source/DisplayCluster/Private/Render/Viewport/Preview/DisplayClusterViewportManagerPreview_Render.cpp b/Engine/Plugins/Runtime/nDisplay/Source/DisplayCluster/Private/Render/Viewport/Preview/DisplayClusterViewportManagerPreview_Render.cpp
--- a/Engine/Plugins/Runtime/nDisplay/Source/DisplayCluster/Private/Render/Viewport/Preview/DisplayClusterViewportManagerPreview_Render.cpp
+++ b/Engine/Plugins/Runtime/nDisplay/Source/DisplayCluster/Private/Render/Viewport/Preview/DisplayClusterViewportManagerPreview_Render.cpp
@@ -173,7 +173,6 @@ bool FDisplayClusterViewportManagerPreview::InitializeClusterNodePreview(const E
 							if (ViewIt.IsViewportContextCanBeRendered() && ViewFamily->RenderTarget)
 							{
 								// Calculate the player's view information.
-								FVector  ViewLocation;
 								FRotator ViewRotation;
 								if (FSceneView* View = ViewportPtr->ViewportPreview->CalcSceneView(*ViewFamily, ViewIt.ContextNum))
 								{
diff --git a/Engine/Plugins/Runtime/nDisplay/Source/DisplayClusterWarp/Private/WarpBlend/Math/DisplayClusterWarpBlendMath_WarpProceduralMesh.h b/Engine/Plugins/Runtime/nDisplay/Source/DisplayClusterWarp/Private/WarpBlend/Math/DisplayClusterWarpBlendMath_WarpProceduralMesh.h
--- a/Engine/Plugins/Runtime/nDisplay/Source/DisplayClusterWarp/Private/WarpBlend/Math/DisplayClusterWarpBlendMath_WarpProceduralMesh.h
+++ b/Engine/Plugins/Runtime/nDisplay/Source/DisplayClusterWarp/Private/WarpBlend/Math/DisplayClusterWarpBlendMath_WarpProceduralMesh.h
@@ -41,7 +41,7 @@ public:
 
 		// Calc static normal and plane
 		bool bAverageNormalValid = false;
-		FVector AverageNormal;
+		FVector AverageNormal = FVector::ZeroVector;
 
 		for (int32 TriIdx = 0; TriIdx < TriNum; ++TriIdx)
 		{
