Fixed unused or potentially uninitialized variables.

diff --git a/Engine/Plugins/Enterprise/LidarPointCloud/Source/LidarPointCloudRuntime/Private/IO/LidarPointCloudFileIO_ASCII.cpp b/Engine/Plugins/Enterprise/LidarPointCloud/Source/LidarPointCloudRuntime/Private/IO/LidarPointCloudFileIO_ASCII.cpp
--- a/Engine/Plugins/Enterprise/LidarPointCloud/Source/LidarPointCloudRuntime/Private/IO/LidarPointCloudFileIO_ASCII.cpp
+++ b/Engine/Plugins/Enterprise/LidarPointCloud/Source/LidarPointCloudRuntime/Private/IO/LidarPointCloudFileIO_ASCII.cpp
@@ -581,7 +581,7 @@ FVector2D FLidarPointCloudImportSettings_ASCII::ReadFileMinMaxColumns(TArray<int
 		return FVector2D::ZeroVector;
 	}
 
-	FVector2D Result;
+	FVector2D Result = FVector2D::ZeroVector;
 
 	TUniquePtr<FArchive> Reader(IFileManager::Get().CreateFileReader(*Filename));
 	if (Reader)
diff --git a/Engine/Plugins/Experimental/SVGImporter/Source/SVGImporter/Private/ProceduralMeshes/JoinedSVGDynamicMeshComponent.cpp b/Engine/Plugins/Experimental/SVGImporter/Source/SVGImporter/Private/ProceduralMeshes/JoinedSVGDynamicMeshComponent.cpp
--- a/Engine/Plugins/Experimental/SVGImporter/Source/SVGImporter/Private/ProceduralMeshes/JoinedSVGDynamicMeshComponent.cpp
+++ b/Engine/Plugins/Experimental/SVGImporter/Source/SVGImporter/Private/ProceduralMeshes/JoinedSVGDynamicMeshComponent.cpp
@@ -229,7 +229,7 @@ void UJoinedSVGDynamicMeshComponent::LoadMaterialSetParameters()
 
 		if (const TObjectPtr<UMaterialInstanceDynamic>& Material = ShapesMaterials[MaterialID])
 		{
-			FLinearColor ColorToApply;
+			FLinearColor ColorToApply = FLinearColor::Black;
 
 			switch (Coloring)
 			{
diff --git a/Engine/Plugins/Experimental/Text3D/Source/Text3D/Private/Arrangement.cpp b/Engine/Plugins/Experimental/Text3D/Source/Text3D/Private/Arrangement.cpp
--- a/Engine/Plugins/Experimental/Text3D/Source/Text3D/Private/Arrangement.cpp
+++ b/Engine/Plugins/Experimental/Text3D/Source/Text3D/Private/Arrangement.cpp
@@ -196,7 +196,6 @@ bool FArrangement::insert_segment(FVector2d A, FVector2d B, double Tol)
 
 FIndex2i FArrangement::split_segment_at_t(int EID, double T, double Tol)
 {
-    FVector2d V1, V2;
     FIndex2i ev = Graph.GetEdgeV(EID);
     FSegment2d seg = FSegment2d(Graph.GetVertex(ev.A), Graph.GetVertex(ev.B));
 
diff --git a/Engine/Plugins/Runtime/OpenXR/Source/OpenXRInput/Private/OpenXRInput.cpp b/Engine/Plugins/Runtime/OpenXR/Source/OpenXRInput/Private/OpenXRInput.cpp
--- a/Engine/Plugins/Runtime/OpenXR/Source/OpenXRInput/Private/OpenXRInput.cpp
+++ b/Engine/Plugins/Runtime/OpenXR/Source/OpenXRInput/Private/OpenXRInput.cpp
@@ -1353,7 +1353,6 @@ ETrackingStatus FOpenXRInputPlugin::FOpenXRInput::GetControllerTrackingStatus(co
 		XrResult Result = xrGetActionStatePose(Session, &GetInfo, &State);
 		if (XR_SUCCEEDED(Result) && State.isActive)
 		{
-			FQuat Orientation;
 			bool bIsTracked = OpenXRHMD->GetIsTracked(GetDeviceIDForMotionSource(MotionSource));
 			return bIsTracked ? ETrackingStatus::Tracked : ETrackingStatus::NotTracked;
 		}
diff --git a/Engine/Plugins/Runtime/nDisplay/Source/DisplayCluster/Private/Render/Viewport/Preview/DisplayClusterViewportManagerPreview_Render.cpp b/Engine/Plugins/Runtime/nDisplay/Source/DisplayCluster/Private/Render/Viewport/Preview/DisplayClusterViewportManagerPreview_Render.cpp
--- a/Engine/Plugins/Runtime/nDisplay/Source/DisplayCluster/Private/Render/Viewport/Preview/DisplayClusterViewportManagerPreview_Render.cpp
+++ b/Engine/Plugins/Runtime/nDisplay/Source/DisplayCluster/Private/Render/Viewport/Preview/DisplayClusterViewportManagerPreview_Render.cpp
@@ -173,7 +173,6 @@ bool FDisplayClusterViewportManagerPreview::InitializeClusterNodePreview(const E
 							if (ViewIt.IsViewportContextCanBeRendered() && ViewFamily->RenderTarget)
 							{
 								// Calculate the player's view information.
-								FRotator ViewRotation;
 								if (FSceneView* View = ViewportPtr->ViewportPreview->CalcSceneView(*ViewFamily, ViewIt.ContextNum))
 								{
 									// Apply viewport context settings to view (crossGPU, visibility, etc)
diff --git a/Engine/Source/Runtime/Core/Public/Math/Box.h b/Engine/Source/Runtime/Core/Public/Math/Box.h
--- a/Engine/Source/Runtime/Core/Public/Math/Box.h
+++ b/Engine/Source/Runtime/Core/Public/Math/Box.h
@@ -49,6 +49,11 @@ public:
 	/**
 	 * Default constructor.
 	 * Creates a new box with uninitialized extents and marks it as invalid.
+	 *
+	 * NOTE: This default constructor is unlike TVector, TMatrix etc. in that it
+	 *       actually initializes the instance.  Ideally it should be = default;
+	 *       in the same way, but this would break backwards compatibility.  It
+	 *       might be possible if the questionable IsValid field ever gets removed.
 	 */
 	constexpr TBox()
 #if ENABLE_UNINITIALIZED_BOX_DIAGNOSTIC
diff --git a/Engine/Source/Runtime/Core/Public/Math/Vector4.h b/Engine/Source/Runtime/Core/Public/Math/Vector4.h
--- a/Engine/Source/Runtime/Core/Public/Math/Vector4.h
+++ b/Engine/Source/Runtime/Core/Public/Math/Vector4.h
@@ -129,6 +129,10 @@ public:
 	 * @param InY Y Coordinate.
 	 * @param InZ Z Coordinate.
 	 * @param InW W Coordinate.
+	 *
+	 * NOTE: This default constructor is unlike TVector, TMatrix etc. in that it
+	 *       actually initializes the instance.  Ideally it should be = default;
+	 *       in the same way, but this would break backwards compatibility.
 	 */
 	explicit TVector4(T InX = 0.0f, T InY = 0.0f, T InZ = 0.0f, T InW = 1.0f);
 
