nDisplay: Fixed rendering for Unreal Stage App
- Modified IDisplayClusterScenePreview interface.
- Some RootActor properties can be overridden using the FDisplayClusterRootActorPropertyOverrides structure.
- Code cleanup

diff --git a/Engine/Plugins/Runtime/nDisplay/Source/DisplayCluster/Private/Game/EngineClasses/Scene/DisplayClusterRootActor.cpp b/Engine/Plugins/Runtime/nDisplay/Source/DisplayCluster/Private/Game/EngineClasses/Scene/DisplayClusterRootActor.cpp
--- a/Engine/Plugins/Runtime/nDisplay/Source/DisplayCluster/Private/Game/EngineClasses/Scene/DisplayClusterRootActor.cpp
+++ b/Engine/Plugins/Runtime/nDisplay/Source/DisplayCluster/Private/Game/EngineClasses/Scene/DisplayClusterRootActor.cpp
@@ -1,6 +1,7 @@
 // Copyright Epic Games, Inc. All Rights Reserved.
 
 #include "DisplayClusterRootActor.h"
+#include "DisplayClusterRootActorContainers.h"
 
 #include "Async/ParallelFor.h"
 #include "Components/SceneComponent.h"
@@ -1125,15 +1126,6 @@ void ADisplayClusterRootActor::Tick(float DeltaSeconds)
 		NewPreviewSettings.bPreviewEnable = false;
 	}
 
-#if WITH_EDITOR
-	if (!PreviewEnableOverriders.IsEmpty())
-	{
-		// Preview rendering is overridden, so force it on
-		bEnablePreviewInScene = true;
-		NewPreviewSettings.bPreviewEnable = true;
-	}
-#endif
-
 	// Update entire cluster preview rendering
 	if (IDisplayClusterViewportManager* ViewportManager = bEnablePreviewInScene ? GetOrCreateViewportManager() : GetViewportManager())
 	{
@@ -1462,3 +1454,31 @@ void ADisplayClusterRootActor::AddReferencedObjects(UObject* InThis, FReferenceC
 
 	Super::AddReferencedObjects(InThis, Collector);
 }
+
+#define OVERRIDE_PROPERTY(PROPERTY_NAME)\
+	if (InPropertyOverrides.PROPERTY_NAME.IsSet())\
+	{\
+		PROPERTY_NAME = InPropertyOverrides.PROPERTY_NAME.GetValue();\
+	}
+
+void ADisplayClusterRootActor::OverrideRootActorProperties(const FDisplayClusterRootActorPropertyOverrides& InPropertyOverrides)
+{
+	OVERRIDE_PROPERTY(bPreviewInGameEnable);
+	OVERRIDE_PROPERTY(bPreviewInGameRenderFrustum);
+	OVERRIDE_PROPERTY(bPreviewEnable);
+	OVERRIDE_PROPERTY(PreviewRenderTargetRatioMult);
+	OVERRIDE_PROPERTY(bPreviewEnablePostProcess);
+	OVERRIDE_PROPERTY(bPreviewEnableOverlayMaterial);
+	OVERRIDE_PROPERTY(bEnablePreviewTechvis);
+	OVERRIDE_PROPERTY(bEnablePreviewMesh);
+	OVERRIDE_PROPERTY(bEnablePreviewEditableMesh);
+	OVERRIDE_PROPERTY(PreviewSetttingsSource);
+	OVERRIDE_PROPERTY(bFreezePreviewRender);
+	OVERRIDE_PROPERTY(bPreviewICVFXFrustums);
+	OVERRIDE_PROPERTY(PreviewICVFXFrustumsFarDistance);
+	OVERRIDE_PROPERTY(TickPerFrame);
+	OVERRIDE_PROPERTY(ViewportsPerFrame);
+	OVERRIDE_PROPERTY(PreviewMaxTextureDimension);
+}
+
+#undef OVERRIDE_PROPERTY
diff --git a/Engine/Plugins/Runtime/nDisplay/Source/DisplayCluster/Private/Game/EngineClasses/Scene/DisplayClusterRootActorEditor.cpp b/Engine/Plugins/Runtime/nDisplay/Source/DisplayCluster/Private/Game/EngineClasses/Scene/DisplayClusterRootActorEditor.cpp
--- a/Engine/Plugins/Runtime/nDisplay/Source/DisplayCluster/Private/Game/EngineClasses/Scene/DisplayClusterRootActorEditor.cpp
+++ b/Engine/Plugins/Runtime/nDisplay/Source/DisplayCluster/Private/Game/EngineClasses/Scene/DisplayClusterRootActorEditor.cpp
@@ -346,18 +346,6 @@ void ADisplayClusterRootActor::RerunConstructionScripts_Editor()
 	UpdateInnerFrustumPriority();
 }
 
-void ADisplayClusterRootActor::AddPreviewEnableOverride(const uint8* Object)
-{
-	check(Object);
-	PreviewEnableOverriders.Add(Object);
-}
-
-void ADisplayClusterRootActor::RemovePreviewEnableOverride(const uint8* Object)
-{
-	check(Object);
-	PreviewEnableOverriders.Remove(Object);
-}
-
 void ADisplayClusterRootActor::UpdateInnerFrustumPriority()
 {
 	if (InnerFrustumPriority.Num() == 0)
diff --git a/Engine/Plugins/Runtime/nDisplay/Source/DisplayCluster/Private/Game/EngineClasses/Scene/DisplayClusterStageGeometryComponent.cpp b/Engine/Plugins/Runtime/nDisplay/Source/DisplayCluster/Private/Game/EngineClasses/Scene/DisplayClusterStageGeometryComponent.cpp
--- a/Engine/Plugins/Runtime/nDisplay/Source/DisplayCluster/Private/Game/EngineClasses/Scene/DisplayClusterStageGeometryComponent.cpp
+++ b/Engine/Plugins/Runtime/nDisplay/Source/DisplayCluster/Private/Game/EngineClasses/Scene/DisplayClusterStageGeometryComponent.cpp
@@ -57,10 +57,17 @@ bool UDisplayClusterStageGeometryComponent::GetStageDistanceAndNormal(const FVec
 
 		auto GetPixel = [&GeometryMap](uint32 InX, uint32 InY)
 		{
-			uint32 ClampedX = FMath::Clamp(InX, (uint32)0, GeometryMapSize - 1);
-			uint32 ClampedY = FMath::Clamp(InY, (uint32)0, GeometryMapSize - 1);
+			const uint32 ClampedX = FMath::Clamp(InX, (uint32)0, GeometryMapSize - 1);
+			const uint32 ClampedY = FMath::Clamp(InY, (uint32)0, GeometryMapSize - 1);
 
-			return GeometryMap.GeometryData[ClampedY * GeometryMapSize + ClampedX].GetFloats();
+			const uint32 PixelIndex = ClampedY * GeometryMapSize + ClampedX;
+			if (!GeometryMap.GeometryData.IsValidIndex(PixelIndex))
+			{
+				// Deny access to a pixel that is out of range.
+				return FLinearColor::Black;
+			}
+
+			return GeometryMap.GeometryData[PixelIndex].GetFloats();
 		};
 
 		// Use the view matrix of the geometry map to convert from world coordinates to view coordinates, then convert those coordinates into dome projection space
diff --git a/Engine/Plugins/Runtime/nDisplay/Source/DisplayCluster/Public/DisplayClusterRootActor.h b/Engine/Plugins/Runtime/nDisplay/Source/DisplayCluster/Public/DisplayClusterRootActor.h
--- a/Engine/Plugins/Runtime/nDisplay/Source/DisplayCluster/Public/DisplayClusterRootActor.h
+++ b/Engine/Plugins/Runtime/nDisplay/Source/DisplayCluster/Public/DisplayClusterRootActor.h
@@ -43,7 +43,7 @@ class UDisplayClusterStageGeometryComponent;
 class UDisplayClusterStageIsosphereComponent;
 class UDisplayClusterSyncTickComponent;
 class UProceduralMeshComponent;
-
+struct FDisplayClusterRootActorPropertyOverrides;
 
 /**
  * VR root. This contains nDisplay VR hierarchy in the game.
@@ -99,6 +99,8 @@ public:
 	/** Returns true if this RootActor is running in DC mode. */
 	bool IsRunningDisplayCluster() const;
 
+	/** Override properties of this root actor. */
+	void OverrideRootActorProperties(const FDisplayClusterRootActorPropertyOverrides& InPropertyOverrides);
 
 	UDisplayClusterConfigurationData* GetDefaultConfigDataFromAsset() const;
 	UDisplayClusterConfigurationData* GetConfigData() const;
@@ -623,13 +625,15 @@ public:
 	bool ShouldThisFrameOutputPreviewToPostProcessRenderTarget() const { return false; }
 
 	/** Force preview rendering to be enabled regardless of the user's setting until a matching RemovePreviewEnableOverride call is made. */
-	void AddPreviewEnableOverride(const uint8* Object);
+	UE_DEPRECATED(5.5, "This function has been deprecated.")
+	void AddPreviewEnableOverride(const uint8* Object) { };
 
 	/**
 	 * Stop forcing preview rendering to be enabled for this caller. If other objects have called AddPreviewEnableOverride, it will remain
 	 * forced until they have also removed their overrides.
 	 */
-	void RemovePreviewEnableOverride(const uint8* Object);
+	UE_DEPRECATED(5.5, "This function has been deprecated.")
+	void RemovePreviewEnableOverride(const uint8* Object) { };
 
 	UE_DEPRECATED(5.4, "This function has been deprecated.")
 	float GetPreviewRenderTargetRatioMult() const
@@ -682,9 +686,6 @@ private:
 
 	bool bIsSelectedInEditor = false;
 
-	/* Addresses of callers to AddPreviewEnableOverride that haven't removed their overrides yet. */
-	TSet<const uint8*> PreviewEnableOverriders;
-
 	TWeakPtr<IDisplayClusterConfiguratorBlueprintEditor> ToolkitPtr;
 
 	// UE_DEPRECATED 5.4
diff --git a/Engine/Plugins/Runtime/nDisplay/Source/DisplayCluster/Public/DisplayClusterRootActorContainers.h b/Engine/Plugins/Runtime/nDisplay/Source/DisplayCluster/Public/DisplayClusterRootActorContainers.h
new file mode 100644
--- /dev/null
+++ b/Engine/Plugins/Runtime/nDisplay/Source/DisplayCluster/Public/DisplayClusterRootActorContainers.h
@@ -0,0 +1,60 @@
+// Copyright Epic Games, Inc. All Rights Reserved.
+
+#pragma once
+
+#include "CoreMinimal.h"
+#include "DisplayClusterConfigurationTypes_Enums.h"
+
+/**
+ * Optional overrides of root actor settings.
+ */
+struct FDisplayClusterRootActorPropertyOverrides
+{
+	/** Render this DCRA in game for Standalone/Package builds. */
+	TOptional<bool> bPreviewInGameEnable;
+
+	/** Render ICVFX Frustums in game. */
+	TOptional<bool> bPreviewInGameRenderFrustum;
+
+	/** Render the scene and display it as a preview on the nDisplay root actor in the editor.  This will impact editor performance. */
+	TOptional<bool> bPreviewEnable;
+	
+	/** Adjust resolution scaling for the editor preview. */
+	TOptional<float> PreviewRenderTargetRatioMult;
+
+	/** Enable PostProcess for preview. */
+	TOptional<bool> bPreviewEnablePostProcess;
+
+	/** Show overlay material on the preview mesh when preview rendering is enabled (UMeshComponent::OverlayMaterial). */
+	TOptional<bool> bPreviewEnableOverlayMaterial;
+
+	/** Configure the root actor for Techvis rendering with preview components. */
+	TOptional<bool> bEnablePreviewTechvis;
+
+	/** Enable the use of a preview mesh for the preview for this DCRA. */
+	TOptional<bool> bEnablePreviewMesh;
+
+	/** Enable the use of a preview editable mesh for the preview for this DCRA. */
+	TOptional<bool> bEnablePreviewEditableMesh;
+
+	/** Determines where the preview settings will be retrieved from. */
+	TOptional<EDisplayClusterConfigurationRootActorPreviewSettingsSource> PreviewSetttingsSource;
+
+	/** Freeze preview render.  This will impact editor performance. */
+	TOptional<bool> bFreezePreviewRender;
+
+	/** Render ICVFX Frustums */
+	TOptional<bool> bPreviewICVFXFrustums;
+
+	/** Render ICVFX Frustums */
+	TOptional<float> PreviewICVFXFrustumsFarDistance;
+
+	/** Tick Per Frame */
+	TOptional<int32> TickPerFrame;
+
+	/** Max amount of Viewports Per Frame */
+	TOptional<int32> ViewportsPerFrame;
+
+	/** The maximum dimension of any internal texture for preview. Use less memory for large preview viewports */
+	TOptional<int32> PreviewMaxTextureDimension;
+};
diff --git a/Engine/Plugins/Runtime/nDisplay/Source/DisplayClusterLightCardEditor/Private/Viewport/DisplayClusterLightCardEditorViewportClient.cpp b/Engine/Plugins/Runtime/nDisplay/Source/DisplayClusterLightCardEditor/Private/Viewport/DisplayClusterLightCardEditorViewportClient.cpp
--- a/Engine/Plugins/Runtime/nDisplay/Source/DisplayClusterLightCardEditor/Private/Viewport/DisplayClusterLightCardEditorViewportClient.cpp
+++ b/Engine/Plugins/Runtime/nDisplay/Source/DisplayClusterLightCardEditor/Private/Viewport/DisplayClusterLightCardEditorViewportClient.cpp
@@ -17,6 +17,7 @@
 #include "DisplayClusterLightCardEditorUtils.h"
 #include "DisplayClusterProjectionStrings.h"
 #include "DisplayClusterRootActor.h"
+#include "DisplayClusterRootActorContainers.h"
 #include "IDisplayClusterScenePreview.h"
 #include "SceneManagement.h"
 #include "Settings/DisplayClusterLightCardEditorSettings.h"
@@ -1639,14 +1640,6 @@ void FDisplayClusterLightCardEditorViewportClient::UpdatePreviewActor(ADisplayCl
 				ProjectionHelper->SetLevelInstanceRootActor(*RootActorLevelInstance);
 				ProjectionHelper->SetEditorViewportClient(AsShared());
 
-				// Setup custom preview settings for proxy DCRA:
-				RootActorProxy->bPreviewICVFXFrustums = false;
-				RootActorProxy->bEnablePreviewTechvis = false;
-				RootActorProxy->bPreviewEnablePostProcess = true;
-				RootActorProxy->bPreviewEnableOverlayMaterial = false;
-				RootActorProxy->bPreviewEnable = true;
-				RootActorProxy->bPreviewICVFXFrustums = false;
-
 				//@Todo: In the case when RootActorProxy is not recreated each time it is changed:
 				// the DCRA  Copy cfg from RootActorLevelInstance to RootActorProxy every tick (do not recreate local DCRA)
 				if (UDisplayClusterConfigurationData* ProxyConfig = RootActorProxy->GetConfigData())
@@ -1665,8 +1658,22 @@ void FDisplayClusterLightCardEditorViewportClient::UpdatePreviewActor(ADisplayCl
 				});
 			}
 
+			// Use custom settings for root actor.
+			FDisplayClusterRootActorPropertyOverrides PropertyOverrides;
+			{
+				PropertyOverrides.bPreviewICVFXFrustums = false;
+				PropertyOverrides.bEnablePreviewTechvis = false;
+				PropertyOverrides.bPreviewEnableOverlayMaterial = false;
+				PropertyOverrides.bFreezePreviewRender = false;
+
+				PropertyOverrides.bPreviewEnablePostProcess = true;
+				PropertyOverrides.bPreviewEnable = true;
+
+				PropertyOverrides.PreviewSetttingsSource = EDisplayClusterConfigurationRootActorPreviewSettingsSource::RootActor;
+			}
+
 			// Filter out any primitives hidden in game except screen components
-			IDisplayClusterScenePreview::Get().SetRendererRootActor(PreviewRendererId, RootActorProxy.Get());
+			IDisplayClusterScenePreview::Get().SetRendererRootActor(PreviewRendererId, RootActorProxy.Get(), PropertyOverrides);
 
 			if (ProxyType == EDisplayClusterLightCardEditorProxyType::All ||
 				ProxyType == EDisplayClusterLightCardEditorProxyType::StageActor)
diff --git a/Engine/Plugins/Runtime/nDisplay/Source/DisplayClusterScenePreview/Private/DisplayClusterLightCardEditorHelper.cpp b/Engine/Plugins/Runtime/nDisplay/Source/DisplayClusterScenePreview/Private/DisplayClusterLightCardEditorHelper.cpp
--- a/Engine/Plugins/Runtime/nDisplay/Source/DisplayClusterScenePreview/Private/DisplayClusterLightCardEditorHelper.cpp
+++ b/Engine/Plugins/Runtime/nDisplay/Source/DisplayClusterScenePreview/Private/DisplayClusterLightCardEditorHelper.cpp
@@ -9,6 +9,7 @@
 #include "DisplayClusterChromakeyCardActor.h"
 #include "DisplayClusterLightCardActor.h"
 #include "DisplayClusterRootActor.h"
+#include "DisplayClusterRootActorContainers.h"
 #include "Components/DisplayClusterCameraComponent.h"
 
 #include "StageActor/DisplayClusterStageActorTemplate.h"
@@ -230,7 +231,21 @@ void FDisplayClusterLightCardEditorHelper::SetRootActor(ADisplayClusterRootActor
 		return;
 	}
 
-	IDisplayClusterScenePreview::Get().SetRendererRootActor(RendererId, &NewRootActor, false);
+	// Use custom settings for root actor.
+	FDisplayClusterRootActorPropertyOverrides PropertyOverrides;
+	{
+		PropertyOverrides.bPreviewICVFXFrustums = false;
+		PropertyOverrides.bEnablePreviewTechvis = false;
+		PropertyOverrides.bPreviewEnableOverlayMaterial = false;
+		PropertyOverrides.bFreezePreviewRender = false;
+
+		PropertyOverrides.bPreviewEnablePostProcess = true;
+		PropertyOverrides.bPreviewEnable = true;
+
+		PropertyOverrides.PreviewSetttingsSource = EDisplayClusterConfigurationRootActorPreviewSettingsSource::RootActor;
+	}
+
+	IDisplayClusterScenePreview::Get().SetRendererRootActor(RendererId, &NewRootActor, PropertyOverrides);
 }
 
 void FDisplayClusterLightCardEditorHelper::SetLevelInstanceRootActor(ADisplayClusterRootActor& NewRootActor)
diff --git a/Engine/Plugins/Runtime/nDisplay/Source/DisplayClusterScenePreview/Private/DisplayClusterScenePreviewModule.cpp b/Engine/Plugins/Runtime/nDisplay/Source/DisplayClusterScenePreview/Private/DisplayClusterScenePreviewModule.cpp
--- a/Engine/Plugins/Runtime/nDisplay/Source/DisplayClusterScenePreview/Private/DisplayClusterScenePreviewModule.cpp
+++ b/Engine/Plugins/Runtime/nDisplay/Source/DisplayClusterScenePreview/Private/DisplayClusterScenePreviewModule.cpp
@@ -56,24 +56,32 @@ bool FDisplayClusterScenePreviewModule::DestroyRenderer(int32 RendererId)
 {
 	if (FRendererConfig* Config = RendererConfigs.Find(RendererId))
 	{
-		RegisterRootActorEvents(Config->RootActor.Get(), false);
+		// Release the proxy resources that were used by this renderer.
+		ProxyManager.SetSceneRootActorForRenderer(RendererId, nullptr);
+
+		RegisterRootActorEvents(RendererId, *Config, false);
 		RendererConfigs.Remove(RendererId);
 
 		RegisterOrUnregisterGlobalActorEvents();
+
 		return true;
 	}
 
 	return false;
 }
 
-bool FDisplayClusterScenePreviewModule::SetRendererRootActorPath(int32 RendererId, const FString& ActorPath, bool bAutoUpdateLightcards)
+bool FDisplayClusterScenePreviewModule::SetRendererRootActorPath(int32 RendererId, const FString& ActorPath, const FDisplayClusterRootActorPropertyOverrides& InPropertyOverrides, const EDisplayClusterScenePreviewFlags PreviewFlags)
 {
 	if (FRendererConfig* Config = RendererConfigs.Find(RendererId))
 	{
 		Config->RootActorPath = ActorPath;
 
 		ADisplayClusterRootActor* RootActor = FindObject<ADisplayClusterRootActor>(nullptr, *ActorPath);
-		InternalSetRendererRootActor(*Config, RootActor, bAutoUpdateLightcards);
+		InternalSetRendererRootActor(RendererId, *Config, RootActor, PreviewFlags);
+
+		// Use custom properties on root actor
+		Config->RootActorPropertyOverrides = InPropertyOverrides;
+		InternalOverridePropertiesForRendererRootActor(RendererId, *Config);
 
 		return true;
 	}
@@ -81,12 +89,18 @@ bool FDisplayClusterScenePreviewModule::SetRendererRootActorPath(int32 RendererI
 	return false;
 }
 
-bool FDisplayClusterScenePreviewModule::SetRendererRootActor(int32 RendererId, ADisplayClusterRootActor* Actor, bool bAutoUpdateLightcards)
+bool FDisplayClusterScenePreviewModule::SetRendererRootActor(int32 RendererId, ADisplayClusterRootActor* Actor, const FDisplayClusterRootActorPropertyOverrides& InPropertyOverrides, const EDisplayClusterScenePreviewFlags PreviewFlags)
 {
 	if (FRendererConfig* Config = RendererConfigs.Find(RendererId))
 	{
 		Config->RootActorPath.Empty();
-		InternalSetRendererRootActor(*Config, Actor, bAutoUpdateLightcards);
+		Config->RootActorPropertyOverrides = InPropertyOverrides;
+
+		InternalSetRendererRootActor(RendererId, *Config, Actor, PreviewFlags);
+
+		// Use custom properties on root actor
+		Config->RootActorPropertyOverrides = InPropertyOverrides;
+		InternalOverridePropertiesForRendererRootActor(RendererId, *Config);
 
 		return true;
 	}
@@ -98,7 +112,7 @@ ADisplayClusterRootActor* FDisplayClusterScenePreviewModule::GetRendererRootActo
 {
 	if (FRendererConfig* Config = RendererConfigs.Find(RendererId))
 	{
-		return InternalGetRendererRootActor(*Config);
+		return InternalGetRendererRootActorOrProxy(RendererId, *Config);
 	}
 
 	return nullptr;
@@ -198,22 +212,11 @@ bool FDisplayClusterScenePreviewModule::SetRendererRenderSimpleElementsDelegate(
 	return false;
 }
 
-bool FDisplayClusterScenePreviewModule::SetRendererUsePostProcessTexture(int32 RendererId, bool bUsePostProcessTexture)
-{
-	if (FRendererConfig* Config = RendererConfigs.Find(RendererId))
-	{
-		Config->bUsePostProcessTexture = bUsePostProcessTexture;
-		return true;
-	}
-
-	return false;
-}
-
 bool FDisplayClusterScenePreviewModule::Render(int32 RendererId, FDisplayClusterMeshProjectionRenderSettings& RenderSettings, FCanvas& Canvas)
 {
 	if (FRendererConfig* Config = RendererConfigs.Find(RendererId))
 	{
-		return InternalRenderImmediate(*Config, RenderSettings, Canvas);
+		return InternalRenderImmediate(RendererId, *Config, RenderSettings, Canvas);
 	}
 
 	return false;
@@ -245,7 +248,19 @@ bool FDisplayClusterScenePreviewModule::IsRealTimePreviewEnabled() const
 	return bIsRealTimePreviewEnabled;
 }
 
-ADisplayClusterRootActor* FDisplayClusterScenePreviewModule::InternalGetRendererRootActor(FRendererConfig& RendererConfig)
+bool FDisplayClusterScenePreviewModule::IsBlueprintMatchesRendererRootActor(int32 RendererId, FRendererConfig& RendererConfig, UBlueprint* Blueprint)
+{
+#if WITH_EDITOR
+	if (RendererConfig.RootActor.IsValid() && Blueprint == UBlueprint::GetBlueprintFromClass(RendererConfig.RootActor->GetClass()))
+	{
+		return true;
+	}
+#endif
+
+	return false;
+}
+
+ADisplayClusterRootActor* FDisplayClusterScenePreviewModule::InternalGetRendererRootActor(int32 RendererId, FRendererConfig& RendererConfig)
 {
 	if (!RendererConfig.RootActor.IsValid() && !RendererConfig.RootActorPath.IsEmpty())
 	{
@@ -253,13 +268,25 @@ ADisplayClusterRootActor* FDisplayClusterScenePreviewModule::InternalGetRenderer
 		ADisplayClusterRootActor* RootActor = FindObject<ADisplayClusterRootActor>(nullptr, *RendererConfig.RootActorPath);
 		if (RootActor)
 		{
-			InternalSetRendererRootActor(RendererConfig, RootActor, RendererConfig.bAutoUpdateLightcards);
+			InternalSetRendererRootActor(RendererId, RendererConfig, RootActor, RendererConfig.PreviewFlags);
+			InternalOverridePropertiesForRendererRootActor(RendererId, RendererConfig);
 		}
 	}
 
 	return RendererConfig.RootActor.Get();
 }
 
+ADisplayClusterRootActor* FDisplayClusterScenePreviewModule::InternalGetRendererRootActorOrProxy(int32 RendererId, FRendererConfig& RendererConfig)
+{
+	// Note: we call this function first because it can assign a new root actor and proxy.
+	ADisplayClusterRootActor* RootActor = InternalGetRendererRootActor(RendererId, RendererConfig);
+
+	// When proxy is required, get it from ProxyManager.
+	const bool bUseRootActorProxy = EnumHasAnyFlags(RendererConfig.PreviewFlags, EDisplayClusterScenePreviewFlags::UseRootActorProxy);
+
+	return bUseRootActorProxy ? ProxyManager.GetProxyRootActor(RendererId) : RootActor;
+}
+
 bool FDisplayClusterScenePreviewModule::InternalRenderQueued(int32 RendererId, FDisplayClusterMeshProjectionRenderSettings& RenderSettings, TWeakPtr<FCanvas> Canvas,
 	const FIntPoint& Size, FRenderResultDelegate ResultDelegate)
 {
@@ -281,42 +308,61 @@ bool FDisplayClusterScenePreviewModule::InternalRenderQueued(int32 RendererId, F
 	return false;
 }
 
-void FDisplayClusterScenePreviewModule::InternalSetRendererRootActor(FRendererConfig& RendererConfig, ADisplayClusterRootActor* Actor, bool bAutoUpdateLightcards)
+
+void FDisplayClusterScenePreviewModule::InternalOverridePropertiesForRendererRootActor(int32 RendererId, FRendererConfig& RendererConfig)
+{
+	// Override root actor properties:
+	if (ADisplayClusterRootActor* RendererRootActor = InternalGetRendererRootActorOrProxy(RendererId, RendererConfig))
+	{
+		RendererRootActor->OverrideRootActorProperties(RendererConfig.RootActorPropertyOverrides);
+	}
+}
+
+void FDisplayClusterScenePreviewModule::InternalSetRendererRootActor(int32 RendererId, FRendererConfig& RendererConfig, ADisplayClusterRootActor* Actor, const EDisplayClusterScenePreviewFlags PreviewFlags)
 {
 	// Determine these values before we update the config's RootActor/bAutoUpdateLightcards
 	const bool bRootChanged = RendererConfig.RootActor != Actor;
 
-	if (bRootChanged)
+	if (bRootChanged || RendererConfig.PreviewFlags != PreviewFlags)
 	{
 		// Unregister events for the previous cluster
-		if (RendererConfig.bAutoUpdateLightcards && RendererConfig.RootActor.IsValid())
+		const bool bAutoUpdateLightcards = EnumHasAnyFlags(RendererConfig.PreviewFlags, EDisplayClusterScenePreviewFlags::AutoUpdateLightcards);
+		if (bAutoUpdateLightcards)
 		{
-			RegisterRootActorEvents(RendererConfig.RootActor.Get(), false);
+			RegisterRootActorEvents(RendererId, RendererConfig, false);
 		}
 
 		RendererConfig.RootActor = Actor;
-		RendererConfig.bAutoUpdateLightcards = bAutoUpdateLightcards;
-		AutoPopulateScene(RendererConfig);
+		RendererConfig.PreviewFlags = PreviewFlags;
+
+		// Updates the proxy for the new root actor.
+		const bool bUseRootActorProxy = EnumHasAnyFlags(RendererConfig.PreviewFlags, EDisplayClusterScenePreviewFlags::UseRootActorProxy);
+		ProxyManager.SetSceneRootActorForRenderer(RendererId, bUseRootActorProxy ? Actor : nullptr);
+
+		AutoPopulateScene(RendererId, RendererConfig);
 	}
 
-	RegisterRootActorEvents(Actor, true);
+	RegisterRootActorEvents(RendererId, RendererConfig, true);
 	RegisterOrUnregisterGlobalActorEvents();
 }
 
-bool FDisplayClusterScenePreviewModule::InternalRenderImmediate(FRendererConfig& RendererConfig, FDisplayClusterMeshProjectionRenderSettings& RenderSettings, FCanvas& Canvas)
+bool FDisplayClusterScenePreviewModule::InternalRenderImmediate(int32 RendererId, FRendererConfig& RendererConfig, FDisplayClusterMeshProjectionRenderSettings& RenderSettings, FCanvas& Canvas)
 {
 	// Update this so that whoever gets the callback can immediately check whether the nDisplay preview may be out of date
 	UpdateIsRealTimePreviewEnabled();
 
-	UWorld* World = RendererConfig.RootActor.IsValid() ? RendererConfig.RootActor->GetWorld() : nullptr;
+	// Get the Root Actor or proxy for rendering previews.
+	ADisplayClusterRootActor* RootActor = InternalGetRendererRootActorOrProxy(RendererId, RendererConfig);
+	UWorld* World = RootActor ? RootActor->GetWorld() : nullptr;
 	if (!World)
 	{
 		return false;
 	}
 
-	if (RendererConfig.bAutoUpdateLightcards && RendererConfig.bIsSceneDirty)
+	const bool bAutoUpdateLightcards = EnumHasAnyFlags(RendererConfig.PreviewFlags, EDisplayClusterScenePreviewFlags::AutoUpdateLightcards);
+	if (bAutoUpdateLightcards && RendererConfig.bIsSceneDirty)
 	{
-		AutoPopulateScene(RendererConfig);
+		AutoPopulateScene(RendererId, RendererConfig);
 	}
 
 	// Push any deferred render state updates to ensure that light card positions, preview meshes modified above, etc. are up to date
@@ -333,7 +379,8 @@ void FDisplayClusterScenePreviewModule::RegisterOrUnregisterGlobalActorEvents()
 	bool bShouldBeRegistered = false;
 	for (const TPair<int32, FRendererConfig>& ConfigPair : RendererConfigs)
 	{
-		if (ConfigPair.Value.bAutoUpdateLightcards)
+		const bool bAutoUpdateLightcards = EnumHasAnyFlags(ConfigPair.Value.PreviewFlags, EDisplayClusterScenePreviewFlags::AutoUpdateLightcards);
+		if (bAutoUpdateLightcards)
 		{
 			bShouldBeRegistered = true;
 			break;
@@ -368,9 +415,10 @@ void FDisplayClusterScenePreviewModule::RegisterOrUnregisterGlobalActorEvents()
 #endif
 }
 
-void FDisplayClusterScenePreviewModule::RegisterRootActorEvents(ADisplayClusterRootActor* Actor, bool bShouldRegister)
+void FDisplayClusterScenePreviewModule::RegisterRootActorEvents(int32 RendererId, FRendererConfig& RendererConfig, bool bShouldRegister)
 {
 #if WITH_EDITOR
+	ADisplayClusterRootActor* Actor = RendererConfig.RootActor.Get();
 	if (!Actor)
 	{
 		return;
@@ -391,16 +439,18 @@ void FDisplayClusterScenePreviewModule::RegisterRootActorEvents(ADisplayClusterR
 #endif
 }
 
-void FDisplayClusterScenePreviewModule::AutoPopulateScene(FRendererConfig& RendererConfig)
+void FDisplayClusterScenePreviewModule::AutoPopulateScene(int32 RendererId, FRendererConfig& RendererConfig)
 {
-	if (RendererConfig.bAutoUpdateLightcards)
+	const bool bAutoUpdateLightcards = EnumHasAnyFlags(RendererConfig.PreviewFlags, EDisplayClusterScenePreviewFlags::AutoUpdateLightcards);
+	if (bAutoUpdateLightcards)
 	{
 		RendererConfig.Renderer->ClearScene();
 		RendererConfig.AddedActors.Empty();
 		RendererConfig.AutoActors.Empty();
 	}
 	
-	if (ADisplayClusterRootActor* RootActor = InternalGetRendererRootActor(RendererConfig))
+	// The renderer can use a proxy.
+	if (ADisplayClusterRootActor* RootActor = InternalGetRendererRootActorOrProxy(RendererId, RendererConfig))
 	{
 		TArray<FString> ProjectionMeshNames;
 
@@ -417,7 +467,7 @@ void FDisplayClusterScenePreviewModule::AutoPopulateScene(FRendererConfig& Rende
 			return bIsProjectionMesh || bIsScreen;
 		});
 
-		if (RendererConfig.bAutoUpdateLightcards)
+		if (bAutoUpdateLightcards)
 		{
 			// Automatically add the lightcards found on this actor
 			TSet<ADisplayClusterLightCardActor*> LightCards;
@@ -490,6 +540,9 @@ bool FDisplayClusterScenePreviewModule::UpdateIsRealTimePreviewEnabled()
 
 bool FDisplayClusterScenePreviewModule::OnTick(float DeltaTime)
 {
+	// This function calls Tick() for a preview world with proxy root actors, which triggers rendering of previews for them.
+	ProxyManager.Tick(DeltaTime);
+
 	// This loop should break when we either run out of jobs or complete a single job
 	while (!RenderQueue.IsEmpty())
 	{
@@ -520,12 +573,13 @@ bool FDisplayClusterScenePreviewModule::OnTick(float DeltaTime)
 					continue;
 				}
 
-				InternalRenderImmediate(*Config, Job.Settings, *Canvas);
+				InternalRenderImmediate(Job.RendererId, *Config, Job.Settings, *Canvas);
 				Job.ResultDelegate.Execute(RenderTarget);
 				break;
 			}
 
-			if (UWorld* World = Config->RootActor.IsValid() ? Config->RootActor->GetWorld() : nullptr)
+			ADisplayClusterRootActor* RootActor = InternalGetRendererRootActorOrProxy(Job.RendererId, *Config);
+			if (UWorld* World = RootActor ? RootActor->GetWorld() : nullptr)
 			{
 				// We need to provide the render target for this job
 				UTextureRenderTarget2D* RenderTarget = Config->RenderTarget.Get();
@@ -550,7 +604,7 @@ bool FDisplayClusterScenePreviewModule::OnTick(float DeltaTime)
 				FTextureRenderTargetResource* RenderTargetResource = RenderTarget->GameThread_GetRenderTargetResource();
 				FCanvas Canvas(RenderTargetResource, nullptr, FGameTime::GetTimeSinceAppStart(), World->Scene->GetFeatureLevel());
 
-				InternalRenderImmediate(*Config, Job.Settings, Canvas);
+				InternalRenderImmediate(Job.RendererId, *Config, Job.Settings, Canvas);
 				Job.ResultDelegate.Execute(RenderTargetResource);
 				break;
 			}
@@ -576,7 +630,8 @@ void FDisplayClusterScenePreviewModule::OnActorPropertyChanged(UObject* ObjectBe
 	for (TPair<int32, FRendererConfig>& ConfigPair : RendererConfigs)
 	{
 		FRendererConfig& Config = ConfigPair.Value;
-		if (Config.bAutoUpdateLightcards)
+		const bool bAutoUpdateLightcards = EnumHasAnyFlags(Config.PreviewFlags, EDisplayClusterScenePreviewFlags::AutoUpdateLightcards);
+		if (bAutoUpdateLightcards)
 		{
 			if (Config.RootActor == ObjectBeingModified)
 			{
@@ -627,11 +682,13 @@ void FDisplayClusterScenePreviewModule::OnBlueprintCompiled(UBlueprint* Blueprin
 #if WITH_EDITOR
 	for (TPair<int32, FRendererConfig>& ConfigPair : RendererConfigs)
 	{
+		int32 RendererId = ConfigPair.Key;
 		FRendererConfig& Config = ConfigPair.Value;
-		if (Config.RootActor.IsValid() && Blueprint == UBlueprint::GetBlueprintFromClass(Config.RootActor->GetClass()))
+
+		if(IsBlueprintMatchesRendererRootActor(RendererId, Config, Blueprint))
 		{
 			Config.bIsSceneDirty = true;
-			RegisterRootActorEvents(Config.RootActor.Get(), true);
+			RegisterRootActorEvents(RendererId, Config, true);
 		}
 	}
 #endif
diff --git a/Engine/Plugins/Runtime/nDisplay/Source/DisplayClusterScenePreview/Private/DisplayClusterScenePreviewModule.h b/Engine/Plugins/Runtime/nDisplay/Source/DisplayClusterScenePreview/Private/DisplayClusterScenePreviewModule.h
--- a/Engine/Plugins/Runtime/nDisplay/Source/DisplayClusterScenePreview/Private/DisplayClusterScenePreviewModule.h
+++ b/Engine/Plugins/Runtime/nDisplay/Source/DisplayClusterScenePreview/Private/DisplayClusterScenePreviewModule.h
@@ -3,6 +3,8 @@
 #pragma once
 
 #include "IDisplayClusterScenePreview.h"
+#include "DisplayClusterScenePreviewProxyManager.h"
+#include "DisplayClusterRootActorContainers.h"
 
 #include "Containers/Ticker.h"
 #include "DisplayClusterMeshProjectionRenderer.h"
@@ -27,8 +29,8 @@ public:
 	//~ Begin IDisplayClusterScenePreview Interface
 	virtual int32 CreateRenderer() override;
 	virtual bool DestroyRenderer(int32 RendererId) override;
-	virtual bool SetRendererRootActorPath(int32 RendererId, const FString& ActorPath, bool bAutoUpdateLightcards = false) override;
-	virtual bool SetRendererRootActor(int32 RendererId, ADisplayClusterRootActor* Actor, bool bAutoUpdateLightcards = false) override;
+	virtual bool SetRendererRootActorPath(int32 RendererId, const FString& ActorPath, const FDisplayClusterRootActorPropertyOverrides& InPropertyOverrides, const EDisplayClusterScenePreviewFlags PreviewFlags = EDisplayClusterScenePreviewFlags::None) override;
+	virtual bool SetRendererRootActor(int32 RendererId, ADisplayClusterRootActor* Actor, const FDisplayClusterRootActorPropertyOverrides& InPropertyOverrides, const EDisplayClusterScenePreviewFlags PreviewFlags = EDisplayClusterScenePreviewFlags::None) override;
 	virtual ADisplayClusterRootActor* GetRendererRootActor(int32 RendererId) override;
 	virtual bool GetActorsInRendererScene(int32 RendererId, bool bIncludeRoot, TArray<AActor*>& OutActors) override;
 	virtual bool AddActorToRenderer(int32 RendererId, AActor* Actor) override;
@@ -37,7 +39,6 @@ public:
 	virtual bool ClearRendererScene(int32 RendererId) override;
 	virtual bool SetRendererActorSelectedDelegate(int32 RendererId, FDisplayClusterMeshProjectionRenderer::FSelection ActorSelectedDelegate) override;
 	virtual bool SetRendererRenderSimpleElementsDelegate(int32 RendererId, FDisplayClusterMeshProjectionRenderer::FSimpleElementPass RenderSimpleElementsDelegate) override;
-	virtual bool SetRendererUsePostProcessTexture(int32 RendererId, bool bUsePostProcessTexture) override;
 	virtual bool Render(int32 RendererId, FDisplayClusterMeshProjectionRenderSettings& RenderSettings, FCanvas& Canvas) override;
 	virtual bool RenderQueued(int32 RendererId, FDisplayClusterMeshProjectionRenderSettings& RenderSettings, const FIntPoint& Size, FRenderResultDelegate ResultDelegate) override;
 	virtual bool RenderQueued(int32 RendererId, FDisplayClusterMeshProjectionRenderSettings& RenderSettings, const TWeakPtr<FCanvas> Canvas, FRenderResultDelegate ResultDelegate) override;
@@ -63,14 +64,14 @@ private:
 		/** Actors that have been automatically added to the scene. */
 		TArray<TWeakObjectPtr<AActor>> AutoActors;
 
-		/** If true, automatically update the renderer with stage actors belonging to the root actor. */
-		bool bAutoUpdateLightcards = false;
+		/** Special flags that control the behavior of the renderer. */
+		EDisplayClusterScenePreviewFlags PreviewFlags = EDisplayClusterScenePreviewFlags::None;
 
-		/** If true, the scene needs to be updated before the next render. This is only relevant if bAutoUpdateLightcards is true. */
-		bool bIsSceneDirty = true;
+		/** Container with properties to be overridden for the root actor used by this renderer. */
+		FDisplayClusterRootActorPropertyOverrides RootActorPropertyOverrides;
 
-		/** If true, apply post-processing to the nDisplay preview texture and override the preview component with it before rendering. */
-		bool bUsePostProcessTexture = false;
+		/** If true, the scene needs to be updated before the next render. This is only relevant if PreviewFlags.AutoUpdateLightcards is true. */
+		bool bIsSceneDirty = true;
 
 		/** The render target to use for queued renders. */
 		TStrongObjectPtr<UTextureRenderTarget2D> RenderTarget = nullptr;
@@ -107,26 +108,35 @@ private:
 	};
 
 	/** Get the root actor for a config. If the root actor pointer is invalid but we have a path to the actor, try to reacquire a pointer using the path first. */
-	ADisplayClusterRootActor* InternalGetRendererRootActor(FRendererConfig& RendererConfig);
+	ADisplayClusterRootActor* InternalGetRendererRootActor(int32 RendererId, FRendererConfig& RendererConfig);
+
+	/** Get the root actor or proxy. */
+	ADisplayClusterRootActor* InternalGetRendererRootActorOrProxy(int32 RendererId, FRendererConfig& RendererConfig);
 
 	/** Set the root actor for a config, update its scene, and register events accordingly. */
-	void InternalSetRendererRootActor(FRendererConfig& RendererConfig, ADisplayClusterRootActor* Actor, bool bAutoUpdateLightcards);
+	void InternalSetRendererRootActor(int32 RendererId, FRendererConfig& RendererConfig, ADisplayClusterRootActor* Actor, const EDisplayClusterScenePreviewFlags PreviewFlags);
+
+	/** Override properties for the renderer root actor. */
+	void InternalOverridePropertiesForRendererRootActor(int32 RendererId, FRendererConfig& RendererConfig);
 
 	/** Queue a preview to be rendered. */
 	bool InternalRenderQueued(int32 RendererId, FDisplayClusterMeshProjectionRenderSettings& RenderSettings, TWeakPtr<FCanvas> Canvas,
 		const FIntPoint& Size, FRenderResultDelegate ResultDelegate);
 
 	/** Immediately render with the given renderer config and settings to the given canvas. */
-	bool InternalRenderImmediate(FRendererConfig& RendererConfig, FDisplayClusterMeshProjectionRenderSettings& RenderSettings, FCanvas& Canvas);
+	bool InternalRenderImmediate(int32 RendererId, FRendererConfig& RendererConfig, FDisplayClusterMeshProjectionRenderSettings& RenderSettings, FCanvas& Canvas);
 
 	/** Check if any of the tracked root actors are set to auto-update their lightcards and register/unregister event listeners accordingly. */
 	void RegisterOrUnregisterGlobalActorEvents();
 
 	/** Register/unregister to events affecting a cluster root actor. */
-	void RegisterRootActorEvents(ADisplayClusterRootActor* Actor, bool bShouldRegister);
+	void RegisterRootActorEvents(int32 RendererId, FRendererConfig& RendererConfig, bool bShouldRegister);
+
+	/** Returns true if the Blueprint class matches the RootActor class used in the renderer configuration. */
+	bool IsBlueprintMatchesRendererRootActor(int32 RendererId, FRendererConfig& RendererConfig, UBlueprint* Blueprint);
 
 	/** Clear and re-populate a renderer's scene with the root actor and lightcards if applicable. */
-	void AutoPopulateScene(FRendererConfig& RendererConfig);
+	void AutoPopulateScene(int32 RendererId, FRendererConfig& RendererConfig);
 
 	/** Check whether nDisplay preview textures are being updated in real time. */
 	bool UpdateIsRealTimePreviewEnabled();
@@ -168,4 +178,7 @@ private:
 
 	/** Whether nDisplay preview textures are being updated in real time. */
 	bool bIsRealTimePreviewEnabled = false;
+
+	/** Manager for DCRA proxy objects. */
+	FDisplayClusterScenePreviewProxyManager ProxyManager;
 };
\ No newline at end of file
diff --git a/Engine/Plugins/Runtime/nDisplay/Source/DisplayClusterScenePreview/Private/DisplayClusterScenePreviewProxyManager.cpp b/Engine/Plugins/Runtime/nDisplay/Source/DisplayClusterScenePreview/Private/DisplayClusterScenePreviewProxyManager.cpp
new file mode 100644
--- /dev/null
+++ b/Engine/Plugins/Runtime/nDisplay/Source/DisplayClusterScenePreview/Private/DisplayClusterScenePreviewProxyManager.cpp
@@ -0,0 +1,154 @@
+// Copyright Epic Games, Inc. All Rights Reserved.
+
+#include "DisplayClusterScenePreviewProxyManager.h"
+
+#include "Components/DisplayClusterStageGeometryComponent.h"
+
+FDisplayClusterScenePreviewProxyManager::~FDisplayClusterScenePreviewProxyManager()
+{
+	DestroyPreviewWorld();
+}
+
+void FDisplayClusterScenePreviewProxyManager::CreatePreviewWorld()
+{
+#if WITH_EDITOR
+	if (!PreviewWorld)
+	{
+		PreviewWorld = UWorld::CreateWorld(EWorldType::None, false);
+
+		FWorldContext& WorldContext = GEngine->CreateNewWorldContext(EWorldType::None);
+		WorldContext.SetCurrentWorld(PreviewWorld);
+	}
+#endif
+}
+
+void FDisplayClusterScenePreviewProxyManager::DestroyPreviewWorld()
+{
+#if WITH_EDITOR
+	if (PreviewWorld)
+	{
+		GEngine->DestroyWorldContext(PreviewWorld);
+		PreviewWorld->DestroyWorld(false);
+	}
+#endif
+
+	PreviewWorld = nullptr;
+}
+
+void FDisplayClusterScenePreviewProxyManager::Tick(float DeltaTime)
+{
+	if (PreviewWorld)
+	{
+		PreviewWorld->Tick(ELevelTick::LEVELTICK_All, DeltaTime);
+	}
+}
+
+ADisplayClusterRootActor* FDisplayClusterScenePreviewProxyManager::CreateRootActorProxy(int32 RendererId, ADisplayClusterRootActor* SceneRootActor)
+{
+	if (!SceneRootActor)
+	{
+		return nullptr;
+	}
+
+	ADisplayClusterRootActor* RootActorProxy = nullptr;
+
+#if WITH_EDITOR
+	// Create new preview world for DCRA proxy
+	CreatePreviewWorld();
+	if (!PreviewWorld)
+	{
+		return nullptr;
+	}
+
+	// Create a proxy for render
+	FObjectDuplicationParameters DupeActorParameters(SceneRootActor, PreviewWorld->GetCurrentLevel());
+	DupeActorParameters.FlagMask = RF_AllFlags & ~(RF_ArchetypeObject | RF_Transactional); // Keeps archetypes correct in config data.
+	DupeActorParameters.PortFlags = PPF_DuplicateVerbatim;
+	DupeActorParameters.DestName = FName(FString::Printf(TEXT("Preview%i_%s"), RendererId, *SceneRootActor->GetName()));
+
+	RootActorProxy = CastChecked<ADisplayClusterRootActor>(StaticDuplicateObjectEx(DupeActorParameters));
+
+	// Use root actor from scene to render
+	if (IDisplayClusterViewportManager* ViewportManager = RootActorProxy ? RootActorProxy->GetOrCreateViewportManager() : nullptr)
+	{
+		// Using DCRA from the scene for rendering
+		ViewportManager->GetConfiguration().SetRootActor(EDisplayClusterRootActorType::Scene | EDisplayClusterRootActorType::Configuration, SceneRootActor);
+	}
+
+	RootActorProxy->SetFlags(RF_Transient); // This signals to the stage actor it is a proxy
+
+	PreviewWorld->GetCurrentLevel()->AddLoadedActor(RootActorProxy);
+
+	// Draw the geometry map for the proxy stage actor immediately to avoid a race condition where the geometry map could render
+	// before the actor location changes propagate to its component proxies, resulting in an inaccurate proxy geometry map
+	RootActorProxy->GetStageGeometryComponent()->Invalidate(true);
+
+	// Spawned actor will take the transform values from the template, so manually reset them to zero here
+	RootActorProxy->SetActorLocation(FVector::ZeroVector);
+	RootActorProxy->SetActorRotation(FRotator::ZeroRotator);
+
+	// Set translucency sort priority of root actor proxy primitive components so that actors that are flush with screens are rendered on top of them
+	RootActorProxy->ForEachComponent<UPrimitiveComponent>(false, [](UPrimitiveComponent* InPrimitiveComponent)
+	{
+		InPrimitiveComponent->SetTranslucentSortPriority(-10);
+	});
+#endif
+
+	return RootActorProxy;
+}
+
+void FDisplayClusterScenePreviewProxyManager::DestroyRootActorProxy(ADisplayClusterRootActor* ProxyRootActor)
+{
+#if WITH_EDITOR
+	if (PreviewWorld && ProxyRootActor)
+	{
+		PreviewWorld->GetCurrentLevel()->RemoveLoadedActors({ ProxyRootActor });
+	}
+#endif
+}
+
+void FDisplayClusterScenePreviewProxyManager::SetSceneRootActorForRenderer(int32 RendererId, ADisplayClusterRootActor* SceneRootActor)
+{
+	if (RendererProxies.Contains(RendererId))
+	{
+		FRendererProxy& RendererProxy = RendererProxies[RendererId];
+
+		if (RendererProxy.SceneRootActor == SceneRootActor)
+		{
+			// Re-use existing proxy
+			return;
+		}
+
+		// The root agent has changed, destroy the proxy currently in use.
+		DestroyRootActorProxy(RendererProxy.ProxyRootActor.Get());
+
+		RendererProxies.Remove(RendererId);
+	}
+
+	if (ADisplayClusterRootActor* ProxyRootActor = CreateRootActorProxy(RendererId, SceneRootActor))
+	{
+		// Create new proxy
+		FRendererProxy RendererProxy;
+		RendererProxy.SceneRootActor = SceneRootActor;
+		RendererProxy.ProxyRootActor = ProxyRootActor;
+
+		RendererProxies.Emplace(RendererId, RendererProxy);
+	}
+
+
+	// Destroy preview world if not used
+	if (RendererProxies.IsEmpty())
+	{
+		DestroyPreviewWorld();
+	}
+}
+
+ADisplayClusterRootActor* FDisplayClusterScenePreviewProxyManager::GetProxyRootActor(int32 RendererId) const
+{
+	if (RendererProxies.Contains(RendererId))
+	{
+		return RendererProxies[RendererId].ProxyRootActor.Get();
+	}
+
+	return nullptr;
+}
diff --git a/Engine/Plugins/Runtime/nDisplay/Source/DisplayClusterScenePreview/Private/DisplayClusterScenePreviewProxyManager.h b/Engine/Plugins/Runtime/nDisplay/Source/DisplayClusterScenePreview/Private/DisplayClusterScenePreviewProxyManager.h
new file mode 100644
--- /dev/null
+++ b/Engine/Plugins/Runtime/nDisplay/Source/DisplayClusterScenePreview/Private/DisplayClusterScenePreviewProxyManager.h
@@ -0,0 +1,66 @@
+// Copyright Epic Games, Inc. All Rights Reserved.
+
+#pragma once
+
+#include "CoreMinimal.h"
+
+#include "DisplayClusterRootActor.h"
+#include "Engine/World.h"
+
+/**
+ * Creates and handle DCRA proxies.
+ */
+class FDisplayClusterScenePreviewProxyManager
+{
+public:
+	~FDisplayClusterScenePreviewProxyManager();
+
+	/** Assign SceneRootActor to the renderer with the specified ID to create a root proxy actor for it.
+	* @param RendererId - ID of renderer that uses this DCRA
+	* @param RootActor  - the new scene root actor that used by this renderer
+	* 
+	* Note: If RootActor is null, it means that the renderer is no longer using the proxy.
+	*/
+	void SetSceneRootActorForRenderer(int32 RendererId, ADisplayClusterRootActor* RootActor);
+
+	/** Get the DCRA proxy for the renderer by renderer id.
+	* 
+	* @param RendererId - ID of renderer that uses this DCRA
+	* 
+	* @return - A reference to the actor's root proxy, if it exists, or nullptr.
+	* 
+	* Note: You must call SetSceneRootActorForRenderer() for the renderer and the root actor being used before using this function.
+	*/
+	ADisplayClusterRootActor* GetProxyRootActor(int32 RendererId) const;
+
+	/** This function calls Tick() for a preview world with proxy root actors, which triggers rendering of previews for them.
+	* 
+	* @param DeltaTime - has passed to PreviewWorld->Tick().
+	*/
+	void Tick(float DeltaTime);
+
+private:
+	void CreatePreviewWorld();
+	void DestroyPreviewWorld();
+
+	/** Create RootActor proxy. */
+	ADisplayClusterRootActor* CreateRootActorProxy(int32 RendererId, ADisplayClusterRootActor* SceneRootActor);
+
+	void DestroyRootActorProxy(ADisplayClusterRootActor* ProxyRootActor);
+
+private:
+	struct FRendererProxy
+	{
+		// RootActor proxy objects
+		TWeakObjectPtr<ADisplayClusterRootActor> ProxyRootActor;
+
+		// RootActor in scene that used to create DCRA proxy
+		TWeakObjectPtr<ADisplayClusterRootActor> SceneRootActor;
+	};
+
+	/** Registered proxies for renderers by ID. */
+	TMap<int32, FRendererProxy> RendererProxies;
+
+	/** The preview world is only used when using the DCRA proxy. */
+	TObjectPtr<UWorld> PreviewWorld;
+};
diff --git a/Engine/Plugins/Runtime/nDisplay/Source/DisplayClusterScenePreview/Public/DisplayClusterScenePreviewEnums.h b/Engine/Plugins/Runtime/nDisplay/Source/DisplayClusterScenePreview/Public/DisplayClusterScenePreviewEnums.h
new file mode 100644
--- /dev/null
+++ b/Engine/Plugins/Runtime/nDisplay/Source/DisplayClusterScenePreview/Public/DisplayClusterScenePreviewEnums.h
@@ -0,0 +1,20 @@
+// Copyright Epic Games, Inc. All Rights Reserved.
+
+#pragma once
+
+#include "CoreMinimal.h"
+
+/**
+ * Special flags that control the behavior of the renderer.
+*/
+enum class EDisplayClusterScenePreviewFlags : uint8
+{
+	None = 0,
+
+	/** If true, use DCRA proxy for rendering. */
+	UseRootActorProxy = 1 << 0,
+
+	/** If true, automatically update the renderer with stage actors belonging to the root actor. */
+	AutoUpdateLightcards = 1 << 1,
+};
+ENUM_CLASS_FLAGS(EDisplayClusterScenePreviewFlags);
diff --git a/Engine/Plugins/Runtime/nDisplay/Source/DisplayClusterScenePreview/Public/IDisplayClusterScenePreview.h b/Engine/Plugins/Runtime/nDisplay/Source/DisplayClusterScenePreview/Public/IDisplayClusterScenePreview.h
--- a/Engine/Plugins/Runtime/nDisplay/Source/DisplayClusterScenePreview/Public/IDisplayClusterScenePreview.h
+++ b/Engine/Plugins/Runtime/nDisplay/Source/DisplayClusterScenePreview/Public/IDisplayClusterScenePreview.h
@@ -8,8 +8,10 @@
 
 #include "ShowFlags.h"
 #include "DisplayClusterMeshProjectionRenderer.h"
+#include "DisplayClusterScenePreviewEnums.h"
 
 class ADisplayClusterRootActor;
+struct FDisplayClusterRootActorPropertyOverrides;
 
 DECLARE_DELEGATE_OneParam(FRenderResultDelegate, FRenderTarget*);
 
@@ -63,9 +65,10 @@ public:
 	 * 
 	 * @param RendererId The ID of the renderer as returned from CreateRenderer.
 	 * @param ActorPath The path of the ADisplayClusterRootActor actor to preview.
-	 * @param bAutoUpdateLightcards If true, the renderer scene will also be automatically populated with the actor's associated lightcards, including future changes to them.
+	 * @param InPropertyOverrides optional overrides for root actor
+	 * @param PreviewFlags Special flags that control the behavior of the renderer.
 	 */
-	virtual bool SetRendererRootActorPath(int32 RendererId, const FString& ActorPath, bool bAutoUpdateLightcards = false) = 0;
+	virtual bool SetRendererRootActorPath(int32 RendererId, const FString& ActorPath, const FDisplayClusterRootActorPropertyOverrides& InPropertyOverrides, const EDisplayClusterScenePreviewFlags PreviewFlags = EDisplayClusterScenePreviewFlags::None) = 0;
 
 	/**
 	 * Set the root DisplayCluster actor for the renderer with the given ID, which will be added to the preview scene and used to determine the render world.
@@ -74,9 +77,10 @@ public:
 	 * 
 	 * @param RendererId The ID of the renderer as returned from CreateRenderer.
 	 * @param Actor The actor to preview.
-	 * @param bAutoUpdateLightcards If true, the renderer scene will also be automatically populated with the actor's associated lightcards, including future changes to them.
+	 * @param InPropertyOverrides optional overrides for root actor
+	 * @param PreviewFlags Special flags that control the behavior of the renderer.
 	 */
-	virtual bool SetRendererRootActor(int32 RendererId, ADisplayClusterRootActor* Actor, bool bAutoUpdateLightcards = false) = 0;
+	virtual bool SetRendererRootActor(int32 RendererId, ADisplayClusterRootActor* Actor, const FDisplayClusterRootActorPropertyOverrides& InPropertyOverrides, const EDisplayClusterScenePreviewFlags PreviewFlags = EDisplayClusterScenePreviewFlags::None) = 0;
 
 	/**
 	 * Get the root actor of a renderer. 
@@ -145,14 +149,6 @@ public:
 	 */
 	virtual bool SetRendererRenderSimpleElementsDelegate(int32 RendererId, FDisplayClusterMeshProjectionRenderer::FSimpleElementPass RenderSimpleElementsDelegate) = 0;
 
-	/**
-	 * Set whether the renderer should use post-processed nDisplay preview textures.
-	 *
-	 * @param RendererId The ID of the renderer as returned from CreateRenderer.
-	 * @param bUsePostProcessTexture Whether to use post-processed nDisplay preview texture for future renders.
-	 */
-	virtual bool SetRendererUsePostProcessTexture(int32 RendererId, bool bUsePostProcessTexture) = 0;
-
 	/**
 	 * Immediately render a preview.
 	 * 
@@ -186,4 +182,46 @@ public:
 	 * Check whether nDisplay preview textures are being updated in real time.
 	 */
 	virtual bool IsRealTimePreviewEnabled() const = 0;
+
+	/**
+	 * Set the path of the ADisplayClusterRootActor actor for a renderer, which will be added to the preview scene and used to determine the render world.
+	 * Whenever a render is performed with this renderer, it will try to find the actor at that path and use it as the root.
+	 * 
+	 * @param RendererId The ID of the renderer as returned from CreateRenderer.
+	 * @param ActorPath The path of the ADisplayClusterRootActor actor to preview.
+	 * @param bAutoUpdateLightcards If true, the renderer scene will also be automatically populated with the actor's associated lightcards, including future changes to them.
+	 */
+	UE_DEPRECATED(5.5, "This function has been deprecated. Use EDisplayClusterScenePreviewFlags.")
+	virtual bool SetRendererRootActorPath(int32 RendererId, const FString& ActorPath, bool bAutoUpdateLightcards = false)
+	{
+		return false;
+	};
+
+	/**
+	 * Set the root DisplayCluster actor for the renderer with the given ID, which will be added to the preview scene and used to determine the render world.
+	 * Note that if the actor is destroyed (including by a recompile, reload, etc.), you will need to call this again to update the actor.
+	 * This will also clear the path set by SetRendererRootActorPath.
+	 * 
+	 * @param RendererId The ID of the renderer as returned from CreateRenderer.
+	 * @param Actor The actor to preview.
+	 * @param bAutoUpdateLightcards If true, the renderer scene will also be automatically populated with the actor's associated lightcards, including future changes to them.
+	 */
+	UE_DEPRECATED(5.5, "This function has been deprecated.")
+	virtual bool SetRendererRootActor(int32 RendererId, ADisplayClusterRootActor* Actor, bool bAutoUpdateLightcards = false)
+	{
+		return false;
+	};
+
+	/**
+	 * Set whether the renderer should use post-processed nDisplay preview textures.
+	 *
+	 * @param RendererId The ID of the renderer as returned from CreateRenderer.
+	 * @param bUsePostProcessTexture Whether to use post-processed nDisplay preview texture for future renders.
+	 */
+	UE_DEPRECATED(5.5, "This function has been deprecated.")
+	virtual bool SetRendererUsePostProcessTexture(int32 RendererId, bool bUsePostProcessTexture)
+	{
+		return false;
+	};
+
 };
\ No newline at end of file
diff --git a/Engine/Plugins/VirtualProduction/EpicStageApp/Source/EpicStageApp/Private/StageAppRouteHandler.cpp b/Engine/Plugins/VirtualProduction/EpicStageApp/Source/EpicStageApp/Private/StageAppRouteHandler.cpp
--- a/Engine/Plugins/VirtualProduction/EpicStageApp/Source/EpicStageApp/Private/StageAppRouteHandler.cpp
+++ b/Engine/Plugins/VirtualProduction/EpicStageApp/Source/EpicStageApp/Private/StageAppRouteHandler.cpp
@@ -8,6 +8,7 @@
 #include "DisplayClusterLightCardActor.h"
 #include "DisplayClusterLightCardEditorHelper.h"
 #include "DisplayClusterRootActor.h"
+#include "DisplayClusterRootActorContainers.h"
 #include "Engine/Canvas.h"
 #include "Engine/TextureRenderTarget2D.h"
 #include "IDisplayClusterScenePreview.h"
@@ -396,6 +397,31 @@ void FStageAppRouteHandler::RegisterRoute(TUniquePtr<FRemoteControlWebsocketRout
 	Routes.Emplace(MoveTemp(Route));
 }
 
+void FStageAppRouteHandler::SetRendererRootActorPath(const int32 RendererId, const FString& RootActorPath)
+{
+	// Use custom settings for root actor.
+	FDisplayClusterRootActorPropertyOverrides PropertyOverrides;
+	{
+		PropertyOverrides.bPreviewICVFXFrustums = false;
+		PropertyOverrides.bEnablePreviewTechvis = false;
+		PropertyOverrides.bPreviewEnableOverlayMaterial = false;
+		PropertyOverrides.bFreezePreviewRender = false;
+
+		PropertyOverrides.bPreviewEnablePostProcess = true;
+		PropertyOverrides.bPreviewEnable = true;
+
+		PropertyOverrides.PreviewSetttingsSource = EDisplayClusterConfigurationRootActorPreviewSettingsSource::RootActor;
+	}
+
+	IDisplayClusterScenePreview::Get().SetRendererRootActorPath(RendererId, RootActorPath, PropertyOverrides,
+		// When we use preview rendering for an external request with settings different from those used in the scene,
+		// need to create a new DCRA actor and use it as a proxy for custom rendering.
+		EDisplayClusterScenePreviewFlags::UseRootActorProxy
+		// and also continue to use the bAutoUpdateLightcards flag from the previous revision.
+		| EDisplayClusterScenePreviewFlags::AutoUpdateLightcards
+	);
+}
+
 void FStageAppRouteHandler::HandleWebSocketNDisplayPreviewRendererCreate(const FRemoteControlWebSocketMessage& WebSocketMessage)
 {
 	FRCWebSocketNDisplayPreviewRendererCreateBody Body;
@@ -409,12 +435,9 @@ void FStageAppRouteHandler::HandleWebSocketNDisplayPreviewRendererCreate(const F
 	const int32 RendererId = PreviewModule.CreateRenderer();
 	PerRendererDataMapsByClientId.FindOrAdd(WebSocketMessage.ClientId).Emplace(RendererId, RendererId);
 
-	PreviewModule.SetRendererUsePostProcessTexture(RendererId, true);
-
 	if (!Body.RootActorPath.IsEmpty())
 	{
-		PreviewModule.SetRendererRootActorPath(RendererId, Body.RootActorPath, true);
-		BeginForceRootActorPreview(RendererId);
+		SetRendererRootActorPath(RendererId, Body.RootActorPath);
 	}
 
 	ChangePreviewRendererSettings(WebSocketMessage.ClientId, RendererId, Body.Settings);
@@ -436,11 +459,7 @@ void FStageAppRouteHandler::HandleWebSocketNDisplayPreviewRendererSetRoot(const
 		return;
 	}
 
-	EndForceRootActorPreview(Body.RendererId);
-
-	IDisplayClusterScenePreview::Get().SetRendererRootActorPath(Body.RendererId, Body.RootActorPath, true);
-
-	BeginForceRootActorPreview(Body.RendererId);
+	SetRendererRootActorPath(Body.RendererId, Body.RootActorPath);
 }
 
 void FStageAppRouteHandler::HandleWebSocketNDisplayPreviewRendererConfigure(const FRemoteControlWebSocketMessage& WebSocketMessage)
@@ -479,7 +498,6 @@ void FStageAppRouteHandler::HandleWebSocketNDisplayPreviewRendererDestroy(const
 		// Check that this client created the renderer
 		if (PerRendererDataMap->Remove(Body.RendererId) > 0)
 		{
-			EndForceRootActorPreview(Body.RendererId);
 			IDisplayClusterScenePreview::Get().DestroyRenderer(Body.RendererId);
 		}
 	}
@@ -1046,7 +1064,6 @@ void FStageAppRouteHandler::HandleClientDisconnected(FGuid ClientId)
 		{
 			EndActorDrag(PerRendererDataPair.Value, ClientId, PerRendererDataPair.Key, true);
 
-			EndForceRootActorPreview(PerRendererDataPair.Key);
 			PreviewModule.DestroyRenderer(PerRendererDataPair.Key);
 		}
 
@@ -1197,25 +1214,4 @@ void FStageAppRouteHandler::EndActorDrag(FPerRendererData& PerRendererData, cons
 	}
 }
 
-void FStageAppRouteHandler::BeginForceRootActorPreview(int32 RendererId)
-{
-#if WITH_EDITOR
-	if (ADisplayClusterRootActor* RootActor = IDisplayClusterScenePreview::Get().GetRendererRootActor(RendererId))
-	{
-		// Force the new root actor to render previews so the app always gets a render of it
-		RootActor->AddPreviewEnableOverride(reinterpret_cast<uint8*>(this));
-	}
-#endif
-}
-
-void FStageAppRouteHandler::EndForceRootActorPreview(int32 RendererId)
-{
-#if WITH_EDITOR
-	if (ADisplayClusterRootActor* RootActor = IDisplayClusterScenePreview::Get().GetRendererRootActor(RendererId))
-	{
-		RootActor->RemovePreviewEnableOverride(reinterpret_cast<uint8*>(this));
-	};
-#endif
-}
-
 #undef LOCTEXT_NAMESPACE
diff --git a/Engine/Plugins/VirtualProduction/EpicStageApp/Source/EpicStageApp/Private/StageAppRouteHandler.h b/Engine/Plugins/VirtualProduction/EpicStageApp/Source/EpicStageApp/Private/StageAppRouteHandler.h
--- a/Engine/Plugins/VirtualProduction/EpicStageApp/Source/EpicStageApp/Private/StageAppRouteHandler.h
+++ b/Engine/Plugins/VirtualProduction/EpicStageApp/Source/EpicStageApp/Private/StageAppRouteHandler.h
@@ -164,14 +164,8 @@ private:
 	/** End a client's actor drag operation. */
 	void EndActorDrag(FPerRendererData& PreviewData, const FGuid& ClientId, int32 RendererId, bool bEndedByClient);
 
-	/**
-	 * If the preview renderer with the given ID has a root actor, force it to render previews until EndForceRootActorPreview is called
-	 * with the same ID.
-	 */
-	void BeginForceRootActorPreview(int32 RendererId);
-
-	/** If the preview renderer with the given ID has a root actor, stop forcing it to render previews. */
-	void EndForceRootActorPreview(int32 RendererId);
+	/** Set root actor via path. */
+	void SetRendererRootActorPath(const int32 RendererId, const FString& RootActorPath);
 
 private:
 	/** The module for the WebSocket server for which this is handling routes. */
