Removing a number of unnecessary _API uses.

diff --git a/Engine/Plugins/AI/MassAI/Source/MassNavigation/Public/MassNavigationFragments.h b/Engine/Plugins/AI/MassAI/Source/MassNavigation/Public/MassNavigationFragments.h
--- a/Engine/Plugins/AI/MassAI/Source/MassNavigation/Public/MassNavigationFragments.h
+++ b/Engine/Plugins/AI/MassAI/Source/MassNavigation/Public/MassNavigationFragments.h
@@ -112,20 +112,20 @@ struct MASSNAVIGATION_API FMassNavigationObstacleGridCellLocationFragment : publ
 };
 
 
-enum class MASSNAVIGATION_API EMassColliderType : uint8
+enum class EMassColliderType : uint8
 {
 	Circle,
 	Pill,
 };
 
-struct MASSNAVIGATION_API FMassCircleCollider
+struct FMassCircleCollider
 {
 	FMassCircleCollider() = default;
 	FMassCircleCollider(const float Radius) : Radius(Radius) {}
 	float Radius = 0.f;
 };
 
-struct MASSNAVIGATION_API FMassPillCollider
+struct FMassPillCollider
 {
 	FMassPillCollider() = default;
 	FMassPillCollider(const float Radius, const float HalfLength) : Radius(Radius), HalfLength(HalfLength) {}
diff --git a/Engine/Plugins/AudioInsights/Source/AudioInsights/Public/Views/DashboardViewFactory.h b/Engine/Plugins/AudioInsights/Source/AudioInsights/Public/Views/DashboardViewFactory.h
--- a/Engine/Plugins/AudioInsights/Source/AudioInsights/Public/Views/DashboardViewFactory.h
+++ b/Engine/Plugins/AudioInsights/Source/AudioInsights/Public/Views/DashboardViewFactory.h
@@ -12,7 +12,7 @@
 
 namespace UE::Audio::Insights
 {
-	enum class AUDIOINSIGHTS_API EDefaultDashboardTabStack : uint8
+	enum class EDefaultDashboardTabStack : uint8
 	{
 		Viewport,
 		Log,
@@ -22,7 +22,7 @@ namespace UE::Audio::Insights
 		Oscilloscope
 	};
 
-	class AUDIOINSIGHTS_API IDashboardViewFactory
+	class IDashboardViewFactory
 	{
 	public:
 		virtual ~IDashboardViewFactory() = default;
@@ -34,7 +34,7 @@ namespace UE::Audio::Insights
 		virtual TSharedRef<SWidget> MakeWidget() = 0;
 	};
 
-	class AUDIOINSIGHTS_API FTraceDashboardViewFactoryBase : public IDashboardViewFactory
+	class FTraceDashboardViewFactoryBase : public IDashboardViewFactory
 	{
 	public:
 		const TArray<TSharedPtr<FTraceProviderBase>>& GetProviders() const
diff --git a/Engine/Plugins/Editor/StylusInput/Source/StylusInput/Public/IStylusState.h b/Engine/Plugins/Editor/StylusInput/Source/StylusInput/Public/IStylusState.h
--- a/Engine/Plugins/Editor/StylusInput/Source/StylusInput/Public/IStylusState.h
+++ b/Engine/Plugins/Editor/StylusInput/Source/StylusInput/Public/IStylusState.h
@@ -7,7 +7,7 @@
 /**
  * The types of stylus inputs that can be potentially supported by a stylus.
  */
-enum class STYLUSINPUT_API EStylusInputType
+enum class EStylusInputType
 {
 	Position,
 	Z,
@@ -22,7 +22,7 @@ enum class STYLUSINPUT_API EStylusInputType
 /**
  * The current state of a single stylus, as sent by IStylusMessageHandler. 
  */
-class STYLUSINPUT_API FStylusState
+class FStylusState
 {
 public:
 
@@ -115,7 +115,7 @@ private:
 };
 
 /** An input device representing a stylus and its current state. */
-class STYLUSINPUT_API IStylusInputDevice
+class IStylusInputDevice
 {
 public:
 
diff --git a/Engine/Plugins/Experimental/ChaosVehiclesPlugin/Source/ChaosVehicles/Public/ChaosWheeledVehicleMovementComponent.h b/Engine/Plugins/Experimental/ChaosVehiclesPlugin/Source/ChaosVehicles/Public/ChaosWheeledVehicleMovementComponent.h
--- a/Engine/Plugins/Experimental/ChaosVehiclesPlugin/Source/ChaosVehicles/Public/ChaosWheeledVehicleMovementComponent.h
+++ b/Engine/Plugins/Experimental/ChaosVehiclesPlugin/Source/ChaosVehicles/Public/ChaosWheeledVehicleMovementComponent.h
@@ -16,7 +16,7 @@
 UE_DISABLE_OPTIMIZATION
 #endif
 
-struct CHAOSVEHICLES_API FWheeledVehicleDebugParams
+struct FWheeledVehicleDebugParams
 {
 	bool ShowWheelCollisionNormal = false;
 	bool ShowSuspensionRaycasts = false;
@@ -47,7 +47,7 @@ struct CHAOSVEHICLES_API FWheeledVehicleDebugParams
  * Advance through pages using p.Vehicles.NextDebugPage | p.Vehicles.PrevDebugPage which can be hooked
  * up to the keyboard or a controller in blueprint using execCommand
  */
-enum CHAOSVEHICLES_API EDebugPages : uint8
+enum EDebugPages : uint8
 {
 	BasicPage = 0,
 	PerformancePage,
diff --git a/Engine/Plugins/FX/Niagara/Source/NiagaraEditor/Public/ViewModels/NiagaraSystemViewModel.h b/Engine/Plugins/FX/Niagara/Source/NiagaraEditor/Public/ViewModels/NiagaraSystemViewModel.h
--- a/Engine/Plugins/FX/Niagara/Source/NiagaraEditor/Public/ViewModels/NiagaraSystemViewModel.h
+++ b/Engine/Plugins/FX/Niagara/Source/NiagaraEditor/Public/ViewModels/NiagaraSystemViewModel.h
@@ -55,7 +55,7 @@ class UNiagaraUserParametersHierarchyViewModel;
 
 
 /** Defines different editing modes for this system view model. */
-enum class NIAGARAEDITOR_API ENiagaraSystemViewModelEditMode
+enum class ENiagaraSystemViewModelEditMode
 {
 	/** A system asset is being edited.  This assumes that emitters should be inheriting from a base version and that emitter editing will be restricted. */
 	SystemAsset,
diff --git a/Engine/Plugins/Importers/USDImporter/Source/USDUtilities/Public/USDGeomMeshConversion.h b/Engine/Plugins/Importers/USDImporter/Source/USDUtilities/Public/USDGeomMeshConversion.h
--- a/Engine/Plugins/Importers/USDImporter/Source/USDUtilities/Public/USDGeomMeshConversion.h
+++ b/Engine/Plugins/Importers/USDImporter/Source/USDUtilities/Public/USDGeomMeshConversion.h
@@ -434,7 +434,7 @@ namespace UsdUtils
 	USDUTILITIES_API bool IsAnimatedMesh(const pxr::UsdPrim& Prim);
 
 	/** Mesh topology variance as defined by Alembic's MeshTopologyVariance */
-	enum class USDUTILITIES_API EMeshTopologyVariance : uint8
+	enum class EMeshTopologyVariance : uint8
 	{
 		Constant,		// unchanging topology
 		Homogenous,		// connectivity is unchanging, but the positions may change
diff --git a/Engine/Plugins/Importers/USDImporter/Source/USDUtilities/Public/USDPrimConversion.h b/Engine/Plugins/Importers/USDImporter/Source/USDUtilities/Public/USDPrimConversion.h
--- a/Engine/Plugins/Importers/USDImporter/Source/USDUtilities/Public/USDPrimConversion.h
+++ b/Engine/Plugins/Importers/USDImporter/Source/USDUtilities/Public/USDPrimConversion.h
@@ -331,7 +331,7 @@ namespace UnrealToUsd
 	// We use this enum to keep track of which type of baking was already being done for a component, so that if
 	// e.g. we have a track for camera aperture and another for camera focal length, we end up just baking that camera
 	// only once (as only once will be enough to handle all animated properties)
-	enum class USDUTILITIES_API EBakingType : uint8
+	enum class EBakingType : uint8
 	{
 		None = 0,
 		Transform = 1,
diff --git a/Engine/Plugins/Importers/USDImporter/Source/USDUtilities/Public/USDValueConversion.h b/Engine/Plugins/Importers/USDImporter/Source/USDUtilities/Public/USDValueConversion.h
--- a/Engine/Plugins/Importers/USDImporter/Source/USDUtilities/Public/USDValueConversion.h
+++ b/Engine/Plugins/Importers/USDImporter/Source/USDUtilities/Public/USDValueConversion.h
@@ -36,7 +36,7 @@ namespace UE
 namespace UsdUtils
 {
 	/** Source: https://graphics.pixar.com/usd/docs/api/_usd__page__datatypes.html */
-	enum class USDUTILITIES_API EUsdBasicDataTypes
+	enum class EUsdBasicDataTypes
 	{
 		None,
 
@@ -90,7 +90,7 @@ namespace UsdUtils
 
 	/** Corresponds to a value held by a pxr::VtValue converted into UE types. Should handle anything a VtValue can hold, including arrays of values
 	 */
-	struct USDUTILITIES_API FConvertedVtValue
+	struct FConvertedVtValue
 	{
 		TArray<FConvertedVtValueEntry> Entries;
 
diff --git a/Engine/Plugins/Media/PixelStreaming/Source/PixelStreamingEditor/Public/PixelStreamingEditorUtils.h b/Engine/Plugins/Media/PixelStreaming/Source/PixelStreamingEditor/Public/PixelStreamingEditorUtils.h
--- a/Engine/Plugins/Media/PixelStreaming/Source/PixelStreamingEditor/Public/PixelStreamingEditorUtils.h
+++ b/Engine/Plugins/Media/PixelStreaming/Source/PixelStreamingEditor/Public/PixelStreamingEditorUtils.h
@@ -10,7 +10,7 @@ class SWindow;
 
 namespace UE::EditorPixelStreaming
 {
-	enum class PIXELSTREAMINGEDITOR_API EStreamTypes : uint8
+	enum class EStreamTypes : uint8
 	{
 		LevelEditorViewport = 0,
 		Editor = 1
diff --git a/Engine/Plugins/Media/PixelStreaming/Source/PixelStreamingInput/Public/PixelStreamingInputEnums.h b/Engine/Plugins/Media/PixelStreaming/Source/PixelStreamingInput/Public/PixelStreamingInputEnums.h
--- a/Engine/Plugins/Media/PixelStreaming/Source/PixelStreamingInput/Public/PixelStreamingInputEnums.h
+++ b/Engine/Plugins/Media/PixelStreaming/Source/PixelStreamingInput/Public/PixelStreamingInputEnums.h
@@ -8,7 +8,7 @@
  * @brief The data types supported when sending messages across the data channel to/from peers
  *
  */
-enum class PIXELSTREAMINGINPUT_API EPixelStreamingMessageTypes
+enum class EPixelStreamingMessageTypes
 {
 	Uint8,
 	Uint16,
@@ -24,7 +24,7 @@ enum class PIXELSTREAMINGINPUT_API EPixelStreamingMessageTypes
  * @brief The message directions
  *
  */
-enum class PIXELSTREAMINGINPUT_API EPixelStreamingMessageDirection : uint8
+enum class EPixelStreamingMessageDirection : uint8
 {
 	ToStreamer = 0,
 	FromStreamer = 1
@@ -34,7 +34,7 @@ enum class PIXELSTREAMINGINPUT_API EPixelStreamingMessageDirection : uint8
  * @brief The possible actions pixel streaming supports. These actions help differentiate input received from the browser
  *
  */
-enum class PIXELSTREAMINGINPUT_API EPixelStreamingInputAction : uint8
+enum class EPixelStreamingInputAction : uint8
 {
 	X = 0,
 	Y = 1,
@@ -48,7 +48,7 @@ enum class PIXELSTREAMINGINPUT_API EPixelStreamingInputAction : uint8
  * @brief The possible control schemes pixel streaming supports. RouteToWindow routes input at an application level. RouteToWidget routes input to a specific widget, ignoring the rest of the application.
  *
  */
-enum class PIXELSTREAMINGINPUT_API EPixelStreamingInputType : uint8
+enum class EPixelStreamingInputType : uint8
 {
 	RouteToWindow,
 	RouteToWidget
diff --git a/Engine/Plugins/Media/PixelStreaming/Source/PixelStreamingServers/Public/PixelStreamingServers.h b/Engine/Plugins/Media/PixelStreaming/Source/PixelStreamingServers/Public/PixelStreamingServers.h
--- a/Engine/Plugins/Media/PixelStreaming/Source/PixelStreamingServers/Public/PixelStreamingServers.h
+++ b/Engine/Plugins/Media/PixelStreaming/Source/PixelStreamingServers/Public/PixelStreamingServers.h
@@ -21,7 +21,7 @@ namespace UE::PixelStreamingServers
 	/**
 	 * Configuration to control behaviour when launching any of the Pixel Streaming servers.
 	 **/
-	struct PIXELSTREAMINGSERVERS_API FLaunchArgs
+	struct FLaunchArgs
 	{
 		// Arguments passed to the actual server when its process is started.
 		FString ProcessArgs = TEXT("");
@@ -42,7 +42,7 @@ namespace UE::PixelStreamingServers
 	/**
 	 * Endpoints for the various Pixel Streaming servers.
 	 **/
-	enum class PIXELSTREAMINGSERVERS_API EEndpoint
+	enum class EEndpoint
 	{
 		// The websocket signalling url between the server and the UE streamer - e.g. ws://localhost:8888
 		Signalling_Streamer,
@@ -68,7 +68,7 @@ namespace UE::PixelStreamingServers
 	/**
 	 * Interface for all Pixel Streaming servers.
 	 **/
-	class PIXELSTREAMINGSERVERS_API IServer
+	class IServer
 	{
 	public:
 		virtual ~IServer() = default;
diff --git a/Engine/Plugins/Online/OnlineSubsystemSteam/Source/Public/OnlineAuthInterfaceUtilsSteam.h b/Engine/Plugins/Online/OnlineSubsystemSteam/Source/Public/OnlineAuthInterfaceUtilsSteam.h
--- a/Engine/Plugins/Online/OnlineSubsystemSteam/Source/Public/OnlineAuthInterfaceUtilsSteam.h
+++ b/Engine/Plugins/Online/OnlineSubsystemSteam/Source/Public/OnlineAuthInterfaceUtilsSteam.h
@@ -5,7 +5,7 @@
 #include "CoreMinimal.h"
 
 /** Status flags of SteamAuth response codes, mapped 1:1 with SteamAPI */
-enum class ONLINESUBSYSTEMSTEAM_API ESteamAuthResponseCode : uint8
+enum class ESteamAuthResponseCode : uint8
 {
 	OK = 0,
 	NotConnectedToSteam = 1,
@@ -46,7 +46,7 @@ DECLARE_MULTICAST_DELEGATE_ThreeParams(FOnSteamAuthResultWithCode, const class F
 typedef FOnSteamAuthResultWithCode::FDelegate FOnSteamAuthResultWithCodeDelegate;
 
 /** A class for interfacing with SteamAuth outside of the SteamOSS. */
-class ONLINESUBSYSTEMSTEAM_API FOnlineAuthUtilsSteam
+class FOnlineAuthUtilsSteam
 {
 public:
 	
@@ -64,7 +64,7 @@ public:
 	 * @return the true if SteamAuth is enabled.
 	 *
 	 */
-	bool IsSteamAuthEnabled() const;
+	ONLINESUBSYSTEMSTEAM_API bool IsSteamAuthEnabled() const;
 	
 	/** Attach to this delegate to control the behavior of the Steam authentication failure. 
 	 *  This overrides the default behavior (kick). 
diff --git a/Engine/Plugins/Runtime/Harmonix/Source/Harmonix/Public/Harmonix/PropertyUtility.h b/Engine/Plugins/Runtime/Harmonix/Source/Harmonix/Public/Harmonix/PropertyUtility.h
--- a/Engine/Plugins/Runtime/Harmonix/Source/Harmonix/Public/Harmonix/PropertyUtility.h
+++ b/Engine/Plugins/Runtime/Harmonix/Source/Harmonix/Public/Harmonix/PropertyUtility.h
@@ -13,14 +13,14 @@ DEFINE_LOG_CATEGORY_STATIC(LogPropertyUtility, Log, Log)
 namespace Harmonix
 {
 
-enum HARMONIX_API EPostEditAction : uint8
+enum EPostEditAction : uint8
 {
 	DoNothing,
 	UpdateTrivial,
 	UpdateNonTrivial
 };
 
-enum HARMONIX_API EPostEditType : uint8
+enum EPostEditType : uint8
 {
 	None,
 	Trivial,
diff --git a/Engine/Plugins/Runtime/nDisplay/Source/DisplayClusterLightCardEditorShaders/Public/DisplayClusterMeshProjectionRenderer.h b/Engine/Plugins/Runtime/nDisplay/Source/DisplayClusterLightCardEditorShaders/Public/DisplayClusterMeshProjectionRenderer.h
--- a/Engine/Plugins/Runtime/nDisplay/Source/DisplayClusterLightCardEditorShaders/Public/DisplayClusterMeshProjectionRenderer.h
+++ b/Engine/Plugins/Runtime/nDisplay/Source/DisplayClusterLightCardEditorShaders/Public/DisplayClusterMeshProjectionRenderer.h
@@ -22,7 +22,7 @@ class FMeshProjectionPassParameters;
 struct FEngineShowFlags;
 
 /** Indicates which kind of projection is used by the renderer */
-enum DISPLAYCLUSTERLIGHTCARDEDITORSHADERS_API EDisplayClusterMeshProjectionType
+enum EDisplayClusterMeshProjectionType
 {
 	/** Default linear projection */
 	Linear,
@@ -35,7 +35,7 @@ enum DISPLAYCLUSTERLIGHTCARDEDITORSHADERS_API EDisplayClusterMeshProjectionType
 };
 
 /** Indicates the quantity that is output to the canvas by the renderer */
-enum DISPLAYCLUSTERLIGHTCARDEDITORSHADERS_API EDisplayClusterMeshProjectionOutput
+enum EDisplayClusterMeshProjectionOutput
 {
 	/** Outputs the emissive color of the rendered primitives */
 	Color,
@@ -45,7 +45,7 @@ enum DISPLAYCLUSTERLIGHTCARDEDITORSHADERS_API EDisplayClusterMeshProjectionOutpu
 };
 
 /** A filter that allows specific primitive components to be filtered from a render pass */
-class DISPLAYCLUSTERLIGHTCARDEDITORSHADERS_API FDisplayClusterMeshProjectionPrimitiveFilter
+class FDisplayClusterMeshProjectionPrimitiveFilter
 {
 public:
 	DECLARE_DELEGATE_RetVal_OneParam(bool, FPrimitiveFilter, const UPrimitiveComponent*);
@@ -58,14 +58,14 @@ public:
 	FPrimitiveFilter ShouldApplyProjectionDelegate;
 
 	/** Gets whether a primitive component should be filtered out of the render pass or not */
-	bool ShouldRenderPrimitive(const UPrimitiveComponent* InPrimitiveComponent) const;
+	DISPLAYCLUSTERLIGHTCARDEDITORSHADERS_API bool ShouldRenderPrimitive(const UPrimitiveComponent* InPrimitiveComponent) const;
 
 	/** Gets whether a primitive component should be rendered using the current projection type or not */
-	bool ShouldApplyProjection(const UPrimitiveComponent* InPrimitiveComponent) const;
+	DISPLAYCLUSTERLIGHTCARDEDITORSHADERS_API bool ShouldApplyProjection(const UPrimitiveComponent* InPrimitiveComponent) const;
 };
 
 /** Settings for specific mesh projection types  */
-struct DISPLAYCLUSTERLIGHTCARDEDITORSHADERS_API FDisplayClusterMeshProjectionTypeSettings
+struct FDisplayClusterMeshProjectionTypeSettings
 {
 	/** The index of the UV to use when performing a UV projection */
 	uint32 UVProjectionIndex = 0;
@@ -81,7 +81,7 @@ struct DISPLAYCLUSTERLIGHTCARDEDITORSHADERS_API FDisplayClusterMeshProjectionTyp
 };
 
 /** Settings for producing a single render. */
-struct DISPLAYCLUSTERLIGHTCARDEDITORSHADERS_API FDisplayClusterMeshProjectionRenderSettings
+struct FDisplayClusterMeshProjectionRenderSettings
 {
 	/** Camera setup options for the render. */
 	FSceneViewInitOptions ViewInitOptions;
@@ -112,7 +112,7 @@ struct DISPLAYCLUSTERLIGHTCARDEDITORSHADERS_API FDisplayClusterMeshProjectionRen
 };
 
 /** A transform that can be passed around to project and unprojection positions for a specific projection type */
-class DISPLAYCLUSTERLIGHTCARDEDITORSHADERS_API FDisplayClusterMeshProjectionTransform
+class FDisplayClusterMeshProjectionTransform
 {
 public:
 	FDisplayClusterMeshProjectionTransform()
@@ -125,8 +125,8 @@ public:
 		, InvViewMatrix(InViewMatrix.Inverse())
 	{ }
 
-	FVector ProjectPosition(const FVector& WorldPosition) const;
-	FVector UnprojectPosition(const FVector& ProjectedPosition) const;
+	DISPLAYCLUSTERLIGHTCARDEDITORSHADERS_API FVector ProjectPosition(const FVector& WorldPosition) const;
+	DISPLAYCLUSTERLIGHTCARDEDITORSHADERS_API FVector UnprojectPosition(const FVector& ProjectedPosition) const;
 
 private:
 	EDisplayClusterMeshProjectionType Projection;
@@ -135,32 +135,32 @@ private:
 };
 
 /** A renderer that projects meshes to screen space using non-linear projection methods */
-class DISPLAYCLUSTERLIGHTCARDEDITORSHADERS_API FDisplayClusterMeshProjectionRenderer
+class FDisplayClusterMeshProjectionRenderer
 {
 public:
 	/** Clean up any references to the renderer. */
-	~FDisplayClusterMeshProjectionRenderer();
+	DISPLAYCLUSTERLIGHTCARDEDITORSHADERS_API ~FDisplayClusterMeshProjectionRenderer();
 
 	/** Projects a position in view coordinates into the projected view space of the specified projection type */
-	static FVector ProjectViewPosition(const FVector& ViewPosition, EDisplayClusterMeshProjectionType  ProjectionType);
+	DISPLAYCLUSTERLIGHTCARDEDITORSHADERS_API static FVector ProjectViewPosition(const FVector& ViewPosition, EDisplayClusterMeshProjectionType  ProjectionType);
 
 	/** Projects a position in the projected view space of the specified projection type to ordinary view coordinates */
-	static FVector UnprojectViewPosition(const FVector& ProjectedViewPosition, EDisplayClusterMeshProjectionType  ProjectionType);
+	DISPLAYCLUSTERLIGHTCARDEDITORSHADERS_API static FVector UnprojectViewPosition(const FVector& ProjectedViewPosition, EDisplayClusterMeshProjectionType  ProjectionType);
 
 	/** Adds an actor's primitive components to the list of primitives to render */
-	void AddActor(AActor* Actor);
+	DISPLAYCLUSTERLIGHTCARDEDITORSHADERS_API void AddActor(AActor* Actor);
 
 	/** Adds an actor's primitive components to the list of primitives to render, filtering which primitive components get rendered using the specified callback */
-	void AddActor(AActor* Actor, const TFunctionRef<bool(const UPrimitiveComponent*)>& PrimitiveFilter);
+	DISPLAYCLUSTERLIGHTCARDEDITORSHADERS_API void AddActor(AActor* Actor, const TFunctionRef<bool(const UPrimitiveComponent*)>& PrimitiveFilter);
 
 	/** Removes an actor's primitive components from the list of primitives to render */
-	void RemoveActor(AActor* Actor);
+	DISPLAYCLUSTERLIGHTCARDEDITORSHADERS_API void RemoveActor(AActor* Actor);
 
 	/** Clears the list of primitives to render */
-	void ClearScene();
+	DISPLAYCLUSTERLIGHTCARDEDITORSHADERS_API void ClearScene();
 
 	/** Renders its list of primitive components to the specified canvas using the desired projection type. Can be called from the game thread */
-	void Render(FCanvas* Canvas, FSceneInterface* Scene, const FDisplayClusterMeshProjectionRenderSettings& RenderSettings);
+	DISPLAYCLUSTERLIGHTCARDEDITORSHADERS_API void Render(FCanvas* Canvas, FSceneInterface* Scene, const FDisplayClusterMeshProjectionRenderSettings& RenderSettings);
 
 private:
 	/** Constructs the necessary render passes for the default output of the rendered primitives */
diff --git a/Engine/Plugins/VirtualProduction/Rivermax/RivermaxCore/Source/RivermaxCore/Public/RivermaxFormats.h b/Engine/Plugins/VirtualProduction/Rivermax/RivermaxCore/Source/RivermaxCore/Public/RivermaxFormats.h
--- a/Engine/Plugins/VirtualProduction/Rivermax/RivermaxCore/Source/RivermaxCore/Public/RivermaxFormats.h
+++ b/Engine/Plugins/VirtualProduction/Rivermax/RivermaxCore/Source/RivermaxCore/Public/RivermaxFormats.h
@@ -8,7 +8,7 @@
 
 namespace UE::RivermaxCore
 {
-	enum class RIVERMAXCORE_API ESamplingType : uint8
+	enum class ESamplingType : uint8
 	{
 		YUV422_8bit
 		, YUV422_10bit
@@ -102,7 +102,7 @@ namespace UE::RivermaxCore
 		return TEXT("<Unknown ESamplingType>");;
 	}
 
-	struct RIVERMAXCORE_API FVideoFormatInfo
+	struct FVideoFormatInfo
 	{
 		/** Sampling type of this format */
 		ESamplingType Sampling = ESamplingType::RGB_10bit;
@@ -117,13 +117,13 @@ namespace UE::RivermaxCore
 		uint16 PixelGroupCoverage = 0;
 	};
 
-	struct RIVERMAXCORE_API FStandardVideoFormat
+	struct FStandardVideoFormat
 	{
-		static FVideoFormatInfo GetVideoFormatInfo(ESamplingType SamplingType);
-		static bool IsRGB(ESamplingType SamplingType);
-		static bool IsYUV(ESamplingType SamplingType);
-		static bool IsYUV422(ESamplingType SamplingType);
-		static bool IsYUV444(ESamplingType SamplingType);
+		RIVERMAXCORE_API static FVideoFormatInfo GetVideoFormatInfo(ESamplingType SamplingType);
+		RIVERMAXCORE_API static bool IsRGB(ESamplingType SamplingType);
+		RIVERMAXCORE_API static bool IsYUV(ESamplingType SamplingType);
+		RIVERMAXCORE_API static bool IsYUV422(ESamplingType SamplingType);
+		RIVERMAXCORE_API static bool IsYUV444(ESamplingType SamplingType);
 
 	private:
 		static const TArray<FVideoFormatInfo> AllSamplingTypes;
diff --git a/Engine/Plugins/VirtualProduction/Rivermax/RivermaxCore/Source/RivermaxCore/Public/RivermaxTypes.h b/Engine/Plugins/VirtualProduction/Rivermax/RivermaxCore/Source/RivermaxCore/Public/RivermaxTypes.h
--- a/Engine/Plugins/VirtualProduction/Rivermax/RivermaxCore/Source/RivermaxCore/Public/RivermaxTypes.h
+++ b/Engine/Plugins/VirtualProduction/Rivermax/RivermaxCore/Source/RivermaxCore/Public/RivermaxTypes.h
@@ -72,7 +72,7 @@ namespace UE::RivermaxCore
 		return TEXT("<Unknown EFrameLockingMode>");
 	}
 
-	struct RIVERMAXCORE_API FRivermaxInputStreamOptions
+	struct FRivermaxInputStreamOptions
 	{
 		/** Stream FrameRate */
 		FFrameRate FrameRate = { 24,1 };
@@ -102,7 +102,7 @@ namespace UE::RivermaxCore
 		bool bUseGPUDirect = true;
 	};
 
-	struct RIVERMAXCORE_API FRivermaxOutputStreamOptions
+	struct FRivermaxOutputStreamOptions
 	{
 		/** Desired stream resolution */
 		FIntPoint Resolution = { 1920, 1080 };
@@ -144,7 +144,7 @@ namespace UE::RivermaxCore
 		bool bDoFrameCounterTimestamping = true;
 	};
 
-	enum class RIVERMAXCORE_API ERivermaxStreamType : uint8
+	enum class ERivermaxStreamType : uint8
 	{
 		VIDEO_2110_20_STREAM,
 	};
diff --git a/Engine/Source/Developer/TraceInsights/Public/Insights/ITimingViewSession.h b/Engine/Source/Developer/TraceInsights/Public/Insights/ITimingViewSession.h
--- a/Engine/Source/Developer/TraceInsights/Public/Insights/ITimingViewSession.h
+++ b/Engine/Source/Developer/TraceInsights/Public/Insights/ITimingViewSession.h
@@ -8,7 +8,7 @@ class FBaseTimingTrack;
 class FTimingEventsTrack;
 class ITimingEvent;
 class SWidget;
-enum class TRACEINSIGHTS_API ETimingTrackLocation : uint32; // Insights/ViewModels/BaseTimingTrack.h
+enum class ETimingTrackLocation : uint32; // Insights/ViewModels/BaseTimingTrack.h
 
 namespace Insights
 {
@@ -26,7 +26,7 @@ ENUM_CLASS_FLAGS(ETimeChangedFlags);
 
 ////////////////////////////////////////////////////////////////////////////////////////////////////
 
-class TRACEINSIGHTS_API ITimeMarker
+class ITimeMarker
 {
 public:
 	virtual ~ITimeMarker() = default;
@@ -70,7 +70,7 @@ DECLARE_MULTICAST_DELEGATE_OneParam(FTrackRemovedDelegate, const TSharedPtr<cons
 ////////////////////////////////////////////////////////////////////////////////////////////////////
 
 /** Hosts a number of timing view visualizers, represents a session of the timing view. */
-class TRACEINSIGHTS_API ITimingViewSession
+class ITimingViewSession
 {
 public:
 	virtual ~ITimingViewSession() = default;
diff --git a/Engine/Source/Developer/TraceInsights/Public/Insights/ViewModels/BaseTimingTrack.h b/Engine/Source/Developer/TraceInsights/Public/Insights/ViewModels/BaseTimingTrack.h
--- a/Engine/Source/Developer/TraceInsights/Public/Insights/ViewModels/BaseTimingTrack.h
+++ b/Engine/Source/Developer/TraceInsights/Public/Insights/ViewModels/BaseTimingTrack.h
@@ -28,7 +28,7 @@ struct FGeometry;
 
 ////////////////////////////////////////////////////////////////////////////////////////////////////
 
-enum class TRACEINSIGHTS_API ETimingTrackLocation : uint32
+enum class ETimingTrackLocation : uint32
 {
 	None         = 0,
 	Scrollable   = (1 << 0),
@@ -41,7 +41,7 @@ ENUM_CLASS_FLAGS(ETimingTrackLocation);
 
 ////////////////////////////////////////////////////////////////////////////////////////////////////
 
-struct TRACEINSIGHTS_API FTimingTrackOrder
+struct FTimingTrackOrder
 {
 	static constexpr int32 GroupRange = 100000;
 	static constexpr int32 TimeRuler  = -2 * GroupRange;
@@ -56,7 +56,7 @@ struct TRACEINSIGHTS_API FTimingTrackOrder
 
 ////////////////////////////////////////////////////////////////////////////////////////////////////
 
-enum class TRACEINSIGHTS_API ETimingTrackFlags : uint32
+enum class ETimingTrackFlags : uint32
 {
 	None            = 0,
 	IsVisible       = (1 << 0),
@@ -69,7 +69,7 @@ ENUM_CLASS_FLAGS(ETimingTrackFlags);
 
 ////////////////////////////////////////////////////////////////////////////////////////////////////
 
-enum class TRACEINSIGHTS_API EDrawEventMode : uint32
+enum class EDrawEventMode : uint32
 {
 	None = 0,
 
@@ -89,7 +89,7 @@ ENUM_CLASS_FLAGS(EDrawEventMode);
 
 ////////////////////////////////////////////////////////////////////////////////////////////////////
 
-class TRACEINSIGHTS_API ITimingTrackUpdateContext
+class ITimingTrackUpdateContext
 {
 public:
 	virtual const FGeometry& GetGeometry() const = 0;
@@ -105,7 +105,7 @@ public:
 
 ////////////////////////////////////////////////////////////////////////////////////////////////////
 
-class TRACEINSIGHTS_API ITimingTrackDrawContext
+class ITimingTrackDrawContext
 {
 public:
 	virtual const FTimingTrackViewport& GetViewport() const = 0;
diff --git a/Engine/Source/Editor/BlueprintGraph/Classes/K2Node.h b/Engine/Source/Editor/BlueprintGraph/Classes/K2Node.h
--- a/Engine/Source/Editor/BlueprintGraph/Classes/K2Node.h
+++ b/Engine/Source/Editor/BlueprintGraph/Classes/K2Node.h
@@ -412,7 +412,7 @@ public:
 	 */
 	virtual FBlueprintNodeSignature GetSignature() const { return FBlueprintNodeSignature(GetClass()); }
 
-	enum BLUEPRINTGRAPH_API EBaseNodeRefreshPriority
+	enum EBaseNodeRefreshPriority
 	{
 		Low_UsesDependentWildcard = 100,
 		Low_ReceivesDelegateSignature = 150,
diff --git a/Engine/Source/Runtime/ApplicationCore/Public/Windows/WindowsApplication.h b/Engine/Source/Runtime/ApplicationCore/Public/Windows/WindowsApplication.h
--- a/Engine/Source/Runtime/ApplicationCore/Public/Windows/WindowsApplication.h
+++ b/Engine/Source/Runtime/ApplicationCore/Public/Windows/WindowsApplication.h
@@ -31,7 +31,7 @@ DECLARE_MULTICAST_DELEGATE_OneParam(FWindowsApplication_OnWindowCreated, HWND);
 
 namespace ETaskbarProgressState
 {
-	enum APPLICATIONCORE_API Type
+	enum Type
 	{
 		//Stops displaying progress and returns the button to its normal state.
 		NoProgress = 0x0,
diff --git a/Engine/Source/Runtime/Experimental/ChaosVehicles/ChaosVehiclesCore/Public/SimModule/SimulationModuleBase.h b/Engine/Source/Runtime/Experimental/ChaosVehicles/ChaosVehiclesCore/Public/SimModule/SimulationModuleBase.h
--- a/Engine/Source/Runtime/Experimental/ChaosVehicles/ChaosVehiclesCore/Public/SimModule/SimulationModuleBase.h
+++ b/Engine/Source/Runtime/Experimental/ChaosVehicles/ChaosVehiclesCore/Public/SimModule/SimulationModuleBase.h
@@ -123,13 +123,13 @@ namespace Chaos
 	};
 
 
-	enum CHAOSVEHICLESCORE_API eSimModuleState
+	enum eSimModuleState
 	{
 		Disabled,
 		Enabled
 	};
 
-	enum CHAOSVEHICLESCORE_API eSimModuleTypeFlags
+	enum eSimModuleTypeFlags
 	{
 		NonFunctional = (1 << 0),	// bitmask 1,2,4,8
 		Raycast = (1 << 1),	// requires raycast data
@@ -137,7 +137,7 @@ namespace Chaos
 		Velocity = (1 << 3),	// requires velocity data
 	};
 
-	enum CHAOSVEHICLESCORE_API eSimType
+	enum eSimType
 	{
 		Undefined = 0,
 		Chassis,		// linear/angular damping can be applied here
@@ -157,7 +157,7 @@ namespace Chaos
 		TorqueSim
 	};
 
-	enum CHAOSVEHICLESCORE_API EWheelAxis
+	enum EWheelAxis
 	{
 		X,	// X forward
 		Y	// Y forward
