nDisplay: Removed unnecessary dependency on VPUtils plugin.

This was a vestige from light cards actor placement, but it was later removed.

diff --git a/Engine/Plugins/Runtime/nDisplay/nDisplay.uplugin b/Engine/Plugins/Runtime/nDisplay/nDisplay.uplugin
--- a/Engine/Plugins/Runtime/nDisplay/nDisplay.uplugin
+++ b/Engine/Plugins/Runtime/nDisplay/nDisplay.uplugin
@@ -279,10 +279,6 @@
 			"Name": "RemoteControlInterception",
 			"Enabled": true
 		},
-		{
-			"Name": "VirtualProductionUtilities",
-			"Enabled": true
-		},
 		{
 			"Name": "MovieRenderPipeline",
 			"Enabled": true
