nDisplay: fixed crash with nDisplay replication for StaticMeshActors with enabled SimulatePhysics flag which happens due to missing nullptr check required by slightly different replicated scheduling behaviour

diff --git a/Engine/Plugins/Runtime/nDisplay/Source/DisplayClusterReplication/Private/DisplayClusterNetDriver.cpp b/Engine/Plugins/Runtime/nDisplay/Source/DisplayClusterReplication/Private/DisplayClusterNetDriver.cpp
--- a/Engine/Plugins/Runtime/nDisplay/Source/DisplayClusterReplication/Private/DisplayClusterNetDriver.cpp
+++ b/Engine/Plugins/Runtime/nDisplay/Source/DisplayClusterReplication/Private/DisplayClusterNetDriver.cpp
@@ -567,6 +567,11 @@ void UDisplayClusterNetDriver::PostListUpdate(ConsiderListUpdateParams const& Up
 
 		for (int j = 0; j < ClusterReplicationState.FinalSortedCount; j++)
 		{
+			if (ClusterReplicationState.PriorityList[j].ActorInfo == nullptr)
+			{
+				continue;
+			}
+
 			ClusterReplicationState.PriorityList[j].Channel = DisplayClusterConnection->FindActorChannelRef(ClusterReplicationState.PriorityList[j].ActorInfo->WeakActor);
 		}
 
@@ -632,6 +637,11 @@ void UDisplayClusterNetDriver::ListUpdate(ConsiderListUpdateParams const& Update
 		// Update channels for each actor
 		for (int j = 0; j < ClusterReplicationState.FinalSortedCount; j++)
 		{
+			if (ClusterReplicationState.PriorityList[j].ActorInfo == nullptr)
+			{
+				continue;
+			}
+
 			ClusterReplicationState.PriorityList[j].Channel = Connection->FindActorChannelRef(ClusterReplicationState.PriorityList[j].ActorInfo->WeakActor);
 		}
 
