[nDisplay] Media tiles configuration dialog for ICVFX cameras

diff --git a/Engine/Plugins/Runtime/nDisplay/Source/DisplayClusterConfigurator/DisplayClusterConfigurator.Build.cs b/Engine/Plugins/Runtime/nDisplay/Source/DisplayClusterConfigurator/DisplayClusterConfigurator.Build.cs
--- a/Engine/Plugins/Runtime/nDisplay/Source/DisplayClusterConfigurator/DisplayClusterConfigurator.Build.cs
+++ b/Engine/Plugins/Runtime/nDisplay/Source/DisplayClusterConfigurator/DisplayClusterConfigurator.Build.cs
@@ -21,6 +21,7 @@ public class DisplayClusterConfigurator : ModuleRules
 				"AppFramework",
 				"AssetTools",
 				"CinematicCamera",
+				"ClassViewer",
 				"Core",
 				"CoreUObject",
 				"DesktopPlatform",
@@ -37,6 +38,7 @@ public class DisplayClusterConfigurator : ModuleRules
 				"KismetCompiler",
 				"MainFrame",
 				"MediaAssets",
+				"MediaIOCore",
 				"MessageLog",
 				"Networking",
 				"PinnedCommandList",
diff --git a/Engine/Plugins/Runtime/nDisplay/Source/DisplayClusterConfigurator/Private/Views/Details/Media/DisplayClusterConfiguratorICVFXMediaCustomization.cpp b/Engine/Plugins/Runtime/nDisplay/Source/DisplayClusterConfigurator/Private/Views/Details/Media/DisplayClusterConfiguratorICVFXMediaCustomization.cpp
--- a/Engine/Plugins/Runtime/nDisplay/Source/DisplayClusterConfigurator/Private/Views/Details/Media/DisplayClusterConfiguratorICVFXMediaCustomization.cpp
+++ b/Engine/Plugins/Runtime/nDisplay/Source/DisplayClusterConfigurator/Private/Views/Details/Media/DisplayClusterConfiguratorICVFXMediaCustomization.cpp
@@ -2,29 +2,46 @@
 
 #include "Views/Details/Media/DisplayClusterConfiguratorICVFXMediaCustomization.h"
 #include "Views/Details/Media/DisplayClusterConfiguratorMediaUtils.h"
+#include "Views/Details/Media/SMediaTilesConfigurationDialog.h"
 
 #include "DisplayClusterConfigurationTypes_Media.h"
+#include "DisplayClusterConfiguratorBlueprintEditor.h"
+#include "DisplayClusterConfiguratorUtils.h"
+#include "DisplayClusterRootActor.h"
+
+#include "Blueprints/DisplayClusterBlueprint.h"
+#include "Components/DisplayClusterICVFXCameraComponent.h"
 
 #include "DetailWidgetRow.h"
 #include "IDetailChildrenBuilder.h"
 #include "IPropertyUtilities.h"
 #include "PropertyHandle.h"
 
+#include "Framework/Application/SlateApplication.h"
+#include "HAL/PlatformApplicationMisc.h"
+#include "Interfaces/IMainFrameModule.h"
+#include "Misc/MessageDialog.h"
+#include "Modules/ModuleManager.h"
+
 #include "Widgets/Input/SButton.h"
 #include "Widgets/Layout/SBox.h"
 #include "Widgets/Text/STextBlock.h"
 
+#include "MediaSource.h"
+#include "MediaOutput.h"
+
 #define LOCTEXT_NAMESPACE "FDisplayClusterConfiguratorICVFXMediaCustomization"
 
+
 void FDisplayClusterConfiguratorICVFXMediaCustomization::CustomizeChildren(TSharedRef<IPropertyHandle> InPropertyHandle, IDetailChildrenBuilder& InChildBuilder, IPropertyTypeCustomizationUtils& InCustomizationUtils)
 {
 	// SplitType property
 	TSharedPtr<IPropertyHandle> SplitTypeHandle = GET_CHILD_HANDLE(FDisplayClusterConfigurationMediaICVFX, SplitType);
-	check(SplitTypeHandle->IsValidHandle());
+	check(SplitTypeHandle && SplitTypeHandle->IsValidHandle());
 
 	// Layout property
 	TilesLayoutHandle = GET_CHILD_HANDLE(FDisplayClusterConfigurationMediaICVFX, TiledSplitLayout);
-	check(TilesLayoutHandle->IsValidHandle());
+	check(TilesLayoutHandle && TilesLayoutHandle->IsValidHandle());
 
 	// Separate groups specific for every split type available
 	TArray<TSharedPtr<IPropertyHandle>> FullFramePropertyHandles;
@@ -86,36 +103,163 @@ void FDisplayClusterConfiguratorICVFXMediaCustomization::CustomizeChildren(TShar
 	// Create all the property widgets
 	FDisplayClusterConfiguratorBaseTypeCustomization::CustomizeChildren(InPropertyHandle, InChildBuilder, InCustomizationUtils);
 
-	// Create auto-configure button in the bottom
-	AddAutoConfigurationButton(InChildBuilder);
+	// Create 'setup' button
+	if (SplitTypeValue == EDisplayClusterConfigurationMediaSplitType::UniformTiles)
+	{
+		AddSetupButton(InChildBuilder);
+	}
+
+	// Create 'reset' button at the bottom
+	AddResetButton(InChildBuilder);
 }
 
-void FDisplayClusterConfiguratorICVFXMediaCustomization::AddAutoConfigurationButton(IDetailChildrenBuilder& InChildBuilder)
+void FDisplayClusterConfiguratorICVFXMediaCustomization::AddSetupButton(IDetailChildrenBuilder& InChildBuilder)
 {
 	InChildBuilder.AddCustomRow(FText::GetEmpty())
 		.WholeRowContent()
 		[
 			SNew(SBox)
-				.Padding(5.f)
+			.Padding(5.f)
+			[
+				SNew(SButton)
+				.HAlign(HAlign_Center)
+				.OnClicked(this, &FDisplayClusterConfiguratorICVFXMediaCustomization::OnSetupButtonClicked)
 				[
-					SNew(SButton)
-						.HAlign(HAlign_Center)
-						.OnClicked(this, &FDisplayClusterConfiguratorICVFXMediaCustomization::OnAutoConfigureButtonClicked)
-						[
-							SNew(STextBlock)
-								.Font(FAppStyle::GetFontStyle(TEXT("PropertyWindow.NormalFont")))
-								.Text(LOCTEXT("AutoConfigureInputButtonTitle", "Auto-Configure"))
-						]
+					SNew(STextBlock)
+					.Font(FAppStyle::GetFontStyle(TEXT("PropertyWindow.NormalFont")))
+					.Text(LOCTEXT("SetupTilesButtonTitle", "Setup Tiles"))
 				]
+			]
 		];
 }
 
-FReply FDisplayClusterConfiguratorICVFXMediaCustomization::OnAutoConfigureButtonClicked()
+void FDisplayClusterConfiguratorICVFXMediaCustomization::AddResetButton(IDetailChildrenBuilder& InChildBuilder)
+{
+	InChildBuilder.AddCustomRow(FText::GetEmpty())
+		.WholeRowContent()
+		[
+			SNew(SBox)
+			.Padding(5.f)
+			[
+				SNew(SButton)
+				.HAlign(HAlign_Center)
+				.OnClicked(this, &FDisplayClusterConfiguratorICVFXMediaCustomization::OnResetButtonClicked)
+				[
+					SNew(STextBlock)
+					.Font(FAppStyle::GetFontStyle(TEXT("PropertyWindow.NormalFont")))
+					.Text(LOCTEXT("ResetToDefaultsButtonTitle", "Reset To Defaults"))
+				]
+			]
+		];
+}
+
+FReply FDisplayClusterConfiguratorICVFXMediaCustomization::OnSetupButtonClicked()
+{
+	// We're in camera tiles customization so let's get the camera component
+	UDisplayClusterICVFXCameraComponent* ICVFXCamera = Cast<UDisplayClusterICVFXCameraComponent>(EditingObject.Get());
+	if (!ICVFXCamera)
+	{
+		return FReply::Handled();
+	}
+
+	// Get config data of the DCRA owning the camera being edited
+	const UDisplayClusterConfigurationData* ConfigData = GetConfig();
+	if (!ConfigData)
+	{
+		return FReply::Handled();
+	}
+
+	// Nothing to do if no cluster nodes available
+	if (ConfigData->Cluster->Nodes.Num() < 1)
+	{
+		FMessageDialog::Open(EAppMsgType::Ok, LOCTEXT("MessageNoClusterNodesAvailable", "There are no cluster nodes assigned. Can't setup media."));
+		return FReply::Handled();
+	}
+
+	// Camera's media settings
+	FDisplayClusterConfigurationMediaICVFX& MediaSettings = ICVFXCamera->CameraSettings.RenderSettings.Media;
+
+	// Prepare everything for the configuration dialog
+	FMediaTilesConfigurationDialogParameters Parameters;
+	Parameters.Owner        = ICVFXCamera;
+	Parameters.ConfigData   = ConfigData;
+	Parameters.SplitLayout  = &MediaSettings.TiledSplitLayout;
+	Parameters.InputGroups  = &MediaSettings.TiledMediaInputGroups;
+	Parameters.OutputGroups = &MediaSettings.TiledMediaOutputGroups;
+	Parameters.bAutoPreconfigureOutputMapping = true;
+
+	// Instantiate and show the config dialog
+	const TSharedRef<SMediaTilesConfigurationDialog> TilesConfigurationDialog = SNew(SMediaTilesConfigurationDialog, Parameters);
+	if (GEditor)
+	{
+		GEditor->EditorAddModalWindow(TilesConfigurationDialog);
+	}
+
+	// Process configuration results
+	if (TilesConfigurationDialog->WasConfigurationCompleted())
+	{
+		// Notify tile customizers to re-initialize all media we just generated
+		FDisplayClusterConfiguratorMediaUtils::Get().OnMediaResetToDefaults().Broadcast(EditingObject.Get());
+
+		// Redraw property views
+		FPropertyEditorModule& PropertyEditorModule = FModuleManager::LoadModuleChecked<FPropertyEditorModule>("PropertyEditor");
+		PropertyEditorModule.UpdatePropertyViews({ EditingObject.Get() });
+
+		// Set owning package dirty
+		MarkDirty();
+	}
+
+	return FReply::Handled();
+}
+
+FReply FDisplayClusterConfiguratorICVFXMediaCustomization::OnResetButtonClicked()
 {
-	// Notify tile customizators to re-initialize their media subjects
-	FDisplayClusterConfiguratorMediaUtils::Get().OnMediaAutoConfiguration().Broadcast(EditingObject.Get());
+	if (EditingObject.IsValid())
+	{
+		// Notify tile customizers to re-initialize their media objects
+		FDisplayClusterConfiguratorMediaUtils::Get().OnMediaResetToDefaults().Broadcast(EditingObject.Get());
+
+		// Set owning package dirty
+		MarkDirty();
+	}
 
 	return FReply::Handled();
 }
 
+UDisplayClusterConfigurationData* FDisplayClusterConfiguratorICVFXMediaCustomization::GetConfig() const
+{
+	if (UDisplayClusterICVFXCameraComponent* ICVFXCameraComponent = Cast<UDisplayClusterICVFXCameraComponent>(EditingObject))
+	{
+		// For instances
+		if (ADisplayClusterRootActor* DCRA = Cast<ADisplayClusterRootActor>(ICVFXCameraComponent->GetOwner()))
+		{
+			return DCRA->GetConfigData();
+		}
+		// For DCRA configurator
+		else if (FDisplayClusterConfiguratorBlueprintEditor* BlueprintEditor = FDisplayClusterConfiguratorUtils::GetBlueprintEditorForObject(ICVFXCameraComponent))
+		{
+			return BlueprintEditor->GetConfig();
+		}
+	}
+
+	return nullptr;
+}
+
+void FDisplayClusterConfiguratorICVFXMediaCustomization::MarkDirty()
+{
+	if (EditingObject.IsValid())
+	{
+		// Blueprint
+		if (EditingObject->IsInBlueprint())
+		{
+			ModifyBlueprint();
+		}
+		// Instance
+		else
+		{
+			EditingObject->MarkPackageDirty();
+		}
+	}
+}
+
 #undef LOCTEXT_NAMESPACE
diff --git a/Engine/Plugins/Runtime/nDisplay/Source/DisplayClusterConfigurator/Private/Views/Details/Media/DisplayClusterConfiguratorICVFXMediaCustomization.h b/Engine/Plugins/Runtime/nDisplay/Source/DisplayClusterConfigurator/Private/Views/Details/Media/DisplayClusterConfiguratorICVFXMediaCustomization.h
--- a/Engine/Plugins/Runtime/nDisplay/Source/DisplayClusterConfigurator/Private/Views/Details/Media/DisplayClusterConfiguratorICVFXMediaCustomization.h
+++ b/Engine/Plugins/Runtime/nDisplay/Source/DisplayClusterConfigurator/Private/Views/Details/Media/DisplayClusterConfiguratorICVFXMediaCustomization.h
@@ -9,6 +9,8 @@
 
 #include "IPropertyTypeCustomization.h"
 
+class UDisplayClusterConfigurationData;
+
 
 /**
  * Details panel customization for the FDisplayClusterConfigurationMediaICVFX struct.
@@ -29,11 +31,25 @@ protected:
 
 private:
 
-	/** Builds auto-conf button widget. */
-	void AddAutoConfigurationButton(IDetailChildrenBuilder& ChildBuilder);
+	/** Builds setup button widget. */
+	void AddSetupButton(IDetailChildrenBuilder& ChildBuilder);
+
+	/** Builds reset button widget. */
+	void AddResetButton(IDetailChildrenBuilder& ChildBuilder);
+
+	/** Handles setup button clicks. */
+	FReply OnSetupButtonClicked();
+
+	/** Handles reset button clicks. */
+	FReply OnResetButtonClicked();
+
+private:
+
+	/** Returns configuration of a DCRA owning the camera being edited. */
+	UDisplayClusterConfigurationData* GetConfig() const;
 
-	/** Handles auto-configuration button clicks. */
-	FReply OnAutoConfigureButtonClicked();
+	/** Marks package as dirty */
+	void MarkDirty();
 
 private:
 
diff --git a/Engine/Plugins/Runtime/nDisplay/Source/DisplayClusterConfigurator/Private/Views/Details/Media/DisplayClusterConfiguratorMediaCustomizationBase.cpp b/Engine/Plugins/Runtime/nDisplay/Source/DisplayClusterConfigurator/Private/Views/Details/Media/DisplayClusterConfiguratorMediaCustomizationBase.cpp
--- a/Engine/Plugins/Runtime/nDisplay/Source/DisplayClusterConfigurator/Private/Views/Details/Media/DisplayClusterConfiguratorMediaCustomizationBase.cpp
+++ b/Engine/Plugins/Runtime/nDisplay/Source/DisplayClusterConfigurator/Private/Views/Details/Media/DisplayClusterConfiguratorMediaCustomizationBase.cpp
@@ -20,13 +20,13 @@
 FDisplayClusterConfiguratorMediaFullFrameCustomizationBase::FDisplayClusterConfiguratorMediaFullFrameCustomizationBase()
 {
 	// Subscribe for auto-configure event
-	FDisplayClusterConfiguratorMediaUtils::Get().OnMediaAutoConfiguration().AddRaw(this, &FDisplayClusterConfiguratorMediaFullFrameCustomizationBase::OnAutoConfigureRequested);
+	FDisplayClusterConfiguratorMediaUtils::Get().OnMediaResetToDefaults().AddRaw(this, &FDisplayClusterConfiguratorMediaFullFrameCustomizationBase::OnAutoConfigureRequested);
 }
 
 FDisplayClusterConfiguratorMediaFullFrameCustomizationBase::~FDisplayClusterConfiguratorMediaFullFrameCustomizationBase()
 {
 	// Unsubscribe from auto-configure event
-	FDisplayClusterConfiguratorMediaUtils::Get().OnMediaAutoConfiguration().RemoveAll(this);
+	FDisplayClusterConfiguratorMediaUtils::Get().OnMediaResetToDefaults().RemoveAll(this);
 }
 
 void FDisplayClusterConfiguratorMediaFullFrameCustomizationBase::CustomizeChildren(TSharedRef<IPropertyHandle> InPropertyHandle, IDetailChildrenBuilder& InChildBuilder, IPropertyTypeCustomizationUtils& InCustomizationUtils)
@@ -34,16 +34,16 @@ void FDisplayClusterConfiguratorMediaFullFrameCustomizationBase::CustomizeChildr
 	if (CVarMediaAutoInitializationEnabled.GetValueOnGameThread())
 	{
 		// Subscribe for change callbacks
-		if (MediaSubjectHandle->IsValidHandle())
+		if (MediaObjectHandle->IsValidHandle())
 		{
-			MediaSubjectHandle->SetOnPropertyValueChanged(FSimpleDelegate::CreateSP(this, &FDisplayClusterConfiguratorMediaFullFrameCustomizationBase::OnMediaSubjectChanged));
+			MediaObjectHandle->SetOnPropertyValueChanged(FSimpleDelegate::CreateSP(this, &FDisplayClusterConfiguratorMediaFullFrameCustomizationBase::OnMediaObjectChanged));
 		}
 	}
 
 	FDisplayClusterConfiguratorBaseTypeCustomization::CustomizeChildren(InPropertyHandle, InChildBuilder, InCustomizationUtils);
 }
 
-void FDisplayClusterConfiguratorMediaFullFrameCustomizationBase::ModifyMediaSubjectParameters()
+void FDisplayClusterConfiguratorMediaFullFrameCustomizationBase::ModifyMediaObjectParameters()
 {
 	// Nothing to do if disabled by a CVar
 	if (!CVarMediaAutoInitializationEnabled.GetValueOnGameThread())
@@ -58,19 +58,19 @@ void FDisplayClusterConfiguratorMediaFullFrameCustomizationBase::ModifyMediaSubj
 	}
 
 	UObject* Owner = EditingObject.Get();
-	if (!MediaSubjectHandle->IsValidHandle() || !Owner)
+	if (!MediaObjectHandle->IsValidHandle() || !Owner)
 	{
 		return;
 	}
 
-	// Get media subject
-	UObject* NewMediaSubject = nullptr;
-	if (MediaSubjectHandle->GetValue(NewMediaSubject) != FPropertyAccess::Success)
+	// Get media object
+	UObject* NewMediaObject = nullptr;
+	if (MediaObjectHandle->GetValue(NewMediaObject) != FPropertyAccess::Success)
 	{
 		return;
 	}
 
-	if (!NewMediaSubject)
+	if (!NewMediaObject)
 	{
 		return;
 	}
@@ -79,33 +79,29 @@ void FDisplayClusterConfiguratorMediaFullFrameCustomizationBase::ModifyMediaSubj
 	const TArray<IDisplayClusterModularFeatureMediaInitializer*> MediaInitializers = FDisplayClusterConfiguratorMediaUtils::Get().GetMediaInitializers();
 	for (IDisplayClusterModularFeatureMediaInitializer* Initializer : MediaInitializers)
 	{
-		if (Initializer && Initializer->IsMediaSubjectSupported(NewMediaSubject))
+		if (Initializer && Initializer->IsMediaObjectSupported(NewMediaObject))
 		{
-			if (PerformMediaInitialization(Owner, NewMediaSubject, Initializer))
-			{
-				ModifyBlueprint();
-			}
-
+			PerformMediaInitialization(Owner, NewMediaObject, Initializer);
 			break;
 		}
 	}
 }
 
-bool FDisplayClusterConfiguratorMediaFullFrameCustomizationBase::PerformMediaInitialization(UObject* Owner, UObject* MediaSubject, IDisplayClusterModularFeatureMediaInitializer* Initializer)
+bool FDisplayClusterConfiguratorMediaFullFrameCustomizationBase::PerformMediaInitialization(UObject* Owner, UObject* MediaObject, IDisplayClusterModularFeatureMediaInitializer* Initializer)
 {
-	IDisplayClusterModularFeatureMediaInitializer::FMediaSubjectOwnerInfo OwnerInfo;
+	FMediaObjectOwnerInfo OwnerInfo;
 	const bool bGotOwnerData = GetOwnerData(Owner, OwnerInfo);
 
 	if (bGotOwnerData && Initializer)
 	{
-		Initializer->InitializeMediaSubjectForFullFrame(MediaSubject, OwnerInfo);
+		Initializer->InitializeMediaObjectForFullFrame(MediaObject, OwnerInfo);
 		return true;
 	}
 
 	return false;
 }
 
-bool FDisplayClusterConfiguratorMediaFullFrameCustomizationBase::GetOwnerData(const UObject* Owner, IDisplayClusterModularFeatureMediaInitializer::FMediaSubjectOwnerInfo& OutOwnerInfo) const
+bool FDisplayClusterConfiguratorMediaFullFrameCustomizationBase::GetOwnerData(const UObject* Owner, FMediaObjectOwnerInfo& OutOwnerInfo) const
 {
 	if (const UDisplayClusterICVFXCameraComponent* ICVFXCameraComponent = Cast<UDisplayClusterICVFXCameraComponent>(Owner))
 	{
@@ -127,7 +123,7 @@ bool FDisplayClusterConfiguratorMediaFullFrameCustomizationBase::GetOwnerData(co
 	return false;
 }
 
-bool FDisplayClusterConfiguratorMediaFullFrameCustomizationBase::GetOwnerData(const UDisplayClusterICVFXCameraComponent* ICVFXCameraComponent, IDisplayClusterModularFeatureMediaInitializer::FMediaSubjectOwnerInfo& OutOwnerInfo) const
+bool FDisplayClusterConfiguratorMediaFullFrameCustomizationBase::GetOwnerData(const UDisplayClusterICVFXCameraComponent* ICVFXCameraComponent, FMediaObjectOwnerInfo& OutOwnerInfo) const
 {
 	if (const AActor* const OwningActor = GetOwningActor())
 	{
@@ -138,17 +134,11 @@ bool FDisplayClusterConfiguratorMediaFullFrameCustomizationBase::GetOwnerData(co
 		const int32 CamerasAmount = ICVFXCameras.Num();
 		if (CamerasAmount > 0)
 		{
-			// Camera sort predicate
-			struct FCameraSortPredicate
-			{
-				bool operator()(UDisplayClusterICVFXCameraComponent& LHS, UDisplayClusterICVFXCameraComponent& RHS) const
+			// Sort by name to always keep the same ABC-order
+			ICVFXCameras.Sort([](const UDisplayClusterICVFXCameraComponent& LHS, const UDisplayClusterICVFXCameraComponent& RHS)
 				{
 					return LHS.GetName().Compare(RHS.GetName(), ESearchCase::IgnoreCase) < 0;
-				}
-			};
-
-			// Sort by name to always keep the same ABC-order
-			ICVFXCameras.Sort(FCameraSortPredicate());
+				});
 
 			FString OrigCameraName = ICVFXCameraComponent->GetName();
 			OrigCameraName.RemoveFromEnd(TEXT("_GEN_VARIABLE"));
@@ -158,7 +148,7 @@ bool FDisplayClusterConfiguratorMediaFullFrameCustomizationBase::GetOwnerData(co
 			{
 				if (ICVFXCameras[CameraIdx]->GetName().Equals(OrigCameraName, ESearchCase::IgnoreCase))
 				{
-					OutOwnerInfo.OwnerType = IDisplayClusterModularFeatureMediaInitializer::FMediaSubjectOwnerInfo::EMediaSubjectOwnerType::ICVFXCamera;
+					OutOwnerInfo.OwnerType = FMediaObjectOwnerInfo::EMediaObjectOwnerType::ICVFXCamera;
 					OutOwnerInfo.OwnerName = OrigCameraName;
 					OutOwnerInfo.OwnerUniqueIdx = CameraIdx;
 
@@ -171,7 +161,7 @@ bool FDisplayClusterConfiguratorMediaFullFrameCustomizationBase::GetOwnerData(co
 	return false;
 }
 
-bool FDisplayClusterConfiguratorMediaFullFrameCustomizationBase::GetOwnerData(const UDisplayClusterConfigurationViewport* ViewportCfg, IDisplayClusterModularFeatureMediaInitializer::FMediaSubjectOwnerInfo& OutOwnerInfo) const
+bool FDisplayClusterConfiguratorMediaFullFrameCustomizationBase::GetOwnerData(const UDisplayClusterConfigurationViewport* ViewportCfg, FMediaObjectOwnerInfo& OutOwnerInfo) const
 {
 	// Get cluster node that holds the viewport requested
 	if (const UDisplayClusterConfigurationClusterNode* const ClusterNodeCfg = ViewportCfg->GetTypedOuter<UDisplayClusterConfigurationClusterNode>())
@@ -193,13 +183,13 @@ bool FDisplayClusterConfiguratorMediaFullFrameCustomizationBase::GetOwnerData(co
 					{
 						checkSlow(ViewportIdx < TNumericLimits<uint8>::Max());
 
-						OutOwnerInfo.OwnerType = IDisplayClusterModularFeatureMediaInitializer::FMediaSubjectOwnerInfo::EMediaSubjectOwnerType::Viewport;
+						OutOwnerInfo.OwnerType = FMediaObjectOwnerInfo::EMediaObjectOwnerType::Viewport;
 						OutOwnerInfo.OwnerName = ViewportIt.Key;
 						OutOwnerInfo.OwnerUniqueIdx = static_cast<uint8>(ViewportIdx);
 
 						// Here we leverage cluster node data provider to get unique cluster node index
 						{
-							IDisplayClusterModularFeatureMediaInitializer::FMediaSubjectOwnerInfo TempNodeInfo;
+							FMediaObjectOwnerInfo TempNodeInfo;
 							if (GetOwnerData(ClusterNodeCfg, TempNodeInfo))
 							{
 								OutOwnerInfo.ClusterNodeUniqueIdx = TempNodeInfo.ClusterNodeUniqueIdx;
@@ -216,7 +206,7 @@ bool FDisplayClusterConfiguratorMediaFullFrameCustomizationBase::GetOwnerData(co
 	return false;
 }
 
-bool FDisplayClusterConfiguratorMediaFullFrameCustomizationBase::GetOwnerData(const UDisplayClusterConfigurationClusterNode* NodeCfg, IDisplayClusterModularFeatureMediaInitializer::FMediaSubjectOwnerInfo& OutOwnerInfo) const
+bool FDisplayClusterConfiguratorMediaFullFrameCustomizationBase::GetOwnerData(const UDisplayClusterConfigurationClusterNode* NodeCfg, FMediaObjectOwnerInfo& OutOwnerInfo) const
 {
 	// Get cluster object holding the node requested
 	if (const UDisplayClusterConfigurationCluster* const ClusterCfg = NodeCfg->GetTypedOuter<UDisplayClusterConfigurationCluster>())
@@ -238,7 +228,7 @@ bool FDisplayClusterConfiguratorMediaFullFrameCustomizationBase::GetOwnerData(co
 					{
 						checkSlow(NodeIdx < TNumericLimits<uint8>::Max());
 
-						OutOwnerInfo.OwnerType = IDisplayClusterModularFeatureMediaInitializer::FMediaSubjectOwnerInfo::EMediaSubjectOwnerType::Backbuffer;
+						OutOwnerInfo.OwnerType = FMediaObjectOwnerInfo::EMediaObjectOwnerType::Backbuffer;
 						OutOwnerInfo.OwnerName = NodeIt.Key;
 						OutOwnerInfo.OwnerUniqueIdx = static_cast<uint8>(NodeIdx);
 						OutOwnerInfo.ClusterNodeUniqueIdx = static_cast<uint8>(NodeIdx);
@@ -257,14 +247,12 @@ AActor* FDisplayClusterConfiguratorMediaFullFrameCustomizationBase::GetOwningAct
 {
 	if (UDisplayClusterICVFXCameraComponent* ICVFXCameraComponent = Cast<UDisplayClusterICVFXCameraComponent>(EditingObject))
 	{
-		if (AActor* Actor1 = ICVFXCameraComponent->GetOwner())
-		{
-			return Actor1;
-		}
-		else if (AActor* Actor2 = Cast<AActor>(FindRootActor()))
+		// For instances
+		if (AActor* Actor = ICVFXCameraComponent->GetOwner())
 		{
-			return Actor2;
+			return Actor;
 		}
+		// For DCRA configurator
 		else if (FDisplayClusterConfiguratorBlueprintEditor* BlueprintEditor = FDisplayClusterConfiguratorUtils::GetBlueprintEditorForObject(ICVFXCameraComponent))
 		{
 			return BlueprintEditor->GetPreviewActor();
@@ -274,17 +262,14 @@ AActor* FDisplayClusterConfiguratorMediaFullFrameCustomizationBase::GetOwningAct
 	return nullptr;
 }
 
-void FDisplayClusterConfiguratorMediaFullFrameCustomizationBase::OnMediaSubjectChanged()
+void FDisplayClusterConfiguratorMediaFullFrameCustomizationBase::OnMediaObjectChanged()
 {
-	ModifyMediaSubjectParameters();
+	ModifyMediaObjectParameters();
 }
 
 void FDisplayClusterConfiguratorMediaFullFrameCustomizationBase::OnAutoConfigureRequested(UObject* InEditingObject)
 {
-	if (EditingObject == InEditingObject)
-	{
-		ModifyMediaSubjectParameters();
-	}
+	ModifyMediaObjectParameters();
 }
 
 
@@ -307,9 +292,9 @@ void FDisplayClusterConfiguratorMediaTileCustomizationBase::CustomizeChildren(TS
 	FDisplayClusterConfiguratorMediaFullFrameCustomizationBase::CustomizeChildren(InPropertyHandle, InChildBuilder, InCustomizationUtils);
 }
 
-bool FDisplayClusterConfiguratorMediaTileCustomizationBase::PerformMediaInitialization(UObject* Owner, UObject* MediaSubject, IDisplayClusterModularFeatureMediaInitializer* Initializer)
+bool FDisplayClusterConfiguratorMediaTileCustomizationBase::PerformMediaInitialization(UObject* Owner, UObject* MediaObject, IDisplayClusterModularFeatureMediaInitializer* Initializer)
 {
-	IDisplayClusterModularFeatureMediaInitializer::FMediaSubjectOwnerInfo OwnerInfo;
+	FMediaObjectOwnerInfo OwnerInfo;
 	const bool bGotOwnerData = GetOwnerData(Owner, OwnerInfo);
 
 	if (bGotOwnerData)
@@ -320,7 +305,7 @@ bool FDisplayClusterConfiguratorMediaTileCustomizationBase::PerformMediaInitiali
 		{
 			if (Initializer)
 			{
-				Initializer->InitializeMediaSubjectForTile(MediaSubject, OwnerInfo, TilePos);
+				Initializer->InitializeMediaObjectForTile(MediaObject, OwnerInfo, TilePos);
 				return true;
 			}
 		}
@@ -342,5 +327,5 @@ FIntPoint FDisplayClusterConfiguratorMediaTileCustomizationBase::GetEditedTilePo
 
 void FDisplayClusterConfiguratorMediaTileCustomizationBase::OnTilePositionChanged()
 {
-	ModifyMediaSubjectParameters();
+	ModifyMediaObjectParameters();
 }
diff --git a/Engine/Plugins/Runtime/nDisplay/Source/DisplayClusterConfigurator/Private/Views/Details/Media/DisplayClusterConfiguratorMediaCustomizationBase.h b/Engine/Plugins/Runtime/nDisplay/Source/DisplayClusterConfigurator/Private/Views/Details/Media/DisplayClusterConfiguratorMediaCustomizationBase.h
--- a/Engine/Plugins/Runtime/nDisplay/Source/DisplayClusterConfigurator/Private/Views/Details/Media/DisplayClusterConfiguratorMediaCustomizationBase.h
+++ b/Engine/Plugins/Runtime/nDisplay/Source/DisplayClusterConfigurator/Private/Views/Details/Media/DisplayClusterConfiguratorMediaCustomizationBase.h
@@ -30,33 +30,33 @@ protected:
 
 protected:
 
-	/** Entry point to modify media subjects. */
-	void ModifyMediaSubjectParameters();
+	/** Entry point to modify media objects. */
+	void ModifyMediaObjectParameters();
 
 	/** Virtual media initialization depending on media type (full-frame, split, etc.) */
-	virtual bool PerformMediaInitialization(UObject* Owner, UObject* MediaSubject, IDisplayClusterModularFeatureMediaInitializer* Initializer);
+	virtual bool PerformMediaInitialization(UObject* Owner, UObject* MediaObject, IDisplayClusterModularFeatureMediaInitializer* Initializer);
 
 	/** Returns the actor that owns the object being edited. */
 	AActor* GetOwningActor() const;
 
-	/** Returns the name of media subject's owner. */
-	bool GetOwnerData(const UObject* Owner, IDisplayClusterModularFeatureMediaInitializer::FMediaSubjectOwnerInfo& OutOwnerInfo) const;
+	/** Returns the name of media object's owner. */
+	bool GetOwnerData(const UObject* Owner, FMediaObjectOwnerInfo& OutOwnerInfo) const;
 
 private:
 
 	/** Helper data provider for ICVFX cameras. */
-	bool GetOwnerData(const UDisplayClusterICVFXCameraComponent* ICVFXCameraComponent, IDisplayClusterModularFeatureMediaInitializer::FMediaSubjectOwnerInfo& OutOwnerInfo) const;
+	bool GetOwnerData(const UDisplayClusterICVFXCameraComponent* ICVFXCameraComponent, FMediaObjectOwnerInfo& OutOwnerInfo) const;
 
 	/** Helper data provider for viewports. */
-	bool GetOwnerData(const UDisplayClusterConfigurationViewport* ViewportCfg, IDisplayClusterModularFeatureMediaInitializer::FMediaSubjectOwnerInfo& OutOwnerInfo) const;
+	bool GetOwnerData(const UDisplayClusterConfigurationViewport* ViewportCfg, FMediaObjectOwnerInfo& OutOwnerInfo) const;
 
 	/** Helper data provider for cluster nodes. */
-	bool GetOwnerData(const UDisplayClusterConfigurationClusterNode* NodeCfg, IDisplayClusterModularFeatureMediaInitializer::FMediaSubjectOwnerInfo& OutOwnerInfo) const;
+	bool GetOwnerData(const UDisplayClusterConfigurationClusterNode* NodeCfg, FMediaObjectOwnerInfo& OutOwnerInfo) const;
 
 private:
 
 	/** Handles media source/output change callbacks */
-	void OnMediaSubjectChanged();
+	void OnMediaObjectChanged();
 
 	/** Auto-configuration requests handler */
 	void OnAutoConfigureRequested(UObject* EditingObject);
@@ -64,7 +64,7 @@ private:
 protected:
 
 	/** MediaSource or MediaOutput property handle, depending on the child implementation. */
-	TSharedPtr<IPropertyHandle> MediaSubjectHandle;
+	TSharedPtr<IPropertyHandle> MediaObjectHandle;
 };
 
 
@@ -86,7 +86,7 @@ protected:
 	//~ Begin FDisplayClusterConfiguratorMediaFullFrameCustomizationBase
 
 	/** Virtual media initialization depending on media type (full-frame, split, etc.) */
-	virtual bool PerformMediaInitialization(UObject* Owner, UObject* MediaSubject, IDisplayClusterModularFeatureMediaInitializer* Initializer) override;
+	virtual bool PerformMediaInitialization(UObject* Owner, UObject* MediaObject, IDisplayClusterModularFeatureMediaInitializer* Initializer) override;
 
 	//~ End FDisplayClusterConfiguratorMediaFullFrameCustomizationBase
 
diff --git a/Engine/Plugins/Runtime/nDisplay/Source/DisplayClusterConfigurator/Private/Views/Details/Media/DisplayClusterConfiguratorMediaFullFrameCustomization.cpp b/Engine/Plugins/Runtime/nDisplay/Source/DisplayClusterConfigurator/Private/Views/Details/Media/DisplayClusterConfiguratorMediaFullFrameCustomization.cpp
--- a/Engine/Plugins/Runtime/nDisplay/Source/DisplayClusterConfigurator/Private/Views/Details/Media/DisplayClusterConfiguratorMediaFullFrameCustomization.cpp
+++ b/Engine/Plugins/Runtime/nDisplay/Source/DisplayClusterConfigurator/Private/Views/Details/Media/DisplayClusterConfiguratorMediaFullFrameCustomization.cpp
@@ -15,8 +15,8 @@
 void FDisplayClusterConfiguratorMediaFullFrameInputCustomization::CustomizeChildren(TSharedRef<IPropertyHandle> InPropertyHandle, IDetailChildrenBuilder& InChildBuilder, IPropertyTypeCustomizationUtils& InCustomizationUtils)
 {
 	// MediaSource property
-	MediaSubjectHandle = GET_CHILD_HANDLE(FDisplayClusterConfigurationMediaInput, MediaSource);
-	check(MediaSubjectHandle->IsValidHandle());
+	MediaObjectHandle = GET_CHILD_HANDLE(FDisplayClusterConfigurationMediaInput, MediaSource);
+	check(MediaObjectHandle->IsValidHandle());
 
 	FDisplayClusterConfiguratorMediaFullFrameCustomizationBase::CustomizeChildren(InPropertyHandle, InChildBuilder, InCustomizationUtils);
 }
@@ -27,8 +27,8 @@ void FDisplayClusterConfiguratorMediaFullFrameInputCustomization::CustomizeChild
 void FDisplayClusterConfiguratorMediaFullFrameOutputCustomization::CustomizeChildren(TSharedRef<IPropertyHandle> InPropertyHandle, IDetailChildrenBuilder& InChildBuilder, IPropertyTypeCustomizationUtils& InCustomizationUtils)
 {
 	// MediaOutput property
-	MediaSubjectHandle = GET_CHILD_HANDLE(FDisplayClusterConfigurationMediaOutput, MediaOutput);
-	check(MediaSubjectHandle->IsValidHandle());
+	MediaObjectHandle = GET_CHILD_HANDLE(FDisplayClusterConfigurationMediaOutput, MediaOutput);
+	check(MediaObjectHandle->IsValidHandle());
 
 	FDisplayClusterConfiguratorMediaFullFrameCustomizationBase::CustomizeChildren(InPropertyHandle, InChildBuilder, InCustomizationUtils);
 }
diff --git a/Engine/Plugins/Runtime/nDisplay/Source/DisplayClusterConfigurator/Private/Views/Details/Media/DisplayClusterConfiguratorMediaTileCustomization.cpp b/Engine/Plugins/Runtime/nDisplay/Source/DisplayClusterConfigurator/Private/Views/Details/Media/DisplayClusterConfiguratorMediaTileCustomization.cpp
--- a/Engine/Plugins/Runtime/nDisplay/Source/DisplayClusterConfigurator/Private/Views/Details/Media/DisplayClusterConfiguratorMediaTileCustomization.cpp
+++ b/Engine/Plugins/Runtime/nDisplay/Source/DisplayClusterConfigurator/Private/Views/Details/Media/DisplayClusterConfiguratorMediaTileCustomization.cpp
@@ -15,8 +15,8 @@
 void FDisplayClusterConfiguratorMediaInputTileCustomization::CustomizeChildren(TSharedRef<IPropertyHandle> InPropertyHandle, IDetailChildrenBuilder& InChildBuilder, IPropertyTypeCustomizationUtils& InCustomizationUtils)
 {
 	// MediaSource property
-	MediaSubjectHandle = GET_CHILD_HANDLE(FDisplayClusterConfigurationMediaUniformTileInput, MediaSource);
-	check(MediaSubjectHandle->IsValidHandle());
+	MediaObjectHandle = GET_CHILD_HANDLE(FDisplayClusterConfigurationMediaUniformTileInput, MediaSource);
+	check(MediaObjectHandle->IsValidHandle());
 
 	// Tile position property
 	TilePosHandle = GET_CHILD_HANDLE(FDisplayClusterConfigurationMediaUniformTileInput, Position);
@@ -32,8 +32,8 @@ void FDisplayClusterConfiguratorMediaInputTileCustomization::CustomizeChildren(T
 void FDisplayClusterConfiguratorMediaOutputTileCustomization::CustomizeChildren(TSharedRef<IPropertyHandle> InPropertyHandle, IDetailChildrenBuilder& InChildBuilder, IPropertyTypeCustomizationUtils& InCustomizationUtils)
 {
 	// MediaOutput property
-	MediaSubjectHandle = GET_CHILD_HANDLE(FDisplayClusterConfigurationMediaUniformTileOutput, MediaOutput);
-	check(MediaSubjectHandle->IsValidHandle());
+	MediaObjectHandle = GET_CHILD_HANDLE(FDisplayClusterConfigurationMediaUniformTileOutput, MediaOutput);
+	check(MediaObjectHandle->IsValidHandle());
 
 	// Tile position property
 	TilePosHandle = GET_CHILD_HANDLE(FDisplayClusterConfigurationMediaUniformTileOutput, Position);
diff --git a/Engine/Plugins/Runtime/nDisplay/Source/DisplayClusterConfigurator/Private/Views/Details/Media/DisplayClusterConfiguratorMediaUtils.h b/Engine/Plugins/Runtime/nDisplay/Source/DisplayClusterConfigurator/Private/Views/Details/Media/DisplayClusterConfiguratorMediaUtils.h
--- a/Engine/Plugins/Runtime/nDisplay/Source/DisplayClusterConfigurator/Private/Views/Details/Media/DisplayClusterConfiguratorMediaUtils.h
+++ b/Engine/Plugins/Runtime/nDisplay/Source/DisplayClusterConfigurator/Private/Views/Details/Media/DisplayClusterConfiguratorMediaUtils.h
@@ -28,10 +28,10 @@ public:
 	}
 
 	/** Tiled media auto-configuration event */
-	DECLARE_EVENT_OneParam(FDisplayClusterConfiguratorMediaUtils, FDisplayClusterConfiguratorMediaAutoConfigurationRequestedEvent, UObject*);
-	FDisplayClusterConfiguratorMediaAutoConfigurationRequestedEvent& OnMediaAutoConfiguration()
+	DECLARE_EVENT_OneParam(FDisplayClusterConfiguratorMediaUtils, FResetToDefaultsEvent, UObject*);
+	FResetToDefaultsEvent& OnMediaResetToDefaults()
 	{
-		return MediaAutoConfigurationEvent;
+		return MediaResetToDefaultsEvent;
 	}
 
 private:
@@ -39,6 +39,6 @@ private:
 	/** Media initializers (modular features) available. */
 	TArray<IDisplayClusterModularFeatureMediaInitializer*> MediaInitializers;
 
-	/** Tiled media auto-configuration event */
-	FDisplayClusterConfiguratorMediaAutoConfigurationRequestedEvent MediaAutoConfigurationEvent;
+	/** Tiled media reset to defaults event */
+	FResetToDefaultsEvent MediaResetToDefaultsEvent;
 };
diff --git a/Engine/Plugins/Runtime/nDisplay/Source/DisplayClusterConfigurator/Private/Views/Details/Media/SMediaTilesConfigurationDialog.cpp b/Engine/Plugins/Runtime/nDisplay/Source/DisplayClusterConfigurator/Private/Views/Details/Media/SMediaTilesConfigurationDialog.cpp
new file mode 100644
--- /dev/null
+++ b/Engine/Plugins/Runtime/nDisplay/Source/DisplayClusterConfigurator/Private/Views/Details/Media/SMediaTilesConfigurationDialog.cpp
@@ -0,0 +1,1621 @@
+// Copyright Epic Games, Inc. All Rights Reserved.
+
+#include "Views/Details/Media/SMediaTilesConfigurationDialog.h"
+#include "Views/Details/Media/DisplayClusterConfiguratorMediaUtils.h"
+
+#include "DisplayClusterConfiguratorLog.h"
+#include "DisplayClusterConfiguratorStyle.h"
+
+#include "IDisplayClusterModularFeatureMediaInitializer.h"
+
+#include "ClassViewerFilter.h"
+#include "ClassViewerModule.h"
+#include "PropertyCustomizationHelpers.h"
+#include "SClassViewer.h"
+
+#include "MediaOutput.h"
+#include "MediaSource.h"
+
+#include "Framework/Application/MenuStack.h"
+#include "Framework/Application/SlateApplication.h"
+#include "Framework/MultiBox/MultiBoxBuilder.h"
+#include "Misc/MessageDialog.h"
+#include "Modules/ModuleManager.h"
+#include "Styling/AppStyle.h"
+#include "Styling/StyleColors.h"
+#include "UObject/UObjectGlobals.h"
+
+#include "Widgets/Input/SButton.h"
+#include "Widgets/Input/SCheckBox.h"
+#include "Widgets/Input/SComboButton.h"
+#include "Widgets/Layout/SBox.h"
+#include "Widgets/Layout/SScrollBox.h"
+#include "Widgets/Layout/SSpacer.h"
+#include "Widgets/Layout/SUniformGridPanel.h"
+#include "Widgets/Text/STextBlock.h"
+#include "Widgets/Workflow/SWizard.h"
+
+#define LOCTEXT_NAMESPACE "SMediaTilesConfigurationDialog"
+
+
+namespace UE::DisplayClusterConfigurator::MediaTilesConfigurationDialog::Private
+{
+	const FText TextDialogTitle = LOCTEXT("DialogTitle", "Media Tiles Configuration");
+
+	// Page: Layout
+	const FText TextPageLayoutHeader = LOCTEXT("PageTilesLayoutHeader", "Step 1: Choose split layout");
+
+	// Page: Media
+	const FText TextPageMediaHeader = LOCTEXT("PageMediaObjectsHeader", "Step 2: Configure template media source & output");
+	const FText TextPageMediaSource = LOCTEXT("PageMediaSourceRowName", "Media Source:");
+	const FText TextPageMediaOutput = LOCTEXT("PageMediaOutputRowName", "Media Output:");
+	const FText TextPageMediaComboboxNone  = LOCTEXT("PageMediaMediaComboboxItemNone", "None");
+	const FText TextPageMediaStatusOk      = LOCTEXT("PageMediaStatusOk", "Ok");
+	const FText TextPageMediaStatusNotSupp = LOCTEXT("PageMediaStatusNotSupported", "Auto-configuration is not supported");
+	const FText TextPageMediaNotCompatOrNotSupp = LOCTEXT("PageMediaStatusNotSupported", "Not compatible or not supported");
+
+	// Page: Nodes
+	const FText TextPageNodesHeader = LOCTEXT("PageNodesHeader", "Step 3: Allot cluster nodes");
+	const FText TextPageNodesSectionOutputHeader = LOCTEXT("PageNodesSectionOutputHeader", "Tile Senders");
+	const FText TextPageNodesSectionInputHeader  = LOCTEXT("PageNodesSectionInputHeader", "Tile Receivers");
+	const FText TextPageNodesMenuSelectAll       = LOCTEXT("PageNodesMenuSelectAll", "Select All");
+	const FText TextPageNodesMenuDeselectAll     = LOCTEXT("PageNodesMenuDeselectAll", "Deselect All");
+	const FText TextPageNodesMenuSelectAllNoSenders   = LOCTEXT("PageNodesMenuSelectAllNoSenders", "Select All w/o Senders");
+	const FText TextPageNodesMenuSelectAllNoReceivers = LOCTEXT("PageNodesMenuSelectAllNoReceivers", "Select All w/o Receivers");
+	const FText TextPageNodesMenuDeselectSenders      = LOCTEXT("PageNodesMenuSelectDeselectSenders", "Deselect Senders");
+	const FText TextPageNodesMenuDeselectReceivers    = LOCTEXT("PageNodesMenuSelectDeselectReceivers", "Deselect Receivers");
+	const FText TextPageNodesMenuSelectAllOffscreen   = LOCTEXT("PageNodesMenuSelectAllOffscreen", "Select All Offscreen");
+
+	// Page: Finalization
+	const FText TextPageFinalizationHeader = LOCTEXT("PageFinalizationHeader", "Step 4: Output mapping (tile senders)");
+
+	// Feedback message
+	const FText TextMsgConfigurationFeedbackBegin = LOCTEXT("MsgConfigurationFeedbackBegin", "Some issues were found during the configuration process:");
+	const FText TextMsgConfigurationFeedbackEnd   = LOCTEXT("MsgConfigurationFeedbackEnd",   "Ignore and proceed (Yes) or return and fix media configuration (No)?");
+}
+
+void SMediaTilesConfigurationDialog::Construct(const FArguments& InArgs, const FMediaTilesConfigurationDialogParameters& InParameters)
+{
+	InParameters.Validate();
+
+	// Save input parameters first
+	Parameters = InParameters;
+
+	// Initialize internals
+	InitializeInternals();
+
+	// And construct the dialog widget
+	SWindow::Construct(SWindow::FArguments()
+		.Title(UE::DisplayClusterConfigurator::MediaTilesConfigurationDialog::Private::TextDialogTitle)
+		.SizingRule(ESizingRule::Autosized)
+		.SupportsMaximize(false)
+		.SupportsMinimize(false)
+		[
+			SNew(SVerticalBox)
+
+				// Current page hint
+				+SVerticalBox::Slot()
+				.Padding(10, 10, 10, 5)
+				.AutoHeight()
+				[
+					SAssignNew(PageHint, STextBlock)
+					.Font(FAppStyle::Get().GetFontStyle("NormalFontBold"))
+				]
+
+				// Page body
+				+SVerticalBox::Slot()
+				.Padding(10, 5, 10, 10)
+				.AutoHeight()
+				[
+					SAssignNew(Wizard, SWizard)
+					.OnCanceled(this, &SMediaTilesConfigurationDialog::OnCancelButtonClicked)
+					.OnFinished(this, &SMediaTilesConfigurationDialog::OnFinishButtonClicked)
+					.CanFinish(this, &SMediaTilesConfigurationDialog::IsFinishButtonEnabled)
+					.ShowPageList(false)
+
+						// Page: Layout
+						+SWizard::Page()
+						.CanShow(true)
+						.OnEnter(this, &SMediaTilesConfigurationDialog::PageLayout_OnEnter)
+						[
+							PageLayout_Build()
+						]
+
+						// Page: Media
+						+SWizard::Page()
+						.CanShow(this, &SMediaTilesConfigurationDialog::PageMedia_OnCanShow)
+						.OnEnter(this, &SMediaTilesConfigurationDialog::PageMedia_OnEnter)
+						[
+							PageMedia_Build()
+						]
+
+						// Page: Nodes
+						+SWizard::Page()
+						.CanShow(this, &SMediaTilesConfigurationDialog::PageNodes_OnCanShow)
+						.OnEnter(this, &SMediaTilesConfigurationDialog::PageNodes_OnEnter)
+						[
+							PageNodes_Build()
+						]
+
+						// Page: Finalization
+						+SWizard::Page()
+						.CanShow(this, &SMediaTilesConfigurationDialog::PageFinalization_OnCanShow)
+						.OnEnter(this, &SMediaTilesConfigurationDialog::PageFinalization_OnEnter)
+						[
+							PageFinalization_Build()
+						]
+				]
+		]);
+}
+
+bool SMediaTilesConfigurationDialog::WasConfigurationCompleted() const
+{
+	return bConfigurationCompleted;
+}
+
+void SMediaTilesConfigurationDialog::InitializeInternals()
+{
+	// Pre-save some data to simplify future use
+	for (const TPair<FString, TObjectPtr<UDisplayClusterConfigurationClusterNode>>& Node : Parameters.ConfigData->Cluster->Nodes)
+	{
+		// All nodes
+		ClusterNodeIds.Add(Node.Key);
+		// Node to host
+		NodeToHostMap.Emplace(Node.Key, Node.Value->Host);
+		// Nodes per host
+		HostToNodesMap.FindOrAdd(Node.Value->Host).Add(Node.Key);
+		// Nodes allowed to be used for tile rendering
+		NodesAllowedForOutput.Add(Node.Key);
+		// Nodes allowed to be used for tile receiving and compositing
+		NodesAllowedForInput.Add(Node.Key);
+	}
+
+	// Fill output mapping with defaults
+	for (int32 TileX = 0; TileX < MaxTilesAmount; ++TileX)
+	{
+		for (int32 TileY = 0; TileY < MaxTilesAmount; ++TileY)
+		{
+			FOutputMappingData& TileData = OutputMapping.Emplace({TileX, TileY});
+			TileData.ClusterNodes.Reserve(NodesAllowedForOutput.Num());
+		}
+	}
+}
+
+bool SMediaTilesConfigurationDialog::IsFinishButtonEnabled() const
+{
+	// Check if all settings are good
+	const bool bConfigValid = PageLayout_IsConfigurationValid() && PageMedia_IsConfigurationValid() && PageNodes_IsConfigurationValid() && PageFinalization_IsConfigurationValid();
+	// Make sure users have seen the final configuration on the last page
+	const bool bOnLastPage = Wizard->GetCurrentPageIndex() == Wizard->GetNumPages() - 1;
+
+	return bConfigValid && bOnLastPage;
+}
+
+void SMediaTilesConfigurationDialog::OnFinishButtonClicked()
+{
+	TArray<FString> Errors;
+	TArray<FString> Warnings;
+
+	bConfigurationCompleted = ApplyConfiguration(Errors, Warnings);
+
+	// Any warnings/errors/notes?
+	if (bConfigurationCompleted && (!Errors.IsEmpty() || !Warnings.IsEmpty()))
+	{
+		FText Message;
+		GenerateIssuesFoundMessage(Errors, Warnings, Message);
+
+		// Notify user
+		const EAppReturnType::Type Result = FMessageDialog::Open(EAppMsgType::YesNo, Message);
+
+		// User decided to return to configuration. So we don't close this configuration dialog.
+		if (Result == EAppReturnType::No)
+		{
+			return;
+		}
+	}
+
+	// Close this window
+	RequestDestroyWindow();
+}
+
+void SMediaTilesConfigurationDialog::OnCancelButtonClicked()
+{
+	bConfigurationCompleted = false;
+	RequestDestroyWindow();
+}
+
+bool SMediaTilesConfigurationDialog::ApplyConfiguration(TArray<FString>& OutErrors, TArray<FString>& OutWarnings)
+{
+	if (!Parameters.Validate())
+	{
+		return false;
+	}
+
+	// Redirect to a proper handler
+	if (EnumHasAnyFlags(MediaPropagationTypes, EMediaStreamPropagationType::Multicast))
+	{
+		return ApplyConfiguration_Multicast(OutErrors, OutWarnings);
+	}
+	else if (EnumHasAnyFlags(MediaPropagationTypes, EMediaStreamPropagationType::LocalMulticast))
+	{
+		return ApplyConfiguration_LocalMulticast(OutErrors, OutWarnings);
+	}
+	else
+	{
+		return false;
+	}
+}
+
+bool SMediaTilesConfigurationDialog::ApplyConfiguration_Multicast(TArray<FString>& OutErrors, TArray<FString>& OutWarnings)
+{
+	// In multicast, the data is propagated in OneSender-to-MultipleReceivers way. Thus we
+	// can have a single input group with all the receivers because they get same data from
+	// the same senders. Each receiver gets the full set of of tiles, so it's connected to all
+	// the receivers. This allows for all of them to share a single MediaSource for each tile.
+	// There is no limitation on the receivers amount.
+	// 
+	// As for the senders, we don't actually need more than total amount of tiles. Each tile
+	// is produced by a dedicated sender and propagated to all the receivers. Based on
+	// this we're going to have a dedicated output group for every tile/sender.
+	//
+	// INPUT:
+	// [InputGroup]
+	//   Nodes: Receiver0, Receiver1, ..., ReceiverM
+	//   Tiles: 0x0, 0x1, ..., AxB
+	//
+	// OUTPUT:
+	// [OutputGroup0]
+	//   Nodes: Sender0
+	//   Tiles: 0x0
+	// [OutputGroup1]
+	//   Nodes: Sender1
+	//   Tiles: 0x1
+	// ...
+	// [OutputGroupN]
+	//   Nodes: SenderN
+	//   Tiles: AxB
+	//
+	// Where
+	//  - A - tiles amount horizontally
+	//  - B - tiles amount vertically
+	//  - N = A * B == amount of senders
+	//  - M - any amount of receivers
+
+	// Apply tile layout
+	*Parameters.SplitLayout = Accepted + FIntPoint{ 1, 1 };
+
+	// Object flags for new media source/output objects
+	const EObjectFlags MediaObjectFlags = Parameters.Owner->IsInBlueprint() ?
+		RF_Public | RF_Transactional | RF_ArchetypeObject :
+		RF_Public | RF_Transactional;
+
+	//
+	// INPUT setup
+	//
+
+	// Add single input group
+	Parameters.InputGroups->Reset(1);
+	FDisplayClusterConfigurationMediaTiledInputGroup& NewInputGroup = Parameters.InputGroups->AddDefaulted_GetRef();
+
+	// Put all receivers into the group
+	if (!InputSelection.IsEmpty())
+	{
+		NewInputGroup.ClusterNodes.ItemNames = InputSelection.Array();
+	}
+	else
+	{
+		OutErrors.Add(TEXT("No receivers found"));
+	}
+
+	// Setup input tiles (put all the tiles and receiving nodes into a single group)
+	for (int32 TileX = 0; TileX < Parameters.SplitLayout->X; ++TileX)
+	{
+		for (int32 TileY = 0; TileY < Parameters.SplitLayout->Y; ++TileY)
+		{
+			FDisplayClusterConfigurationMediaUniformTileInput& NewTile = NewInputGroup.Tiles.AddDefaulted_GetRef();
+
+			NewTile.Position = { TileX, TileY };
+			NewTile.MediaSource = NewObject<UMediaSource>(Parameters.Owner, MediaSource->GetClass(), NAME_None, MediaObjectFlags);
+		}
+	}
+
+	//
+	// OUTPUT setup
+	//
+
+	const int32 TilesAmount = Parameters.SplitLayout->X * Parameters.SplitLayout->Y;
+
+	// First, remove any existing output mapping
+	Parameters.OutputGroups->Reset(TilesAmount);
+
+	// For each tile, we create an output group with a corresponding sender node
+	for (int32 TileX = 0; TileX < Parameters.SplitLayout->X; ++TileX)
+	{
+		for (int32 TileY = 0; TileY < Parameters.SplitLayout->Y; ++TileY)
+		{
+			// Create new output group
+			FDisplayClusterConfigurationMediaTiledOutputGroup& NewOutputGroup = Parameters.OutputGroups->AddDefaulted_GetRef();
+
+			// Add cluster nodes bound to this tile
+			const FIntPoint Tile{ TileX, TileY };
+			const int32 SenderNodesAmount = OutputMapping[Tile].ClusterNodes.Num();
+
+			if (SenderNodesAmount >= 1)
+			{
+				const FString& SenderNodeId = *OutputMapping[Tile].ClusterNodes.CreateConstIterator();
+				NewOutputGroup.ClusterNodes.ItemNames.Add(SenderNodeId);
+
+				// Normally we should have a single sender
+				if (SenderNodesAmount > 1)
+				{
+					OutWarnings.Add(FString::Printf(TEXT("Too many (%d) senders found for tile '%dx%d'. A single one is set up: %s"),
+						SenderNodesAmount, TileX, TileY, *SenderNodeId));
+				}
+			}
+			else
+			{
+				OutErrors.Add(FString::Printf(TEXT("No senders found for tile '%dx%d'"), TileX, TileY));
+			}
+
+			// Create new tile in this group
+			FDisplayClusterConfigurationMediaUniformTileOutput& NewTileInGroup = NewOutputGroup.Tiles.AddDefaulted_GetRef();
+			NewTileInGroup.Position = { TileX, TileY };
+			NewTileInGroup.MediaOutput = NewObject<UMediaOutput>(Parameters.Owner, MediaOutput->GetClass(), NAME_None, MediaObjectFlags);
+		}
+	}
+
+	return true;
+}
+
+bool SMediaTilesConfigurationDialog::ApplyConfiguration_LocalMulticast(TArray<FString>& OutErrors, TArray<FString>& OutWarnings)
+{
+	// Local multicast implies every host produces and consumes its own set of tiles. This means we can
+	// group inputs and outputs per every host like this:
+	//
+	// [Host A]
+	//   - SenderA1,   SenderA2,   ..., SenderAN    <== separate output groups for Host A
+	//   - ReceiverA1, ReceiverA2, ..., ReceiverAN  <== single input group for Host A
+	// [Host B]
+	//   - SenderB1,   SenderB2,   ..., SenderBN    <== separate output groups for Host A
+	//   - ReceiverB1, ReceiverB2, ..., ReceiverBN  <== single input group for Host A
+	//
+	// In other words. Assuming there are N tiles, each host that has at least one tile
+	// receiver must also have N tile senders (or less, but some senders would have to render
+	// multiple tiles in this case).
+
+	// Apply tile layout
+	*Parameters.SplitLayout = Accepted + FIntPoint{ 1, 1 };
+
+	// All receivers involved
+	const TSet<FString> AllReceivers = InputSelection;
+
+	// All senders involved
+	TSet<FString> AllSenders;
+	for (const TPair<FIntPoint, FOutputMappingData>& TileData : OutputMapping)
+	{
+		AllSenders.Append(TileData.Value.ClusterNodes);
+	}
+
+	// All nodes (senders and receivers) being used in tiled media
+	TSet<FString> HostsWithMedia;
+	{
+		const TSet<FString> AllNodes = AllSenders.Union(AllReceivers);
+		for (const FString& NodeId : AllNodes)
+		{
+			if (const FString* Host = NodeToHostMap.Find(NodeId))
+			{
+				HostsWithMedia.Add(*Host);
+			}
+		}
+	}
+
+	// Object flags for new media source/output objects
+	const EObjectFlags MediaObjectFlags = Parameters.Owner->IsInBlueprint() ?
+		RF_Public | RF_Transactional | RF_ArchetypeObject :
+		RF_Public | RF_Transactional;
+
+	const int32 TilesAmount = Parameters.SplitLayout->X * Parameters.SplitLayout->Y;
+
+	// Reset any existing data
+	Parameters.InputGroups->Reset(HostsWithMedia.Num());
+	Parameters.OutputGroups->Reset(HostsWithMedia.Num() * TilesAmount);
+
+	// Now, we can generate an input group and corresponding amount (amount of tiles) of output
+	// groups for every host.
+	for (const FString& Host : HostsWithMedia)
+	{
+		//
+		// INPUT setup
+		//
+
+		// Add input group
+		FDisplayClusterConfigurationMediaTiledInputGroup& NewInputGroup = Parameters.InputGroups->AddDefaulted_GetRef();
+
+		// Assing all the receivers on this host to the group
+		const TSet<FString> ReceiversInGroup = AllReceivers.Intersect(HostToNodesMap[Host]);
+		NewInputGroup.ClusterNodes.ItemNames = ReceiversInGroup.Array();
+
+		// We expect at least one to be set
+		if (NewInputGroup.ClusterNodes.ItemNames.IsEmpty())
+		{
+			OutErrors.Add(FString::Printf(TEXT("No receivers found on host '%s'"), *Host));
+		}
+
+		// Setup input tiles (setup all the tiles in the group)
+		for (int32 TileX = 0; TileX < Parameters.SplitLayout->X; ++TileX)
+		{
+			for (int32 TileY = 0; TileY < Parameters.SplitLayout->Y; ++TileY)
+			{
+				FDisplayClusterConfigurationMediaUniformTileInput& NewTile = NewInputGroup.Tiles.AddDefaulted_GetRef();
+
+				NewTile.Position = { TileX, TileY };
+				NewTile.MediaSource = NewObject<UMediaSource>(Parameters.Owner, MediaSource->GetClass(), NAME_None, MediaObjectFlags);
+			}
+		}
+
+		//
+		// OUTPUT setup
+		//
+
+		// For each tile, we create an output group with all the nodes bound to this tile
+		for (int32 TileX = 0; TileX < Parameters.SplitLayout->X; ++TileX)
+		{
+			for (int32 TileY = 0; TileY < Parameters.SplitLayout->Y; ++TileY)
+			{
+				// Create new output group
+				FDisplayClusterConfigurationMediaTiledOutputGroup& NewOutputGroup = Parameters.OutputGroups->AddDefaulted_GetRef();
+
+				// Get all senders of this tile on this particualar host
+				const TSet<FString>& AllSendersMappedToThisTile = OutputMapping[{TileX, TileY}].ClusterNodes;
+				const TSet<FString>  AllSendersOnThisHost = AllSenders.Intersect(HostToNodesMap[Host]);
+				const TSet<FString>  SendersOfThisTile = AllSendersOnThisHost.Intersect(AllSendersMappedToThisTile);
+				
+				// Find a sender for this tile
+				const int32 AmountOfSendersAssignedForThisTile = SendersOfThisTile.Num();
+				if (AmountOfSendersAssignedForThisTile >= 1)
+				{
+					const FString& SenderNodeId = *SendersOfThisTile.CreateConstIterator();
+					NewOutputGroup.ClusterNodes.ItemNames.Add(SenderNodeId);
+
+					// Normally we should have a single sender
+					if (AmountOfSendersAssignedForThisTile > 1)
+					{
+						OutWarnings.Add(FString::Printf(TEXT("Too many (%d) senders found for tile '%dx%d' on host '%s'"),
+							AmountOfSendersAssignedForThisTile, TileX, TileY, *Host));
+					}
+				}
+				else
+				{
+					OutErrors.Add(FString::Printf(TEXT("No senders found for tile '%dx%d' on host '%s'"),
+						TileX, TileY, *Host));
+				}
+
+				// Create new tile in this group
+				FDisplayClusterConfigurationMediaUniformTileOutput& NewTileInGroup = NewOutputGroup.Tiles.AddDefaulted_GetRef();
+				NewTileInGroup.Position = { TileX, TileY };
+				NewTileInGroup.MediaOutput = NewObject<UMediaOutput>(Parameters.Owner, MediaOutput->GetClass(), NAME_None, MediaObjectFlags);
+			}
+		}
+	}
+
+	return true;
+}
+
+void SMediaTilesConfigurationDialog::GenerateIssuesFoundMessage(const TArray<FString>& Errors, const TArray<FString>& Warnings, FText& OutMessage) const
+{
+	FString Temp;
+	Temp.Reserve(512);
+
+	Temp += UE::DisplayClusterConfigurator::MediaTilesConfigurationDialog::Private::TextMsgConfigurationFeedbackBegin.ToString();
+	Temp += TEXT("\n\n");
+
+	if (!Errors.IsEmpty())
+	{
+		Temp += TEXT("Errors:\n");
+		for (const FString& Message : Errors)
+		{
+			Temp += FString::Printf(TEXT("- %s\n"), *Message);
+		}
+	}
+
+	if (!Warnings.IsEmpty())
+	{
+		Temp += TEXT("Warnings:\n");
+		for (const FString& Message : Warnings)
+		{
+			Temp += FString::Printf(TEXT("- %s\n"), *Message);
+		}
+	}
+
+	Temp += TEXT("\n\n");
+	Temp += UE::DisplayClusterConfigurator::MediaTilesConfigurationDialog::Private::TextMsgConfigurationFeedbackEnd.ToString();
+
+	OutMessage = FText::FromString(Temp);
+}
+
+bool SMediaTilesConfigurationDialog::IsMediaObjectSupported(const UObject* MediaObject) const
+{
+	const TArray<IDisplayClusterModularFeatureMediaInitializer*> MediaInitializers = FDisplayClusterConfiguratorMediaUtils::Get().GetMediaInitializers();
+	const bool bSupported = MediaInitializers.ContainsByPredicate([MediaObject](IDisplayClusterModularFeatureMediaInitializer* Initializer)
+		{
+			return Initializer ? Initializer->IsMediaObjectSupported(MediaObject) : false;
+		});
+
+	return bSupported;
+}
+
+bool SMediaTilesConfigurationDialog::AreMediaObjectsCompatible(const UObject* InMediaSource, const UObject* InMediaOutput) const
+{
+	const TArray<IDisplayClusterModularFeatureMediaInitializer*> MediaInitializers = FDisplayClusterConfiguratorMediaUtils::Get().GetMediaInitializers();
+	const bool bCompatible = MediaInitializers.ContainsByPredicate([InMediaSource, InMediaOutput](IDisplayClusterModularFeatureMediaInitializer* Initializer)
+		{
+			return Initializer ? Initializer->AreMediaObjectsCompatible(InMediaSource, InMediaOutput) : false;
+		});
+
+	return bCompatible;
+}
+
+bool SMediaTilesConfigurationDialog::GetMediaPropagationTypes(const UObject* InMediaSource, const UObject* InMediaOutput, EMediaStreamPropagationType& OutPropagationFlags) const
+{
+	// Find appropriate media initializer and get media type
+	const TArray<IDisplayClusterModularFeatureMediaInitializer*>& MediaInitializers = FDisplayClusterConfiguratorMediaUtils::Get().GetMediaInitializers();
+	for (IDisplayClusterModularFeatureMediaInitializer* MediaInitializer : MediaInitializers)
+	{
+		// Find corresponding media initializer, and get its media propagation compatibilities
+		if (MediaInitializer && MediaInitializer->GetSupportedMediaPropagationTypes(InMediaSource, InMediaOutput, OutPropagationFlags))
+		{
+			return true;
+		}
+	}
+
+	return false;
+}
+
+
+/**
+ * Layout page
+ */
+TSharedRef<SWidget> SMediaTilesConfigurationDialog::PageLayout_Build()
+{
+	TSharedRef<SUniformGridPanel> LayoutGrid = SNew(SUniformGridPanel)
+		.SlotPadding(1)
+		.MinDesiredSlotWidth(LayoutGridCellSize)
+		.MinDesiredSlotHeight(LayoutGridCellSize);
+
+	// Build the grid
+	for (int32 TileX = 0; TileX < MaxTilesAmount; ++TileX)
+	{
+		for (int32 TileY = 0; TileY < MaxTilesAmount; ++TileY)
+		{
+			// Current tile
+			const FIntPoint Tile = { TileX, TileY };
+
+			// Instantiate button at XY
+			TSharedRef<SButton> GridButton = SNew(SButton)
+				.HAlign(HAlign_Center)
+				.VAlign(VAlign_Center)
+				.Text(FText::FromString(FString::Printf(TEXT("%dx%d"), TileX + 1, TileY + 1)))
+				.OnClicked(FOnClicked::CreateLambda([this, Tile]()
+					{
+						PageLayout_OnGridCellClicked(Tile);
+						return FReply::Handled();
+					}))
+				.OnHovered(FSimpleDelegate::CreateLambda([this, Tile]()
+					{
+						PageLayout_OnGridCellHovered(Tile);
+					}))
+				.OnUnhovered(FSimpleDelegate::CreateLambda([this]()
+					{
+						PageLayout_OnGridCellUnhovered();
+					}));
+
+			// Store it internally
+			LayoutGridButtons.Emplace(Tile, GridButton);
+
+			// Add the button to layout
+			LayoutGrid->AddSlot(TileX, TileY)
+			[
+				SNew(SBox)
+				[
+					GridButton
+				]
+			];
+		}
+	}
+
+	return LayoutGrid;
+}
+
+void SMediaTilesConfigurationDialog::PageLayout_OnEnter()
+{
+	PageHint->SetText(UE::DisplayClusterConfigurator::MediaTilesConfigurationDialog::Private::TextPageLayoutHeader);
+}
+
+bool SMediaTilesConfigurationDialog::PageLayout_IsConfigurationValid() const
+{
+	// Anything other than 1x1 is valid
+	return Accepted.X > 0 || Accepted.Y > 0;
+}
+
+void SMediaTilesConfigurationDialog::PageLayout_OnGridCellClicked(const FIntPoint& Tile)
+{
+	const bool bInvalidChoice = (Tile == FIntPoint::ZeroValue);
+
+	Accepted = bInvalidChoice ? FIntPoint{ -1, -1 } : Tile;
+	
+	PageLayout_RefreshGridLayout();
+}
+
+void SMediaTilesConfigurationDialog::PageLayout_OnGridCellHovered(const FIntPoint& Tile)
+{
+	Hovered = Tile;
+	PageLayout_RefreshGridLayout();
+}
+
+void SMediaTilesConfigurationDialog::PageLayout_OnGridCellUnhovered()
+{
+	Hovered = { -1, -1 };
+	PageLayout_RefreshGridLayout();
+}
+
+void SMediaTilesConfigurationDialog::PageLayout_RefreshGridLayout()
+{
+	const FButtonStyle& ButtonStyle = FAppStyle::Get().GetWidgetStyle<FButtonStyle>("Button");
+
+	const FSlateColor ColorAccepted  { COLOR("#00FF00FF") };
+	const FSlateColor ColorProposed1 { COLOR("#00FF0080") };
+	const FSlateColor ColorProposed2 { FStyleColors::Foreground };
+	const FSlateColor ColorDefault   { ButtonStyle.NormalForeground };
+	const FSlateColor ColorInvalid   { FLinearColor::Red };
+
+	const bool bMatchesAccepted = (Accepted == Hovered);
+
+	for (int32 XPos = 0; XPos < MaxTilesAmount; ++XPos)
+	{
+		for (int32 YPos = 0; YPos < MaxTilesAmount; ++YPos)
+		{
+			const FSlateColor* NewColor = nullptr;
+
+			// 1x1 not allowed
+			if (XPos == 0 && YPos == 0 && Hovered == FIntPoint::ZeroValue)
+			{
+				NewColor = &ColorInvalid;
+			}
+			// Cells proposed to select again (if doesn't match the currently accepted region)
+			else if (XPos <= Hovered.X && YPos <= Hovered.Y && XPos <= Accepted.X && YPos <= Accepted.Y && !bMatchesAccepted)
+			{
+				NewColor = &ColorProposed1;
+			}
+			// Cells proposed to select (if doesn't match the currently accepted region)
+			else if (XPos <= Hovered.X && YPos <= Hovered.Y && !bMatchesAccepted)
+			{
+				NewColor = &ColorProposed2;
+			}
+			// Currently selected
+			else if (XPos <= Accepted.X && YPos <= Accepted.Y)
+			{
+				NewColor = &ColorAccepted;
+			}
+			// Remaining cells
+			else
+			{
+				NewColor = &ColorDefault;
+			}
+
+			LayoutGridButtons[{XPos, YPos}]->SetBorderBackgroundColor(*NewColor);
+		}
+	}
+}
+
+
+/**
+ * Media objects page
+ */
+TSharedRef<SWidget> SMediaTilesConfigurationDialog::PageMedia_Build()
+{
+	TSharedRef<SVerticalBox> MediaObjectsPage = SNew(SVerticalBox);
+
+	// Add input and output sections
+	for (int32 Idx = 0; Idx < 2; ++Idx)
+	{
+		MediaObjectsPage->AddSlot()
+		.Padding(5, 5)
+		[
+			SNew(SBorder)
+			.Visibility(EVisibility::All)
+			.Content()
+			[
+				SNew(SVerticalBox)
+
+					// Media object row
+					+SVerticalBox::Slot()
+					.Padding(5, 5)
+					[
+						SNew(SHorizontalBox)
+
+							// Field name
+							+SHorizontalBox::Slot()
+							.AutoWidth()
+							.HAlign(HAlign_Left)
+							.VAlign(VAlign_Center)
+							[
+								SNew(STextBlock)
+								.Font(FAppStyle::Get().GetFontStyle("NormalFontBold"))
+								.Text( Idx == 0 ?
+									UE::DisplayClusterConfigurator::MediaTilesConfigurationDialog::Private::TextPageMediaSource :
+									UE::DisplayClusterConfigurator::MediaTilesConfigurationDialog::Private::TextPageMediaOutput)
+							]
+
+							+SHorizontalBox::Slot()
+							.MaxWidth(20)
+							[
+								SNew(SSpacer)
+							]
+
+							// Field value
+							+SHorizontalBox::Slot()
+							.AutoWidth()
+							.HAlign(HAlign_Fill)
+							.VAlign(VAlign_Center)
+							[
+								SNew(SComboButton)
+								.ComboButtonStyle(FAppStyle::Get(), "SimpleComboButton")
+								.ContentPadding(0)
+								.HasDownArrow(true)
+								.OnGetMenuContent_Lambda([this, Idx]()
+								{
+									return PageMedia_BuildClassPicker(Idx == 0 ? UMediaSource::StaticClass() : UMediaOutput::StaticClass());
+								})
+								.ButtonContent()
+								[
+									SAssignNew(Idx == 0 ? MediaSourceSelectedName : MediaOutputSelectedName, STextBlock)
+									.Text(UE::DisplayClusterConfigurator::MediaTilesConfigurationDialog::Private::TextPageMediaComboboxNone)
+									.ColorAndOpacity(FStyleColors::Foreground)
+								]
+							]
+					]
+
+					// Media object status
+					+SVerticalBox::Slot()
+					.Padding(5, 5)
+					[
+						SNew(SHorizontalBox)
+
+							+SHorizontalBox::Slot()
+							.AutoWidth()
+							.VAlign(VAlign_Center)
+							.HAlign(HAlign_Left)
+							.Padding(2.f)
+							[
+								SAssignNew(Idx == 0 ? MediaSourceStatusImage : MediaOutputStatusImage, SImage)
+								.Image(nullptr)
+							]
+
+							+SHorizontalBox::Slot()
+							.AutoWidth()
+							.VAlign(VAlign_Center)
+							.HAlign(HAlign_Left)
+							[
+								SAssignNew(Idx == 0 ? MediaSourceStatusText : MediaOutputStatusText, STextBlock)
+							]
+					]
+
+			]
+		];
+	}
+
+	return MediaObjectsPage;
+}
+
+TSharedRef<SWidget> SMediaTilesConfigurationDialog::PageMedia_BuildClassPicker(UClass* FilterClass)
+{
+	// Auxiliary class filter implementation
+	class FDisplayClusterConfiguratorMediaTypeFilter final : public IClassViewerFilter
+	{
+	public:
+		explicit FDisplayClusterConfiguratorMediaTypeFilter(UClass* InRequiredBaseClass = nullptr)
+			: RequiredBaseClass(InRequiredBaseClass)
+		{ }
+
+		virtual bool IsClassAllowed(const FClassViewerInitializationOptions& InInitOptions, const UClass* InClass, TSharedRef<FClassViewerFilterFuncs> InFilterFuncs) override
+		{
+			static const EClassFlags FilterFlags = EClassFlags::CLASS_Abstract | EClassFlags::CLASS_Deprecated | EClassFlags::CLASS_Hidden | EClassFlags::CLASS_HideDropDown;
+			return InClass && !InClass->HasAnyClassFlags(FilterFlags) && (RequiredBaseClass ? InClass->IsChildOf(RequiredBaseClass) : true);
+		}
+
+		virtual bool IsUnloadedClassAllowed(const FClassViewerInitializationOptions& InInitOptions, const TSharedRef<const IUnloadedBlueprintData> InUnloadedClassData, TSharedRef<FClassViewerFilterFuncs> InFilterFuncs) override
+		{
+			return false;
+		}
+
+	private:
+		/** Classes must have this base class to pass the filter. */
+		UClass* RequiredBaseClass = nullptr;
+	};
+
+	FClassViewerModule& ClassViewerModule = FModuleManager::LoadModuleChecked<FClassViewerModule>("ClassViewer");
+
+	FClassViewerInitializationOptions Options;
+	Options.Mode = EClassViewerMode::ClassPicker;
+	Options.NameTypeToDisplay = EClassViewerNameTypeToDisplay::DisplayName;
+	Options.bShowNoneOption = true;
+	Options.bIsActorsOnly = false;
+	Options.bShowUnloadedBlueprints = false;
+	Options.ClassFilters.Add(MakeShared<FDisplayClusterConfiguratorMediaTypeFilter>(FilterClass));
+
+	const bool bIsMediaSource = (FilterClass == UMediaSource::StaticClass());
+
+	// Instantiate class viewer widget
+	const TSharedRef<SWidget> ClassViewer = ClassViewerModule.CreateClassViewer(
+		Options,
+		FOnClassPicked::CreateSP(this, FilterClass == UMediaSource::StaticClass() ?
+			&SMediaTilesConfigurationDialog::PageMedia_OnMediaSourceChanged :
+			&SMediaTilesConfigurationDialog::PageMedia_OnMediaOutputChanged)
+		);
+
+	return SNew(SBox)
+		[
+			ClassViewer
+		];
+}
+
+bool SMediaTilesConfigurationDialog::PageMedia_OnCanShow() const
+{
+	return PageLayout_IsConfigurationValid();
+}
+
+void SMediaTilesConfigurationDialog::PageMedia_OnEnter()
+{
+	PageHint->SetText(UE::DisplayClusterConfigurator::MediaTilesConfigurationDialog::Private::TextPageMediaHeader);
+}
+
+bool SMediaTilesConfigurationDialog::PageMedia_IsConfigurationValid() const
+{
+	// Media source and output must be valid and compatible
+	const bool bMediaObjectsCompatible = AreMediaObjectsCompatible(MediaSource.Get(), MediaOutput.Get());
+	return bMediaObjectsCompatible;
+}
+
+void SMediaTilesConfigurationDialog::PageMedia_OnMediaSourceChanged(UClass* InNewClass)
+{
+	// Process new choice
+	if (InNewClass && InNewClass->IsChildOf<UMediaSource>())
+	{
+		MediaSource.Reset(NewObject<UMediaSource>(GetTransientPackage(), InNewClass));
+		MediaSourceSelectedName->SetText(InNewClass->GetDisplayNameText());
+	}
+	else
+	{
+		MediaSource.Reset();
+		MediaSourceSelectedName->SetText(UE::DisplayClusterConfigurator::MediaTilesConfigurationDialog::Private::TextPageMediaComboboxNone);
+	}
+
+	// Update feedback
+	PageMedia_UpdateMessage();
+
+	FSlateApplication::Get().DismissAllMenus();
+}
+
+void SMediaTilesConfigurationDialog::PageMedia_OnMediaOutputChanged(UClass* InNewClass)
+{
+	// Process new choice
+	if (InNewClass && InNewClass->IsChildOf<UMediaOutput>())
+	{
+		MediaOutput.Reset(NewObject<UMediaOutput>(GetTransientPackage(), InNewClass));
+		MediaOutputSelectedName->SetText(InNewClass->GetDisplayNameText());
+	}
+	else
+	{
+		MediaOutput.Reset();
+		MediaOutputSelectedName->SetText(UE::DisplayClusterConfigurator::MediaTilesConfigurationDialog::Private::TextPageMediaComboboxNone);
+	}
+
+	// Update feedback
+	PageMedia_UpdateMessage();
+
+	FSlateApplication::Get().DismissAllMenus();
+}
+
+void SMediaTilesConfigurationDialog::PageMedia_UpdateMessage()
+{
+	if (!MediaSource)
+	{
+		MediaSourceStatusImage->SetImage(nullptr);
+		MediaSourceStatusText->SetText(FText::GetEmpty());
+	}
+
+	if (!MediaOutput)
+	{
+		MediaOutputStatusImage->SetImage(nullptr);
+		MediaOutputStatusText->SetText(FText::GetEmpty());
+	}
+
+	const FSlateBrush* StatusImageOk   = FAppStyle::Get().GetBrush("EditorViewport.LightingOnlyMode");
+	const FSlateBrush* StatusImageWarn = FAppStyle::Get().GetBrush("Level.LightingScenarioNotIcon16x");
+
+	// Invalidate media propagation type. It will be set to a proper value later if media objects are good.
+	MediaPropagationTypes = EMediaStreamPropagationType::None;
+
+	// Only one object is chosen
+	if((MediaSource && !MediaOutput) || (!MediaSource && MediaOutput))
+	{
+		const UObject* CurrentObject = (MediaSource.IsValid() ? Cast<UObject>(MediaSource.Get()) : Cast<UObject>(MediaOutput.Get()));
+
+		const TSharedPtr<SImage>&     StatusImage = (MediaSource.IsValid() ? MediaSourceStatusImage : MediaOutputStatusImage);
+		const TSharedPtr<STextBlock>& StatusText  = (MediaSource.IsValid() ? MediaSourceStatusText : MediaOutputStatusText);
+
+		const bool bIsObjectSupported = IsMediaObjectSupported(CurrentObject);
+
+		StatusImage->SetImage(bIsObjectSupported ? StatusImageOk : StatusImageWarn);
+
+		StatusText->SetColorAndOpacity(bIsObjectSupported ? FLinearColor::Green : FLinearColor::Yellow);
+
+		StatusText->SetText(bIsObjectSupported ?
+			UE::DisplayClusterConfigurator::MediaTilesConfigurationDialog::Private::TextPageMediaStatusOk :
+			UE::DisplayClusterConfigurator::MediaTilesConfigurationDialog::Private::TextPageMediaStatusNotSupp);
+		
+	}
+	// Both objects are chosen
+	else if(MediaSource && MediaOutput)
+	{
+		const bool bAreObjectsCompatible = AreMediaObjectsCompatible(MediaSource.Get(), MediaOutput.Get());
+
+		TSharedPtr<SImage>*     StatusImage[] = { &MediaSourceStatusImage, &MediaOutputStatusImage };
+		TSharedPtr<STextBlock>* StatusText[]  = { &MediaSourceStatusText , &MediaOutputStatusText };
+
+		for (int32 Idx = 0; Idx < 2; ++Idx)
+		{
+			(*StatusImage[Idx])->SetImage(bAreObjectsCompatible ? StatusImageOk : StatusImageWarn);
+
+			(*StatusText[Idx])->SetColorAndOpacity(bAreObjectsCompatible ? FLinearColor::Green : FLinearColor::Yellow);
+
+			(*StatusText[Idx])->SetText(bAreObjectsCompatible ?
+				UE::DisplayClusterConfigurator::MediaTilesConfigurationDialog::Private::TextPageMediaStatusOk :
+				UE::DisplayClusterConfigurator::MediaTilesConfigurationDialog::Private::TextPageMediaNotCompatOrNotSupp);
+		}
+
+		// Update media propagation types
+		if (bAreObjectsCompatible)
+		{
+			GetMediaPropagationTypes(MediaSource.Get(), MediaOutput.Get(), MediaPropagationTypes);
+		}
+	}
+}
+
+
+/**
+ * Nodes selection page
+ */
+TSharedRef<SWidget> SMediaTilesConfigurationDialog::PageNodes_Build()
+{
+	return SNew(SHorizontalBox)
+
+		// Output nodes that produce the tiles
+		+SHorizontalBox::Slot()
+		.Padding(2)
+		[
+			SNew(SBorder)
+			.OnMouseButtonDown_Lambda([this](const FGeometry& Geometry, const FPointerEvent& MouseEvent)
+				{
+					return PageNodes_OnMouseButtonDown(Geometry, MouseEvent, true);
+				})
+			[
+				PageNodes_BuildSelSection(
+					UE::DisplayClusterConfigurator::MediaTilesConfigurationDialog::Private::TextPageNodesSectionOutputHeader,
+					NodesAllowedForOutput,
+					OutputSelection)
+			]
+		]
+
+		// Input nodes that receive all the tiles and compose the full frame
+		+SHorizontalBox::Slot()
+		.Padding(2)
+		[
+			SNew(SBorder)
+			.OnMouseButtonDown_Lambda([this](const FGeometry& Geometry, const FPointerEvent& MouseEvent)
+				{
+					return PageNodes_OnMouseButtonDown(Geometry, MouseEvent, false);
+				})
+			[
+				PageNodes_BuildSelSection(
+					UE::DisplayClusterConfigurator::MediaTilesConfigurationDialog::Private::TextPageNodesSectionInputHeader,
+					NodesAllowedForInput,
+					InputSelection)
+			]
+		];
+}
+
+TSharedRef<SWidget> SMediaTilesConfigurationDialog::PageNodes_BuildSelSection(const FText& Header, const TSet<FString>& AllowedItems, TSet<FString>& CurrentSelection)
+{
+	TSharedRef<SVerticalBox> Section = SNew(SVerticalBox);
+
+	// Output header
+	Section->AddSlot()
+	.Padding(5)
+	.AutoHeight()
+	.HAlign(HAlign_Center)
+	.VAlign(VAlign_Top)
+	[
+		SNew(STextBlock)
+		.Font(FAppStyle::Get().GetFontStyle("NormalFontBold"))
+		.Text(Header)
+	];
+
+	Section->AddSlot()
+	.MaxHeight(10)
+	[
+		SNew(SSpacer)
+	];
+
+	// Output body
+	Section->AddSlot()
+	.Padding(5, 3)
+	.MaxHeight(400)
+	.AutoHeight()
+	[
+		SNew(SScrollBox)
+		.Orientation(EOrientation::Orient_Vertical)
+
+			+SScrollBox::Slot()
+			[
+				PageNodes_BuildNodePicker(AllowedItems, CurrentSelection)
+			]
+	];
+
+	return Section;
+}
+
+TSharedRef<SWidget> SMediaTilesConfigurationDialog::PageNodes_BuildNodePicker(const TSet<FString>& AllowedItems, TSet<FString>& CurrentSelection)
+{
+	TSharedRef<SVerticalBox> NodePicker = SNew(SVerticalBox);
+
+	// Generate output nodes list
+	for (const FString& NodeId : AllowedItems)
+	{
+		NodePicker->AddSlot()
+			.AutoHeight()
+			.Padding(3)
+			[
+				SNew(SCheckBox)
+				.Padding(5)
+				.Style(FAppStyle::Get(), "ToggleButtonCheckbox")
+				.CheckBoxContentUsesAutoWidth(true)
+				.IsChecked_Lambda([NodeId, &CurrentSelection]()
+					{
+						return CurrentSelection.Contains(NodeId) ? ECheckBoxState::Checked : ECheckBoxState::Unchecked;
+					})
+				.OnCheckStateChanged_Lambda([NodeId, &CurrentSelection](ECheckBoxState NewState)
+					{
+						if (NewState == ECheckBoxState::Checked)
+						{
+							CurrentSelection.Add(NodeId);
+						}
+						else
+						{
+							CurrentSelection.Remove(NodeId);
+						}
+					})
+				.Content()
+				[
+					SNew(STextBlock)
+					.Text(FText::FromString(NodeId))
+					.AutoWrapText(true)
+				]
+			];
+	}
+
+	return NodePicker;
+}
+
+FReply SMediaTilesConfigurationDialog::PageNodes_OnMouseButtonDown(const FGeometry& Geometry, const FPointerEvent& PointerEvent, bool bOutputSelection)
+{
+	if (PointerEvent.GetEffectingButton() == EKeys::RightMouseButton)
+	{
+		PageNodes_CreateContextMenu(PointerEvent.GetScreenSpacePosition(), bOutputSelection);
+		return FReply::Handled();
+	}
+
+	return FReply::Unhandled();
+}
+
+void SMediaTilesConfigurationDialog::PageNodes_CreateContextMenu(const FVector2D& CursorPosition, bool bOutputSelection)
+{
+	constexpr bool bCloseAfterSelection = true;
+	FMenuBuilder MenuBuilder(bCloseAfterSelection, nullptr);
+
+	// Section: general commands
+	{
+		MenuBuilder.BeginSection("General");
+
+		MenuBuilder.AddWidget(
+			SNew(STextBlock)
+			.Text(FText::FromString("General")),
+			FText(),
+			true
+		);
+
+		// Select All
+		MenuBuilder.AddMenuEntry(
+			UE::DisplayClusterConfigurator::MediaTilesConfigurationDialog::Private::TextPageNodesMenuSelectAll,
+			FText(),
+			FSlateIcon(),
+			FUIAction(
+				FExecuteAction::CreateSP(this, &SMediaTilesConfigurationDialog::PageNodes_Menu_OnChangeSelectionAll, bOutputSelection, true)
+			)
+		);
+
+		// Deselect all
+		MenuBuilder.AddMenuEntry(
+			UE::DisplayClusterConfigurator::MediaTilesConfigurationDialog::Private::TextPageNodesMenuDeselectAll,
+			FText(),
+			FSlateIcon(),
+			FUIAction(
+				FExecuteAction::CreateSP(this, &SMediaTilesConfigurationDialog::PageNodes_Menu_OnChangeSelectionAll, bOutputSelection, false)
+			)
+		);
+
+		MenuBuilder.EndSection();
+	}
+
+	// Section: extra commands
+	{
+		MenuBuilder.BeginSection("Extra");
+
+		MenuBuilder.AddWidget(
+			SNew(STextBlock)
+			.Text(FText::FromString("Extra")),
+			FText(),
+			true
+		);
+
+		// Select All w/o counterpart
+		MenuBuilder.AddMenuEntry(
+			bOutputSelection ?
+				UE::DisplayClusterConfigurator::MediaTilesConfigurationDialog::Private::TextPageNodesMenuSelectAllNoReceivers :
+				UE::DisplayClusterConfigurator::MediaTilesConfigurationDialog::Private::TextPageNodesMenuSelectAllNoSenders,
+			FText(),
+			FSlateIcon(),
+			FUIAction(
+				FExecuteAction::CreateSP(this, &SMediaTilesConfigurationDialog::PageNodes_Menu_OnSelectAllExceptOfCounterpart, bOutputSelection)
+			)
+		);
+
+		// Deselect all from counterpart
+		MenuBuilder.AddMenuEntry(
+			bOutputSelection ?
+				UE::DisplayClusterConfigurator::MediaTilesConfigurationDialog::Private::TextPageNodesMenuDeselectReceivers :
+				UE::DisplayClusterConfigurator::MediaTilesConfigurationDialog::Private::TextPageNodesMenuDeselectSenders,
+			FText(),
+			FSlateIcon(),
+			FUIAction(
+				FExecuteAction::CreateSP(this, &SMediaTilesConfigurationDialog::PageNodes_Menu_OnDeselectAllFromCounterpart, bOutputSelection)
+			)
+		);
+
+		// Output context only
+		if (bOutputSelection)
+		{
+			// Select all offscreen nodes
+			MenuBuilder.AddMenuEntry(
+				UE::DisplayClusterConfigurator::MediaTilesConfigurationDialog::Private::TextPageNodesMenuSelectAllOffscreen,
+				FText(),
+				FSlateIcon(),
+				FUIAction(
+					FExecuteAction::CreateSP(this, &SMediaTilesConfigurationDialog::PageNodes_Menu_OnSelectAllOffscreen, bOutputSelection)
+				)
+			);
+		}
+
+		MenuBuilder.EndSection();
+	}
+
+	// Show the context menu
+	FSlateApplication::Get().PushMenu(
+		SharedThis(this),
+		FWidgetPath(),
+		MenuBuilder.MakeWidget(),
+		CursorPosition,
+		FPopupTransitionEffect(FPopupTransitionEffect::ContextMenu));
+}
+
+void SMediaTilesConfigurationDialog::PageNodes_Menu_OnChangeSelectionAll(bool bOutputSelection, bool bSelected)
+{
+	// A set being edited
+	TSet<FString>& CurrentSet = bOutputSelection ? OutputSelection : InputSelection;
+	// Allowed items for the set being edited
+	const TSet<FString>& AllowedSet = bOutputSelection ? NodesAllowedForOutput : NodesAllowedForInput;
+
+	// Clear first either it's 'select all' or 'deselect all'
+	CurrentSet.Reset();
+
+	// And 'select all' if it was actually requested
+	if (bSelected)
+	{
+		CurrentSet = AllowedSet;
+	}
+}
+
+void SMediaTilesConfigurationDialog::PageNodes_Menu_OnSelectAllExceptOfCounterpart(bool bOutputSelection)
+{
+	// A set being edited
+	TSet<FString>& SetA = bOutputSelection ? OutputSelection : InputSelection;
+	// A counterpart set
+	const TSet<FString>& SetB = bOutputSelection ? InputSelection : OutputSelection;
+	// Allowed items for the set being edited
+	const TSet<FString>& AllowedSetA = bOutputSelection ? NodesAllowedForOutput : NodesAllowedForInput;
+	// Allowed items for the counterpart set
+	const TSet<FString>& AllowedSetB = bOutputSelection ? NodesAllowedForInput : NodesAllowedForOutput;
+
+	// Select all in A that aren't currently selected in B
+	// Result = AllowedA - SelectedB
+	SetA = AllowedSetA.Difference(SetB);
+}
+
+void SMediaTilesConfigurationDialog::PageNodes_Menu_OnDeselectAllFromCounterpart(bool bOutputSelection)
+{
+	// A set being edited
+	TSet<FString>& SetA = bOutputSelection ? OutputSelection : InputSelection;
+	// A counterpart set
+	const TSet<FString>& SetB = bOutputSelection ? InputSelection : OutputSelection;
+
+	// In set A, deselect everything that is in set B
+	SetA = SetA.Difference(SetB);
+}
+
+void SMediaTilesConfigurationDialog::PageNodes_Menu_OnSelectAllOffscreen(bool bOutputSelection)
+{
+	// A set being edited
+	TSet<FString>& SetA = bOutputSelection ? OutputSelection : InputSelection;
+
+	// Iterate through all the nodes in config
+	for (const TPair<FString, TObjectPtr<UDisplayClusterConfigurationClusterNode>>& NodeIt : Parameters.ConfigData->Cluster->Nodes)
+	{
+		// Offscreen?
+		if (NodeIt.Value && NodeIt.Value->bRenderHeadless)
+		{
+			SetA.Remove(NodeIt.Key);
+		}
+	}
+}
+
+bool SMediaTilesConfigurationDialog::PageNodes_OnCanShow() const
+{
+	return PageMedia_IsConfigurationValid();
+}
+
+void SMediaTilesConfigurationDialog::PageNodes_OnEnter()
+{
+	PageHint->SetText(UE::DisplayClusterConfigurator::MediaTilesConfigurationDialog::Private::TextPageNodesHeader);
+}
+
+bool SMediaTilesConfigurationDialog::PageNodes_IsConfigurationValid() const
+{
+	// We need at least one input node to be selected. There is no requirement for output nodes
+	// because those can be set up on the next page manually.
+	const bool bHasInputAssigned = !InputSelection.IsEmpty();
+
+	return bHasInputAssigned;
+}
+
+
+/**
+ * Finalization page
+ */
+TSharedRef<SWidget> SMediaTilesConfigurationDialog::PageFinalization_Build()
+{
+	// Output mapping grid
+	return SNew(SBorder)
+	[
+		SAssignNew(OutputMappingGrid, SUniformGridPanel)
+		.SlotPadding(2)
+		.MinDesiredSlotWidth(OutputMappingGridCellSize)
+		.MinDesiredSlotHeight(OutputMappingGridCellSize)
+	];
+}
+
+TSharedRef<SWidget> SMediaTilesConfigurationDialog::PageFinalization_BuildOutputTileButtonContent(const FIntPoint& Tile)
+{
+	return SNew(SBox)
+		.MaxDesiredHeight(OutputMappingGridCellSize)
+		.Content()
+		[
+			SNew(SScrollBox)
+			.Orientation(EOrientation::Orient_Vertical)
+
+				+SScrollBox::Slot()
+				[
+					SAssignNew(OutputMapping[Tile].ButtonContentWidget, SVerticalBox)
+				]
+		];
+}
+
+TSharedRef<SWidget> SMediaTilesConfigurationDialog::PageFinalization_BuildOutputTileButtonDropdownContent(const FIntPoint& Tile)
+{
+	TSharedRef<SVerticalBox> NodePicker = SNew(SVerticalBox);
+
+	// Generate nodes list
+	for (const FString& NodeId : ClusterNodeIds)
+	{
+		NodePicker->AddSlot()
+			.AutoHeight()
+			.Padding(1)
+			[
+				SNew(SCheckBox)
+				.Padding(1)
+				.Style(FAppStyle::Get(), "ToggleButtonCheckbox")
+				.CheckBoxContentUsesAutoWidth(true)
+				.IsChecked_Lambda([this, Tile, NodeId]()
+					{
+						const bool bSelected = OutputMapping[Tile].ClusterNodes.Contains(NodeId);
+						return bSelected ? ECheckBoxState::Checked : ECheckBoxState::Unchecked;
+					})
+				.OnCheckStateChanged_Lambda([this, Tile, NodeId](ECheckBoxState NewState)
+					{
+						if (NewState == ECheckBoxState::Checked)
+						{
+							OutputMapping[Tile].ClusterNodes.Add(NodeId);
+						}
+						else
+						{
+							OutputMapping[Tile].ClusterNodes.Remove(NodeId);
+						}
+
+						// Update grid cell content
+						PageFinalization_UpdateTileButtonContent(Tile);
+					})
+				.Content()
+				[
+					SNew(SBox)
+					.Padding(3)
+					.HAlign(HAlign_Left)
+					.VAlign(VAlign_Top)
+					[
+						SNew(STextBlock)
+						.Text(FText::FromString(NodeId))
+						.TextStyle(&FCoreStyle::Get().GetWidgetStyle<FTextBlockStyle>("SmallText"))
+						.Font(FAppStyle::GetFontStyle("TinyText"))
+						.AutoWrapText(true)
+					]
+				]
+			];
+	}
+
+	return SNew(SBox)
+		.MaxDesiredHeight(600)
+		.Content()
+		[
+			SNew(SScrollBox)
+			.Orientation(EOrientation::Orient_Vertical)
+
+				+SScrollBox::Slot()
+				[
+					NodePicker
+				]
+		];
+}
+
+bool SMediaTilesConfigurationDialog::PageFinalization_OnCanShow() const
+{
+	return PageNodes_IsConfigurationValid();
+}
+
+void SMediaTilesConfigurationDialog::PageFinalization_OnEnter()
+{
+	PageHint->SetText(UE::DisplayClusterConfigurator::MediaTilesConfigurationDialog::Private::TextPageFinalizationHeader);
+
+	// Clean up any current mapping
+	PageFinalization_ResetOutputMapping();
+
+	// Pre-configure output mapping
+	if (Parameters.bAutoPreconfigureOutputMapping)
+	{
+		PageFinalization_PresetupOutputs();
+	}
+
+	// Re-generate grid to properly reflect current configuration
+	PageFinalization_RegenerateGrid();
+}
+
+void SMediaTilesConfigurationDialog::PageFinalization_ResetOutputMapping()
+{
+	// Clean up current mapping
+	for (TPair<FIntPoint, FOutputMappingData>& TileData : OutputMapping)
+	{
+		TileData.Value.ClusterNodes.Reset();
+	}
+}
+
+void SMediaTilesConfigurationDialog::PageFinalization_PresetupOutputs()
+{
+	// Pre-configure output mapping based on media type previously chosen
+	if (EnumHasAnyFlags(MediaPropagationTypes, EMediaStreamPropagationType::Multicast))
+	{
+		PageFinalization_PresetupOutputs_Multicast();
+	}
+	else if (EnumHasAnyFlags(MediaPropagationTypes, EMediaStreamPropagationType::LocalMulticast))
+	{
+		PageFinalization_PresetupOutputs_LocalMulticast();
+	}
+	else
+	{
+		unimplemented();
+	}
+}
+
+void SMediaTilesConfigurationDialog::PageFinalization_PresetupOutputs_LocalMulticast()
+{
+	// With LocalMulticast, we basically have per-host tile propagation. So every host
+	// that has at least one tile receiver, must also have the senders that produce
+	// full set of tiles.
+
+	// Max amount of senders according to layout (e.g. would be 6 nodes for 3x2 layout)
+	const int32 MaxSendersPerHost = (Accepted.X + 1) * (Accepted.Y + 1);
+
+	// Pre-build senders per-host mapping
+	TMultiMap<FString, FString> HostToSendersMap;
+	for (const FString& NodeId : OutputSelection)
+	{
+		if (const FString* Host = NodeToHostMap.Find(NodeId))
+		{
+			HostToSendersMap.Add(*Host, NodeId);
+		}
+	}
+
+	// Get list of hosts where we have at least one sender
+	TSet<FString> SenderHosts;
+	HostToSendersMap.GetKeys(SenderHosts);
+
+	// Map tile senders for every host
+	for (const FString& SenderHost : SenderHosts)
+	{
+		TArray<FString> SenderNodes;
+		HostToSendersMap.MultiFind(SenderHost, SenderNodes);
+
+		// Sort alphabetically. This may be handy if user has some naming convention which is kind of typical.
+		SenderNodes.Sort();
+
+		// Counter of senders already mapped
+		int32 AssignedAmount = 0;
+
+		// For each sender at host...
+		for (const FString& SenderNodeId : SenderNodes)
+		{
+			// A tile we're going to assign for this sender
+			const int32 TileX = AssignedAmount % (Accepted.X + 1);
+			const int32 TileY = AssignedAmount / (Accepted.X + 1);
+			const FIntPoint Tile { TileX, TileY };
+
+			// Map this sender to the tile
+			OutputMapping[Tile].ClusterNodes.Add(SenderNodeId);
+
+			// Ignore other senders
+			if (++AssignedAmount >= MaxSendersPerHost)
+			{
+				break;
+			}
+		}
+	}
+}
+
+void SMediaTilesConfigurationDialog::PageFinalization_PresetupOutputs_Multicast()
+{
+	// Multicast implies every tile is produced once by a dedicated node then propagated
+	// to all the receivers. This means we can simply pre-configure output mapping
+	// by binding N tiles to N senders using 1-1 relation, and any amount of receivers.
+	// We don't care about host-node mapping in this case as it's true multicast.
+
+	// Max amount of senders according to layout (e.g. would be 6 nodes for 3x2 layout)
+	const int32 MaxSendersPerHost = (Accepted.X + 1) * (Accepted.Y + 1);
+
+	// Counter of senders already mapped
+	int32 AssignedAmount = 0;
+
+	// Iterate through the current output selection
+	for (const FString& SenderNodeId : OutputSelection)
+	{
+		// A tile we're going to assign for this sender
+		const int32 TileX = AssignedAmount % (Accepted.X + 1);
+		const int32 TileY = AssignedAmount / (Accepted.X + 1);
+		const FIntPoint Tile{ TileX, TileY };
+
+		// Map this sender to the tile
+		OutputMapping[Tile].ClusterNodes.Add(SenderNodeId);
+
+		// Ignore other senders
+		if (++AssignedAmount >= MaxSendersPerHost)
+		{
+			break;
+		}
+	}
+}
+
+void SMediaTilesConfigurationDialog::PageFinalization_RegenerateGrid()
+{
+	// Remove slots if there are any
+	OutputMappingGrid->ClearChildren();
+
+	// Create new slots
+	for (int32 TileX = 0; TileX <= Accepted.X; ++TileX)
+	{
+		for (int32 TileY = 0; TileY <= Accepted.Y; ++TileY)
+		{
+			const FIntPoint Tile{ TileX, TileY };
+			FOutputMappingData& TileOutputMapping = OutputMapping[Tile];
+
+			// Instantiate button at XY
+			TSharedRef<SComboButton> GridButton = SNew(SComboButton)
+				.HAlign(HAlign_Left)
+				.VAlign(VAlign_Top)
+				.ContentPadding(3)
+				.HasDownArrow(true)
+				.ForegroundColor_Lambda([this, Tile]()
+					{
+						return PageFinalization_GetButtonColor(Tile);
+					})
+				.ButtonColorAndOpacity_Lambda([this, Tile]()
+					{
+						return PageFinalization_GetButtonColor(Tile);
+					})
+				.OnGetMenuContent_Lambda([this, Tile]()
+					{
+						return PageFinalization_BuildOutputTileButtonDropdownContent(Tile);
+					})
+				.ButtonContent()
+				[
+					PageFinalization_BuildOutputTileButtonContent(Tile)
+				];
+
+			// Store it internally
+			TileOutputMapping.ButtonWidget = GridButton;
+
+			// Grid cell (button)
+			OutputMappingGrid->AddSlot(TileX, TileY)
+				[
+					SNew(SBox)
+					[
+						GridButton
+					]
+				];
+
+			// Force button content update
+			PageFinalization_UpdateTileButtonContent(Tile);
+		}
+	}
+}
+
+bool SMediaTilesConfigurationDialog::PageFinalization_IsConfigurationValid() const
+{
+	// There are no strict rules to validate current configuration at this step. The GUI
+	// has pre-configured output mapping, plus user may want to change the configuration
+	// in any way. User even may clean everything or partially and configure the rest
+	// manually in the property editor. Therefore always return true.
+	return true;
+}
+
+void SMediaTilesConfigurationDialog::PageFinalization_UpdateTileButtonContent(const FIntPoint& Tile)
+{
+	FOutputMappingData& TileData = OutputMapping[Tile];
+
+	// Clean all
+	TileData.ButtonContentWidget->ClearChildren();
+
+	// Rebuild children
+	for (const FString& NodeId : TileData.ClusterNodes)
+	{
+		TileData.ButtonContentWidget->AddSlot()
+			[
+				SNew(STextBlock)
+				.Text(FText::FromString(NodeId))
+				.TextStyle(FAppStyle::Get(), "SmallButtonText")
+			];
+	}
+}
+
+FSlateColor SMediaTilesConfigurationDialog::PageFinalization_GetButtonColor(const FIntPoint& Tile)
+{
+	// Visualize whether a tile has any output nodes assigned
+	const bool bNodesAssigned = !OutputMapping[Tile].ClusterNodes.IsEmpty();
+	return bNodesAssigned ? FLinearColor::Green : FLinearColor::Yellow;
+}
+
+#undef LOCTEXT_NAMESPACE
diff --git a/Engine/Plugins/Runtime/nDisplay/Source/DisplayClusterConfigurator/Private/Views/Details/Media/SMediaTilesConfigurationDialog.h b/Engine/Plugins/Runtime/nDisplay/Source/DisplayClusterConfigurator/Private/Views/Details/Media/SMediaTilesConfigurationDialog.h
new file mode 100644
--- /dev/null
+++ b/Engine/Plugins/Runtime/nDisplay/Source/DisplayClusterConfigurator/Private/Views/Details/Media/SMediaTilesConfigurationDialog.h
@@ -0,0 +1,361 @@
+// Copyright Epic Games, Inc. All Rights Reserved.
+
+#pragma once
+
+#include "Widgets/SWindow.h"
+#include "Widgets/DeclarativeSyntaxSupport.h"
+#include "Styling/SlateColor.h"
+
+#include "DisplayClusterconfigurationTypes.h"
+
+
+class SButton;
+class SComboButton;
+class SImage;
+class STextBlock;
+class SUniformGridPanel;
+class SWidget;
+class SWizard;
+
+class UMediaOutput;
+class UMediaSource;
+
+enum class EMediaStreamPropagationType : uint8;
+
+
+/**
+ * Tiled media configuration dialog parameters.
+ */
+struct FMediaTilesConfigurationDialogParameters
+{
+	/** Configuration data of a DCRA being edited. */
+	const UDisplayClusterConfigurationData* ConfigData = nullptr;
+
+	/** An object owning the configuration data being edited. */
+	TObjectPtr<UObject> Owner;
+
+	/** Whether output mapping should be pre-configured automatically. */
+	bool bAutoPreconfigureOutputMapping = true;
+
+	/** Layout ref of a DCRA's entity being edited. */
+	FIntPoint* SplitLayout = nullptr;
+
+	/** Input groups ref of a DCRA's entity being edited. */
+	TArray<FDisplayClusterConfigurationMediaTiledInputGroup>* InputGroups = nullptr;
+
+	/** Output groups ref of a DCRA's entity being edited. */
+	TArray<FDisplayClusterConfigurationMediaTiledOutputGroup>* OutputGroups = nullptr;
+
+public:
+
+	/** Config parameters validation. */
+	bool Validate() const
+	{
+		const bool bIsValid = ConfigData && Owner && SplitLayout && InputGroups && OutputGroups;
+		checkSlow(bIsValid);
+		return bIsValid;
+	}
+};
+
+
+/**
+ * Configuration dialog for tiled media
+ */
+class SMediaTilesConfigurationDialog
+	: public SWindow
+{
+public:
+	SLATE_BEGIN_ARGS(SMediaTilesConfigurationDialog)
+	{}
+	SLATE_END_ARGS()
+
+	void Construct(const FArguments& InArgs, const FMediaTilesConfigurationDialogParameters& InParameters);
+
+public:
+
+	/** Returns true if configuration was completed. */
+	bool WasConfigurationCompleted() const;
+
+private:
+
+	/** Initializes internal data. */
+	void InitializeInternals();
+
+	/** Returns true if 'Finish' button should be active. */
+	bool IsFinishButtonEnabled() const;
+
+	/** 'Finish' button handler. */
+	void OnFinishButtonClicked();
+
+	/** 'Cancel' button handler. */
+	void OnCancelButtonClicked();
+
+private:
+
+	/** Apply configuration entry point. */
+	bool ApplyConfiguration(TArray<FString>& OutErrors, TArray<FString>& OutWarnings);
+
+	/** Apply configuration for 'Multicast' media. */
+	bool ApplyConfiguration_Multicast(TArray<FString>& OutErrors, TArray<FString>& OutWarnings);
+
+	/** Apply configuration for 'LocalMulticast' media. */
+	bool ApplyConfiguration_LocalMulticast(TArray<FString>& OutErrors, TArray<FString>& OutWarnings);
+
+private:
+
+	/** Helper function to generate a message with the issues found. */
+	void GenerateIssuesFoundMessage(const TArray<FString>& Errors, const TArray<FString>& Warnings, FText& OutMessage) const;
+
+	/** Helper function to check if specific media object is supported for auto-configuration */
+	bool IsMediaObjectSupported(const UObject* MediaObject) const;
+
+	/** Helper function to check if specific media objects are compatible to each other */
+	bool AreMediaObjectsCompatible(const UObject* InMediaSource, const UObject* InMediaOutput) const;
+
+	/** Helper function to get media propagation types from a media initializer */
+	bool GetMediaPropagationTypes(const UObject* InMediaSource, const UObject* InMediaOutput, EMediaStreamPropagationType& OutPropagationFlags) const;
+
+
+	/**
+	 * Layout page
+	 */
+private:
+
+	/** Builds layout page. */
+	TSharedRef<SWidget> PageLayout_Build();
+
+	/** Callback on layout page enter. */
+	void PageLayout_OnEnter();
+
+	/** Returns true if layout configuration is valid. */
+	bool PageLayout_IsConfigurationValid() const;
+
+	/** Cell click handler. */
+	void PageLayout_OnGridCellClicked(const FIntPoint& Tile);
+
+	/** Cell hover handler. */
+	void PageLayout_OnGridCellHovered(const FIntPoint& Tile);
+
+	/** Cell unhover handler. */
+	void PageLayout_OnGridCellUnhovered();
+
+	/** Refresh current cell states. */
+	void PageLayout_RefreshGridLayout();
+
+	/**
+	 * Media objects page
+	 */
+private:
+
+	/** Builds media objects page. */
+	TSharedRef<SWidget> PageMedia_Build();
+
+	/** Builds media class picker widget. */
+	TSharedRef<SWidget> PageMedia_BuildClassPicker(UClass* FilterClass);
+
+	/** Callback to check if we can switch to the media objects configuration page. */
+	bool PageMedia_OnCanShow() const;
+
+	/** Callback on media objects page enter. */
+	void PageMedia_OnEnter();
+
+	/** Returns true if media objects configuration is valid. */
+	bool PageMedia_IsConfigurationValid() const;
+
+	/** Media source change handler. */
+	void PageMedia_OnMediaSourceChanged(UClass* InNewClass);
+
+	/** Media output change handler. */
+	void PageMedia_OnMediaOutputChanged(UClass* InNewClass);
+
+	/** Updates configuration message/feedback. */
+	void PageMedia_UpdateMessage();
+
+	/**
+	 * Nodes selection page
+	 */
+private:
+
+	/** Builds nodes selection page. */
+	TSharedRef<SWidget> PageNodes_Build();
+
+	/** Builds intput/output section. */
+	TSharedRef<SWidget> PageNodes_BuildSelSection(const FText& Header, const TSet<FString>& AllowedItems, TSet<FString>& CurrentSelection);
+
+	/** Builds node picker widget for input/output mapping. */
+	TSharedRef<SWidget> PageNodes_BuildNodePicker(const TSet<FString>& AllowedItems, TSet<FString>& CurrentSelection);
+
+	/** Mouse click handler (for context menu). */
+	FReply PageNodes_OnMouseButtonDown(const FGeometry& Geometry, const FPointerEvent& PointerEvent, bool bOutputSelection);
+
+	/** Mouse click handler (for context menu). */
+	void PageNodes_CreateContextMenu(const FVector2D& CursorPosition, bool bOutputSelection);
+
+	/** Context menu handler: Select/Deselect All. */
+	void PageNodes_Menu_OnChangeSelectionAll(bool bOutputSelection, bool bSelected);
+
+	/** Context menu handler: Select all items in A that aren't selected in B (A-Bsel). */
+	void PageNodes_Menu_OnSelectAllExceptOfCounterpart(bool bOutputSelection);
+
+	/** Context menu handler: Deselect all items in A that are selected in B (Asel-Bsel). */
+	void PageNodes_Menu_OnDeselectAllFromCounterpart(bool bOutputSelection);
+
+	/** Context menu handler: Select all offscreen nodes. */
+	void PageNodes_Menu_OnSelectAllOffscreen(bool bOutputSelection);
+
+	/** Callback to check if we can switch to the nodes selection page. */
+	bool PageNodes_OnCanShow() const;
+
+	/** Callback on nodes selection page enter. */
+	void PageNodes_OnEnter();
+
+	/** Returns true if nodes selection is valid. */
+	bool PageNodes_IsConfigurationValid() const;
+
+	/**
+	 * Finalization page
+	 */
+private:
+
+	/** Builds finalization page. */
+	TSharedRef<SWidget> PageFinalization_Build();
+
+	/** Builds tile button content. */
+	TSharedRef<SWidget> PageFinalization_BuildOutputTileButtonContent(const FIntPoint& Tile);
+
+	/** Builds output mapping dropdown content. */
+	TSharedRef<SWidget> PageFinalization_BuildOutputTileButtonDropdownContent(const FIntPoint& Tile);
+
+	/** Callback to check if we can switch to the media groups configuration page. */
+	bool PageFinalization_OnCanShow() const;
+
+	/** Callback on media groups page enter. */
+	void PageFinalization_OnEnter();
+
+	/** Reset current output mapping. */
+	void PageFinalization_ResetOutputMapping();
+
+	/** Pre-configures output mapping. */
+	void PageFinalization_PresetupOutputs();
+
+	/** Pre-configures output mapping for local multicast (like SharedMemoryMedia). */
+	void PageFinalization_PresetupOutputs_LocalMulticast();
+
+	/** Pre-configures output mapping for multicast (like Rivermax). */
+	void PageFinalization_PresetupOutputs_Multicast();
+
+	/** Regenerates grid data to match current layout. */
+	void PageFinalization_RegenerateGrid();
+
+	/** Returns true if media groups configuration is valid. */
+	bool PageFinalization_IsConfigurationValid() const;
+
+	/** Updates cell content. */
+	void PageFinalization_UpdateTileButtonContent(const FIntPoint& Tile);
+
+	/** Returns cell color for its current state. */
+	FSlateColor PageFinalization_GetButtonColor(const FIntPoint& Tile);
+
+private:
+
+	/** Tile layout constraints. */
+	static constexpr int32 MaxTilesAmount = 4;
+
+	/** External configuration parameters. */
+	FMediaTilesConfigurationDialogParameters Parameters;
+
+	/** List of all cluster node IDs. */
+	TSet<FString> ClusterNodeIds;
+
+	/** Node-to-host map. */
+	TMap<FString, FString> NodeToHostMap;
+
+	/** Cluster nodes per host. */
+	TMap<FString, TSet<FString>> HostToNodesMap;
+
+	/** Whether user completed configuration. */
+	bool bConfigurationCompleted = false;
+
+	/** Page hint. */
+	TSharedPtr<STextBlock> PageHint;
+
+	/** Wizard widget. */
+	TSharedPtr<SWizard> Wizard;
+
+
+	/** [Layout] Button widgets representing tile layout. */
+	TMap<FIntPoint, TSharedPtr<SButton>> LayoutGridButtons;
+
+	/** [Layout] Layout grid cell size. */
+	const int32 LayoutGridCellSize = 80;
+
+	/** [Layout] Current 'hover' postion. */
+	FIntPoint Hovered = { -1, -1 };
+
+	/** [Layout] Last 'accepted' position. */
+	FIntPoint Accepted = { -1, -1 };
+
+
+	/** [Media] Status image for media source. */
+	TSharedPtr<SImage> MediaSourceStatusImage;
+
+	/** [Media] Status image for media output. */
+	TSharedPtr<SImage> MediaOutputStatusImage;
+
+	/** [Media] Status text for media source. */
+	TSharedPtr<STextBlock> MediaSourceStatusText;
+
+	/** [Media] Status text for media output. */
+	TSharedPtr<STextBlock> MediaOutputStatusText;
+
+	/** [Media] Currently selected media source name (button text). */
+	TSharedPtr<STextBlock> MediaSourceSelectedName;
+
+	/** [Media] Currently selected media output name (button text). */
+	TSharedPtr<STextBlock> MediaOutputSelectedName;
+
+	/** [Media] Media source template choosed by user. */
+	TStrongObjectPtr<UMediaSource> MediaSource;
+
+	/** [Media] Media output template choosed by user. */
+	TStrongObjectPtr<UMediaOutput> MediaOutput;
+
+	/** [Media] Stream propagation type (unicast, multicast, etc.). */
+	EMediaStreamPropagationType MediaPropagationTypes = static_cast<EMediaStreamPropagationType>(0);
+
+
+	/** [Nodes] Nodes allowed to be picked for input (allowed for receiving). */
+	TSet<FString> NodesAllowedForInput;
+
+	/** [Nodes] Nodes allowed to be picked for output (allowed for sending). */
+	TSet<FString> NodesAllowedForOutput;
+
+	/** [Nodes] Input nodes selection. */
+	TSet<FString> InputSelection;
+
+	/** [Nodes] Output nodes selection. */
+	TSet<FString> OutputSelection;
+
+
+	/** [Finalization] Output mapping widget. */
+	TSharedPtr<SUniformGridPanel> OutputMappingGrid;
+
+	/** [Finalization] Output mapping grid cell size. */
+	const int32 OutputMappingGridCellSize = 150;
+
+	/** [Finalization] Helper container for output tile mapping. */
+	struct FOutputMappingData
+	{
+		/** A drop-down button widget representing a tile (grid cell) */
+		TSharedPtr<SComboButton> ButtonWidget;
+
+		/** A drop-down content widget for the button above */
+		TSharedPtr<SVerticalBox> ButtonContentWidget;
+
+		/** Cluster nodes associated with this tile (grid cell) */
+		TSet<FString> ClusterNodes;
+	};
+
+	/** [Finalization] Tile output mapping. */
+	TMap<FIntPoint, FOutputMappingData> OutputMapping;
+};
diff --git a/Engine/Plugins/Runtime/nDisplay/Source/SharedMemoryMediaEditor/Private/ModularFeatures/SharedMemoryMediaInitializerFeature.cpp b/Engine/Plugins/Runtime/nDisplay/Source/SharedMemoryMediaEditor/Private/ModularFeatures/SharedMemoryMediaInitializerFeature.cpp
--- a/Engine/Plugins/Runtime/nDisplay/Source/SharedMemoryMediaEditor/Private/ModularFeatures/SharedMemoryMediaInitializerFeature.cpp
+++ b/Engine/Plugins/Runtime/nDisplay/Source/SharedMemoryMediaEditor/Private/ModularFeatures/SharedMemoryMediaInitializerFeature.cpp
@@ -6,27 +6,53 @@
 #include "SharedMemoryMediaSource.h"
 
 
-bool FSharedMemoryMediaInitializerFeature::IsMediaSubjectSupported(const UObject* MediaSubject)
+bool FSharedMemoryMediaInitializerFeature::IsMediaObjectSupported(const UObject* MediaObject)
 {
-	if (MediaSubject)
+	if (MediaObject)
 	{
-		return MediaSubject->IsA<USharedMemoryMediaSource>() || MediaSubject->IsA<USharedMemoryMediaOutput>();
+		return MediaObject->IsA<USharedMemoryMediaSource>() || MediaObject->IsA<USharedMemoryMediaOutput>();
 	}
 
 	return false;
 }
 
-static constexpr const TCHAR* GetMediaPrefix(const IDisplayClusterModularFeatureMediaInitializer::FMediaSubjectOwnerInfo::EMediaSubjectOwnerType OwnerType)
+bool FSharedMemoryMediaInitializerFeature::AreMediaObjectsCompatible(const UObject* MediaSource, const UObject* MediaOutput)
+{
+	if (MediaSource && MediaOutput)
+	{
+		return MediaSource->IsA<USharedMemoryMediaSource>() && MediaOutput->IsA<USharedMemoryMediaOutput>();
+	}
+
+	return false;
+}
+
+bool FSharedMemoryMediaInitializerFeature::GetSupportedMediaPropagationTypes(const UObject* MediaSource, const UObject* MediaOutput, EMediaStreamPropagationType& OutPropagationTypes)
+{
+	if (!IsMediaObjectSupported(MediaSource) ||
+		!IsMediaObjectSupported(MediaOutput) ||
+		!AreMediaObjectsCompatible(MediaSource, MediaOutput))
+	{
+		return false;
+	}
+
+	OutPropagationTypes =
+		EMediaStreamPropagationType::LocalUnicast |
+		EMediaStreamPropagationType::LocalMulticast;
+
+	return true;
+}
+
+static constexpr const TCHAR* GetMediaPrefix(const FMediaObjectOwnerInfo::EMediaObjectOwnerType OwnerType)
 {
 	switch (OwnerType)
 	{
-	case IDisplayClusterModularFeatureMediaInitializer::FMediaSubjectOwnerInfo::EMediaSubjectOwnerType::ICVFXCamera:
+	case FMediaObjectOwnerInfo::EMediaObjectOwnerType::ICVFXCamera:
 		return TEXT("icam");
 
-	case IDisplayClusterModularFeatureMediaInitializer::FMediaSubjectOwnerInfo::EMediaSubjectOwnerType::Viewport:
+	case FMediaObjectOwnerInfo::EMediaObjectOwnerType::Viewport:
 		return TEXT("vp");
 
-	case IDisplayClusterModularFeatureMediaInitializer::FMediaSubjectOwnerInfo::EMediaSubjectOwnerType::Backbuffer:
+	case FMediaObjectOwnerInfo::EMediaObjectOwnerType::Backbuffer:
 		return TEXT("node");
 
 	default:
@@ -34,17 +60,17 @@ static constexpr const TCHAR* GetMediaPrefix(const IDisplayClusterModularFeature
 	}
 }
 
-void FSharedMemoryMediaInitializerFeature::InitializeMediaSubjectForTile(UObject* MediaSubject, const FMediaSubjectOwnerInfo& OnwerInfo, const FIntPoint& TilePos)
+void FSharedMemoryMediaInitializerFeature::InitializeMediaObjectForTile(UObject* MediaObject, const FMediaObjectOwnerInfo& OnwerInfo, const FIntPoint& TilePos)
 {
 	const FString UniqueName = FString::Printf(TEXT("%s@%s_tile_%d:%d"), GetMediaPrefix(OnwerInfo.OwnerType), *OnwerInfo.OwnerName, TilePos.X, TilePos.Y);
 
-	if (USharedMemoryMediaSource* SMMediaSource = Cast<USharedMemoryMediaSource>(MediaSubject))
+	if (USharedMemoryMediaSource* SMMediaSource = Cast<USharedMemoryMediaSource>(MediaObject))
 	{
 		SMMediaSource->UniqueName   = UniqueName;
 		SMMediaSource->bZeroLatency = true;
 		SMMediaSource->Mode         = ESharedMemoryMediaSourceMode::Framelocked;
 	}
-	else if (USharedMemoryMediaOutput* SMMediaOutput = Cast<USharedMemoryMediaOutput>(MediaSubject))
+	else if (USharedMemoryMediaOutput* SMMediaOutput = Cast<USharedMemoryMediaOutput>(MediaObject))
 	{
 		SMMediaOutput->UniqueName   = UniqueName;
 		SMMediaOutput->bInvertAlpha = true;
@@ -53,17 +79,17 @@ void FSharedMemoryMediaInitializerFeature::InitializeMediaSubjectForTile(UObject
 	}
 }
 
-void FSharedMemoryMediaInitializerFeature::InitializeMediaSubjectForFullFrame(UObject* MediaSubject, const FMediaSubjectOwnerInfo& OnwerInfo)
+void FSharedMemoryMediaInitializerFeature::InitializeMediaObjectForFullFrame(UObject* MediaObject, const FMediaObjectOwnerInfo& OnwerInfo)
 {
 	const FString UniqueName = FString::Printf(TEXT("%s@%s"), GetMediaPrefix(OnwerInfo.OwnerType), *OnwerInfo.OwnerName);
 
-	if (USharedMemoryMediaSource* SMMediaSource = Cast<USharedMemoryMediaSource>(MediaSubject))
+	if (USharedMemoryMediaSource* SMMediaSource = Cast<USharedMemoryMediaSource>(MediaObject))
 	{
 		SMMediaSource->UniqueName   = UniqueName;
 		SMMediaSource->bZeroLatency = true;
 		SMMediaSource->Mode         = ESharedMemoryMediaSourceMode::Framelocked;
 	}
-	else if (USharedMemoryMediaOutput* SMMediaOutput = Cast<USharedMemoryMediaOutput>(MediaSubject))
+	else if (USharedMemoryMediaOutput* SMMediaOutput = Cast<USharedMemoryMediaOutput>(MediaObject))
 	{
 		SMMediaOutput->UniqueName   = UniqueName;
 		SMMediaOutput->bInvertAlpha = true;
diff --git a/Engine/Plugins/Runtime/nDisplay/Source/SharedMemoryMediaEditor/Private/ModularFeatures/SharedMemoryMediaInitializerFeature.h b/Engine/Plugins/Runtime/nDisplay/Source/SharedMemoryMediaEditor/Private/ModularFeatures/SharedMemoryMediaInitializerFeature.h
--- a/Engine/Plugins/Runtime/nDisplay/Source/SharedMemoryMediaEditor/Private/ModularFeatures/SharedMemoryMediaInitializerFeature.h
+++ b/Engine/Plugins/Runtime/nDisplay/Source/SharedMemoryMediaEditor/Private/ModularFeatures/SharedMemoryMediaInitializerFeature.h
@@ -15,8 +15,10 @@ class FSharedMemoryMediaInitializerFeature
 public:
 
 	//~ Begin IDisplayClusterModularFeatureMediaInitializer
-	virtual bool IsMediaSubjectSupported(const UObject* MediaSubject) override;
-	virtual void InitializeMediaSubjectForTile(UObject* MediaSubject, const FMediaSubjectOwnerInfo& OnwerInfo, const FIntPoint& TilePos) override;
-	virtual void InitializeMediaSubjectForFullFrame(UObject* MediaSubject, const FMediaSubjectOwnerInfo& OnwerInfo) override;
+	virtual bool IsMediaObjectSupported(const UObject* MediaObject) override;
+	virtual bool AreMediaObjectsCompatible(const UObject* MediaSource, const UObject* MediaOutput) override;
+	virtual bool GetSupportedMediaPropagationTypes(const UObject* MediaSource, const UObject* MediaOutput, EMediaStreamPropagationType& OutPropagationTypes) override;
+	virtual void InitializeMediaObjectForTile(UObject* MediaObject, const FMediaObjectOwnerInfo& OnwerInfo, const FIntPoint& TilePos) override;
+	virtual void InitializeMediaObjectForFullFrame(UObject* MediaObject, const FMediaObjectOwnerInfo& OnwerInfo) override;
 	//~ End IDisplayClusterModularFeatureMediaInitializer
 };
diff --git a/Engine/Plugins/Runtime/nDisplayModularFeatures/Source/DisplayClusterModularFeaturesEditor/Public/IDisplayClusterModularFeatureMediaInitializer.h b/Engine/Plugins/Runtime/nDisplayModularFeatures/Source/DisplayClusterModularFeaturesEditor/Public/IDisplayClusterModularFeatureMediaInitializer.h
--- a/Engine/Plugins/Runtime/nDisplayModularFeatures/Source/DisplayClusterModularFeaturesEditor/Public/IDisplayClusterModularFeatureMediaInitializer.h
+++ b/Engine/Plugins/Runtime/nDisplayModularFeatures/Source/DisplayClusterModularFeaturesEditor/Public/IDisplayClusterModularFeatureMediaInitializer.h
@@ -5,6 +5,53 @@
 #include "CoreMinimal.h"
 #include "Features/IModularFeature.h"
 
+/**
+ * Container to carry the info about media object's owner
+ */
+struct FMediaObjectOwnerInfo
+{
+	/**
+	 * Type of the media object's owner
+	 */
+	enum class EMediaObjectOwnerType : uint8
+	{
+		ICVFXCamera = 0,
+		Viewport,
+		Backbuffer
+	};
+
+
+	/** Owner name (ICVFX camera component name, viewport or node name) */
+	FString OwnerName;
+
+	/** Owner type (ICVFX camera component name, viewport or node name) */
+	EMediaObjectOwnerType OwnerType;
+
+	/** Optional unique index of the cluster node holding the owner object */
+	TOptional<uint8> ClusterNodeUniqueIdx;
+
+	/**
+	 * Unique index of the owner
+	 *   Camera     - within a config
+	 *   Viewport   - within a cluster node
+	 *   Backbuffer - within a config
+	 */
+	uint8 OwnerUniqueIdx = 0;
+};
+
+/**
+ * Media stream propagation types
+ */
+enum class EMediaStreamPropagationType : uint8
+{
+	None = 0,
+	LocalUnicast      = 1 << 0, // Single sender and single receiver within the same host
+	LocalMulticast    = 1 << 1, // Single sender and multiple receivers within the same host
+	Unicast           = 1 << 2, // Single sender and single receiver on different hosts
+	Multicast         = 1 << 3  // Single sender and multiple receivers on different hosts
+};
+ENUM_CLASS_FLAGS(EMediaStreamPropagationType);
+
 
 /**
  * Base class for nDisplay media initializer implementations.
@@ -18,68 +65,51 @@ public:
 
 public:
 
-	/**
-	 * Container to carrier the info about media subject's owner
-	 */
-	struct FMediaSubjectOwnerInfo
-	{
-		/**
-		 * Type of the media subject's owner
-		 */
-		enum class EMediaSubjectOwnerType : uint8
-		{
-			ICVFXCamera = 0,
-			Viewport,
-			Backbuffer
-		};
-
-
-		/** Owner name (ICVFX camera component name, viewport or node name) */
-		FString OwnerName;
-
-		/** Owner type (ICVFX camera component name, viewport or node name) */
-		EMediaSubjectOwnerType OwnerType;
-
-		/** Optional unique index of the cluster node holding the owner object */
-		TOptional<uint8> ClusterNodeUniqueIdx = 0;
-
-		/** 
-		 * Unique index of the owner
-		 *   Camera     - within a config
-		 *   Viewport   - within a cluster node
-		 *   Backbuffer - within a config
-		 */
-		uint8 OwnerUniqueIdx = 0;
-	};
+	virtual ~IDisplayClusterModularFeatureMediaInitializer() = default;
 
 public:
 
-	virtual ~IDisplayClusterModularFeatureMediaInitializer() = default;
+	/**
+	 * Checks if media object is supported by the initializer
+	 *
+	 * @param MediaObject - UMediaSource or UMediaOutput instance
+	 * @return true if media object supported
+	 */
+	virtual bool IsMediaObjectSupported(const UObject* MediaObject) = 0;
 
-public:
+	/**
+	 * Checks if media source and output are compatible and can be paired
+	 *
+	 * @param MediaSource - UMediaSource instance
+	 * @param MediaOutput - UMediaOutput instance
+	 * @return true if media objects are compatible with each other
+	 */
+	virtual bool AreMediaObjectsCompatible(const UObject* MediaSource, const UObject* MediaOutput) = 0;
 
 	/**
-	 * Checks if media subject is supported by the initializer
+	 * Provides stream supported media propagation types (local/global, unicast/multicast, etc.)
 	 *
-	 * @param MediaSubject - UMediaSource or UMediaOutput instance
-	 * @return true if media subject supported
+	 * @param MediaSource - UMediaSource instance
+	 * @param MediaOutput - UMediaOutput instance
+	 * @param OutPropagationTypes - [out] supported types mask
+	 * @return false if media objects are invalid, incompatible or not supported
 	 */
-	virtual bool IsMediaSubjectSupported(const UObject* MediaSubject) = 0;
+	virtual bool GetSupportedMediaPropagationTypes(const UObject* MediaSource, const UObject* MediaOutput, EMediaStreamPropagationType& OutPropagationTypes) = 0;
 
 	/**
-	 * Performs initialization of a media subject for tiled input/output
+	 * Performs initialization of a media object for tiled input/output
 	 *
-	 * @param MediaSubject - UMediaSource or UMediaOutput instance
-	 * @param OwnerInfo    - Additional information about the object holding the media subject
-	 * @param TilePos      - TIle XY-position
+	 * @param MediaObject - UMediaSource or UMediaOutput instance
+	 * @param OwnerInfo   - Additional information about the object holding the media object
+	 * @param TilePos     - Tile XY-position
 	 */
-	virtual void InitializeMediaSubjectForTile(UObject* MediaSubject, const FMediaSubjectOwnerInfo& OnwerInfo, const FIntPoint& TilePos) = 0;
+	virtual void InitializeMediaObjectForTile(UObject* MediaObject, const FMediaObjectOwnerInfo& OnwerInfo, const FIntPoint& TilePos) = 0;
 
 	/**
-	 * Performs initialization of a media subject for full frame input/output
+	 * Performs initialization of a media object for full frame input/output
 	 *
-	 * @param MediaSubject - UMediaSource or UMediaOutput instance
-	 * @param OwnerInfo    - Additional information about the object holding the media subject
+	 * @param MediaObject - UMediaSource or UMediaOutput instance
+	 * @param OwnerInfo   - Additional information about the object holding the media object
 	 */
-	virtual void InitializeMediaSubjectForFullFrame(UObject* MediaSubject, const FMediaSubjectOwnerInfo& OnwerInfo) = 0;
+	virtual void InitializeMediaObjectForFullFrame(UObject* MediaObject, const FMediaObjectOwnerInfo& OnwerInfo) = 0;
 };
diff --git a/Engine/Plugins/VirtualProduction/Rivermax/RivermaxMedia/Source/RivermaxMediaEditor/Private/ModularFeatures/RivermaxMediaInitializerFeature.cpp b/Engine/Plugins/VirtualProduction/Rivermax/RivermaxMedia/Source/RivermaxMediaEditor/Private/ModularFeatures/RivermaxMediaInitializerFeature.cpp
--- a/Engine/Plugins/VirtualProduction/Rivermax/RivermaxMedia/Source/RivermaxMediaEditor/Private/ModularFeatures/RivermaxMediaInitializerFeature.cpp
+++ b/Engine/Plugins/VirtualProduction/Rivermax/RivermaxMedia/Source/RivermaxMediaEditor/Private/ModularFeatures/RivermaxMediaInitializerFeature.cpp
@@ -11,19 +11,47 @@
 #include "RivermaxMediaOutput.h"
 
 
-bool FRivermaxMediaInitializerFeature::IsMediaSubjectSupported(const UObject* MediaSubject)
+bool FRivermaxMediaInitializerFeature::IsMediaObjectSupported(const UObject* MediaObject)
 {
-	if (MediaSubject)
+	if (MediaObject)
 	{
-		return MediaSubject->IsA<URivermaxMediaSource>() || MediaSubject->IsA<URivermaxMediaOutput>();
+		return MediaObject->IsA<URivermaxMediaSource>() || MediaObject->IsA<URivermaxMediaOutput>();
 	}
 
 	return false;
 }
 
-void FRivermaxMediaInitializerFeature::InitializeMediaSubjectForTile(UObject* MediaSubject, const FMediaSubjectOwnerInfo& OnwerInfo, const FIntPoint& TilePos)
+bool FRivermaxMediaInitializerFeature::AreMediaObjectsCompatible(const UObject* MediaSource, const UObject* MediaOutput)
 {
-	if (URivermaxMediaSource* RivermaxMediaSource = Cast<URivermaxMediaSource>(MediaSubject))
+	if (MediaSource && MediaOutput)
+	{
+		return MediaSource->IsA<URivermaxMediaSource>() && MediaOutput->IsA<URivermaxMediaOutput>();
+	}
+
+	return false;
+}
+
+bool FRivermaxMediaInitializerFeature::GetSupportedMediaPropagationTypes(const UObject* MediaSource, const UObject* MediaOutput, EMediaStreamPropagationType& OutPropagationTypes)
+{
+	if (!IsMediaObjectSupported(MediaSource) ||
+		!IsMediaObjectSupported(MediaOutput) ||
+		!AreMediaObjectsCompatible(MediaSource, MediaOutput))
+	{
+		return false;
+	}
+
+	OutPropagationTypes =
+		EMediaStreamPropagationType::LocalUnicast |
+		EMediaStreamPropagationType::LocalMulticast |
+		EMediaStreamPropagationType::Unicast |
+		EMediaStreamPropagationType::Multicast;
+
+	return true;
+}
+
+void FRivermaxMediaInitializerFeature::InitializeMediaObjectForTile(UObject* MediaObject, const FMediaObjectOwnerInfo& OnwerInfo, const FIntPoint& TilePos)
+{
+	if (URivermaxMediaSource* RivermaxMediaSource = Cast<URivermaxMediaSource>(MediaObject))
 	{
 		RivermaxMediaSource->PlayerMode          = ERivermaxPlayerMode::Framelock;
 		RivermaxMediaSource->bUseZeroLatency     = true;
@@ -37,7 +65,7 @@ void FRivermaxMediaInitializerFeature::InitializeMediaSubjectForTile(UObject* Me
 		RivermaxMediaSource->bIsSRGBInput        = false;
 		RivermaxMediaSource->bUseGPUDirect       = true;
 	}
-	else if (URivermaxMediaOutput* RivermaxMediaOutput = Cast<URivermaxMediaOutput>(MediaSubject))
+	else if (URivermaxMediaOutput* RivermaxMediaOutput = Cast<URivermaxMediaOutput>(MediaObject))
 	{
 		RivermaxMediaOutput->AlignmentMode       = ERivermaxMediaAlignmentMode::FrameCreation;
 		RivermaxMediaOutput->bDoContinuousOutput = false;
@@ -55,9 +83,9 @@ void FRivermaxMediaInitializerFeature::InitializeMediaSubjectForTile(UObject* Me
 	}
 }
 
-void FRivermaxMediaInitializerFeature::InitializeMediaSubjectForFullFrame(UObject* MediaSubject, const FMediaSubjectOwnerInfo& OnwerInfo)
+void FRivermaxMediaInitializerFeature::InitializeMediaObjectForFullFrame(UObject* MediaObject, const FMediaObjectOwnerInfo& OnwerInfo)
 {
-	if (URivermaxMediaSource* RivermaxMediaSource = Cast<URivermaxMediaSource>(MediaSubject))
+	if (URivermaxMediaSource* RivermaxMediaSource = Cast<URivermaxMediaSource>(MediaObject))
 	{
 		RivermaxMediaSource->PlayerMode          = ERivermaxPlayerMode::Framelock;
 		RivermaxMediaSource->bUseZeroLatency     = true;
@@ -71,7 +99,7 @@ void FRivermaxMediaInitializerFeature::InitializeMediaSubjectForFullFrame(UObjec
 		RivermaxMediaSource->bIsSRGBInput        = false;
 		RivermaxMediaSource->bUseGPUDirect       = true;
 	}
-	else if (URivermaxMediaOutput* RivermaxMediaOutput = Cast<URivermaxMediaOutput>(MediaSubject))
+	else if (URivermaxMediaOutput* RivermaxMediaOutput = Cast<URivermaxMediaOutput>(MediaObject))
 	{
 		RivermaxMediaOutput->FrameLockingMode      = ERivermaxFrameLockingMode::BlockOnReservation;
 		RivermaxMediaOutput->PresentationQueueSize = 2;
@@ -84,8 +112,8 @@ void FRivermaxMediaInitializerFeature::InitializeMediaSubjectForFullFrame(UObjec
 
 		switch (OnwerInfo.OwnerType)
 		{
-		case FMediaSubjectOwnerInfo::EMediaSubjectOwnerType::ICVFXCamera:
-		case FMediaSubjectOwnerInfo::EMediaSubjectOwnerType::Viewport:
+		case FMediaObjectOwnerInfo::EMediaObjectOwnerType::ICVFXCamera:
+		case FMediaObjectOwnerInfo::EMediaObjectOwnerType::Viewport:
 			RivermaxMediaOutput->AlignmentMode       = ERivermaxMediaAlignmentMode::FrameCreation;
 			RivermaxMediaOutput->bDoContinuousOutput = false;
 			RivermaxMediaOutput->bDoFrameCounterTimestamping = true;
@@ -93,7 +121,7 @@ void FRivermaxMediaInitializerFeature::InitializeMediaSubjectForFullFrame(UObjec
 			RivermaxMediaOutput->bUseGPUDirect       = true;
 			break;
 
-		case FMediaSubjectOwnerInfo::EMediaSubjectOwnerType::Backbuffer:
+		case FMediaObjectOwnerInfo::EMediaObjectOwnerType::Backbuffer:
 			RivermaxMediaOutput->AlignmentMode       = ERivermaxMediaAlignmentMode::AlignmentPoint;
 			RivermaxMediaOutput->bDoContinuousOutput = true;
 			RivermaxMediaOutput->bDoFrameCounterTimestamping = false;
@@ -142,7 +170,7 @@ FString FRivermaxMediaInitializerFeature::GenerateStreamAddress(uint8 OwnerUniqu
 	return FString::Printf(TEXT("228.%u.%u.%u"), AddressOffsetForTiles + OwnerUniqueIdx, TilePos.X, TilePos.Y);
 }
 
-FString FRivermaxMediaInitializerFeature::GenerateStreamAddress(uint8 ClusterNodeUniqueIdx, uint8 OwnerUniqueIdx, const FMediaSubjectOwnerInfo::EMediaSubjectOwnerType OwnerType) const
+FString FRivermaxMediaInitializerFeature::GenerateStreamAddress(uint8 ClusterNodeUniqueIdx, uint8 OwnerUniqueIdx, const FMediaObjectOwnerInfo::EMediaObjectOwnerType OwnerType) const
 {
 	// 228.0.*.* - 228.199.*.* - range for full-frame media (max 200 objects). But could be extended up to the limit if no tiles used.
 	return FString::Printf(TEXT("228.%u.%u.%u"), ClusterNodeUniqueIdx, static_cast<uint8>(OwnerType), OwnerUniqueIdx);
diff --git a/Engine/Plugins/VirtualProduction/Rivermax/RivermaxMedia/Source/RivermaxMediaEditor/Private/ModularFeatures/RivermaxMediaInitializerFeature.h b/Engine/Plugins/VirtualProduction/Rivermax/RivermaxMedia/Source/RivermaxMediaEditor/Private/ModularFeatures/RivermaxMediaInitializerFeature.h
--- a/Engine/Plugins/VirtualProduction/Rivermax/RivermaxMedia/Source/RivermaxMediaEditor/Private/ModularFeatures/RivermaxMediaInitializerFeature.h
+++ b/Engine/Plugins/VirtualProduction/Rivermax/RivermaxMedia/Source/RivermaxMediaEditor/Private/ModularFeatures/RivermaxMediaInitializerFeature.h
@@ -15,9 +15,11 @@ class FRivermaxMediaInitializerFeature
 public:
 
 	//~ Begin IDisplayClusterModularFeatureMediaInitializer
-	virtual bool IsMediaSubjectSupported(const UObject* MediaSubject) override;
-	virtual void InitializeMediaSubjectForTile(UObject* MediaSubject, const FMediaSubjectOwnerInfo& OnwerInfo, const FIntPoint& TilePos) override;
-	virtual void InitializeMediaSubjectForFullFrame(UObject* MediaSubject, const FMediaSubjectOwnerInfo& OnwerInfo) override;
+	virtual bool IsMediaObjectSupported(const UObject* MediaObject) override;
+	virtual bool AreMediaObjectsCompatible(const UObject* MediaSource, const UObject* MediaOutput) override;
+	virtual bool GetSupportedMediaPropagationTypes(const UObject* MediaSource, const UObject* MediaOutput, EMediaStreamPropagationType& OutPropagationTypes) override;
+	virtual void InitializeMediaObjectForTile(UObject* MediaObject, const FMediaObjectOwnerInfo& OnwerInfo, const FIntPoint& TilePos) override;
+	virtual void InitializeMediaObjectForFullFrame(UObject* MediaObject, const FMediaObjectOwnerInfo& OnwerInfo) override;
 	//~ End IDisplayClusterModularFeatureMediaInitializer
 
 private:
@@ -29,5 +31,5 @@ private:
 	FString GenerateStreamAddress(uint8 OwnerUniqueIdx, const FIntPoint& TilePos) const;
 
 	/** Generates stream address based on the function parameters */
-	FString GenerateStreamAddress(uint8 ClusterNodeUniqueIdx, uint8 OwnerUniqueIdx, const FMediaSubjectOwnerInfo::EMediaSubjectOwnerType OwnerType) const;
+	FString GenerateStreamAddress(uint8 ClusterNodeUniqueIdx, uint8 OwnerUniqueIdx, const FMediaObjectOwnerInfo::EMediaObjectOwnerType OwnerType) const;
 };
