[nDisplay] Fixed CIS issue after tile configuration dialog submission

diff --git a/Engine/Plugins/Runtime/nDisplay/Source/DisplayClusterConfigurator/Private/Views/Details/Media/SMediaTilesConfigurationDialog.h b/Engine/Plugins/Runtime/nDisplay/Source/DisplayClusterConfigurator/Private/Views/Details/Media/SMediaTilesConfigurationDialog.h
--- a/Engine/Plugins/Runtime/nDisplay/Source/DisplayClusterConfigurator/Private/Views/Details/Media/SMediaTilesConfigurationDialog.h
+++ b/Engine/Plugins/Runtime/nDisplay/Source/DisplayClusterConfigurator/Private/Views/Details/Media/SMediaTilesConfigurationDialog.h
@@ -6,7 +6,7 @@
 #include "Widgets/DeclarativeSyntaxSupport.h"
 #include "Styling/SlateColor.h"
 
-#include "DisplayClusterconfigurationTypes.h"
+#include "DisplayClusterConfigurationTypes.h"
 
 
 class SButton;
