[nDisplay] Fixed CIS issue, LOCTEXT conflict

diff --git a/Engine/Plugins/Runtime/nDisplay/Source/DisplayClusterConfigurator/Private/Views/Details/Media/SMediaTilesConfigurationDialog.cpp b/Engine/Plugins/Runtime/nDisplay/Source/DisplayClusterConfigurator/Private/Views/Details/Media/SMediaTilesConfigurationDialog.cpp
--- a/Engine/Plugins/Runtime/nDisplay/Source/DisplayClusterConfigurator/Private/Views/Details/Media/SMediaTilesConfigurationDialog.cpp
+++ b/Engine/Plugins/Runtime/nDisplay/Source/DisplayClusterConfigurator/Private/Views/Details/Media/SMediaTilesConfigurationDialog.cpp
@@ -52,7 +52,7 @@ namespace UE::DisplayClusterConfigurator::MediaTilesConfigurationDialog::Private
 	const FText TextPageMediaComboboxNone  = LOCTEXT("PageMediaMediaComboboxItemNone", "None");
 	const FText TextPageMediaStatusOk      = LOCTEXT("PageMediaStatusOk", "Ok");
 	const FText TextPageMediaStatusNotSupp = LOCTEXT("PageMediaStatusNotSupported", "Auto-configuration is not supported");
-	const FText TextPageMediaNotCompatOrNotSupp = LOCTEXT("PageMediaStatusNotSupported", "Not compatible or not supported");
+	const FText TextPageMediaNotCompatOrNotSupp = LOCTEXT("PageMediaStatusNotCompatible", "Not compatible or not supported");
 
 	// Page: Nodes
 	const FText TextPageNodesHeader = LOCTEXT("PageNodesHeader", "Step 3: Allot cluster nodes");
