nDisplay: Fixed FrustumFit math
- reworked symmetric frustum math
- code cleanup

diff --git a/Engine/Plugins/Runtime/nDisplay/Source/DisplayCluster/Private/Render/Viewport/DisplayClusterViewport_ViewPoint.cpp b/Engine/Plugins/Runtime/nDisplay/Source/DisplayCluster/Private/Render/Viewport/DisplayClusterViewport_ViewPoint.cpp
--- a/Engine/Plugins/Runtime/nDisplay/Source/DisplayCluster/Private/Render/Viewport/DisplayClusterViewport_ViewPoint.cpp
+++ b/Engine/Plugins/Runtime/nDisplay/Source/DisplayCluster/Private/Render/Viewport/DisplayClusterViewport_ViewPoint.cpp
@@ -201,7 +201,9 @@ float FDisplayClusterViewport::GetStereoEyeOffsetDistance(const uint32 InContext
 		const float EyeOffsetValues[] = { -EyeOffset, 0.f, EyeOffset };
 
 		// Decode current eye type
-		const EDisplayClusterEyeType EyeType = (Contexts.Num() == 1)
+		// This function should work correctly even if the viewport context data is not currently initialized.
+		const int32 ViewPerViewportAmount = Configuration->GetRenderFrameSettings().GetViewPerViewportAmount();
+		const EDisplayClusterEyeType EyeType = (ViewPerViewportAmount < 2)
 			? EDisplayClusterEyeType::Mono
 			: (InContextNum == 0) ? EDisplayClusterEyeType::StereoLeft : EDisplayClusterEyeType::StereoRight;
 
diff --git a/Engine/Plugins/Runtime/nDisplay/Source/DisplayClusterWarp/Private/Policy/InFrustumFit/DisplayClusterWarpInFrustumFitPolicy.cpp b/Engine/Plugins/Runtime/nDisplay/Source/DisplayClusterWarp/Private/Policy/InFrustumFit/DisplayClusterWarpInFrustumFitPolicy.cpp
--- a/Engine/Plugins/Runtime/nDisplay/Source/DisplayClusterWarp/Private/Policy/InFrustumFit/DisplayClusterWarpInFrustumFitPolicy.cpp
+++ b/Engine/Plugins/Runtime/nDisplay/Source/DisplayClusterWarp/Private/Policy/InFrustumFit/DisplayClusterWarpInFrustumFitPolicy.cpp
@@ -17,24 +17,6 @@
 #include "Components/MeshComponent.h"
 #include "Materials/MaterialInstanceDynamic.h"
 
-// Experimental: Enable static view direction for mpcdi 2d
-int32 GDisplayClusterWarpInFrustumFitPolicyUseStaticViewDirectionForMPCDIProfile2D = 1;
-static FAutoConsoleVariableRef CVarDisplayClusterWarpInFrustumFitPolicyUseStaticViewDirectionForMPCDIProfile2D(
-	TEXT("nDisplay.warp.InFrustumFit.UseStaticViewDirectionForMPCDIProfile2D"),
-	GDisplayClusterWarpInFrustumFitPolicyUseStaticViewDirectionForMPCDIProfile2D,
-	TEXT("Experimental: Enable static view direction for mpcdi 2d (0 - disable)\n"),
-	ECVF_Default
-);
-
-// Experimental: Enable single view target for group of viewports
-int32 GDisplayClusterWarpInFrustumFitPolicyUseGroupViewTarget = 1;
-static FAutoConsoleVariableRef CVarDisplayClusterWarpInFrustumFitPolicyUseGroupViewTarget(
-	TEXT("nDisplay.warp.InFrustumFit.UseGroupViewTarget"),
-	GDisplayClusterWarpInFrustumFitPolicyUseGroupViewTarget,
-	TEXT("Experimental: Enable single view target for group of viewports (0 - disable)\n"),
-	ECVF_Default
-);
-
 int32 GDisplayClusterWarpInFrustumFitPolicyDrawFrustum = 0;
 static FAutoConsoleVariableRef CVarDisplayClusterWarpInFrustumFitPolicyDrawFrustum(
 	TEXT("nDisplay.warp.InFrustumFit.DrawFrustum"),
@@ -59,102 +41,140 @@ const FString& FDisplayClusterWarpInFrustumFitPolicy::GetType() const
 
 void FDisplayClusterWarpInFrustumFitPolicy::HandleNewFrame(const TArray<TSharedPtr<IDisplayClusterViewport, ESPMode::ThreadSafe>>& InViewports)
 {
-	// Todo: get real scale
-	const float WorldToMeters = 100.0f;
-
-	const float WorldScale = WorldToMeters / 100.f;
-	
-	bGeometryContextsUpdated = false;
-	if (GDisplayClusterWarpInFrustumFitPolicyUseGroupViewTarget)
+	// Reset all special data that was used in the prev frame.
+	// In the new frame, we have to find a new solution because the viewer's position or geometry may have changed.
+	// The number of viewports can also be changed at runtime. This changes the shape of the united geometry.
 	{
-		// Calculate AABB for group of viewports:
-		GroupAABBox = FDisplayClusterWarpAABB();
-		for (const TSharedPtr<IDisplayClusterViewport, ESPMode::ThreadSafe>& Viewport : InViewports)
-		{
-			if (Viewport.IsValid() && Viewport->GetProjectionPolicy().IsValid())
-			{
-				TSharedPtr<IDisplayClusterWarpBlend, ESPMode::ThreadSafe> WarpBlend;
-				if (Viewport->GetProjectionPolicy()->GetWarpBlendInterface(WarpBlend))
-				{
-					if (WarpBlend->UpdateGeometryContext(WorldScale))
-					{
-						GroupAABBox.UpdateAABB(WarpBlend->GetGeometryContext().AABBox);
-					}
-				}
-			}
-		}
+		// Invalidate geometry cache for the new frame.
+		bGeometryContextsUpdated = false;
 
-		bGeometryContextsUpdated = true;
+		// Reset united geometry AABB
+		OptUnitedGeometryWorldAABB.Reset();
+
+		// Reset the solution from the previous frame.
+		OptUnitedGeometryWarpProjection.Reset();
+		OptOverrideWorldViewTarget.Reset();
 	}
 
-	const uint32 ContextNum = 0; // calculate for a single context
+	if (InViewports.IsEmpty() || !InViewports[0].IsValid())
+	{
+		return;
+	}
 
-	bValidGroupFrustum = false;
+	UDisplayClusterInFrustumFitCameraComponent* ConfigurationCameraComponent = Cast<UDisplayClusterInFrustumFitCameraComponent>(InViewports[0]->GetViewPointCameraComponent(EDisplayClusterRootActorType::Configuration));
+	UDisplayClusterInFrustumFitCameraComponent* SceneCameraComponent = Cast<UDisplayClusterInFrustumFitCameraComponent>(InViewports[0]->GetViewPointCameraComponent(EDisplayClusterRootActorType::Scene));
+	if(!ConfigurationCameraComponent || !SceneCameraComponent)
+	{
+		return;
+	}
 
-	// Recalculate warp projection angles
-	GroupGeometryWarpProjection.ResetProjectionAngles();
-	SymmetricForwardCorrection.Reset();
+	// calculate GroupFrustum for a single context
+	const uint32 ContextNum = 0;
 
-	bool bCanCalcFrustumContext = true;
-	bool bHasFixedViewDirection = false;
+	const float WorldToMeters = InViewports[0]->GetConfiguration().GetWorldToMeters();
+	const float WorldScale = WorldToMeters / 100.f;
 
-	// Use center of GroupAABB as target viewpoint, and setup to use same viewprojection plane for all viewports:
-	for (const TSharedPtr<IDisplayClusterViewport, ESPMode::ThreadSafe>& Viewport : InViewports)
+	switch (ConfigurationCameraComponent->CameraViewTarget)
+	{
+	case EDisplayClusterWarpCameraViewTarget::GeometricCenter:
 	{
-		if (bCanCalcFrustumContext && Viewport.IsValid() && Viewport->GetProjectionPolicy().IsValid())
+		// In this case, we need to find a symmetric frustum
+		// 
+		// The united geometry frustum is built from geometric points projected onto a special plane.
+		// This plane is called the 'projection plane' and is created from two quantities: the view direction vector and the eye position.
+		// Therefore, when we change the view direction vector, it leads to a change in the "projection plane" and, then, to new frustum values.
+		// When we set out to create a symmetrical frustum, we need to solve this math problem.
+		// An easy way is to do this in a few iterations and stop when we find a suitable view direction that provides a nearly symmetrical frustum with acceptable precision.
+
+		// Set the center of the AABB of united geometry as the view target:
 		{
-			// Note: This code is partially copied from FDisplayClusterProjectionMPCDIPolicy::CalculateView().
-
-			// Override viewpoint
-			// MPCDI always expects the location of the viewpoint component (eye location from the real world)
-			FVector ViewOffset = FVector::ZeroVector;
-			FVector InOutViewLocation;
-			FRotator InOutViewRotation;
-			if (!Viewport->GetViewPointCameraEye(ContextNum, InOutViewLocation, InOutViewRotation, ViewOffset))
+			// Calculate AABB for group of viewports:
+			FDisplayClusterWarpAABB UnitedGeometryWorldAABBox;
+			if (!CalcUnitedGeometryWorldAABBox(InViewports, WorldScale, UnitedGeometryWorldAABBox))
 			{
-				continue;
+				return;
 			}
+			bGeometryContextsUpdated = true;
+			OptUnitedGeometryWorldAABB = UnitedGeometryWorldAABBox;
 
-			if (UDisplayClusterInFrustumFitCameraComponent* ConfigurationCameraComponent = Cast<UDisplayClusterInFrustumFitCameraComponent>(Viewport->GetViewPointCameraComponent(EDisplayClusterRootActorType::Configuration)))
-			{
-				bHasFixedViewDirection = ConfigurationCameraComponent->CameraViewTarget == EDisplayClusterWarpCameraViewTarget::MatchViewOrigin;
-			}
+			// in the first iteration, use the center of the united AABB geometry as the view target.
+			OptOverrideWorldViewTarget = UnitedGeometryWorldAABBox.GetCenter();
+		}
 
-			TSharedPtr<IDisplayClusterWarpBlend, ESPMode::ThreadSafe> WarpBlend;
-			if (Viewport->GetProjectionPolicy()->GetWarpBlendInterface(WarpBlend))
+		// Iterate frustum to a nearly symmetrical frustum with acceptable precision:
+		{
+			FSymmetricFrustumData FrustumData(WorldScale, SceneCameraComponent->GetComponentTransform());
+			while (true)
 			{
-				const USceneComponent* const OriginComp = Viewport->GetProjectionPolicy()->GetOriginComponent();
-
-				TSharedPtr<FDisplayClusterWarpEye, ESPMode::ThreadSafe> WarpEye = MakeShared<FDisplayClusterWarpEye, ESPMode::ThreadSafe>(Viewport, 0);
-
-				WarpEye->World2LocalTransform = (OriginComp ? OriginComp->GetComponentTransform() : FTransform::Identity);
-
-				// Get our base camera location and view offset in local space (MPCDI space)
-				WarpEye->ViewPoint.Location = WarpEye->World2LocalTransform.InverseTransformPosition(InOutViewLocation - ViewOffset);
-				WarpEye->ViewPoint.EyeOffset = WarpEye->World2LocalTransform.InverseTransformPosition(InOutViewLocation) - WarpEye->ViewPoint.Location;
-				WarpEye->ViewPoint.Rotation = WarpEye->World2LocalTransform.InverseTransformRotation(InOutViewRotation.Quaternion()).Rotator();
-
-				WarpEye->WorldScale = WorldScale;
-
-				WarpEye->WarpPolicy = SharedThis(this);
+				if (CalcUnitedGeometrySymmetricFrustum(InViewports, ContextNum, FrustumData))
+				{
+					// A symmetric frustum has been found, use it to fit
+					OptUnitedGeometryWarpProjection = *FrustumData.NewUnitedSymmetricWarpProjection;
+					break;
+				}
 
-				if (!WarpBlend->CalcFrustumContext(WarpEye))
+				if (FrustumData.IterationNum == INDEX_NONE)
 				{
-					bCanCalcFrustumContext = false;
+					// The maximum iteration number has been reached.
+					if (FrustumData.NewUnitedSymmetricWarpProjection.IsSet())
+					{
+						// Use last valid projection
+						OptUnitedGeometryWarpProjection = *FrustumData.NewUnitedSymmetricWarpProjection;
+						break;
+					}
+
+					// no valid projection found.
+					break;
 				}
-				else
+
+				if (!FrustumData.NewWorldViewTarget.IsSet())
 				{
-					const FDisplayClusterWarpData& WarpData = WarpBlend->GetWarpData(ContextNum);
-					GroupGeometryWarpProjection.ExpandProjectionAngles(WarpData.GeometryWarpProjection);
+					// no valid projection found.
+					break;
 				}
+
+
+				// Set the new location of the view target and perform the next iteration.
+				OptOverrideWorldViewTarget = *FrustumData.NewWorldViewTarget;
 			}
 		}
+
+		break;
 	}
 
-	bValidGroupFrustum = bCanCalcFrustumContext;
+	case EDisplayClusterWarpCameraViewTarget::MatchViewOrigin:
+	{
+		// In this case, the viewing direction is the X-axis of the ViewPoint component.
+		// This value is obtained in the FDisplayClusterWarpInFrustumFitPolicy::BeginCalcFrustum() function.
+		FDisplayClusterWarpProjection UnitedWarpProjection;
+		if (CalcUnitedGeometryFrustum(InViewports, ContextNum, WorldScale, UnitedWarpProjection))
+		{
+			bGeometryContextsUpdated = true;
+
+			// If the view target is set to a fixed value instead of being computed by the group AABB, we do not want to alter the view direction,
+			// but make the frustum symmetric around that fixed direction. This involves expanding the asymmetric frustum so that its left and right,
+			// top and bottom angles are equal
+
+			const double MaxHorizontal = FMath::Max(FMath::Abs(UnitedWarpProjection.Left), FMath::Abs(UnitedWarpProjection.Right));
+			const double MaxVertical = FMath::Max(FMath::Abs(UnitedWarpProjection.Top), FMath::Abs(UnitedWarpProjection.Bottom));
+
+			FDisplayClusterWarpProjection UnitedWarpSymmetricProjection = UnitedWarpProjection;
+
+			UnitedWarpSymmetricProjection.Left = -MaxHorizontal;
+			UnitedWarpSymmetricProjection.Right = MaxHorizontal;
 
-	// Recalculate the group frustum so that it is symmetric
-	MakeGroupFrustumSymmetrical(bHasFixedViewDirection);
+			UnitedWarpSymmetricProjection.Bottom = -MaxVertical;
+			UnitedWarpSymmetricProjection.Top = MaxVertical;
+
+			// use this united  frustum for fit:
+			OptUnitedGeometryWarpProjection = UnitedWarpSymmetricProjection;
+		}
+		break;
+	}
+
+	default:
+		break;
+	}
 }
 
 void FDisplayClusterWarpInFrustumFitPolicy::Tick(IDisplayClusterViewportManager* InViewportManager, float DeltaSeconds)
@@ -218,84 +238,6 @@ bool FDisplayClusterWarpInFrustumFitPolicy::HasPreviewEditableMesh(IDisplayClust
 	return false;
 }
 
-void FDisplayClusterWarpInFrustumFitPolicy::BeginCalcFrustum(IDisplayClusterViewport* InViewport, const uint32 ContextNum)
-{
-	if (InViewport && InViewport->GetProjectionPolicy().IsValid())
-	{
-		TSharedPtr<IDisplayClusterWarpBlend, ESPMode::ThreadSafe> WarpBlend;
-		if (InViewport->GetProjectionPolicy()->GetWarpBlendInterface(WarpBlend))
-		{
-			if (UDisplayClusterInFrustumFitCameraComponent* ConfigurationCameraComponent = Cast<UDisplayClusterInFrustumFitCameraComponent>(InViewport->GetViewPointCameraComponent(EDisplayClusterRootActorType::Configuration)))
-			{
-				FDisplayClusterWarpData& WarpData = WarpBlend->GetWarpData(ContextNum);
-				if (WarpData.WarpEye.IsValid())
-				{
-					// geometry context already updated.
-					WarpData.WarpEye->bUpdateGeometryContext = !bGeometryContextsUpdated;
-
-					if (ConfigurationCameraComponent->CameraViewTarget == EDisplayClusterWarpCameraViewTarget::MatchViewOrigin)
-					{
-						WarpData.WarpEye->OverrideViewDirection = WarpData.WarpEye->ViewPoint.Rotation.RotateVector(FVector::XAxisVector);
-					}
-					else
-					{
-						if (GDisplayClusterWarpInFrustumFitPolicyUseGroupViewTarget)
-						{
-							// If we have a correction to the view forward vector to make the frustum symmetric, apply it
-							if (SymmetricForwardCorrection.IsSet())
-							{
-								const FVector AABBForward = (GroupAABBox.GetCenter() - WarpData.WarpEye->ViewPoint.GetEyeLocation()).GetSafeNormal();
-								WarpData.WarpEye->OverrideViewDirection = SymmetricForwardCorrection->RotateVector(AABBForward);
-							}
-							else
-							{
-								// Use the same view target for all viewports
-								WarpData.WarpEye->OverrideViewTarget = GroupAABBox.GetCenter();
-							}
-						}
-					}
-
-					if (GDisplayClusterWarpInFrustumFitPolicyUseStaticViewDirectionForMPCDIProfile2D)
-					{
-						// [experimental] use static view direction for MPCDI profile 2D
-						if (WarpBlend->GetWarpProfileType() == EDisplayClusterWarpProfileType::warp_2D)
-						{
-							WarpData.WarpEye->OverrideViewDirection = FVector(1, 0, 0);
-						}
-					}
-
-					// Todo: This feature now not supported here. need to be fixed.
-					WarpData.bEnabledRotateFrustumToFitContextSize = false;
-				}
-			}
-		}
-	}
-}
-
-void FDisplayClusterWarpInFrustumFitPolicy::EndCalcFrustum(IDisplayClusterViewport* InViewport, const uint32 ContextNum)
-{
-	if (bValidGroupFrustum && InViewport && InViewport->GetProjectionPolicy().IsValid())
-	{
-		TSharedPtr<IDisplayClusterWarpBlend, ESPMode::ThreadSafe> WarpBlend;
-		if (InViewport->GetProjectionPolicy()->GetWarpBlendInterface(WarpBlend))
-		{
-			// Change warp settings:
-			FDisplayClusterWarpData& WarpData = WarpBlend->GetWarpData(ContextNum);
-
-			// Apply camera frustum fitting:
-			FDisplayClusterWarpProjection NewWarpProjection = ApplyInFrustumFit(InViewport, WarpData.WarpEye->World2LocalTransform, WarpData.WarpProjection);
-			if (NewWarpProjection.IsValidProjection())
-			{
-				WarpData.WarpProjection = NewWarpProjection;
-
-				// The warp policy Tick() function uses warp data, and it must be sure that this data is updated in the previous frame.
-				//.This value must be set to true from the EndCalcFrustum() warp policy function when changes are made to this structure.
-				WarpData.bHasWarpPolicyChanges = true;
-			}
-		}
-	}
-}
-
 void FDisplayClusterWarpInFrustumFitPolicy::OnUpdateDisplayDeviceMeshAndMaterialInstance(IDisplayClusterViewportPreview& InViewportPreview, const EDisplayClusterDisplayDeviceMeshType InMeshType, const EDisplayClusterDisplayDeviceMaterialType InMaterialType, UMeshComponent* InMeshComponent, UMaterialInstanceDynamic* InMeshMaterialInstance) const
 {
 	// The preview material used for editable meshes requires a set of unique parameters that are set from the warp policy.
@@ -356,11 +298,11 @@ void FDisplayClusterWarpInFrustumFitPolicy::DrawDebugGroupBoundingBox(ADisplayCl
 	// DCRA uses its own LineBatcher
 	ULineBatchComponent* LineBatcher = SceneRootActor ? SceneRootActor->GetLineBatchComponent() : nullptr;
 	UWorld* World = SceneRootActor ? SceneRootActor->GetWorld() : nullptr;
-	if (LineBatcher && World)
+	if (LineBatcher && World && OptUnitedGeometryWorldAABB.IsSet())
 	{
 		const float Thickness = 1.f;
 		const float PointSize = 5.f;
-		const FBox WorldBox = GroupAABBox.TransformBy(SceneRootActor->GetActorTransform());
+		const FBox WorldBox = *OptUnitedGeometryWorldAABB;
 
 		LineBatcher->DrawBox(WorldBox.GetCenter(), WorldBox.GetExtent(), Color, 0, SDPG_World, Thickness);
 		LineBatcher->DrawPoint(WorldBox.GetCenter(), Color, PointSize, SDPG_World);
@@ -369,6 +311,12 @@ void FDisplayClusterWarpInFrustumFitPolicy::DrawDebugGroupBoundingBox(ADisplayCl
 
 void FDisplayClusterWarpInFrustumFitPolicy::DrawDebugGroupFrustum(ADisplayClusterRootActor* SceneRootActor, UDisplayClusterInFrustumFitCameraComponent* CameraComponent, const FLinearColor& Color)
 {
+	if (!OptUnitedGeometryWarpProjection.IsSet())
+	{
+		// A united frastum is required.
+		return;
+	}
+
 	// DCRA uses its own LineBatcher
 	ULineBatchComponent* LineBatcher = SceneRootActor ? SceneRootActor->GetLineBatchComponent() : nullptr;
 	if (LineBatcher && CameraComponent)
@@ -385,30 +333,32 @@ void FDisplayClusterWarpInFrustumFitPolicy::DrawDebugGroupFrustum(ADisplayCluste
 			const float NearPlane = 10;
 			const float FarPlane = 1000;
 
+		
 			const FVector CameraLoc = CameraComponent->GetComponentLocation();
 			FVector ViewDirection;
-
-			if (ConfigurationCameraComponent.CameraViewTarget == EDisplayClusterWarpCameraViewTarget::MatchViewOrigin)
-			{
-				ViewDirection = CameraComponent->GetComponentRotation().RotateVector(FVector::XAxisVector);
-			}
-			else
 			{
-				const FBox WorldBox = GroupAABBox.TransformBy(SceneRootActor->GetActorTransform());
-				ViewDirection = (WorldBox.GetCenter() - CameraComponent->GetComponentLocation()).GetSafeNormal();
-
-				if (SymmetricForwardCorrection.IsSet())
+				if (ConfigurationCameraComponent.CameraViewTarget == EDisplayClusterWarpCameraViewTarget::MatchViewOrigin)
+				{
+					ViewDirection = CameraComponent->GetComponentRotation().RotateVector(FVector::XAxisVector);
+				}
+				else if(OptUnitedGeometryWorldAABB.IsSet())
+				{
+					const FVector ViewTarget = OptOverrideWorldViewTarget.IsSet() ? *OptOverrideWorldViewTarget : OptUnitedGeometryWorldAABB->GetCenter();
+					ViewDirection = (ViewTarget - CameraComponent->GetComponentLocation()).GetSafeNormal();
+				}
+				else
 				{
-					ViewDirection = SymmetricForwardCorrection->RotateVector(ViewDirection);
+					return;
 				}
 			}
 
+			FDisplayClusterWarpProjection UnitedGeometryWarpProjection = *OptUnitedGeometryWarpProjection;
 
 			const FRotator ViewRotator = ViewDirection.ToOrientationRotator();
-			const FVector FrustumTopLeft = ViewRotator.RotateVector(FVector(GroupGeometryWarpProjection.ZNear, GroupGeometryWarpProjection.Left, GroupGeometryWarpProjection.Top) / GroupGeometryWarpProjection.ZNear);
-			const FVector FrustumTopRight = ViewRotator.RotateVector(FVector(GroupGeometryWarpProjection.ZNear, GroupGeometryWarpProjection.Right, GroupGeometryWarpProjection.Top) / GroupGeometryWarpProjection.ZNear);
-			const FVector FrustumBottomLeft = ViewRotator.RotateVector(FVector(GroupGeometryWarpProjection.ZNear, GroupGeometryWarpProjection.Left, GroupGeometryWarpProjection.Bottom) / GroupGeometryWarpProjection.ZNear);
-			const FVector FrustumBottomRight = ViewRotator.RotateVector(FVector(GroupGeometryWarpProjection.ZNear, GroupGeometryWarpProjection.Right, GroupGeometryWarpProjection.Bottom) / GroupGeometryWarpProjection.ZNear);
+			const FVector FrustumTopLeft     = ViewRotator.RotateVector(FVector(UnitedGeometryWarpProjection.ZNear, UnitedGeometryWarpProjection.Left,  UnitedGeometryWarpProjection.Top)    / UnitedGeometryWarpProjection.ZNear);
+			const FVector FrustumTopRight    = ViewRotator.RotateVector(FVector(UnitedGeometryWarpProjection.ZNear, UnitedGeometryWarpProjection.Right, UnitedGeometryWarpProjection.Top)    / UnitedGeometryWarpProjection.ZNear);
+			const FVector FrustumBottomLeft  = ViewRotator.RotateVector(FVector(UnitedGeometryWarpProjection.ZNear, UnitedGeometryWarpProjection.Left,  UnitedGeometryWarpProjection.Bottom) / UnitedGeometryWarpProjection.ZNear);
+			const FVector FrustumBottomRight = ViewRotator.RotateVector(FVector(UnitedGeometryWarpProjection.ZNear, UnitedGeometryWarpProjection.Right, UnitedGeometryWarpProjection.Bottom) / UnitedGeometryWarpProjection.ZNear);
 
 			const FVector FrustumVertices[8] =
 			{
diff --git a/Engine/Plugins/Runtime/nDisplay/Source/DisplayClusterWarp/Private/Policy/InFrustumFit/DisplayClusterWarpInFrustumFitPolicy.h b/Engine/Plugins/Runtime/nDisplay/Source/DisplayClusterWarp/Private/Policy/InFrustumFit/DisplayClusterWarpInFrustumFitPolicy.h
--- a/Engine/Plugins/Runtime/nDisplay/Source/DisplayClusterWarp/Private/Policy/InFrustumFit/DisplayClusterWarpInFrustumFitPolicy.h
+++ b/Engine/Plugins/Runtime/nDisplay/Source/DisplayClusterWarp/Private/Policy/InFrustumFit/DisplayClusterWarpInFrustumFitPolicy.h
@@ -36,18 +36,68 @@ public:
 	//~~ End IDisplayClusterWarpPolicy
 
 private:
-	/** Return camera component that used for configuration.. */
-	const UDisplayClusterInFrustumFitCameraComponent* GetConfigurationInFrustumFitCameraComponent(IDisplayClusterViewport* InViewport) const;
-
-	/** Converts an asymmetric group frustum to a symmetric group frustum and computes the correction rotator to the view forward vector needed for a symmetric frustum  */
-	void MakeGroupFrustumSymmetrical(bool bFixedViewTarget);
-
 	/** Apply frustum fit to the specified warp projection */
-	FDisplayClusterWarpProjection ApplyInFrustumFit(IDisplayClusterViewport* InViewport, const FTransform& World2OriginTransform, const FDisplayClusterWarpProjection& InWarpProjection);
+	FDisplayClusterWarpProjection ApplyInFrustumFit(IDisplayClusterViewport* InViewport, const FTransform& World2OriginTransform, const FDisplayClusterWarpProjection& InWarpProjection) const;
 
 	/** Find final projection scale.*/
-	FVector2D FindFrustumFit(const EDisplayClusterWarpCameraProjectionMode InProjectionMode, const FVector2D& InCameraFOV, const FVector2D& InGeometryFOV);
+	FVector2D FindFrustumFit(const EDisplayClusterWarpCameraProjectionMode InProjectionMode, const FVector2D& InCameraFOV, const FVector2D& InGeometryFOV) const;
+
+	/** Return true, if this viewport context can be processed. */
+	bool GetWarpBlendAPI(IDisplayClusterViewport* InViewport, const uint32 ContextNum, TSharedPtr<IDisplayClusterWarpBlend, ESPMode::ThreadSafe>& OutWarpBlend, TObjectPtr<UDisplayClusterInFrustumFitCameraComponent>& OutConfigurationCameraComponent) const;
+
+	/** Calculate united AABB for group of viewports: */
+	bool CalcUnitedGeometryWorldAABBox(const TArray<TSharedPtr<IDisplayClusterViewport, ESPMode::ThreadSafe>>& InViewports, const float WorldScale, FDisplayClusterWarpAABB& OutUnitedGeometryWorldAABBox) const;
+
+	/** Calculate united geometry frustum for group of viewports
+	* 
+	* @param InViewports - viewports with warp geometry
+	* @param InContextNum - context number (i.e., eye index)
+	* @param WorldScale - world scale
+	* @param OutUnitedGeometryWarpProjection (out) calculated
+	*/
+	bool CalcUnitedGeometryFrustum(const TArray<TSharedPtr<IDisplayClusterViewport, ESPMode::ThreadSafe>>& InViewports, const int32 InContextNum, const float WorldScale, FDisplayClusterWarpProjection& OutUnitedGeometryWarpProjection);
 
+private:
+	/** Auxiliary container for the CalcUnitedGeometrySymmetricFrustum() function. */
+	struct FSymmetricFrustumData
+	{
+		FSymmetricFrustumData(const float InWorldScale, const FTransform& InCameraComponent2WorldTransform)
+			: WorldScale(InWorldScale), CameraComponent2WorldTransform(InCameraComponent2WorldTransform)
+		{ }
+
+	public:
+		// World scale value
+		const float WorldScale;
+
+		// Transform from the FrustumFit component space to the world space.
+		const FTransform CameraComponent2WorldTransform;
+
+		// Index of the current iteration. Used to stop the search when the maximum number of iterations is reached.
+		// Note:
+		//   The united geometry frustum is built from geometric points projected onto a special plane.
+		//   This plane is called the 'projection plane' and is created from two quantities: the view direction vector and the eye position.
+		//   Therefore, when we change the view direction vector, it leads to a change in the "projection plane" and, then, to new frustum values.
+		//   When we set out to create a symmetrical frustum, we need to solve this math problem.
+		//   An easy way is to do this in a few iterations and stop when we find a suitable view direction that provides a nearly symmetrical frustum with acceptable precision.
+		int32 IterationNum = 0;
+
+		// United symmetric WarpProjection
+		TOptional<FDisplayClusterWarpProjection> NewUnitedSymmetricWarpProjection;
+
+		// new location of the viewing target in the world space.
+		TOptional<FVector> NewWorldViewTarget;
+	};
+
+private:
+	/** Calculate united symmetric frustum for group of viewports:
+	*
+	* @param InViewports - viewports with warp geometry
+	* @param InContextNum - context number (i.e., eye index)
+	* @param InOutSymmetricFrustumData - (in, out) data used by the iteration step logic.
+	*/
+	bool CalcUnitedGeometrySymmetricFrustum(const TArray<TSharedPtr<IDisplayClusterViewport, ESPMode::ThreadSafe>>& InViewports, const int32 InContextNum, FSymmetricFrustumData& InOutSymmetricFrustumData);
+
+private:
 #if WITH_EDITOR
 	/** Renders a debug visualization for the group bounding box */
 	void DrawDebugGroupBoundingBox(ADisplayClusterRootActor* RootActor, const FLinearColor& Color);
@@ -58,14 +108,18 @@ private:
 
 private:
 	// Warp projection data
-	FDisplayClusterWarpProjection GroupGeometryWarpProjection;
+	TOptional<FDisplayClusterWarpProjection> OptUnitedGeometryWarpProjection;
 
-	/** The AABB of the group computed in DCRA space */
-	FDisplayClusterWarpAABB GroupAABBox;
+	// United geometry AABB
+	// (The world space is used because each viewport can use its own origin.)
+	TOptional<FDisplayClusterWarpAABB> OptUnitedGeometryWorldAABB;
 
-	/** If set, stores the correction rotator to apply to the view direction to make the group frustum symmetrical */
-	TOptional<FRotator> SymmetricForwardCorrection;
+	// The view target position (world space).
+	// This value is used to obtain the direction of view and to build the "projection plane".
+	// (The world space is used because each viewport can use its own origin.)
+	TOptional<FVector> OptOverrideWorldViewTarget;
 
-	bool bValidGroupFrustum = false;
+	// true if the geometry cache inside the WarpBlend api has been updated.
+	// This is used to optimize performance to avoid multiple updates when we call the WarpBlend api multiple times per frame.
 	bool bGeometryContextsUpdated = false;
 };
diff --git a/Engine/Plugins/Runtime/nDisplay/Source/DisplayClusterWarp/Private/Policy/InFrustumFit/DisplayClusterWarpInFrustumFitPolicy_Math.cpp b/Engine/Plugins/Runtime/nDisplay/Source/DisplayClusterWarp/Private/Policy/InFrustumFit/DisplayClusterWarpInFrustumFitPolicy_Math.cpp
--- a/Engine/Plugins/Runtime/nDisplay/Source/DisplayClusterWarp/Private/Policy/InFrustumFit/DisplayClusterWarpInFrustumFitPolicy_Math.cpp
+++ b/Engine/Plugins/Runtime/nDisplay/Source/DisplayClusterWarp/Private/Policy/InFrustumFit/DisplayClusterWarpInFrustumFitPolicy_Math.cpp
@@ -3,10 +3,14 @@
 #include "Policy/InFrustumFit/DisplayClusterWarpInFrustumFitPolicy.h"
 
 #include "Render/Viewport/IDisplayClusterViewport.h"
+#include "Render/Projection/IDisplayClusterProjectionPolicy.h"
+
+#include "IDisplayClusterWarpBlend.h"
+
 #include "Components/DisplayClusterInFrustumFitCameraComponent.h"
 #include "Camera/CameraTypes.h"
+
 #include "HAL/IConsoleManager.h"
-#include "IDisplayClusterWarpBlend.h"
 
 // Debug: enable camera position fittig
 int32 GDisplayClusterWarpInFrustumFitPolicyEnableCameraPositionFit = 1;
@@ -26,54 +30,144 @@ static FAutoConsoleVariableRef CVarDisplayClusterWarpInFrustumFitPolicyEnablePro
 	ECVF_Default
 );
 
-void FDisplayClusterWarpInFrustumFitPolicy::MakeGroupFrustumSymmetrical(bool bFixedViewTarget)
+// Experimental: Enable static view direction for mpcdi 2d
+int32 GDisplayClusterWarpInFrustumFitPolicyUseStaticViewDirectionForMPCDIProfile2D = 1;
+static FAutoConsoleVariableRef CVarDisplayClusterWarpInFrustumFitPolicyUseStaticViewDirectionForMPCDIProfile2D(
+	TEXT("nDisplay.warp.InFrustumFit.UseStaticViewDirectionForMPCDIProfile2D"),
+	GDisplayClusterWarpInFrustumFitPolicyUseStaticViewDirectionForMPCDIProfile2D,
+	TEXT("Experimental: Enable static view direction for mpcdi 2d (0 - disable)\n"),
+	ECVF_Default
+);
+
+int32 GDisplayClusterWarpInFrustumFitPolicySymmetricMaxIterations = 10;
+static FAutoConsoleVariableRef CVarDisplayClusterWarpInFrustumFitPolicySymmetricMaxIterations(
+	TEXT("nDisplay.warp.InFrustumFit.Symmetric.MaxIterations"),
+	GDisplayClusterWarpInFrustumFitPolicySymmetricMaxIterations,
+	TEXT("Maximum number of iterations to find a symmetric frustum (10 by default)\n"),
+	ECVF_Default
+);
+
+float GDisplayClusterWarpInFrustumFitPolicySymmetricPrecision = 0.5f;
+static FAutoConsoleVariableRef CVarDisplayClusterWarpInFrustumFitPolicySymmetricPrecision(
+	TEXT("nDisplay.warp.InFrustumFit.Symmetric.Precision"),
+	GDisplayClusterWarpInFrustumFitPolicySymmetricPrecision,
+	TEXT("The symmetrical frustum is accepted when this accuracy is achieved (0.5 degree by default) \n"),
+	ECVF_Default
+);
+
+bool FDisplayClusterWarpInFrustumFitPolicy::GetWarpBlendAPI(IDisplayClusterViewport* InViewport, const uint32 ContextNum, TSharedPtr<IDisplayClusterWarpBlend, ESPMode::ThreadSafe>& OutWarpBlend, TObjectPtr<UDisplayClusterInFrustumFitCameraComponent>& OutConfigurationCameraComponent) const
 {
-	if (bValidGroupFrustum)
+	if (!InViewport || !InViewport->GetProjectionPolicy().IsValid())
 	{
-		if (bFixedViewTarget)
-		{
-			// If the view target is set to a fixed value instead of being computed by the group AABB, we do not want to alter the view direction,
-			// but make the frustum symmetric around that fixed direction. This involves expanding the asymmetric frustum so that its left and right,
-			// top and bottom angles are equal
+		return false;
+	}
+
+	// Stereo is currently not supported by InFrustumFit
+	if (ContextNum != 0)
+	{
+		return false;
+	}
+
+
+	if (!InViewport->GetProjectionPolicy()->GetWarpBlendInterface(OutWarpBlend))
+	{
+		return false;
+	}
+
+	OutConfigurationCameraComponent = Cast<UDisplayClusterInFrustumFitCameraComponent>(InViewport->GetViewPointCameraComponent(EDisplayClusterRootActorType::Configuration));
+	if (!OutConfigurationCameraComponent)
+	{
+		return false;
+	}
+
+	return true;
+}
+
+void FDisplayClusterWarpInFrustumFitPolicy::BeginCalcFrustum(IDisplayClusterViewport* InViewport, const uint32 ContextNum)
+{
+	TSharedPtr<IDisplayClusterWarpBlend, ESPMode::ThreadSafe> WarpBlend;
+	TObjectPtr<UDisplayClusterInFrustumFitCameraComponent> ConfigurationCameraComponent;
+	if (!GetWarpBlendAPI(InViewport, ContextNum, WarpBlend, ConfigurationCameraComponent))
+	{
+		// WarpBlend api and InFrustumFit component are required.
+		return;
+	}
+
+	FDisplayClusterWarpData& WarpData = WarpBlend->GetWarpData(ContextNum);
 
-			const float MaxHorizontal = FMath::Max(FMath::Abs(GroupGeometryWarpProjection.Left), FMath::Abs(GroupGeometryWarpProjection.Right));
-			const float MaxVertical = FMath::Max(FMath::Abs(GroupGeometryWarpProjection.Top), FMath::Abs(GroupGeometryWarpProjection.Bottom));
+	// Todo: This feature now not supported here. need to be fixed.
+	WarpData.bEnabledRotateFrustumToFitContextSize = false;
 
-			GroupGeometryWarpProjection.Left = -MaxHorizontal;
-			GroupGeometryWarpProjection.Right = MaxHorizontal;
+	// Override view direction for entire group
+	if (WarpData.WarpEye.IsValid())
+	{
+		WarpData.WarpEye->bUpdateGeometryContext = !bGeometryContextsUpdated;
 
-			GroupGeometryWarpProjection.Bottom = -MaxVertical;
-			GroupGeometryWarpProjection.Top = MaxVertical;
+		// Overrides the view direction vector to build a custom "projection plane".
+		if (GDisplayClusterWarpInFrustumFitPolicyUseStaticViewDirectionForMPCDIProfile2D && WarpBlend->GetWarpProfileType() == EDisplayClusterWarpProfileType::warp_2D)
+		{
+			// [experimental] use static view direction for MPCDI profile 2D
+			WarpData.WarpEye->OverrideViewDirection = FVector(1, 0, 0);
+		}
+		else if (ConfigurationCameraComponent->CameraViewTarget == EDisplayClusterWarpCameraViewTarget::MatchViewOrigin)
+		{
+			// Use view direction from the ViewPoint component
+			// Note: WarpEye already in the Origin space.
+			WarpData.WarpEye->OverrideViewDirection = WarpData.WarpEye->ViewPoint.Rotation.RotateVector(FVector::XAxisVector);
+		}
+		else if (OptOverrideWorldViewTarget.IsSet())
+		{
+			if (const USceneComponent* const OriginComp = InViewport->GetProjectionPolicy()->GetOriginComponent())
+			{
+				// WarpBlend math uses the Origin space.
+				const FTransform World2OriginTransform = OriginComp->GetComponentTransform();
+
+				// Transform the world space position of the view target to the Origin space.
+				const FVector OriginViewTarget = World2OriginTransform.InverseTransformPosition(*OptOverrideWorldViewTarget);
+				WarpData.WarpEye->OverrideViewTarget = OriginViewTarget;
+			}
 		}
 		else
 		{
-			// Otherwise, symmetrize the view frustum by computing a new view direction within the existing frustum
-			// First, compute the change in the view forward vector needed to make the frustum symmetric. This is done my storing a rotation on the original view forward vector
-			// computed from the group AABB that can then be applied to the each viewport's view direction when the render pass is performed for that viewport. The rotation
-			// angles can be computed by computing the difference between the current frustum's minimum edges and the half angle of the frustum's FOV
-			FDisplayClusterWarpProjection GroupWarpProjectionInDegrees = GroupGeometryWarpProjection;
-			GroupWarpProjectionInDegrees.ConvertProjectionAngles(EDisplayClusterWarpAngleUnit::Degrees);
+			check(false); // United frustum can't be builded properly in this case.
+		}
+	}
+}
 
-			const FVector2D GroupHalfFOVDegrees = 0.5 * FVector2D(GroupWarpProjectionInDegrees.Right - GroupWarpProjectionInDegrees.Left, GroupWarpProjectionInDegrees.Top - GroupWarpProjectionInDegrees.Bottom);
-			const FVector2D CorrectionAngles = GroupHalfFOVDegrees - FVector2D(FMath::Abs(GroupWarpProjectionInDegrees.Left), FMath::Abs(GroupWarpProjectionInDegrees.Bottom));
+void FDisplayClusterWarpInFrustumFitPolicy::EndCalcFrustum(IDisplayClusterViewport* InViewport, const uint32 ContextNum)
+{
+	if (!OptUnitedGeometryWarpProjection.IsSet())
+	{
+		// Frustum fit is applied only if OptUnitedGeometryWarpProjection is defined.
+		return;
+	}
 
-			SymmetricForwardCorrection = FRotator(CorrectionAngles.Y, CorrectionAngles.X, 0);
+	TSharedPtr<IDisplayClusterWarpBlend, ESPMode::ThreadSafe> WarpBlend;
+	TObjectPtr<UDisplayClusterInFrustumFitCameraComponent> ConfigurationCameraComponent;
+	if (!GetWarpBlendAPI(InViewport, ContextNum, WarpBlend, ConfigurationCameraComponent))
+	{
+		// WarpBlend api and InFrustumFit component are required.
+		return;
+	}
 
-			// Finally, make the frustum angles themselves symmetric
-			GroupWarpProjectionInDegrees.Left = -0.5 * (GroupWarpProjectionInDegrees.Right - GroupWarpProjectionInDegrees.Left);
-			GroupWarpProjectionInDegrees.Right = -GroupWarpProjectionInDegrees.Left;
+	FDisplayClusterWarpData& WarpData = WarpBlend->GetWarpData(ContextNum);
 
-			GroupWarpProjectionInDegrees.Bottom = -0.5 * (GroupWarpProjectionInDegrees.Top - GroupWarpProjectionInDegrees.Bottom);
-			GroupWarpProjectionInDegrees.Top = -GroupWarpProjectionInDegrees.Bottom;
+	// Apply camera frustum fitting:
+	FDisplayClusterWarpProjection NewWarpProjection = ApplyInFrustumFit(InViewport, WarpData.WarpEye->World2LocalTransform, WarpData.WarpProjection);
+	if (NewWarpProjection.IsValidProjection())
+	{
+		WarpData.WarpProjection = NewWarpProjection;
 
-			GroupWarpProjectionInDegrees.ConvertProjectionAngles(EDisplayClusterWarpAngleUnit::Default);
-			GroupGeometryWarpProjection = GroupWarpProjectionInDegrees;
-		}
+		// The warp policy Tick() function uses warp data, and it must be sure that this data is updated in the previous frame.
+		//.This value must be set to true from the EndCalcFrustum() warp policy function when changes are made to this structure.
+		WarpData.bHasWarpPolicyChanges = true;
 	}
 }
 
-FDisplayClusterWarpProjection FDisplayClusterWarpInFrustumFitPolicy::ApplyInFrustumFit(IDisplayClusterViewport* InViewport, const FTransform& World2OriginTransform, const FDisplayClusterWarpProjection& InWarpProjection)
+FDisplayClusterWarpProjection FDisplayClusterWarpInFrustumFitPolicy::ApplyInFrustumFit(IDisplayClusterViewport* InViewport, const FTransform& World2OriginTransform, const FDisplayClusterWarpProjection& InWarpProjection) const
 {
+	check(OptUnitedGeometryWarpProjection.IsSet());
+
 	UDisplayClusterInFrustumFitCameraComponent* SceneCameraComponent = Cast<UDisplayClusterInFrustumFitCameraComponent>(InViewport->GetViewPointCameraComponent(EDisplayClusterRootActorType::Scene));
 	if (!SceneCameraComponent)
 	{
@@ -100,36 +194,38 @@ FDisplayClusterWarpProjection FDisplayClusterWarpInFrustumFitPolicy::ApplyInFrus
 	// Fit the frustum to the rules:
 
 	FDisplayClusterWarpProjection ViewportAngles = InWarpProjection;
-	FDisplayClusterWarpProjection GroupAngles = GroupGeometryWarpProjection;
+	FDisplayClusterWarpProjection UnitedGeometryWarpProjection = *OptUnitedGeometryWarpProjection;
 
-	const FVector2D GeometryFOV(FMath::Abs(GroupAngles.Right - GroupAngles.Left), FMath::Abs(GroupAngles.Top - GroupAngles.Bottom));
+	const FVector2D GeometryFOV(FMath::Abs(
+		UnitedGeometryWarpProjection.Right - UnitedGeometryWarpProjection.Left),
+		FMath::Abs(UnitedGeometryWarpProjection.Top - UnitedGeometryWarpProjection.Bottom));
 
 	// Convert frustum angles to group FOV space, normalized to 0..1
-	const FVector2D ViewportMin((ViewportAngles.Left - GroupAngles.Left) / GeometryFOV.X, (ViewportAngles.Bottom - GroupAngles.Bottom) / GeometryFOV.Y);
-	const FVector2D ViewportMax((ViewportAngles.Right - GroupAngles.Left) / GeometryFOV.X, (ViewportAngles.Top - GroupAngles.Bottom) / GeometryFOV.Y);
+	const FVector2D ViewportMin((ViewportAngles.Left - UnitedGeometryWarpProjection.Left) / GeometryFOV.X, (ViewportAngles.Bottom - UnitedGeometryWarpProjection.Bottom) / GeometryFOV.Y);
+	const FVector2D ViewportMax((ViewportAngles.Right - UnitedGeometryWarpProjection.Left) / GeometryFOV.X, (ViewportAngles.Top - UnitedGeometryWarpProjection.Bottom) / GeometryFOV.Y);
 
 	// And convert back to camera space:
 	const float CameraHalfFOVDegrees = CameraViewInfo.FOV * 0.5f;
-	const float CameraHalfFOVProjection = GroupGeometryWarpProjection.ConvertDegreesToProjection(CameraHalfFOVDegrees);
+	const float CameraHalfFOVProjection = UnitedGeometryWarpProjection.ConvertDegreesToProjection(CameraHalfFOVDegrees);
 	const FVector2D CameraHalfFOV(CameraHalfFOVProjection, CameraHalfFOVProjection / CameraViewInfo.AspectRatio);
 	const FVector2D GeometryHalfFOV = GeometryFOV * 0.5;
 
 	// Receive configuration from InCfgComponent
 	const FVector2D FinalHalfFOV = FindFrustumFit(ConfigurationCameraComponent.CameraProjectionMode, CameraHalfFOV, GeometryHalfFOV);
 
-	// Sample code: convert back to projection angles
 	if (GDisplayClusterWarpInFrustumFitPolicyEnableProjectionAnglesFit)
 	{
-		OutWarpProjection.Left = -FinalHalfFOV.X + ViewportMin.X * FinalHalfFOV.X * 2;
-		OutWarpProjection.Right = -FinalHalfFOV.X + ViewportMax.X * FinalHalfFOV.X * 2;
-		OutWarpProjection.Top = -FinalHalfFOV.Y + ViewportMax.Y * FinalHalfFOV.Y * 2;
+		// Sample code: convert back to projection angles
+		OutWarpProjection.Left   = -FinalHalfFOV.X + ViewportMin.X * FinalHalfFOV.X * 2;
+		OutWarpProjection.Right  = -FinalHalfFOV.X + ViewportMax.X * FinalHalfFOV.X * 2;
+		OutWarpProjection.Top    = -FinalHalfFOV.Y + ViewportMax.Y * FinalHalfFOV.Y * 2;
 		OutWarpProjection.Bottom = -FinalHalfFOV.Y + ViewportMin.Y * FinalHalfFOV.Y * 2;
 	}
 
 	return OutWarpProjection;
 }
 
-FVector2D FDisplayClusterWarpInFrustumFitPolicy::FindFrustumFit(const EDisplayClusterWarpCameraProjectionMode InProjectionMode, const FVector2D& InCameraFOV, const FVector2D& InGeometryFOV)
+FVector2D FDisplayClusterWarpInFrustumFitPolicy::FindFrustumFit(const EDisplayClusterWarpCameraProjectionMode InProjectionMode, const FVector2D& InCameraFOV, const FVector2D& InGeometryFOV) const
 {
 	FVector2D OutFOV(InCameraFOV);
 
@@ -170,3 +266,169 @@ FVector2D FDisplayClusterWarpInFrustumFitPolicy::FindFrustumFit(const EDisplayCl
 
 	return OutFOV;
 }
+
+bool FDisplayClusterWarpInFrustumFitPolicy::CalcUnitedGeometryWorldAABBox(const TArray<TSharedPtr<IDisplayClusterViewport, ESPMode::ThreadSafe>>& InViewports, const float WorldScale, FDisplayClusterWarpAABB& OutUnitedGeometryWorldAABBox) const
+{
+	// Reset AABB
+	OutUnitedGeometryWorldAABBox = FDisplayClusterWarpAABB();
+
+	for (const TSharedPtr<IDisplayClusterViewport, ESPMode::ThreadSafe>& Viewport : InViewports)
+	{
+		TSharedPtr<IDisplayClusterWarpBlend, ESPMode::ThreadSafe> WarpBlend;
+		if (!Viewport.IsValid()
+			|| !Viewport->GetProjectionPolicy().IsValid()
+			|| !Viewport->GetProjectionPolicy()->GetWarpBlendInterface(WarpBlend)
+			|| !WarpBlend->UpdateGeometryContext(WorldScale))
+		{
+			return false;
+		}
+
+		const USceneComponent* const OriginComp = Viewport->GetProjectionPolicy()->GetOriginComponent();
+		if (!OriginComp)
+		{
+			return false;
+		}
+
+		// Transform from the Origin component space to the world space.
+		const FBox WorldSpaceAABBox = WarpBlend->GetGeometryContext().AABBox.TransformBy(OriginComp->GetComponentTransform());
+		OutUnitedGeometryWorldAABBox.UpdateAABB(WorldSpaceAABBox);
+	}
+
+	return true;
+}
+
+bool FDisplayClusterWarpInFrustumFitPolicy::CalcUnitedGeometryFrustum(const TArray<TSharedPtr<IDisplayClusterViewport, ESPMode::ThreadSafe>>& InViewports, const int32 InContextNum, const float WorldScale, FDisplayClusterWarpProjection& OutUnitedGeometryWarpProjection)
+{
+	// The united geometry frustum is built from geometric points projected onto a special plane.
+	// This plane is called the 'projection plane' and is created from two quantities: the view direction vector and the eye position.
+
+	check(!OptUnitedGeometryWarpProjection.IsSet()); // when we calculate the united geometry frustum, we must disable the frustum fit logic.
+
+	OutUnitedGeometryWarpProjection.ResetProjectionAngles();
+
+	for (const TSharedPtr<IDisplayClusterViewport, ESPMode::ThreadSafe>& Viewport : InViewports)
+	{
+		TSharedPtr<IDisplayClusterWarpBlend, ESPMode::ThreadSafe> WarpBlend;
+		if (!Viewport.IsValid()
+			|| !Viewport->GetProjectionPolicy().IsValid()
+			|| !Viewport->GetProjectionPolicy()->GetWarpBlendInterface(WarpBlend))
+		{
+			return false;
+		}
+
+		// Note: This code is partially copied from FDisplayClusterProjectionMPCDIPolicy::CalculateView().
+		// Override viewpoint - MPCDI always expects the location of the viewpoint component (eye location from the real world)
+		FVector ViewOffset = FVector::ZeroVector;
+		FVector ViewLocation;
+		FRotator ViewRotation;
+		if (!Viewport->GetViewPointCameraEye(InContextNum, ViewLocation, ViewRotation, ViewOffset))
+		{
+			return false;
+		}
+
+		const USceneComponent* const OriginComp = Viewport->GetProjectionPolicy()->GetOriginComponent();
+		if (!OriginComp)
+		{
+			return false;
+		}
+
+		// Setup eye data:
+		TSharedPtr<FDisplayClusterWarpEye, ESPMode::ThreadSafe> WarpEye = MakeShared<FDisplayClusterWarpEye, ESPMode::ThreadSafe>(Viewport, 0);
+
+		WarpEye->World2LocalTransform = (OriginComp ? OriginComp->GetComponentTransform() : FTransform::Identity);
+
+		// Get our base camera location and view offset in local space (MPCDI space)
+		WarpEye->ViewPoint.Location = WarpEye->World2LocalTransform.InverseTransformPosition(ViewLocation - ViewOffset);
+		WarpEye->ViewPoint.EyeOffset = WarpEye->World2LocalTransform.InverseTransformPosition(ViewLocation) - WarpEye->ViewPoint.Location;
+		WarpEye->ViewPoint.Rotation = WarpEye->World2LocalTransform.InverseTransformRotation(ViewRotation.Quaternion()).Rotator();
+
+		WarpEye->WorldScale = WorldScale;
+		WarpEye->WarpPolicy = SharedThis(this);
+
+		if (!WarpBlend->CalcFrustumContext(WarpEye))
+		{
+			return false;
+		}
+
+		const FDisplayClusterWarpData& WarpData = WarpBlend->GetWarpData(InContextNum);
+
+		OutUnitedGeometryWarpProjection.ExpandProjectionAngles(WarpData.GeometryWarpProjection);
+	}
+
+	return true;
+}
+
+bool FDisplayClusterWarpInFrustumFitPolicy::CalcUnitedGeometrySymmetricFrustum(const TArray<TSharedPtr<IDisplayClusterViewport, ESPMode::ThreadSafe>>& InViewports, const int32 InContextNum, FSymmetricFrustumData& InOutSymmetricFrustumData)
+{
+	check(OptOverrideWorldViewTarget.IsSet()); // this function requires this value
+
+	if (++InOutSymmetricFrustumData.IterationNum > GDisplayClusterWarpInFrustumFitPolicySymmetricMaxIterations)
+	{
+		// we are reaching the iteration limit, so we can use the value from the previous iteration step.
+		InOutSymmetricFrustumData.IterationNum = INDEX_NONE;
+
+		return false;
+	}
+
+	// Begin new iteration
+	InOutSymmetricFrustumData.NewWorldViewTarget.Reset();
+	InOutSymmetricFrustumData.NewUnitedSymmetricWarpProjection.Reset();
+
+	// Calculate the united  frustum for the current view target.
+	FDisplayClusterWarpProjection UnitedGeometryWarpProjection;
+	if (!CalcUnitedGeometryFrustum(InViewports, InContextNum, InOutSymmetricFrustumData.WorldScale, UnitedGeometryWarpProjection))
+	{
+		// Frustum not exist
+		return false;
+	}
+
+	const double OffsetH = UnitedGeometryWarpProjection.Left + (UnitedGeometryWarpProjection.Right - UnitedGeometryWarpProjection.Left) * 0.5;
+	const double OffsetV = UnitedGeometryWarpProjection.Bottom + (UnitedGeometryWarpProjection.Top - UnitedGeometryWarpProjection.Bottom) * 0.5;
+
+	const double OffsetDegreesH = UnitedGeometryWarpProjection.ConvertProjectionToDegrees(OffsetH);
+	const double OffsetDegreesV = UnitedGeometryWarpProjection.ConvertProjectionToDegrees(OffsetV);
+
+	// Check the precision.
+	if (FMath::Max(FMath::Abs(OffsetDegreesH), FMath::Abs(OffsetDegreesV)) <= FMath::Abs(GDisplayClusterWarpInFrustumFitPolicySymmetricPrecision))
+	{
+		// A symmetric frustum has been found:
+		InOutSymmetricFrustumData.NewUnitedSymmetricWarpProjection = UnitedGeometryWarpProjection;
+
+		return true;
+	}
+
+	// Calculate the new location of the view target from the values of the united geometry frustum.
+	const FVector ViewTarget    = InOutSymmetricFrustumData.CameraComponent2WorldTransform.InverseTransformPosition(*OptOverrideWorldViewTarget);
+	const double ViewTargetSize = ViewTarget.Length();
+
+	// Create a projection space matrix from the view direction vector.
+	// This code is similar to the code used inside the WarpBlend api in the FDisplayClusterWarpBlendMath_Frustum::ImplCalcViewProjectionMatrices() function.
+	const FVector ViewDirection = ViewTarget.GetSafeNormal();
+	const FMatrix Projection2Local = (FMath::Abs(ViewDirection.Z) < (1.f - KINDA_SMALL_NUMBER))
+		? FRotationMatrix::MakeFromXZ(ViewDirection, FVector(0.f, 0.f, 1.f))
+		: FRotationMatrix::MakeFromXY(ViewDirection, FVector(0.f, 1.f, 0.f));
+
+	// Calc new view direction vector in projection space
+	const FRotator CorrectionRotator(OffsetDegreesV, OffsetDegreesH, 0);
+	const FVector ProjectionNewViewDirection = CorrectionRotator.RotateVector(FVector::XAxisVector);
+
+	// Transform a new view direction vector from projection space to local space.
+	const FVector NewViewDirection
+		= Projection2Local.GetUnitAxis(EAxis::X) * ProjectionNewViewDirection.X
+		+ Projection2Local.GetUnitAxis(EAxis::Y) * ProjectionNewViewDirection.Y
+		+ Projection2Local.GetUnitAxis(EAxis::Z) * ProjectionNewViewDirection.Z;
+
+	// Get the new position of the view target in the local space.
+	const FVector NewViewTarget = NewViewDirection * ViewTargetSize;
+
+	// Transform the location of the view target from local space to world space.
+	const FVector NewWorldViewTarget = InOutSymmetricFrustumData.CameraComponent2WorldTransform.TransformPosition(NewViewTarget);
+
+	// Use new view target for the next iteration:
+	InOutSymmetricFrustumData.NewWorldViewTarget = NewWorldViewTarget;
+
+	// Save last valid projection
+	InOutSymmetricFrustumData.NewUnitedSymmetricWarpProjection = UnitedGeometryWarpProjection;
+
+	return false;
+}
diff --git a/Engine/Plugins/Runtime/nDisplay/Source/DisplayClusterWarp/Public/Containers/DisplayClusterWarpAABB.h b/Engine/Plugins/Runtime/nDisplay/Source/DisplayClusterWarp/Public/Containers/DisplayClusterWarpAABB.h
--- a/Engine/Plugins/Runtime/nDisplay/Source/DisplayClusterWarp/Public/Containers/DisplayClusterWarpAABB.h
+++ b/Engine/Plugins/Runtime/nDisplay/Source/DisplayClusterWarp/Public/Containers/DisplayClusterWarpAABB.h
@@ -14,6 +14,13 @@ struct FDisplayClusterWarpAABB
 		: FBox(FVector(FLT_MAX, FLT_MAX, FLT_MAX), FVector(-FLT_MAX, -FLT_MAX, -FLT_MAX))
 	{ }
 
+	/** Expand the value in AABB using the input AABB. */
+	inline void UpdateAABB(const FBox& InAABB)
+	{
+		operator+=(InAABB);
+	}
+
+	/** Expand the value in AABB with the new values. */
 	template<typename FArg>
 	inline void UpdateAABB(const FArg X, const FArg Y, const FArg Z)
 	{
@@ -26,6 +33,7 @@ struct FDisplayClusterWarpAABB
 		Max.Z = FMath::Max(Max.Z, Z);
 	}
 
+	/** Expand value in AABB with a new point. */
 	inline void UpdateAABB(const FVector4f& InPts)
 	{
 		if (InPts.W > 0)
@@ -33,12 +41,14 @@ struct FDisplayClusterWarpAABB
 			UpdateAABB(InPts.X, InPts.Y, InPts.Z);
 		}
 	}
+	/** Expand value in AABB with a new point. */
 
 	inline void UpdateAABB(const FVector& InPts)
 	{
 		UpdateAABB(InPts.X, InPts.Y, InPts.Z);
 	}
 
+	/** Expand the value in AABB using the input point list. */
 	inline void UpdateAABB(const TArray<FVector>& InPoints)
 	{
 		for (const FVector& Pts : InPoints)
@@ -47,9 +57,10 @@ struct FDisplayClusterWarpAABB
 		}
 	}
 
+	/** Expand the value in AABB using the input AABB. */
 	inline void UpdateAABB(const FDisplayClusterWarpAABB& InAABB)
 	{
-		operator+=(InAABB);
+		UpdateAABB(InAABB);
 	}
 
 	/** The AABB defined by 8 points
diff --git a/Engine/Plugins/Runtime/nDisplay/Source/DisplayClusterWarp/Public/Containers/DisplayClusterWarpContext.h b/Engine/Plugins/Runtime/nDisplay/Source/DisplayClusterWarp/Public/Containers/DisplayClusterWarpContext.h
--- a/Engine/Plugins/Runtime/nDisplay/Source/DisplayClusterWarp/Public/Containers/DisplayClusterWarpContext.h
+++ b/Engine/Plugins/Runtime/nDisplay/Source/DisplayClusterWarp/Public/Containers/DisplayClusterWarpContext.h
@@ -84,7 +84,7 @@ struct FDisplayClusterWarpGeometryContext
 {
 	FMatrix GeometryToOrigin = FMatrix::Identity;
 
-	// AABB of geometry
+	// AABB of the warp geometry (in the space of the Origin component).
 	FDisplayClusterWarpAABB AABBox;
 
 	// Static surface average normal for this region
