[nDisplay] Fixed CIS issue #473177

diff --git a/Engine/Plugins/Runtime/nDisplay/Source/DisplayClusterWarp/Private/Policy/InFrustumFit/DisplayClusterWarpInFrustumFitPolicy_Math.cpp b/Engine/Plugins/Runtime/nDisplay/Source/DisplayClusterWarp/Private/Policy/InFrustumFit/DisplayClusterWarpInFrustumFitPolicy_Math.cpp
--- a/Engine/Plugins/Runtime/nDisplay/Source/DisplayClusterWarp/Private/Policy/InFrustumFit/DisplayClusterWarpInFrustumFitPolicy_Math.cpp
+++ b/Engine/Plugins/Runtime/nDisplay/Source/DisplayClusterWarp/Private/Policy/InFrustumFit/DisplayClusterWarpInFrustumFitPolicy_Math.cpp
@@ -291,7 +291,7 @@ bool FDisplayClusterWarpInFrustumFitPolicy::CalcUnitedGeometryWorldAABBox(const
 
 		// Transform from the Origin component space to the world space.
 		const FBox WorldSpaceAABBox = WarpBlend->GetGeometryContext().AABBox.TransformBy(OriginComp->GetComponentTransform());
-		OutUnitedGeometryWorldAABBox.UpdateAABB(WorldSpaceAABBox);
+		OutUnitedGeometryWorldAABBox +=WorldSpaceAABBox;
 	}
 
 	return true;
diff --git a/Engine/Plugins/Runtime/nDisplay/Source/DisplayClusterWarp/Public/Containers/DisplayClusterWarpAABB.h b/Engine/Plugins/Runtime/nDisplay/Source/DisplayClusterWarp/Public/Containers/DisplayClusterWarpAABB.h
--- a/Engine/Plugins/Runtime/nDisplay/Source/DisplayClusterWarp/Public/Containers/DisplayClusterWarpAABB.h
+++ b/Engine/Plugins/Runtime/nDisplay/Source/DisplayClusterWarp/Public/Containers/DisplayClusterWarpAABB.h
@@ -14,12 +14,6 @@ struct FDisplayClusterWarpAABB
 		: FBox(FVector(FLT_MAX, FLT_MAX, FLT_MAX), FVector(-FLT_MAX, -FLT_MAX, -FLT_MAX))
 	{ }
 
-	/** Expand the value in AABB using the input AABB. */
-	inline void UpdateAABB(const FBox& InAABB)
-	{
-		operator+=(InAABB);
-	}
-
 	/** Expand the value in AABB with the new values. */
 	template<typename FArg>
 	inline void UpdateAABB(const FArg X, const FArg Y, const FArg Z)
@@ -60,7 +54,7 @@ struct FDisplayClusterWarpAABB
 	/** Expand the value in AABB using the input AABB. */
 	inline void UpdateAABB(const FDisplayClusterWarpAABB& InAABB)
 	{
-		UpdateAABB(InAABB);
+		operator+=(InAABB);
 	}
 
 	/** The AABB defined by 8 points
