nDisplay: Show summary in asset tooltip.

Also hides the illegible json from the content browser NDC asset tooltip.

diff --git a/Engine/Plugins/Runtime/nDisplay/Source/DisplayCluster/Private/Blueprints/DisplayClusterBlueprint.cpp b/Engine/Plugins/Runtime/nDisplay/Source/DisplayCluster/Private/Blueprints/DisplayClusterBlueprint.cpp
--- a/Engine/Plugins/Runtime/nDisplay/Source/DisplayCluster/Private/Blueprints/DisplayClusterBlueprint.cpp
+++ b/Engine/Plugins/Runtime/nDisplay/Source/DisplayCluster/Private/Blueprints/DisplayClusterBlueprint.cpp
@@ -4,6 +4,7 @@
 #include "Blueprints/DisplayClusterBlueprintGeneratedClass.h"
 #include "DisplayClusterRootActor.h"
 #include "Components/DisplayClusterCameraComponent.h"
+#include "Components/DisplayClusterICVFXCameraComponent.h"
 #include "Components/DisplayClusterScreenComponent.h"
 
 #include "IDisplayClusterConfiguration.h"
@@ -78,6 +79,284 @@ void UDisplayClusterBlueprint::UpdateConfigExportProperty()
 	}
 }
 
+
+void UDisplayClusterBlueprint::UpdateSummaryProperty()
+{
+	if (!ConfigData || !ConfigData->Cluster)
+	{
+		Summary = TEXT("No data!");
+		return;
+	}
+
+	TArray<FString> Lines;
+
+	// Description
+
+	if (ConfigData->Info.Description.Len())
+	{
+		Lines.Add(ConfigData->Info.Description);
+		Lines.Add(TEXT(""));
+	}
+
+	// Settings
+
+	Lines.Add(TEXT("Settings:"));
+	Lines.Add(TEXT("--------"));
+	Lines.Add(TEXT(""));
+
+	Lines.Add(FString::Printf(TEXT("Sync Policy: %s"), *ConfigData->Cluster->Sync.RenderSyncPolicy.Type));
+	Lines.Add(FString::Printf(TEXT("Follow Local Player Camera: %s"), ConfigData->bFollowLocalPlayerCamera ? TEXT("Yes") : TEXT("No")));
+	Lines.Add(FString::Printf(TEXT("Viewport Screen Percentage Multiplier: %.2f"), ConfigData->RenderFrameSettings.ClusterBufferRatioMult));
+
+	Lines.Add(TEXT(""));
+
+	// Cluster
+
+	Lines.Add(TEXT("Cluster:"));
+	Lines.Add(TEXT("-------"));
+	Lines.Add(TEXT(""));
+
+	TSet<FString> Hosts;
+	int32 NumNodes = 0;
+	int32 NumHeadlessNodes = 0;
+	int32 NumFullscreenNodes = 0;
+	int32 NumViewports = 0;
+
+	TSet<FString> ViewportMedias;
+	TSet<FString> NodeMedias;
+	TSet<FString> IcvfxCameraMedias;
+
+	for (const TPair<FString, TObjectPtr<UDisplayClusterConfigurationClusterNode>>& NodePair : ConfigData->Cluster->Nodes)
+	{
+		const FString& NodeId = NodePair.Key;
+		const TObjectPtr<UDisplayClusterConfigurationClusterNode> Node = NodePair.Value;
+
+		if (!ensure(Node))
+		{
+			continue;
+		}
+
+		Hosts.Add(Node->Host);
+		NumNodes++;
+
+		if (Node->MediaSettings.bEnable)
+		{
+			for (const FDisplayClusterConfigurationMediaOutput& MediaOutput : Node->MediaSettings.MediaOutputs)
+			{
+				if (MediaOutput.MediaOutput)
+				{
+					FString MediaName = MediaOutput.MediaOutput.GetClass()->GetName();
+					MediaName.RemoveFromEnd(TEXT("Output"));
+					NodeMedias.Add(MediaName);
+				}
+			}
+		}
+
+		if (Node->bRenderHeadless)
+		{
+			NumHeadlessNodes++;
+		}
+		else
+		{
+			if (Node->bIsFullscreen)
+			{
+				NumFullscreenNodes++;
+			}
+		}
+
+		for (const TPair<FString, TObjectPtr<UDisplayClusterConfigurationViewport>>& ViewportPair : Node->Viewports)
+		{
+			TObjectPtr<UDisplayClusterConfigurationViewport> Viewport = ViewportPair.Value;
+
+			if (!ensure(Viewport))
+			{
+				continue;
+			}
+
+			NumViewports++;
+
+			if (Viewport->RenderSettings.Media.bEnable)
+			{
+				if (Viewport->RenderSettings.Media.MediaInput.MediaSource)
+				{
+					FString MediaName = Viewport->RenderSettings.Media.MediaInput.MediaSource->GetClass()->GetName();
+					MediaName.RemoveFromEnd(TEXT("Source"));
+					ViewportMedias.Add(MediaName);
+				}
+
+				for (const FDisplayClusterConfigurationMediaOutput& MediaOutput : Viewport->RenderSettings.Media.MediaOutputs)
+				{
+					if (!MediaOutput.MediaOutput)
+					{
+						continue;
+					}
+
+					FString MediaName = MediaOutput.MediaOutput->GetClass()->GetName();
+					MediaName.RemoveFromEnd(TEXT("Output"));
+					ViewportMedias.Add(MediaName);
+				}
+			}
+		}
+	}
+
+	Lines.Add(FString::Printf(TEXT("Hosts: %d"), Hosts.Num()));
+	Lines.Add(FString::Printf(TEXT("Nodes: %d (%d Headless, %d Fullscreen)"), NumNodes, NumHeadlessNodes, NumFullscreenNodes));
+	Lines.Add(FString::Printf(TEXT("Viewports: %d"), NumViewports));
+
+	// Here we find the icvfx camera templates in the blueprint, using  the SimpleConstructionScript.
+
+	if (SimpleConstructionScript)
+	{
+		TMap<FString, int32> CamerasByMediaOrSplit;
+
+		for (USCS_Node* Node : SimpleConstructionScript->GetAllNodes())
+		{
+			if (UDisplayClusterICVFXCameraComponent* IcvfxCamera = Cast<UDisplayClusterICVFXCameraComponent>(Node->ComponentTemplate))
+			{
+				// We're intentionally including disabled cameras in the count, but not disabled Media in them.
+
+				const FDisplayClusterConfigurationMediaICVFX& MediaSettings = IcvfxCamera->GetCameraSettingsICVFX().RenderSettings.Media;
+
+				if (MediaSettings.bEnable)
+				{
+					switch (MediaSettings.SplitType)
+					{
+					case EDisplayClusterConfigurationMediaSplitType::FullFrame:
+					{
+						int32& FullFrameCameras = CamerasByMediaOrSplit.FindOrAdd(TEXT("Full Frame"), 0);
+						FullFrameCameras++;
+
+						// Gather the media types used.
+
+						for (const FDisplayClusterConfigurationMediaOutputGroup& OutputGroup : MediaSettings.MediaOutputGroups)
+						{
+							if (!OutputGroup.MediaOutput)
+							{
+								continue;
+							}
+
+							FString MediaName = OutputGroup.MediaOutput->GetClass()->GetName();
+							MediaName.RemoveFromEnd(TEXT("Output"));
+							IcvfxCameraMedias.Add(MediaName);
+						}
+
+						for (const FDisplayClusterConfigurationMediaInputGroup& InputGroup : MediaSettings.MediaInputGroups)
+						{
+							if (!InputGroup.MediaSource)
+							{
+								continue;
+							}
+
+							FString MediaName = InputGroup.MediaSource->GetClass()->GetName();
+							MediaName.RemoveFromEnd(TEXT("Source"));
+							IcvfxCameraMedias.Add(MediaName);
+						}
+
+						break;
+					}
+					case EDisplayClusterConfigurationMediaSplitType::UniformTiles:
+					{
+						const FString TiledString = FString::Printf(TEXT("Tiled %dx%d"), MediaSettings.TiledSplitLayout.X, MediaSettings.TiledSplitLayout.Y);
+						int32& TiledCameras = CamerasByMediaOrSplit.FindOrAdd(TiledString, 0);
+						TiledCameras++;
+
+						// Gather the media types used.
+
+						for (const FDisplayClusterConfigurationMediaTiledInputGroup& InputGroup : MediaSettings.TiledMediaInputGroups)
+						{
+							for (const FDisplayClusterConfigurationMediaUniformTileInput& Tile : InputGroup.Tiles)
+							{
+								if (!Tile.MediaSource)
+								{
+									continue;
+								}
+
+								FString MediaName = Tile.MediaSource->GetClass()->GetName();
+								MediaName.RemoveFromEnd(TEXT("Source"));
+								IcvfxCameraMedias.Add(MediaName);
+							}
+						}
+
+						for (const FDisplayClusterConfigurationMediaTiledOutputGroup& OutputGroup : MediaSettings.TiledMediaOutputGroups)
+						{
+							for (const FDisplayClusterConfigurationMediaUniformTileOutput& Tile : OutputGroup.Tiles)
+							{
+								if (!Tile.MediaOutput)
+								{
+									continue;
+								}
+
+								FString MediaName = Tile.MediaOutput->GetClass()->GetName();
+								MediaName.RemoveFromEnd(TEXT("Output"));
+								IcvfxCameraMedias.Add(MediaName);
+							}
+						}
+
+						break;
+					}
+					default:
+						checkNoEntry();
+					}
+				}
+				else
+				{
+					int32& NoMediaCameras = CamerasByMediaOrSplit.FindOrAdd(TEXT("No Media"), 0);
+					NoMediaCameras++;
+				}
+			}
+		}
+
+		FString IcvfxCamerasLine = TEXT("ICVFX Cameras: ");
+
+		if (CamerasByMediaOrSplit.Num())
+		{
+			TArray<FString> TypeValues;
+
+			for (const TPair<FString, int32>& IcvfxCameraPair : CamerasByMediaOrSplit)
+			{
+				TypeValues.Add(FString::Printf(TEXT("%d (%s)"), IcvfxCameraPair.Value, *IcvfxCameraPair.Key));
+			}
+
+			IcvfxCamerasLine += FString::Join(TypeValues, TEXT(", "));
+		}
+		else
+		{
+			IcvfxCamerasLine += TEXT("None");
+		}
+
+		Lines.Add(IcvfxCamerasLine);
+	}
+
+	Lines.Add(TEXT(""));
+
+	// Media
+
+	if (IcvfxCameraMedias.Num() || ViewportMedias.Num() || NodeMedias.Num())
+	{
+		Lines.Add(TEXT("Media:"));
+		Lines.Add(TEXT("------"));
+		Lines.Add(TEXT(""));
+
+		if (NodeMedias.Num())
+		{
+			Lines.Add(TEXT("Node Media: ") + FString::Join(NodeMedias, TEXT(", ")));
+		}
+
+		if (ViewportMedias.Num())
+		{
+			Lines.Add(TEXT("Viewport Media: ") + FString::Join(ViewportMedias, TEXT(", ")));
+		}
+
+		if (IcvfxCameraMedias.Num())
+		{
+			Lines.Add(TEXT("ICVFX Camera Media: ") + FString::Join(IcvfxCameraMedias, TEXT(", ")));
+		}
+	}
+
+	Summary = FString::Join(Lines, TEXT("\n"));
+}
+
+
 namespace DisplayClusterBlueprint
 {
 	void SendAnalytics(const FString& EventName, const UDisplayClusterConfigurationData* const ConfigData)
@@ -122,6 +401,7 @@ void UDisplayClusterBlueprint::PreSave(FObjectPreSaveContext SaveContext)
 	Super::PreSave(SaveContext);
 
 	UpdateConfigExportProperty();
+	UpdateSummaryProperty();
 	DisplayClusterBlueprint::SendAnalytics(TEXT("Usage.nDisplay.ConfigSaved"), ConfigData);
 
 #if WITH_EDITOR
@@ -181,6 +461,7 @@ void UDisplayClusterBlueprint::PostLoad()
 	{
 		const FString LoadedConfigExport = ConfigExport;
 		UpdateConfigExportProperty();
+		UpdateSummaryProperty(); // Note: No need to mark the asset dirty if the generated Summary has changed since it is not being used externally.
 
 		if (!ConfigExport.Equals(LoadedConfigExport))
 		{
@@ -469,3 +750,13 @@ void UDisplayClusterBlueprint::CleanupConfigMaps(UDisplayClusterConfigurationDat
 		Map->Remove(InvalidKey);
 	}
 }
+
+
+void UDisplayClusterBlueprint::GetAssetRegistryTags(FAssetRegistryTagsContext Context) const
+{
+	Super::GetAssetRegistryTags(Context);
+
+	// Add ConfigExport to the tags so that it is asset searchable.
+	Context.AddTag(FAssetRegistryTag(TEXT("ConfigExport"), ConfigExport, FAssetRegistryTag::TT_Hidden));
+}
+
diff --git a/Engine/Plugins/Runtime/nDisplay/Source/DisplayCluster/Public/Blueprints/DisplayClusterBlueprint.h b/Engine/Plugins/Runtime/nDisplay/Source/DisplayCluster/Public/Blueprints/DisplayClusterBlueprint.h
--- a/Engine/Plugins/Runtime/nDisplay/Source/DisplayCluster/Public/Blueprints/DisplayClusterBlueprint.h
+++ b/Engine/Plugins/Runtime/nDisplay/Source/DisplayCluster/Public/Blueprints/DisplayClusterBlueprint.h
@@ -3,6 +3,8 @@
 
 #include "DisplayClusterConfigurationTypes.h"
 #include "Engine/Blueprint.h"
+#include "UObject/AssetRegistryTagsContext.h"
+
 #include "DisplayClusterBlueprint.generated.h"
 
 class USCS_Node;
@@ -33,6 +35,7 @@ public:
 	//~ Begin UObject Interface
 	virtual void PreSave(FObjectPreSaveContext SaveContext) override;
 	virtual void PostLoad() override;
+	virtual void GetAssetRegistryTags(FAssetRegistryTagsContext Context) const override;
 	//~ End UObject Interface
 
 	class UDisplayClusterBlueprintGeneratedClass* GetGeneratedClass() const;
@@ -58,11 +61,18 @@ public:
 	//** Updates the ConfigExport property. Called when saving the asset.
 	void UpdateConfigExportProperty();
 
+	//** Updates the Summary property. Called when saving the asset.
+	void UpdateSummaryProperty();
+
 public:
-	// Holds the last saved config export. In the AssetRegistry to allow parsing without loading.
-	UPROPERTY(AssetRegistrySearchable)
+	// Holds the last saved config export. Added to the AssetRegistry to allow parsing without loading.
+	UPROPERTY()
 	FString ConfigExport;
 
+	// Cache of the summary of this asset. Will be used as asset registry tag and show up in the asset tooltip.
+	UPROPERTY(AssetRegistrySearchable)
+	FString Summary;
+
 private:
 
 	/** Returns normal component name. We get component name from a SCS component template which has _GEN_VARIABLE postfix. */
