[nDisplay] A couple of quick fixes to the tile configuration dialog
- Fixed media parameters auto-configuration after configuration process is finished
- Added "Select All" and "Deselect All" offscreen&onscreen to both input and output sections
- Fixed 'loopback' media configurations (not supported)

diff --git a/Engine/Plugins/Runtime/nDisplay/Source/DisplayClusterConfigurator/Private/Views/Details/Media/SMediaTilesConfigurationDialog.cpp b/Engine/Plugins/Runtime/nDisplay/Source/DisplayClusterConfigurator/Private/Views/Details/Media/SMediaTilesConfigurationDialog.cpp
--- a/Engine/Plugins/Runtime/nDisplay/Source/DisplayClusterConfigurator/Private/Views/Details/Media/SMediaTilesConfigurationDialog.cpp
+++ b/Engine/Plugins/Runtime/nDisplay/Source/DisplayClusterConfigurator/Private/Views/Details/Media/SMediaTilesConfigurationDialog.cpp
@@ -64,7 +64,6 @@ namespace UE::DisplayClusterConfigurator::MediaTilesConfigurationDialog::Private
 	const FText TextPageNodesMenuSelectAllNoReceivers = LOCTEXT("PageNodesMenuSelectAllNoReceivers", "Select All w/o Receivers");
 	const FText TextPageNodesMenuDeselectSenders      = LOCTEXT("PageNodesMenuSelectDeselectSenders", "Deselect Senders");
 	const FText TextPageNodesMenuDeselectReceivers    = LOCTEXT("PageNodesMenuSelectDeselectReceivers", "Deselect Receivers");
-	const FText TextPageNodesMenuSelectAllOffscreen   = LOCTEXT("PageNodesMenuSelectAllOffscreen", "Select All Offscreen");
 
 	// Page: Finalization
 	const FText TextPageFinalizationHeader = LOCTEXT("PageFinalizationHeader", "Step 4: Output mapping (tile senders)");
@@ -158,16 +157,25 @@ void SMediaTilesConfigurationDialog::InitializeInternals()
 	// Pre-save some data to simplify future use
 	for (const TPair<FString, TObjectPtr<UDisplayClusterConfigurationClusterNode>>& Node : Parameters.ConfigData->Cluster->Nodes)
 	{
-		// All nodes
-		ClusterNodeIds.Add(Node.Key);
-		// Node to host
-		NodeToHostMap.Emplace(Node.Key, Node.Value->Host);
-		// Nodes per host
-		HostToNodesMap.FindOrAdd(Node.Value->Host).Add(Node.Key);
-		// Nodes allowed to be used for tile rendering
-		NodesAllowedForOutput.Add(Node.Key);
-		// Nodes allowed to be used for tile receiving and compositing
-		NodesAllowedForInput.Add(Node.Key);
+		if (Node.Value)
+		{
+			// All nodes
+			ClusterNodeIds.Add(Node.Key);
+			// Node to host
+			NodeToHostMap.Emplace(Node.Key, Node.Value->Host);
+			// Nodes per host
+			HostToNodesMap.FindOrAdd(Node.Value->Host).Add(Node.Key);
+			// Nodes allowed to be used for tile rendering
+			NodesAllowedForOutput.Add(Node.Key);
+			// Nodes allowed to be used for tile receiving and compositing
+			NodesAllowedForInput.Add(Node.Key);
+
+			// Is offscreen?
+			if (Node.Value->bRenderHeadless)
+			{
+				OffscreenNodes.Add(Node.Key);
+			}
+		}
 	}
 
 	// Fill output mapping with defaults
@@ -250,7 +258,7 @@ bool SMediaTilesConfigurationDialog::ApplyConfiguration_Multicast(TArray<FString
 {
 	// In multicast, the data is propagated in OneSender-to-MultipleReceivers way. Thus we
 	// can have a single input group with all the receivers because they get same data from
-	// the same senders. Each receiver gets the full set of of tiles, so it's connected to all
+	// the same senders. Each receiver gets the full set of tiles, so it's connected to all
 	// the receivers. This allows for all of them to share a single MediaSource for each tile.
 	// There is no limitation on the receivers amount.
 	// 
@@ -280,6 +288,17 @@ bool SMediaTilesConfigurationDialog::ApplyConfiguration_Multicast(TArray<FString
 	//  - B - tiles amount vertically
 	//  - N = A * B == amount of senders
 	//  - M - any amount of receivers
+	//
+	// ==================================================================================================
+	// HOWEVER!
+	// Currently, nDisplay doesn't allow any loopback-like setups. So it's not allowed the same node
+	// to output a tile, and consume the same tile. Also it's not allowed to have both output and
+	// input assigned to the same tile (passthrough-like). This requires us to remove any input mapping
+	// from the tiles that have already output assigned on the same node. To simplify final configuration,
+	// we'll have a separate input group for each node. Each group will have the full set of tiles that
+	// follow the limitations mentioned above.
+	// Hope it's temporary and we can get back to a single input group soon.
+	// ==================================================================================================
 
 	// Apply tile layout
 	*Parameters.SplitLayout = Accepted + FIntPoint{ 1, 1 };
@@ -289,33 +308,43 @@ bool SMediaTilesConfigurationDialog::ApplyConfiguration_Multicast(TArray<FString
 		RF_Public | RF_Transactional | RF_ArchetypeObject :
 		RF_Public | RF_Transactional;
 
+	const int32 TilesAmount = Parameters.SplitLayout->X * Parameters.SplitLayout->Y;
+
 	//
 	// INPUT setup
 	//
 
-	// Add single input group
-	Parameters.InputGroups->Reset(1);
-	FDisplayClusterConfigurationMediaTiledInputGroup& NewInputGroup = Parameters.InputGroups->AddDefaulted_GetRef();
-
-	// Put all receivers into the group
-	if (!InputSelection.IsEmpty())
-	{
-		NewInputGroup.ClusterNodes.ItemNames = InputSelection.Array();
-	}
-	else
+	if (InputSelection.IsEmpty())
 	{
 		OutErrors.Add(TEXT("No receivers found"));
 	}
 
-	// Setup input tiles (put all the tiles and receiving nodes into a single group)
-	for (int32 TileX = 0; TileX < Parameters.SplitLayout->X; ++TileX)
+	// For each receiver, create a separate input group
+	Parameters.InputGroups->Reset(InputSelection.Num());
+	for (const FString& ReceiverId : InputSelection)
 	{
-		for (int32 TileY = 0; TileY < Parameters.SplitLayout->Y; ++TileY)
+		FDisplayClusterConfigurationMediaTiledInputGroup& NewInputGroup = Parameters.InputGroups->AddDefaulted_GetRef();
+
+		// One receiver per group
+		NewInputGroup.ClusterNodes.ItemNames.Add(ReceiverId);
+
+		// Setup input tiles
+		for (int32 TileX = 0; TileX < Parameters.SplitLayout->X; ++TileX)
 		{
-			FDisplayClusterConfigurationMediaUniformTileInput& NewTile = NewInputGroup.Tiles.AddDefaulted_GetRef();
+			for (int32 TileY = 0; TileY < Parameters.SplitLayout->Y; ++TileY)
+			{
+				const FIntPoint Tile{ TileX, TileY };
+
+				// Don't allow 'loopback'
+				const bool bReceiverHasOutputAssignedForThisTile = OutputMapping[Tile].ClusterNodes.Contains(ReceiverId);
+				if (!bReceiverHasOutputAssignedForThisTile)
+				{
+					FDisplayClusterConfigurationMediaUniformTileInput& NewTile = NewInputGroup.Tiles.AddDefaulted_GetRef();
 
-			NewTile.Position = { TileX, TileY };
-			NewTile.MediaSource = NewObject<UMediaSource>(Parameters.Owner, MediaSource->GetClass(), NAME_None, MediaObjectFlags);
+					NewTile.Position = Tile;
+					NewTile.MediaSource = NewObject<UMediaSource>(Parameters.Owner, MediaSource->GetClass(), NAME_None, MediaObjectFlags);
+				}
+			}
 		}
 	}
 
@@ -323,8 +352,6 @@ bool SMediaTilesConfigurationDialog::ApplyConfiguration_Multicast(TArray<FString
 	// OUTPUT setup
 	//
 
-	const int32 TilesAmount = Parameters.SplitLayout->X * Parameters.SplitLayout->Y;
-
 	// First, remove any existing output mapping
 	Parameters.OutputGroups->Reset(TilesAmount);
 
@@ -359,7 +386,7 @@ bool SMediaTilesConfigurationDialog::ApplyConfiguration_Multicast(TArray<FString
 
 			// Create new tile in this group
 			FDisplayClusterConfigurationMediaUniformTileOutput& NewTileInGroup = NewOutputGroup.Tiles.AddDefaulted_GetRef();
-			NewTileInGroup.Position = { TileX, TileY };
+			NewTileInGroup.Position = Tile;
 			NewTileInGroup.MediaOutput = NewObject<UMediaOutput>(Parameters.Owner, MediaOutput->GetClass(), NAME_None, MediaObjectFlags);
 		}
 	}
@@ -382,12 +409,23 @@ bool SMediaTilesConfigurationDialog::ApplyConfiguration_LocalMulticast(TArray<FS
 	// In other words. Assuming there are N tiles, each host that has at least one tile
 	// receiver must also have N tile senders (or less, but some senders would have to render
 	// multiple tiles in this case).
+	//
+	// ==================================================================================================
+	// HOWEVER!
+	// Currently, nDisplay doesn't allow any loopback-like setups. So it's not allowed the same node
+	// to output a tile, and consume the same tile. Also it's not allowed to have both output and
+	// input assigned to the same tile (passthrough-like). This requires us to remove any input mapping
+	// from the tiles that have already output assigned on the same node. To simplify final configuration,
+	// we'll have a separate input group for each node. Each group will have the full set of tiles that
+	// follow the limitations mentioned above.
+	// Hope it's temporary and we can get back to a single input group soon.
+	// ==================================================================================================
 
 	// Apply tile layout
 	*Parameters.SplitLayout = Accepted + FIntPoint{ 1, 1 };
 
 	// All receivers involved
-	const TSet<FString> AllReceivers = InputSelection;
+	const TSet<FString>& AllReceivers = InputSelection;
 
 	// All senders involved
 	TSet<FString> AllSenders;
@@ -417,39 +455,51 @@ bool SMediaTilesConfigurationDialog::ApplyConfiguration_LocalMulticast(TArray<FS
 	const int32 TilesAmount = Parameters.SplitLayout->X * Parameters.SplitLayout->Y;
 
 	// Reset any existing data
-	Parameters.InputGroups->Reset(HostsWithMedia.Num());
+	Parameters.InputGroups->Reset(AllReceivers.Num());
 	Parameters.OutputGroups->Reset(HostsWithMedia.Num() * TilesAmount);
 
-	// Now, we can generate an input group and corresponding amount (amount of tiles) of output
-	// groups for every host.
+	// Now, we can generate per-receiver input groups and per-tile output groups for every host.
 	for (const FString& Host : HostsWithMedia)
 	{
 		//
 		// INPUT setup
 		//
 
-		// Add input group
-		FDisplayClusterConfigurationMediaTiledInputGroup& NewInputGroup = Parameters.InputGroups->AddDefaulted_GetRef();
-
-		// Assing all the receivers on this host to the group
-		const TSet<FString> ReceiversInGroup = AllReceivers.Intersect(HostToNodesMap[Host]);
-		NewInputGroup.ClusterNodes.ItemNames = ReceiversInGroup.Array();
+		// Find all receivers on this particular host
+		const TSet<FString> ReceiversOnThisHost = HostToNodesMap[Host].Intersect(AllReceivers);
 
 		// We expect at least one to be set
-		if (NewInputGroup.ClusterNodes.ItemNames.IsEmpty())
+		if (ReceiversOnThisHost.IsEmpty())
 		{
 			OutErrors.Add(FString::Printf(TEXT("No receivers found on host '%s'"), *Host));
 		}
 
-		// Setup input tiles (setup all the tiles in the group)
-		for (int32 TileX = 0; TileX < Parameters.SplitLayout->X; ++TileX)
+		// For each receiver on this host, add a new input group
+		for (const FString& ReceiverId : ReceiversOnThisHost)
 		{
-			for (int32 TileY = 0; TileY < Parameters.SplitLayout->Y; ++TileY)
+			// Add input group
+			FDisplayClusterConfigurationMediaTiledInputGroup& NewInputGroup = Parameters.InputGroups->AddDefaulted_GetRef();
+
+			// Assing all the receivers on this host to the group
+			NewInputGroup.ClusterNodes.ItemNames.Add(ReceiverId);
+
+			// Setup input tiles (setup all the tiles in the group)
+			for (int32 TileX = 0; TileX < Parameters.SplitLayout->X; ++TileX)
 			{
-				FDisplayClusterConfigurationMediaUniformTileInput& NewTile = NewInputGroup.Tiles.AddDefaulted_GetRef();
+				for (int32 TileY = 0; TileY < Parameters.SplitLayout->Y; ++TileY)
+				{
+					const FIntPoint Tile{ TileX, TileY };
+
+					// Don't allow 'loopback'
+					const bool bReceiverHasOutputAssignedForThisTile = OutputMapping[Tile].ClusterNodes.Contains(ReceiverId);
+					if (!bReceiverHasOutputAssignedForThisTile)
+					{
+						FDisplayClusterConfigurationMediaUniformTileInput& NewTile = NewInputGroup.Tiles.AddDefaulted_GetRef();
 
-				NewTile.Position = { TileX, TileY };
-				NewTile.MediaSource = NewObject<UMediaSource>(Parameters.Owner, MediaSource->GetClass(), NAME_None, MediaObjectFlags);
+						NewTile.Position = Tile;
+						NewTile.MediaSource = NewObject<UMediaSource>(Parameters.Owner, MediaSource->GetClass(), NAME_None, MediaObjectFlags);
+					}
+				}
 			}
 		}
 
@@ -462,11 +512,13 @@ bool SMediaTilesConfigurationDialog::ApplyConfiguration_LocalMulticast(TArray<FS
 		{
 			for (int32 TileY = 0; TileY < Parameters.SplitLayout->Y; ++TileY)
 			{
+				const FIntPoint Tile{ TileX, TileY };
+
 				// Create new output group
 				FDisplayClusterConfigurationMediaTiledOutputGroup& NewOutputGroup = Parameters.OutputGroups->AddDefaulted_GetRef();
 
 				// Get all senders of this tile on this particualar host
-				const TSet<FString>& AllSendersMappedToThisTile = OutputMapping[{TileX, TileY}].ClusterNodes;
+				const TSet<FString>& AllSendersMappedToThisTile = OutputMapping[Tile].ClusterNodes;
 				const TSet<FString>  AllSendersOnThisHost = AllSenders.Intersect(HostToNodesMap[Host]);
 				const TSet<FString>  SendersOfThisTile = AllSendersOnThisHost.Intersect(AllSendersMappedToThisTile);
 				
@@ -492,7 +544,7 @@ bool SMediaTilesConfigurationDialog::ApplyConfiguration_LocalMulticast(TArray<FS
 
 				// Create new tile in this group
 				FDisplayClusterConfigurationMediaUniformTileOutput& NewTileInGroup = NewOutputGroup.Tiles.AddDefaulted_GetRef();
-				NewTileInGroup.Position = { TileX, TileY };
+				NewTileInGroup.Position = Tile;
 				NewTileInGroup.MediaOutput = NewObject<UMediaOutput>(Parameters.Owner, MediaOutput->GetClass(), NAME_None, MediaObjectFlags);
 			}
 		}
@@ -1125,11 +1177,13 @@ void SMediaTilesConfigurationDialog::PageNodes_CreateContextMenu(const FVector2D
 
 	// Section: general commands
 	{
-		MenuBuilder.BeginSection("General");
+		const FName SectionName = TEXT("General");
+
+		MenuBuilder.BeginSection(SectionName);
 
 		MenuBuilder.AddWidget(
 			SNew(STextBlock)
-			.Text(FText::FromString("General")),
+			.Text(FText::FromName(SectionName)),
 			FText(),
 			true
 		);
@@ -1157,13 +1211,87 @@ void SMediaTilesConfigurationDialog::PageNodes_CreateContextMenu(const FVector2D
 		MenuBuilder.EndSection();
 	}
 
+	// Section: onscreen
+	{
+		const FName SectionName = TEXT("On-screen nodes");
+
+		MenuBuilder.BeginSection(SectionName);
+
+		MenuBuilder.AddWidget(
+			SNew(STextBlock)
+			.Text(FText::FromName(SectionName)),
+			FText(),
+			true
+		);
+
+		// Select all onscreen
+		MenuBuilder.AddMenuEntry(
+			UE::DisplayClusterConfigurator::MediaTilesConfigurationDialog::Private::TextPageNodesMenuSelectAll,
+			FText(),
+			FSlateIcon(),
+			FUIAction(
+				FExecuteAction::CreateSP(this, &SMediaTilesConfigurationDialog::PageNodes_Menu_OnSelectDeselectAllOnOffscreen, bOutputSelection, true, false)
+			)
+		);
+
+		// Deselect all onscreen
+		MenuBuilder.AddMenuEntry(
+			UE::DisplayClusterConfigurator::MediaTilesConfigurationDialog::Private::TextPageNodesMenuDeselectAll,
+			FText(),
+			FSlateIcon(),
+			FUIAction(
+				FExecuteAction::CreateSP(this, &SMediaTilesConfigurationDialog::PageNodes_Menu_OnSelectDeselectAllOnOffscreen, bOutputSelection, false, false)
+			)
+		);
+
+		MenuBuilder.EndSection();
+	}
+
+	// Section: offscreen
+	{
+		const FName SectionName = TEXT("Off-screen nodes");
+
+		MenuBuilder.BeginSection(SectionName);
+
+		MenuBuilder.AddWidget(
+			SNew(STextBlock)
+			.Text(FText::FromName(SectionName)),
+			FText(),
+			true
+		);
+
+		// Select all offscreen
+		MenuBuilder.AddMenuEntry(
+			UE::DisplayClusterConfigurator::MediaTilesConfigurationDialog::Private::TextPageNodesMenuSelectAll,
+			FText(),
+			FSlateIcon(),
+			FUIAction(
+				FExecuteAction::CreateSP(this, &SMediaTilesConfigurationDialog::PageNodes_Menu_OnSelectDeselectAllOnOffscreen, bOutputSelection, true, true)
+			)
+		);
+
+		// Deselect all offscreen
+		MenuBuilder.AddMenuEntry(
+			UE::DisplayClusterConfigurator::MediaTilesConfigurationDialog::Private::TextPageNodesMenuDeselectAll,
+			FText(),
+			FSlateIcon(),
+			FUIAction(
+				FExecuteAction::CreateSP(this, &SMediaTilesConfigurationDialog::PageNodes_Menu_OnSelectDeselectAllOnOffscreen, bOutputSelection, false, true)
+			)
+		);
+
+		MenuBuilder.EndSection();
+	}
+
 	// Section: extra commands
 	{
-		MenuBuilder.BeginSection("Extra");
+		const FName SectionName = TEXT("Extra");
+
+		MenuBuilder.BeginSection(SectionName);
 
 		MenuBuilder.AddWidget(
 			SNew(STextBlock)
-			.Text(FText::FromString("Extra")),
+			.Text(FText::FromName(SectionName)),
 			FText(),
 			true
 		);
@@ -1192,20 +1320,6 @@ void SMediaTilesConfigurationDialog::PageNodes_CreateContextMenu(const FVector2D
 			)
 		);
 
-		// Output context only
-		if (bOutputSelection)
-		{
-			// Select all offscreen nodes
-			MenuBuilder.AddMenuEntry(
-				UE::DisplayClusterConfigurator::MediaTilesConfigurationDialog::Private::TextPageNodesMenuSelectAllOffscreen,
-				FText(),
-				FSlateIcon(),
-				FUIAction(
-					FExecuteAction::CreateSP(this, &SMediaTilesConfigurationDialog::PageNodes_Menu_OnSelectAllOffscreen, bOutputSelection)
-				)
-			);
-		}
-
 		MenuBuilder.EndSection();
 	}
 
@@ -1243,8 +1357,6 @@ void SMediaTilesConfigurationDialog::PageNodes_Menu_OnSelectAllExceptOfCounterpa
 	const TSet<FString>& SetB = bOutputSelection ? InputSelection : OutputSelection;
 	// Allowed items for the set being edited
 	const TSet<FString>& AllowedSetA = bOutputSelection ? NodesAllowedForOutput : NodesAllowedForInput;
-	// Allowed items for the counterpart set
-	const TSet<FString>& AllowedSetB = bOutputSelection ? NodesAllowedForInput : NodesAllowedForOutput;
 
 	// Select all in A that aren't currently selected in B
 	// Result = AllowedA - SelectedB
@@ -1262,19 +1374,29 @@ void SMediaTilesConfigurationDialog::PageNodes_Menu_OnDeselectAllFromCounterpart
 	SetA = SetA.Difference(SetB);
 }
 
-void SMediaTilesConfigurationDialog::PageNodes_Menu_OnSelectAllOffscreen(bool bOutputSelection)
+void SMediaTilesConfigurationDialog::PageNodes_Menu_OnSelectDeselectAllOnOffscreen(bool bOutputSelection, bool bSelect, bool bOffscreen)
 {
 	// A set being edited
-	TSet<FString>& SetA = bOutputSelection ? OutputSelection : InputSelection;
+	TSet<FString>& CurrentSet = bOutputSelection ? OutputSelection : InputSelection;
+	// Allowed items for the set being edited
+	const TSet<FString>& AllowedSet = bOutputSelection ? NodesAllowedForOutput : NodesAllowedForInput;
+
+	// Find nodes that we're going to select or deselect
+	const TSet<FString> DesiredNodes = bOffscreen ?
+		// All offscreen nodes allowed for this section (input or output)
+		AllowedSet.Intersect(OffscreenNodes) :
+		// All onscreen nodes allowed for this section (input or output)
+		AllowedSet.Intersect(ClusterNodeIds.Difference(OffscreenNodes));
 
-	// Iterate through all the nodes in config
-	for (const TPair<FString, TObjectPtr<UDisplayClusterConfigurationClusterNode>>& NodeIt : Parameters.ConfigData->Cluster->Nodes)
+	if (bSelect)
 	{
-		// Offscreen?
-		if (NodeIt.Value && NodeIt.Value->bRenderHeadless)
-		{
-			SetA.Remove(NodeIt.Key);
-		}
+		// Select all
+		CurrentSet.Append(DesiredNodes);
+	}
+	else
+	{
+		// Deselect all
+		CurrentSet = CurrentSet.Difference(DesiredNodes);
 	}
 }
 
diff --git a/Engine/Plugins/Runtime/nDisplay/Source/DisplayClusterConfigurator/Private/Views/Details/Media/SMediaTilesConfigurationDialog.h b/Engine/Plugins/Runtime/nDisplay/Source/DisplayClusterConfigurator/Private/Views/Details/Media/SMediaTilesConfigurationDialog.h
--- a/Engine/Plugins/Runtime/nDisplay/Source/DisplayClusterConfigurator/Private/Views/Details/Media/SMediaTilesConfigurationDialog.h
+++ b/Engine/Plugins/Runtime/nDisplay/Source/DisplayClusterConfigurator/Private/Views/Details/Media/SMediaTilesConfigurationDialog.h
@@ -200,8 +200,8 @@ private:
 	/** Context menu handler: Deselect all items in A that are selected in B (Asel-Bsel). */
 	void PageNodes_Menu_OnDeselectAllFromCounterpart(bool bOutputSelection);
 
-	/** Context menu handler: Select all offscreen nodes. */
-	void PageNodes_Menu_OnSelectAllOffscreen(bool bOutputSelection);
+	/** Context menu handler: Select all onscreen/offscreen nodes. */
+	void PageNodes_Menu_OnSelectDeselectAllOnOffscreen(bool bOutputSelection, bool bSelect, bool bOffscreen);
 
 	/** Callback to check if we can switch to the nodes selection page. */
 	bool PageNodes_OnCanShow() const;
@@ -267,6 +267,9 @@ private:
 	/** List of all cluster node IDs. */
 	TSet<FString> ClusterNodeIds;
 
+	/** List of all offscreen nodes. */
+	TSet<FString> OffscreenNodes;
+
 	/** Node-to-host map. */
 	TMap<FString, FString> NodeToHostMap;
 
