nDisplay: Disabled Hair on LC renders.

diff --git a/Engine/Plugins/Runtime/nDisplay/Source/DisplayCluster/Private/Render/Viewport/DisplayClusterViewportManager.cpp b/Engine/Plugins/Runtime/nDisplay/Source/DisplayCluster/Private/Render/Viewport/DisplayClusterViewportManager.cpp
--- a/Engine/Plugins/Runtime/nDisplay/Source/DisplayCluster/Private/Render/Viewport/DisplayClusterViewportManager.cpp
+++ b/Engine/Plugins/Runtime/nDisplay/Source/DisplayCluster/Private/Render/Viewport/DisplayClusterViewportManager.cpp
@@ -722,6 +722,7 @@ FSceneViewFamily::ConstructionValues FDisplayClusterViewportManager::CreateViewF
 		InEngineShowFlags.SetSeparateTranslucency(0);
 		InEngineShowFlags.SetHMDDistortion(0);
 		InEngineShowFlags.SetOnScreenDebug(0);
+		InEngineShowFlags.SetHair(0);
 
 		InEngineShowFlags.SetLumenReflections(0);
 		InEngineShowFlags.SetLumenGlobalIllumination(0);
diff --git a/Engine/Plugins/Runtime/nDisplay/Source/DisplayClusterShaders/Private/Shaders/DisplayClusterShadersPreprocess_UVLightCards.cpp b/Engine/Plugins/Runtime/nDisplay/Source/DisplayClusterShaders/Private/Shaders/DisplayClusterShadersPreprocess_UVLightCards.cpp
--- a/Engine/Plugins/Runtime/nDisplay/Source/DisplayClusterShaders/Private/Shaders/DisplayClusterShadersPreprocess_UVLightCards.cpp
+++ b/Engine/Plugins/Runtime/nDisplay/Source/DisplayClusterShaders/Private/Shaders/DisplayClusterShadersPreprocess_UVLightCards.cpp
@@ -215,6 +215,7 @@ bool FDisplayClusterShadersPreprocess_UVLightCards::RenderPreprocess_UVLightCard
 		EngineShowFlags.SetSeparateTranslucency(0);
 		EngineShowFlags.SetHMDDistortion(0);
 		EngineShowFlags.SetOnScreenDebug(0);
+		EngineShowFlags.SetHair(0);
 
 		EngineShowFlags.SetLumenReflections(0);
 		EngineShowFlags.SetLumenGlobalIllumination(0);
