[nDisplay] Fixed tiles auto-configuration after tile dialog is closed
- Forcing property views refresh in another way

diff --git a/Engine/Plugins/Runtime/nDisplay/Source/DisplayClusterConfigurator/Private/Views/Details/Media/DisplayClusterConfiguratorICVFXMediaCustomization.cpp b/Engine/Plugins/Runtime/nDisplay/Source/DisplayClusterConfigurator/Private/Views/Details/Media/DisplayClusterConfiguratorICVFXMediaCustomization.cpp
--- a/Engine/Plugins/Runtime/nDisplay/Source/DisplayClusterConfigurator/Private/Views/Details/Media/DisplayClusterConfiguratorICVFXMediaCustomization.cpp
+++ b/Engine/Plugins/Runtime/nDisplay/Source/DisplayClusterConfigurator/Private/Views/Details/Media/DisplayClusterConfiguratorICVFXMediaCustomization.cpp
@@ -199,8 +199,10 @@ FReply FDisplayClusterConfiguratorICVFXMediaCustomization::OnSetupButtonClicked(
 	if (TilesConfigurationDialog->WasConfigurationCompleted())
 	{
 		// Redraw property views
-		FPropertyEditorModule& PropertyEditorModule = FModuleManager::LoadModuleChecked<FPropertyEditorModule>("PropertyEditor");
-		PropertyEditorModule.UpdatePropertyViews({ EditingObject.Get() });
+		if (PropertyUtilities.IsValid())
+		{
+			PropertyUtilities.Pin()->ForceRefresh();
+		}
 
 		// Notify tile customizers to re-initialize all media we just generated
 		FDisplayClusterConfiguratorMediaUtils::Get().OnMediaResetToDefaults().Broadcast(EditingObject.Get());
