nDisplay: fixed UnrealStageApp rendering issues.
- Implemented new rendering features that allow components from multiple worlds to be rendered into the same output texture.
- The logic has been reworked.
- Fixed propagation of preview mesh position from scene to proxy.
- Code cleanup.

diff --git a/Engine/Plugins/Runtime/nDisplay/Source/DisplayCluster/Private/Game/EngineClasses/Scene/DisplayClusterRootActor.cpp b/Engine/Plugins/Runtime/nDisplay/Source/DisplayCluster/Private/Game/EngineClasses/Scene/DisplayClusterRootActor.cpp
--- a/Engine/Plugins/Runtime/nDisplay/Source/DisplayCluster/Private/Game/EngineClasses/Scene/DisplayClusterRootActor.cpp
+++ b/Engine/Plugins/Runtime/nDisplay/Source/DisplayCluster/Private/Game/EngineClasses/Scene/DisplayClusterRootActor.cpp
@@ -1058,6 +1058,18 @@ void ADisplayClusterRootActor::Tick(float DeltaSeconds)
 		}
 	}
 
+	TickPreviewRenderer();
+
+	SetLightCardOwnership();
+
+	Super::Tick(DeltaSeconds);
+}
+
+void ADisplayClusterRootActor::TickPreviewRenderer()
+{
+	const bool bIsPrimaryRootActor = IsPrimaryRootActor();
+	const bool bIsRunningDisplayCluster = IsRunningDisplayCluster();
+
 	// Get preview settings from the this actor current source
 	// as they may already be configured externally, so just get the preview settings from the current source
 	FDisplayClusterViewport_PreviewSettings NewPreviewSettings = GetPreviewSettings();
@@ -1138,10 +1150,6 @@ void ADisplayClusterRootActor::Tick(float DeltaSeconds)
 
 	// Update RootActor visibility for game
 	SetActorHiddenInGame(!bEnablePreviewInScene);
-
-	SetLightCardOwnership();
-
-	Super::Tick(DeltaSeconds);
 }
 
 FDisplayClusterViewport_PreviewSettings ADisplayClusterRootActor::GetPreviewSettings(bool bIgnorePreviewSetttingsSource) const
diff --git a/Engine/Plugins/Runtime/nDisplay/Source/DisplayCluster/Private/Render/Viewport/Preview/DisplayClusterViewportManagerPreview.cpp b/Engine/Plugins/Runtime/nDisplay/Source/DisplayCluster/Private/Render/Viewport/Preview/DisplayClusterViewportManagerPreview.cpp
--- a/Engine/Plugins/Runtime/nDisplay/Source/DisplayCluster/Private/Render/Viewport/Preview/DisplayClusterViewportManagerPreview.cpp
+++ b/Engine/Plugins/Runtime/nDisplay/Source/DisplayCluster/Private/Render/Viewport/Preview/DisplayClusterViewportManagerPreview.cpp
@@ -82,6 +82,50 @@ void FDisplayClusterViewportManagerPreview::OnPreviewRenderTick()
 		return;
 	}
 
+	// Special case for RootActorProxy object (proxy always refers to external DCRA in the scene)
+	ADisplayClusterRootActor* RootActorProxy = Configuration->GetRootActor(EDisplayClusterRootActorType::Preview);
+	if (RootActorProxy != SceneRootActor)
+	{
+		// Since this DCRA is a proxy and the DCRA on the scene is a separate object,
+		// it does not pass component positions and properties from the scene to the proxy.
+		// When the cluster preview was rendered for RootActorProxy, the component properties from
+		// the scene were already retrieved via the EDisplayClusterRootActorType::Scene reference.
+		// This was done via a link to the RootActor in the scene, which means the RootActorProxy components did not change.
+		// But RootActorProxy is used in custom preview rendering, meaning it uses its own mesh components,
+		// so their positions must be synchronized with the scene components.
+		//
+		// Move the RootActorProxy components to the same relative position as the components inside the RootActor in the scene.
+		RootActorProxy->ForEachComponent<UPrimitiveComponent>(true, [&](UPrimitiveComponent* ProxyComponent)
+			{
+				if (!IsValid(ProxyComponent))
+				{
+					return;
+				}
+
+				SceneRootActor->ForEachComponent<UPrimitiveComponent>(true, [&](UPrimitiveComponent* SceneComponent)
+					{
+						if (!IsValid(SceneComponent))
+						{
+							return;
+						}
+
+						const FName SceneComponentName = SceneComponent->GetFName();
+						const FName ProxyComponentName = ProxyComponent->GetFName();
+
+						// Propagate the component transformation from the scene to the proxy.
+						if (SceneComponentName == ProxyComponentName)
+						{
+							const FTransform NewRelativeTransform = SceneComponent->GetRelativeTransform();
+							const FTransform OldRelativeTransform = ProxyComponent->GetRelativeTransform();
+							if (!NewRelativeTransform.Equals(OldRelativeTransform, UE_KINDA_SMALL_NUMBER))
+							{
+								ProxyComponent->SetRelativeTransform(NewRelativeTransform);
+							}
+						}
+					});
+			});
+	}
+
 	const FDisplayClusterRenderFrameSettings& RenderFrameSetting = Configuration->GetRenderFrameSettings();
 
 	// Update preview RTTs correspond to 'TickPerFrame' value
diff --git a/Engine/Plugins/Runtime/nDisplay/Source/DisplayCluster/Public/DisplayClusterRootActor.h b/Engine/Plugins/Runtime/nDisplay/Source/DisplayCluster/Public/DisplayClusterRootActor.h
--- a/Engine/Plugins/Runtime/nDisplay/Source/DisplayCluster/Public/DisplayClusterRootActor.h
+++ b/Engine/Plugins/Runtime/nDisplay/Source/DisplayCluster/Public/DisplayClusterRootActor.h
@@ -119,6 +119,9 @@ public:
 	const FDisplayClusterConfigurationICVFX_StageSettings& GetStageSettings() const;
 	const FDisplayClusterConfigurationRenderFrame& GetRenderFrameSettings() const;
 
+	/** The Tick() function that calls the preview renderer for this root actor. */
+	void TickPreviewRenderer();
+
 	/** Returns the current rendering mode of this DCRA (not a value from configuration).
 	 * This value can be overridden from DCRenderDevice or other rendering subsystems (e.g. Preview).
 	 */
diff --git a/Engine/Plugins/Runtime/nDisplay/Source/DisplayCluster/Public/Misc/DisplayClusterObjectRef.h b/Engine/Plugins/Runtime/nDisplay/Source/DisplayCluster/Public/Misc/DisplayClusterObjectRef.h
--- a/Engine/Plugins/Runtime/nDisplay/Source/DisplayCluster/Public/Misc/DisplayClusterObjectRef.h
+++ b/Engine/Plugins/Runtime/nDisplay/Source/DisplayCluster/Public/Misc/DisplayClusterObjectRef.h
@@ -38,6 +38,18 @@ public:
 		return !ActorClassName.IsEmpty() && !ActorName.IsNone();
 	}
 
+	/** Returns the current weak ptr of the scene actor object. */
+	const TWeakObjectPtr<AActor> GetSceneActorWeakPtr() const
+	{
+		FScopeLock lock(&DataGuard);
+		if (!IsDefinedSceneActor())
+		{
+			return nullptr;
+		}
+
+		return ActorPtr;
+	}
+
 	// Return actor object ptr.
 	// For killed object ptr, reset and find actor new object ptr by name and save to [mutable] ActorPtr
 	AActor* GetOrFindSceneActor() const
@@ -119,7 +131,8 @@ private:
 			UWorld* CurrentWorld = nullptr;
 
 #if WITH_EDITOR
-			if (GIsEditor)
+			// Prevent a crash when exiting the UE application and calling this function after deleting GEditor.
+			if (GIsEditor && IsValid(GEditor))
 			{
 				CurrentWorld = GEditor->GetEditorWorldContext().World();
 			}
@@ -127,12 +140,13 @@ private:
 #endif
 			if (UGameEngine* GameEngine = Cast<UGameEngine>(GEngine))
 			{
-				CurrentWorld = GameEngine->GetGameWorld();
+				CurrentWorld = IsValid(GameEngine) ? GameEngine->GetGameWorld() : nullptr;
 			}
 
-			if (!CurrentWorld)
+			if (!IsValid(CurrentWorld))
 			{
 				WorldPtr.Reset();
+
 				return false;
 			}
 
diff --git a/Engine/Plugins/Runtime/nDisplay/Source/DisplayClusterLightCardEditorShaders/Private/DisplayClusterMeshProjectionRenderer.cpp b/Engine/Plugins/Runtime/nDisplay/Source/DisplayClusterLightCardEditorShaders/Private/DisplayClusterMeshProjectionRenderer.cpp
--- a/Engine/Plugins/Runtime/nDisplay/Source/DisplayClusterLightCardEditorShaders/Private/DisplayClusterMeshProjectionRenderer.cpp
+++ b/Engine/Plugins/Runtime/nDisplay/Source/DisplayClusterLightCardEditorShaders/Private/DisplayClusterMeshProjectionRenderer.cpp
@@ -21,6 +21,37 @@
 #include "DataDrivenShaderPlatformInfo.h"
 #include "ScreenPass.h"
 
+namespace UE::DisplayClusterMeshProjectionRenderer
+{
+	/**
+	* Whether this primitive component should be rendered in this view.
+	* Each 'View' belongs to a 'ViewFamily' that belongs to the 'World'.
+	* UE can only render one world for one "View".
+	*/
+	static inline bool ShouldRenderPrimitiveComponentForView(const UPrimitiveComponent* InPrimitiveComponent, const FSceneView* InView)
+	{
+		const FSceneViewFamily* ViewFamily = InView ? InView->Family : nullptr;
+		if (const FSceneInterface* SceneInterface = ViewFamily ? ViewFamily->Scene : nullptr)
+		{
+			// Get primitive info from the scene by primitive ID.
+			if (FPrimitiveSceneInfo* PrimitiveSceneInfo = SceneInterface->GetPrimitiveSceneInfo(InPrimitiveComponent->GetPrimitiveSceneId()))
+			{
+				// The scene primitive identifier is simply a numeric value that can be the same in multiple worlds.
+				// To check if the primitive component is the same instance, let's compare references to FPrimitiveSceneProxy.
+				if (PrimitiveSceneInfo->Proxy == InPrimitiveComponent->GetSceneProxy())
+				{
+					return true;
+				}
+			}
+
+			// This is a component belonging to a different world; skip it for the current rendering pass.
+			return false;
+		}
+
+		return true;
+	}
+};
+
 //////////////////////////////////////////////////////////////////////////
 // Mesh Pass Processor
 
@@ -1561,7 +1592,16 @@ void FDisplayClusterMeshProjectionRenderer::RenderPrimitives_RenderThread(const
 	DrawDynamicMeshPass(*View, RHICmdList, [View, &RenderSettings, bTranslucencyPass, this](FDynamicPassMeshDrawListContext* DynamicMeshPassContext)
 	{
 		TArray<FSceneProxyElement> PrimitiveSceneProxies;
-		GetSceneProxies(PrimitiveSceneProxies, RenderSettings);
+		GetSceneProxies(PrimitiveSceneProxies, RenderSettings, [View](const UPrimitiveComponent* InPrimitiveComponent)
+			{
+				return UE::DisplayClusterMeshProjectionRenderer::ShouldRenderPrimitiveComponentForView(InPrimitiveComponent, View);
+			});
+
+		if (PrimitiveSceneProxies.IsEmpty())
+		{
+			// Skip if there is nothing to render.
+			return;
+		}
 
 		FMeshProjectionPassProcessor<ProjectionType> ProjectionMeshProcessor(nullptr, View, DynamicMeshPassContext, RenderSettings, bTranslucencyPass);
 		FMeshProjectionPassProcessor<EDisplayClusterMeshProjectionType::Linear> LinearMeshProcessor(nullptr, View, DynamicMeshPassContext, RenderSettings, bTranslucencyPass);
@@ -1601,7 +1641,16 @@ void FDisplayClusterMeshProjectionRenderer::RenderHitProxies_RenderThread(const
 	DrawDynamicMeshPass(*View, RHICmdList, [View, &RenderSettings, this](FDynamicPassMeshDrawListContext* DynamicMeshPassContext)
 	{
 		TArray<FSceneProxyElement> PrimitiveSceneProxies;
-		GetSceneProxies(PrimitiveSceneProxies, RenderSettings);
+		GetSceneProxies(PrimitiveSceneProxies, RenderSettings, [View](const UPrimitiveComponent* InPrimitiveComponent)
+			{
+				return UE::DisplayClusterMeshProjectionRenderer::ShouldRenderPrimitiveComponentForView(InPrimitiveComponent, View);
+			});
+
+		if (PrimitiveSceneProxies.IsEmpty())
+		{
+			// Skip if there is nothing to render.
+			return;
+		}
 
 		FMeshProjectionHitProxyPassProcessor<ProjectionType> ProjectionMeshProcessor(nullptr, View, DynamicMeshPassContext, RenderSettings);
 		FMeshProjectionHitProxyPassProcessor<EDisplayClusterMeshProjectionType::Linear> LinearMeshProcessor(nullptr, View, DynamicMeshPassContext, RenderSettings);
@@ -1636,7 +1685,16 @@ void FDisplayClusterMeshProjectionRenderer::RenderNormals_RenderThread(const FSc
 	DrawDynamicMeshPass(*View, RHICmdList, [View, &RenderSettings, this](FDynamicPassMeshDrawListContext* DynamicMeshPassContext)
 	{
 		TArray<FSceneProxyElement> PrimitiveSceneProxies;
-		GetSceneProxies(PrimitiveSceneProxies, RenderSettings);
+		GetSceneProxies(PrimitiveSceneProxies, RenderSettings, [View](const UPrimitiveComponent* InPrimitiveComponent)
+			{
+				return UE::DisplayClusterMeshProjectionRenderer::ShouldRenderPrimitiveComponentForView(InPrimitiveComponent, View);
+			});
+
+		if (PrimitiveSceneProxies.IsEmpty())
+		{
+			// Skip if there is nothing to render.
+			return;
+		}
 
 		FMeshProjectionPassProcessor<ProjectionType, EDisplayClusterMeshProjectionOutput::Normals> ProjectionMeshProcessor(nullptr, View, DynamicMeshPassContext, RenderSettings);
 		ProjectionMeshProcessor.SetStencilValue(1);
@@ -1677,11 +1735,22 @@ void FDisplayClusterMeshProjectionRenderer::RenderSelection_RenderThread(const F
 	DrawDynamicMeshPass(*View, RHICmdList, [View, &RenderSettings, this](FDynamicPassMeshDrawListContext* DynamicMeshPassContext)
 	{
 		TArray<FSceneProxyElement> PrimitiveSceneProxies;
-		GetSceneProxies(PrimitiveSceneProxies, RenderSettings, [](const UPrimitiveComponent* PrimitiveComponent)
+		GetSceneProxies(PrimitiveSceneProxies, RenderSettings, [View](const UPrimitiveComponent* InPrimitiveComponent)
 		{
-			return PrimitiveComponent->SceneProxy->IsSelected();
+			if (!UE::DisplayClusterMeshProjectionRenderer::ShouldRenderPrimitiveComponentForView(InPrimitiveComponent, View))
+			{
+				return false;
+			}
+
+			return InPrimitiveComponent->SceneProxy->IsSelected();
 		});
 
+		if (PrimitiveSceneProxies.IsEmpty())
+		{
+			// Skip if there is nothing to render.
+			return;
+		}
+
 		FMeshProjectionSelectionPassProcessor<ProjectionType> ProjectionMeshProcessor(nullptr, View, DynamicMeshPassContext, RenderSettings);
 		FMeshProjectionSelectionPassProcessor<EDisplayClusterMeshProjectionType::Linear> LinearMeshProcessor(nullptr, View, DynamicMeshPassContext, RenderSettings);
 
@@ -1738,6 +1807,178 @@ void FDisplayClusterMeshProjectionRenderer::GetSceneProxies(TArray<FSceneProxyEl
 	}
 }
 
+void FDisplayClusterMeshProjectionRenderer::RenderScenes(FCanvas* Canvas, const TArray<FSceneInterface*>& InScenes, const FDisplayClusterMeshProjectionRenderSettings& RenderSettings)
+{
+	Canvas->Flush_GameThread();
+
+	/**
+	* This code is based on the existing FDisplayClusterMeshProjectionRenderer::Render() function.
+	* The idea is to create code similar to the existing code.
+	*/
+	FRenderTarget* RenderTarget = Canvas->GetRenderTarget();
+	const bool bIsHitTesting = Canvas->IsHitTesting();
+	FHitProxyConsumer* HitProxyConsumer = Canvas->GetHitProxyConsumer();
+
+	UE::RenderCommandPipe::FSyncScope SyncScope;
+
+	ENQUEUE_RENDER_COMMAND(FDrawProjectedMeshes)(
+		[RenderTarget, Scenes = MoveTempIfPossible(InScenes), RenderSettings, bIsHitTesting, HitProxyConsumer, this](FRHICommandListImmediate& RHICmdList)
+		{
+			FRDGBuilder GraphBuilder(RHICmdList);
+
+			TArray<TSharedRef<FSceneViewFamilyContext>> ViewFamilies;
+			TArray<TSharedRef<FScenePrimitiveRenderingContextScopeHelper>> ScopeHelpers;
+			TArray<const FSceneView*> SceneViews;
+			for (FSceneInterface* Scene : Scenes)
+			{
+				ViewFamilies.Add(
+					MakeShared<FSceneViewFamilyContext>(FSceneViewFamily::ConstructionValues(
+						RenderTarget,
+						Scene,
+						RenderSettings.EngineShowFlags)
+						.SetTime(FGameTime::GetTimeSinceAppStart()))
+				);
+
+				FSceneViewFamilyContext& ViewFamily = ViewFamilies.Last().Get();
+
+				if (Scene)
+				{
+					Scene->IncrementFrameNumber();
+					ViewFamily.FrameNumber = Scene->GetFrameNumber();
+				}
+				else
+				{
+					ViewFamily.FrameNumber = GFrameNumber;
+				}
+
+				ViewFamily.EngineShowFlags.SetHitProxies(bIsHitTesting);
+
+				ScopeHelpers.Add(MakeShared<FScenePrimitiveRenderingContextScopeHelper>(GetRendererModule().BeginScenePrimitiveRendering(GraphBuilder, &ViewFamily)));
+
+				FSceneViewInitOptions NewInitOptions(RenderSettings.ViewInitOptions);
+				NewInitOptions.ViewFamily = &ViewFamily;
+
+				GetRendererModule().CreateAndInitSingleView(RHICmdList, &ViewFamily, &NewInitOptions);
+				SceneViews.Add(ViewFamily.Views[0]);
+			}
+
+			FRDGTextureRef OutputTexture = GraphBuilder.RegisterExternalTexture(CreateRenderTarget(RenderTarget->GetRenderTargetTexture(), TEXT("ViewRenderTarget")));
+			FRenderTargetBinding OutputRenderTargetBinding(OutputTexture, ERenderTargetLoadAction::ELoad);
+			FMeshProjectionElementCollector ElementCollector(HitProxyConsumer);
+
+			switch (RenderSettings.RenderType)
+			{
+			case EDisplayClusterMeshProjectionOutput::Color:
+				if (ViewFamilies[0]->EngineShowFlags.HitProxies)
+				{
+					RenderHitProxyOutput(GraphBuilder, SceneViews[0], RenderSettings, OutputRenderTargetBinding, HitProxyConsumer);
+				}
+				else
+				{
+					RenderColorOutputs(GraphBuilder, SceneViews, RenderSettings, OutputRenderTargetBinding);
+				}
+
+				if (RenderSimpleElementsDelegate.IsBound())
+				{
+					RenderSimpleElementsDelegate.Execute(SceneViews[0], &ElementCollector);
+					AddSimpleElementPass(GraphBuilder, SceneViews[0], OutputRenderTargetBinding, ElementCollector);
+				}
+				break;
+
+			case EDisplayClusterMeshProjectionOutput::Normals:
+				RenderNormalsOutput(GraphBuilder, SceneViews[0], RenderSettings, OutputRenderTargetBinding);
+				break;
+			}
+
+			GraphBuilder.Execute();
+		});
+}
+
+void FDisplayClusterMeshProjectionRenderer::RenderColorOutputs(FRDGBuilder& GraphBuilder,
+	const TArray<const FSceneView*>& Views,
+	const FDisplayClusterMeshProjectionRenderSettings& RenderSettings,
+	FRenderTargetBinding& OutputRenderTargetBinding)
+{
+	FRDGTextureRef ColorTexture = GraphBuilder.CreateTexture(OutputRenderTargetBinding.GetTexture()->Desc, TEXT("DisplayClusterMeshProjection.ColorTexture"));
+
+	const FRDGTextureDesc DepthDesc = FRDGTextureDesc::Create2D(OutputRenderTargetBinding.GetTexture()->Desc.Extent, PF_DepthStencil, FClearValueBinding::DepthFar, TexCreate_DepthStencilTargetable | TexCreate_ShaderResource);
+	FRDGTextureRef DepthTexture = GraphBuilder.CreateTexture(DepthDesc, TEXT("DisplayClusterMeshProjection.DepthTexture"));
+
+	FRenderTargetBinding ColorRenderTargetBinding(ColorTexture, ERenderTargetLoadAction::EClear);
+	FDepthStencilBinding DepthStencilBinding(DepthTexture, ERenderTargetLoadAction::EClear, ERenderTargetLoadAction::ENoAction, FExclusiveDepthStencil::DepthWrite_StencilNop);
+
+	for (const FSceneView* ViewIt : Views)
+	{
+		AddBaseRenderPass(GraphBuilder, ViewIt, RenderSettings, ColorRenderTargetBinding, DepthStencilBinding);
+
+		ColorRenderTargetBinding.SetLoadAction(ERenderTargetLoadAction::ELoad);
+		DepthStencilBinding.SetDepthLoadAction(ERenderTargetLoadAction::ELoad);
+	}
+
+	for (const FSceneView* ViewIt : Views)
+	{
+		AddTranslucencyRenderPass(GraphBuilder, ViewIt, RenderSettings, ColorRenderTargetBinding, DepthStencilBinding);
+	}
+
+#if WITH_EDITOR
+	if (RenderSettings.EngineShowFlags.SelectionOutline)
+	{
+		const FRDGTextureDesc SelectionDepthDesc = FRDGTextureDesc::Create2D(OutputRenderTargetBinding.GetTexture()->Desc.Extent, PF_DepthStencil, FClearValueBinding::DepthFar, TexCreate_DepthStencilTargetable | TexCreate_ShaderResource);
+		FRDGTextureRef SelectionDepthTexture = GraphBuilder.CreateTexture(SelectionDepthDesc, TEXT("DisplayClusterMeshProjection.SelectionDepthTexture"));
+		FDepthStencilBinding SelectionDepthStencilBinding(SelectionDepthTexture, ERenderTargetLoadAction::EClear, ERenderTargetLoadAction::EClear, FExclusiveDepthStencil::DepthWrite_StencilWrite);
+
+		for (const FSceneView* ViewIt : Views)
+		{
+			AddSelectionDepthRenderPass(GraphBuilder, ViewIt, RenderSettings, SelectionDepthStencilBinding);
+		}
+
+		for (const FSceneView* ViewIt : Views)
+		{
+			AddSelectionOutlineScreenPass(GraphBuilder, ViewIt, OutputRenderTargetBinding, ColorTexture, DepthTexture, SelectionDepthTexture);
+		}
+	}
+	else
+#endif
+		// Copy the scene color to the output render target
+	{
+		FCopyRectPS::FParameters* ScreenPassParameters = GraphBuilder.AllocParameters<FCopyRectPS::FParameters>();
+		ScreenPassParameters->InputTexture = ColorTexture;
+		ScreenPassParameters->InputSampler = TStaticSamplerState<>::GetRHI();
+		ScreenPassParameters->RenderTargets[0] = OutputRenderTargetBinding;
+
+		FGlobalShaderMap* GlobalShaderMap = GetGlobalShaderMap(Views[0]->FeatureLevel);
+		TShaderMapRef<FScreenPassVS> ScreenPassVS(GlobalShaderMap);
+		TShaderMapRef<FCopyRectPS> CopyPixelShader(GlobalShaderMap);
+		if (!ScreenPassVS.IsValid() || !CopyPixelShader.IsValid())
+		{
+			// Always check if shaders are available on the current platform and hardware
+			return;
+		}
+
+		FRHIBlendState* DefaultBlendState = FScreenPassPipelineState::FDefaultBlendState::GetRHI();
+		const FScreenPassTextureViewport RegionViewport(OutputRenderTargetBinding.GetTexture());
+
+		GraphBuilder.AddPass(
+			RDG_EVENT_NAME("MeshProjectionRenderer::CopyColorTexture"),
+			ScreenPassParameters,
+			ERDGPassFlags::Raster,
+			[View = Views[0], ScreenPassVS, CopyPixelShader, RegionViewport, ScreenPassParameters, DefaultBlendState](FRHICommandList& RHICmdList)
+			{
+				DrawScreenPass(
+					RHICmdList,
+					*View,
+					RegionViewport,
+					RegionViewport,
+					FScreenPassPipelineState(ScreenPassVS, CopyPixelShader, DefaultBlendState),
+					EScreenPassDrawFlags::None,
+					[&](FRHICommandList&)
+					{
+						SetShaderParameters(RHICmdList, CopyPixelShader, CopyPixelShader.GetPixelShader(), *ScreenPassParameters);
+					});
+			});
+	}
+}
+
 // Explicit template specializations for each projection type
 #define PROJECTION_TYPE_TEMPLATE_SPECIALIZATION(FuncName, ProjectionType, ...) template void FDisplayClusterMeshProjectionRenderer::FuncName<ProjectionType>(__VA_ARGS__);
 #define CREATE_PROJECTION_TYPE_TEMPLATE_SPECIALIZATIONS(FuncName, ...) \
diff --git a/Engine/Plugins/Runtime/nDisplay/Source/DisplayClusterLightCardEditorShaders/Public/DisplayClusterMeshProjectionRenderer.h b/Engine/Plugins/Runtime/nDisplay/Source/DisplayClusterLightCardEditorShaders/Public/DisplayClusterMeshProjectionRenderer.h
--- a/Engine/Plugins/Runtime/nDisplay/Source/DisplayClusterLightCardEditorShaders/Public/DisplayClusterMeshProjectionRenderer.h
+++ b/Engine/Plugins/Runtime/nDisplay/Source/DisplayClusterLightCardEditorShaders/Public/DisplayClusterMeshProjectionRenderer.h
@@ -159,9 +159,17 @@ public:
 	/** Clears the list of primitives to render */
 	DISPLAYCLUSTERLIGHTCARDEDITORSHADERS_API void ClearScene();
 
-	/** Renders its list of primitive components to the specified canvas using the desired projection type. Can be called from the game thread */
+	/**
+	* Renders its list of primitive components to the specified canvas using the desired projection type. Can be called from either the game thread or the rendering thread.
+	*/
 	DISPLAYCLUSTERLIGHTCARDEDITORSHADERS_API void Render(FCanvas* Canvas, FSceneInterface* Scene, const FDisplayClusterMeshProjectionRenderSettings& RenderSettings);
 
+	/**
+	* Renders its list of primitive components to the specified canvas using the desired projection type. Can be called from either the game thread or the rendering thread.
+	* Render scene by scene
+	*/
+	DISPLAYCLUSTERLIGHTCARDEDITORSHADERS_API void RenderScenes(FCanvas* Canvas, const TArray<FSceneInterface*>& Scenes, const FDisplayClusterMeshProjectionRenderSettings& RenderSettings);
+
 private:
 	/** Constructs the necessary render passes for the default output of the rendered primitives */
 	void RenderColorOutput(FRDGBuilder& GraphBuilder,
@@ -169,6 +177,12 @@ private:
 		const FDisplayClusterMeshProjectionRenderSettings& RenderSettings,
 		FRenderTargetBinding& OutputRenderTargetBinding);
 
+	/** Constructs the necessary render passes for the default output of the rendered primitives */
+	void RenderColorOutputs(FRDGBuilder& GraphBuilder,
+		const TArray<const FSceneView*>& Views,
+		const FDisplayClusterMeshProjectionRenderSettings& RenderSettings,
+		FRenderTargetBinding& OutputRenderTargetBinding);
+
 	/** Constructs the necessary render passes for the hit proxy output of the rendered primitives */
 	void RenderHitProxyOutput(FRDGBuilder& GraphBuilder,
 		const FSceneView* View,
diff --git a/Engine/Plugins/Runtime/nDisplay/Source/DisplayClusterScenePreview/Private/DisplayClusterScenePreviewModule.cpp b/Engine/Plugins/Runtime/nDisplay/Source/DisplayClusterScenePreview/Private/DisplayClusterScenePreviewModule.cpp
--- a/Engine/Plugins/Runtime/nDisplay/Source/DisplayClusterScenePreview/Private/DisplayClusterScenePreviewModule.cpp
+++ b/Engine/Plugins/Runtime/nDisplay/Source/DisplayClusterScenePreview/Private/DisplayClusterScenePreviewModule.cpp
@@ -1,6 +1,7 @@
 // Copyright Epic Games, Inc. All Rights Reserved.
 
 #include "DisplayClusterScenePreviewModule.h"
+#include "DisplayClusterScenePreviewProxyManager.h"
 
 #include "Blueprints/DisplayClusterBlueprintLib.h"
 #include "CanvasTypes.h"
@@ -27,13 +28,20 @@ static TAutoConsoleVariable<float> CVarDisplayClusterScenePreviewRenderTickDelay
 );
 
 void FDisplayClusterScenePreviewModule::StartupModule()
+{ }
+
+void FDisplayClusterScenePreviewModule::ShutdownModule()
 {
+	Release();
 }
 
-void FDisplayClusterScenePreviewModule::ShutdownModule()
+void FDisplayClusterScenePreviewModule::Release()
 {
-	FTSTicker::GetCoreTicker().RemoveTicker(RenderTickerHandle);
-	RenderTickerHandle.Reset();
+	if (RenderTickerHandle.IsValid())
+	{
+		FTSTicker::GetCoreTicker().RemoveTicker(RenderTickerHandle);
+		RenderTickerHandle.Reset();
+	}
 
 	TArray<int32> RendererIds;
 	RendererConfigs.GenerateKeyArray(RendererIds);
@@ -41,66 +49,66 @@ void FDisplayClusterScenePreviewModule::ShutdownModule()
 	{
 		DestroyRenderer(RendererId);
 	}
+
+	// Release the RootActor and PreviewWorld proxies.
+	ProxyManagerPtr.Reset();
 }
 
-int32 FDisplayClusterScenePreviewModule::CreateRenderer()
+#if WITH_EDITOR
+void FDisplayClusterScenePreviewModule::OnEditorClosed()
 {
-	FRendererConfig& Config = RendererConfigs.Add(NextRendererId);
+	// On editor close, Exit() should run to clean up, but this happens very late.
+	Release();
 
-	Config.Renderer = MakeShared<FDisplayClusterMeshProjectionRenderer>();
-
-	return NextRendererId++;
+	if (EditorClosedEventHandle.IsValid() && GEditor)
+	{
+		GEditor->OnEditorClose().Remove(EditorClosedEventHandle);
+	}
 }
+#endif // #if WITH_EDITOR
 
-bool FDisplayClusterScenePreviewModule::DestroyRenderer(int32 RendererId)
+FDisplayClusterScenePreviewProxyManager& FDisplayClusterScenePreviewModule::GetProxyManager()
 {
-	if (FRendererConfig* Config = RendererConfigs.Find(RendererId))
+#if WITH_EDITOR
+	if (!EditorClosedEventHandle.IsValid() && GEditor)
 	{
-		// Release the proxy resources that were used by this renderer.
-		ProxyManager.SetSceneRootActorForRenderer(RendererId, nullptr);
-
-		RegisterRootActorEvents(RendererId, *Config, false);
-		RendererConfigs.Remove(RendererId);
-
-		RegisterOrUnregisterGlobalActorEvents();
+		EditorClosedEventHandle = GEditor->OnEditorClose().AddRaw(this, &FDisplayClusterScenePreviewModule::OnEditorClosed);
+	}
+#endif //#if WITH_EDITOR
 
-		return true;
+	if (!ProxyManagerPtr.IsValid())
+	{
+		ProxyManagerPtr = MakeUnique<FDisplayClusterScenePreviewProxyManager>();
 	}
 
-	return false;
-}
+	check(ProxyManagerPtr.IsValid());
 
-bool FDisplayClusterScenePreviewModule::SetRendererRootActorPath(int32 RendererId, const FString& ActorPath, const FDisplayClusterRootActorPropertyOverrides& InPropertyOverrides, const EDisplayClusterScenePreviewFlags PreviewFlags)
-{
-	if (FRendererConfig* Config = RendererConfigs.Find(RendererId))
-	{
-		Config->RootActorPath = ActorPath;
+	return *ProxyManagerPtr;
+}
 
-		ADisplayClusterRootActor* RootActor = FindObject<ADisplayClusterRootActor>(nullptr, *ActorPath);
-		InternalSetRendererRootActor(RendererId, *Config, RootActor, PreviewFlags);
 
-		// Use custom properties on root actor
-		Config->RootActorPropertyOverrides = InPropertyOverrides;
-		InternalOverridePropertiesForRendererRootActor(RendererId, *Config);
+int32 FDisplayClusterScenePreviewModule::CreateRenderer()
+{
+	FRendererConfig& Config = RendererConfigs.Add(NextRendererId);
 
-		return true;
-	}
+	Config.Renderer = MakeShared<FDisplayClusterMeshProjectionRenderer>();
 
-	return false;
+	return NextRendererId++;
 }
 
-bool FDisplayClusterScenePreviewModule::SetRendererRootActor(int32 RendererId, ADisplayClusterRootActor* Actor, const FDisplayClusterRootActorPropertyOverrides& InPropertyOverrides, const EDisplayClusterScenePreviewFlags PreviewFlags)
+bool FDisplayClusterScenePreviewModule::DestroyRenderer(int32 RendererId)
 {
 	if (FRendererConfig* Config = RendererConfigs.Find(RendererId))
 	{
-		Config->RootActorPath.Empty();
-		Config->RootActorPropertyOverrides = InPropertyOverrides;
+		// Release the proxy resources that were used by this renderer.
+		GetProxyManager().SetSceneRootActorForRenderer(RendererId, nullptr, EDisplayClusterScenePreviewFlags::None);
 
-		InternalSetRendererRootActor(RendererId, *Config, Actor, PreviewFlags);
+		RegisterRootActorEvents(RendererId, *Config, false /* bShouldRegister */);
 
-		// Use custom properties on root actor
-		Config->RootActorPropertyOverrides = InPropertyOverrides;
-		InternalOverridePropertiesForRendererRootActor(RendererId, *Config);
+		// Remove this config.
+		RendererConfigs.Remove(RendererId);
+
+		RegisterOrUnregisterGlobalActorEvents();
 
 		return true;
 	}
@@ -108,108 +116,81 @@ bool FDisplayClusterScenePreviewModule::SetRendererRootActor(int32 RendererId, A
 	return false;
 }
 
-ADisplayClusterRootActor* FDisplayClusterScenePreviewModule::GetRendererRootActor(int32 RendererId)
+bool FDisplayClusterScenePreviewModule::SetRendererRootActorPath(int32 RendererId, const FString& ActorPath, const FDisplayClusterRootActorPropertyOverrides& InPropertyOverrides, const EDisplayClusterScenePreviewFlags PreviewFlags)
 {
-	if (FRendererConfig* Config = RendererConfigs.Find(RendererId))
+	if (FRendererConfig* RendererConfig = RendererConfigs.Find(RendererId))
 	{
-		return InternalGetRendererRootActorOrProxy(RendererId, *Config);
-	}
-
-	return nullptr;
-}
+		// Use custom properties on root actor
+		RendererConfig->RootActorPropertyOverrides = InPropertyOverrides;
+		EnumAddFlags(RendererConfig->Flags, EDisplayClusterRendererConfigFlags::RootActorPropertyOverridesModified);
 
-bool FDisplayClusterScenePreviewModule::GetActorsInRendererScene(int32 RendererId, bool bIncludeRoot, TArray<AActor*>& OutActors)
-{
-	if (FRendererConfig* Config = RendererConfigs.Find(RendererId))
-	{
-		if (bIncludeRoot)
+		// Determine these values before we update the config's RootActor
+		if (RendererConfig->PreviewFlags != PreviewFlags)
 		{
-			OutActors.Add(GetRendererRootActor(RendererId));
+			RendererConfig->PreviewFlags = PreviewFlags;
+			EnumAddFlags(RendererConfig->Flags, EDisplayClusterRendererConfigFlags::PreviewFlagsModified);
 		}
 
-		for (const TWeakObjectPtr<AActor>& Actor : Config->AddedActors)
+		// Update root actor path
+		if (RendererConfig->RootActorPath != ActorPath)
 		{
-			if (!Actor.IsValid())
-			{
-				continue;
-			}
-
-			OutActors.Add(Actor.Get());
+			RendererConfig->RootActorPath = ActorPath;
+			EnumAddFlags(RendererConfig->Flags, EDisplayClusterRendererConfigFlags::RootActorPathModified);
 		}
 
-		return true;
-	}
-
-	return false;
-}
+		InternalUpdateRendererConfig(RendererId, *RendererConfig);
 
-bool FDisplayClusterScenePreviewModule::AddActorToRenderer(int32 RendererId, AActor* Actor)
-{
-	if (FRendererConfig* Config = RendererConfigs.Find(RendererId))
-	{
-		Config->Renderer->AddActor(Actor);
-		Config->AddedActors.Add(Actor);
 		return true;
 	}
 
 	return false;
 }
 
-bool FDisplayClusterScenePreviewModule::AddActorToRenderer(int32 RendererId, AActor* Actor, const TFunctionRef<bool(const UPrimitiveComponent*)>& PrimitiveFilter)
+bool FDisplayClusterScenePreviewModule::SetRendererRootActor(int32 RendererId, ADisplayClusterRootActor* Actor, const FDisplayClusterRootActorPropertyOverrides& InPropertyOverrides, const EDisplayClusterScenePreviewFlags PreviewFlags)
 {
-	if (FRendererConfig* Config = RendererConfigs.Find(RendererId))
+	if (FRendererConfig* RendererConfig = RendererConfigs.Find(RendererId))
 	{
-		Config->Renderer->AddActor(Actor, PrimitiveFilter);
-		return true;
-	}
+		// Use custom properties on root actor
+		RendererConfig->RootActorPropertyOverrides = InPropertyOverrides;
+		EnumAddFlags(RendererConfig->Flags, EDisplayClusterRendererConfigFlags::RootActorPropertyOverridesModified);
 
-	return false;
-}
+		// Determine these values before we update the config's RootActor
+		if (RendererConfig->PreviewFlags != PreviewFlags)
+		{
+			RendererConfig->PreviewFlags = PreviewFlags;
+			EnumAddFlags(RendererConfig->Flags, EDisplayClusterRendererConfigFlags::PreviewFlagsModified);
+		}
 
-bool FDisplayClusterScenePreviewModule::RemoveActorFromRenderer(int32 RendererId, AActor* Actor)
-{
-	if (FRendererConfig* Config = RendererConfigs.Find(RendererId))
-	{
-		Config->Renderer->RemoveActor(Actor);
-		return true;
-	}
+		// Clear root actor path
+		RendererConfig->RootActorPath.Empty();
 
-	return false;
-}
+		InternalSetRendererRootActor(RendererId, *RendererConfig, Actor);
+		InternalUpdateRendererConfig(RendererId, *RendererConfig);
 
-bool FDisplayClusterScenePreviewModule::ClearRendererScene(int32 RendererId)
-{
-	if (FRendererConfig* Config = RendererConfigs.Find(RendererId))
-	{
-		Config->Renderer->ClearScene();
-		Config->AddedActors.Empty();
-		Config->AutoActors.Empty();
 		return true;
 	}
 
 	return false;
 }
 
-bool FDisplayClusterScenePreviewModule::SetRendererActorSelectedDelegate(int32 RendererId, FDisplayClusterMeshProjectionRenderer::FSelection ActorSelectedDelegate)
+ADisplayClusterRootActor* FDisplayClusterScenePreviewModule::GetRendererRootActor(int32 RendererId)
 {
 	if (FRendererConfig* Config = RendererConfigs.Find(RendererId))
 	{
-		Config->Renderer->ActorSelectedDelegate = ActorSelectedDelegate;
-		return true;
+		return InternalGetRendererRootActor(RendererId, *Config);
 	}
 
-	return false;
+	return nullptr;
 }
 
-bool FDisplayClusterScenePreviewModule::SetRendererRenderSimpleElementsDelegate(int32 RendererId, FDisplayClusterMeshProjectionRenderer::FSimpleElementPass RenderSimpleElementsDelegate)
+ADisplayClusterRootActor* FDisplayClusterScenePreviewModule::GetRendererRootActorOrProxy(int32 RendererId)
 {
 	if (FRendererConfig* Config = RendererConfigs.Find(RendererId))
 	{
-		Config->Renderer->RenderSimpleElementsDelegate = RenderSimpleElementsDelegate;
-		return true;
+		return InternalGetRendererRootActorOrProxy(RendererId, *Config);
 	}
 
-	return false;
+	return nullptr;
 }
 
 bool FDisplayClusterScenePreviewModule::Render(int32 RendererId, FDisplayClusterMeshProjectionRenderSettings& RenderSettings, FCanvas& Canvas)
@@ -243,37 +224,61 @@ bool FDisplayClusterScenePreviewModule::RenderQueued(int32 RendererId, FDisplayC
 	return InternalRenderQueued(RendererId, RenderSettings, Canvas, RenderTarget->GetSizeXY(), ResultDelegate);
 }
 
+bool FDisplayClusterScenePreviewModule::InternalRenderQueued(int32 RendererId, FDisplayClusterMeshProjectionRenderSettings& RenderSettings, TWeakPtr<FCanvas> Canvas,
+	const FIntPoint& Size, FRenderResultDelegate ResultDelegate)
+{
+	if (FRendererConfig* Config = RendererConfigs.Find(RendererId))
+	{
+		RenderQueue.Enqueue(FPreviewRenderJob(RendererId, RenderSettings, Size, Canvas, ResultDelegate));
+
+		if (!RenderTickerHandle.IsValid())
+		{
+			RenderTickerHandle = FTSTicker::GetCoreTicker().AddTicker(
+				FTickerDelegate::CreateRaw(this, &FDisplayClusterScenePreviewModule::OnTick),
+				CVarDisplayClusterScenePreviewRenderTickDelay.GetValueOnGameThread()
+			);
+		}
+
+		return true;
+	}
+
+	return false;
+}
+
 bool FDisplayClusterScenePreviewModule::IsRealTimePreviewEnabled() const
 {
 	return bIsRealTimePreviewEnabled;
 }
 
-bool FDisplayClusterScenePreviewModule::IsBlueprintMatchesRendererRootActor(int32 RendererId, FRendererConfig& RendererConfig, UBlueprint* Blueprint)
+void FDisplayClusterScenePreviewModule::InternalSetRendererRootActor(int32 RendererId, FRendererConfig& RendererConfig, ADisplayClusterRootActor* Actor)
 {
-#if WITH_EDITOR
-	if (RendererConfig.RootActor.IsValid() && Blueprint == UBlueprint::GetBlueprintFromClass(RendererConfig.RootActor->GetClass()))
+	const bool bRootChanged = !RendererConfig.IsRootActorEquals(Actor);
+	if (bRootChanged)
 	{
-		return true;
-	}
-#endif
+		// Unregister events for current root actor
+		RegisterRootActorEvents(RendererId, RendererConfig, false /* bShouldRegister */);
 
-	return false;
+		RendererConfig.SetRootActor(Actor);
+	}
 }
 
 ADisplayClusterRootActor* FDisplayClusterScenePreviewModule::InternalGetRendererRootActor(int32 RendererId, FRendererConfig& RendererConfig)
 {
-	if (!RendererConfig.RootActor.IsValid() && !RendererConfig.RootActorPath.IsEmpty())
+	ADisplayClusterRootActor* SceneRootActor = RendererConfig.GetRootActor();
+
+	const bool bRequiredRootActor = !IsValid(SceneRootActor)
+								|| !RendererConfig.IsDefinedRootActor()
+								|| EnumHasAnyFlags(RendererConfig.Flags, EDisplayClusterRendererConfigFlags::RootActorPathModified);
+	if (bRequiredRootActor && !RendererConfig.RootActorPath.IsEmpty())
 	{
-		// Try to find the actor by its path
+		// If we don't have a RootActor, but we do have a RootActorPath, use it to find the RootActor object.
 		ADisplayClusterRootActor* RootActor = FindObject<ADisplayClusterRootActor>(nullptr, *RendererConfig.RootActorPath);
-		if (RootActor)
-		{
-			InternalSetRendererRootActor(RendererId, RendererConfig, RootActor, RendererConfig.PreviewFlags);
-			InternalOverridePropertiesForRendererRootActor(RendererId, RendererConfig);
-		}
+		InternalSetRendererRootActor(RendererId, RendererConfig, RootActor);
+
+		SceneRootActor = RendererConfig.GetRootActor();
 	}
 
-	return RendererConfig.RootActor.Get();
+	return IsValid(SceneRootActor) ? SceneRootActor : nullptr;
 }
 
 ADisplayClusterRootActor* FDisplayClusterScenePreviewModule::InternalGetRendererRootActorOrProxy(int32 RendererId, FRendererConfig& RendererConfig)
@@ -284,91 +289,184 @@ ADisplayClusterRootActor* FDisplayClusterScenePreviewModule::InternalGetRenderer
 	// When proxy is required, get it from ProxyManager.
 	const bool bUseRootActorProxy = EnumHasAnyFlags(RendererConfig.PreviewFlags, EDisplayClusterScenePreviewFlags::UseRootActorProxy);
 
-	return bUseRootActorProxy ? ProxyManager.GetProxyRootActor(RendererId) : RootActor;
+	return bUseRootActorProxy ? GetProxyManager().GetProxyRootActor(RendererId) : RootActor;
 }
 
-bool FDisplayClusterScenePreviewModule::InternalRenderQueued(int32 RendererId, FDisplayClusterMeshProjectionRenderSettings& RenderSettings, TWeakPtr<FCanvas> Canvas,
-	const FIntPoint& Size, FRenderResultDelegate ResultDelegate)
+void FDisplayClusterScenePreviewModule::InternalUpdateRendererConfig(int32 RendererId, FRendererConfig& RendererConfig)
 {
-	if (FRendererConfig* Config = RendererConfigs.Find(RendererId))
+	// Get RootActor from scene
+	// Call this function at the very beginning because it can internally change the RendererConfig.Flags.
+	ADisplayClusterRootActor* SceneRootActor = InternalGetRendererRootActor(RendererId, RendererConfig);
+
+	// Check that the RootActor proxy still has the same reference to the RootActor in the scene as the renderer.
+	if (EnumHasAnyFlags(RendererConfig.PreviewFlags, EDisplayClusterScenePreviewFlags::UseRootActorProxy))
 	{
-		RenderQueue.Enqueue(FPreviewRenderJob(RendererId, RenderSettings, Size, Canvas, ResultDelegate));
+		ADisplayClusterRootActor* PrevSceneRootActor = GetProxyManager().GetSceneRootActor(RendererId);
+		if (PrevSceneRootActor != SceneRootActor)
+		{
+			// The RootActor proxy is no longer valid, have to be re-created
+			EnumAddFlags(RendererConfig.Flags, EDisplayClusterRendererConfigFlags::UpdateRootActorRef);
+		}
+	}
 
-		if (!RenderTickerHandle.IsValid())
+	const bool bAutoUpdateStageActors = EnumHasAnyFlags(RendererConfig.PreviewFlags, EDisplayClusterScenePreviewFlags::AutoUpdateStageActors);
+	if (bAutoUpdateStageActors)
+	{
+		// If the RootActor is missing or any changes have occurred, we must clear the rendering scene
+		if (!SceneRootActor || RendererConfig.Flags != EDisplayClusterRendererConfigFlags::None)
 		{
-			RenderTickerHandle = FTSTicker::GetCoreTicker().AddTicker(
-				FTickerDelegate::CreateRaw(this, &FDisplayClusterScenePreviewModule::OnTick),
-				CVarDisplayClusterScenePreviewRenderTickDelay.GetValueOnGameThread()
-			);
+			RendererConfig.ClearRendererScene();
 		}
+	}
 
-		return true;
+	// If the flags have not changed, do nothing
+	if (RendererConfig.Flags == EDisplayClusterRendererConfigFlags::None)
+	{
+		return;
 	}
 
-	return false;
-}
+	// Optionally, update the proxy when change the RootActor.
+	if (EnumHasAnyFlags(RendererConfig.Flags, EDisplayClusterRendererConfigFlags::UpdateRootActorRef
+											| EDisplayClusterRendererConfigFlags::PreviewFlagsModified
+											| EDisplayClusterRendererConfigFlags::RootActorBeingModified))
+	{
+		// Updates the proxy for the new root actor.
+		if (EnumHasAnyFlags(RendererConfig.PreviewFlags, EDisplayClusterScenePreviewFlags::UseRootActorProxy))
+		{
+			// Create new RootActor proxy
+			GetProxyManager().SetSceneRootActorForRenderer(RendererId, SceneRootActor, RendererConfig.PreviewFlags);
 
+			EnumAddFlags(RendererConfig.Flags, EDisplayClusterRendererConfigFlags::UpdateRootActorProxyRef);
+		}
+		else
+		{
+			// Remove proxy if exists.
+			GetProxyManager().SetSceneRootActorForRenderer(RendererId, nullptr, RendererConfig.PreviewFlags);
+		}
 
-void FDisplayClusterScenePreviewModule::InternalOverridePropertiesForRendererRootActor(int32 RendererId, FRendererConfig& RendererConfig)
-{
-	// Override root actor properties:
-	if (ADisplayClusterRootActor* RendererRootActor = InternalGetRendererRootActorOrProxy(RendererId, RendererConfig))
+		// Update global events
+		RegisterRootActorEvents(RendererId, RendererConfig, true /* bShouldRegister */);
+		RegisterOrUnregisterGlobalActorEvents();
+	}
+
+	// RootActor property overrides has been changed
+	if (EnumHasAnyFlags(RendererConfig.Flags, EDisplayClusterRendererConfigFlags::RootActorPropertyOverridesModified | EDisplayClusterRendererConfigFlags::UpdateRootActorRef | EDisplayClusterRendererConfigFlags::UpdateRootActorProxyRef))
 	{
-		RendererRootActor->OverrideRootActorProperties(RendererConfig.RootActorPropertyOverrides);
+		if (ADisplayClusterRootActor* RendererRootActorOrProxy = InternalGetRendererRootActorOrProxy(RendererId, RendererConfig))
+		{
+			RendererRootActorOrProxy->OverrideRootActorProperties(RendererConfig.RootActorPropertyOverrides);
+		}
 	}
+
+	AutoPopulateScene(RendererId, RendererConfig);
+
+	// Reset the flags at the end.
+	RendererConfig.Flags = EDisplayClusterRendererConfigFlags::None;
 }
 
-void FDisplayClusterScenePreviewModule::InternalSetRendererRootActor(int32 RendererId, FRendererConfig& RendererConfig, ADisplayClusterRootActor* Actor, const EDisplayClusterScenePreviewFlags PreviewFlags)
+void FDisplayClusterScenePreviewModule::AutoPopulateScene(int32 RendererId, FRendererConfig& RendererConfig)
 {
-	// Determine these values before we update the config's RootActor/bAutoUpdateLightcards
-	const bool bRootChanged = RendererConfig.RootActor != Actor;
-
-	if (bRootChanged || RendererConfig.PreviewFlags != PreviewFlags)
+	// The renderer can use a proxy.
+	ADisplayClusterRootActor* RootActor = InternalGetRendererRootActor(RendererId, RendererConfig);
+	ADisplayClusterRootActor* RootActorProxy = InternalGetRendererRootActorOrProxy(RendererId, RendererConfig);
+	if (RootActor && RootActorProxy)
 	{
-		// Unregister events for the previous cluster
-		const bool bAutoUpdateLightcards = EnumHasAnyFlags(RendererConfig.PreviewFlags, EDisplayClusterScenePreviewFlags::AutoUpdateLightcards);
-		if (bAutoUpdateLightcards)
+		TArray<FString> ProjectionMeshNames;
+
+		if (UDisplayClusterConfigurationData* Config = RootActor->GetConfigData())
 		{
-			RegisterRootActorEvents(RendererId, RendererConfig, false);
+			Config->GetReferencedMeshNames(ProjectionMeshNames);
 		}
 
-		RendererConfig.RootActor = Actor;
-		RendererConfig.PreviewFlags = PreviewFlags;
+		RendererConfig.AddActorToRenderer(RootActorProxy, [&ProjectionMeshNames](const UPrimitiveComponent* PrimitiveComponent)
+			{
+				// Filter out any primitive component that isn't a projection mesh (a static mesh that has a Mesh projection configured for it) or a screen component
+				const bool bIsProjectionMesh = PrimitiveComponent->IsA<UStaticMeshComponent>() && ProjectionMeshNames.Contains(PrimitiveComponent->GetName());
+				const bool bIsScreen = PrimitiveComponent->IsA<UDisplayClusterScreenComponent>();
+				return bIsProjectionMesh || bIsScreen;
+			});
+
+		const bool bAutoUpdateStageActors = EnumHasAnyFlags(RendererConfig.PreviewFlags, EDisplayClusterScenePreviewFlags::AutoUpdateStageActors);
+		if (bAutoUpdateStageActors)
+		{
+			// Automatically add the lightcards found on this actor
+			TSet<ADisplayClusterLightCardActor*> LightCards;
+			UDisplayClusterBlueprintLib::FindLightCardsForRootActor(RootActor, LightCards);
+
+			TSet<ADisplayClusterChromakeyCardActor*> ChromaKeyCards;
+			UDisplayClusterBlueprintLib::FindChromakeyCardsForRootActor(RootActor, ChromaKeyCards);
+			LightCards.Append(reinterpret_cast<TSet<ADisplayClusterLightCardActor*>&>(ChromaKeyCards));
 
-		// Updates the proxy for the new root actor.
-		const bool bUseRootActorProxy = EnumHasAnyFlags(RendererConfig.PreviewFlags, EDisplayClusterScenePreviewFlags::UseRootActorProxy);
-		ProxyManager.SetSceneRootActorForRenderer(RendererId, bUseRootActorProxy ? Actor : nullptr);
+			TSet<AActor*> Actors;
+			Actors.Reserve(LightCards.Num());
+			for (ADisplayClusterLightCardActor* LightCard : LightCards)
+			{
+				Actors.Add(LightCard);
+			}
 
-		AutoPopulateScene(RendererId, RendererConfig);
-	}
+			// Also check for any non-lightcard actors in the world that are valid to control from ICVFX editors
+			if (UWorld* World = RootActor->GetWorld())
+			{
+				for (const TWeakObjectPtr<AActor> WeakActor : TActorRange<AActor>(World))
+				{
+					if (WeakActor.IsValid() && WeakActor->Implements<UDisplayClusterStageActor>() && !WeakActor->IsA<ADisplayClusterLightCardActor>())
+					{
+						Actors.Add(WeakActor.Get());
+					}
+				}
+			}
 
-	RegisterRootActorEvents(RendererId, RendererConfig, true);
-	RegisterOrUnregisterGlobalActorEvents();
+			for (AActor* Actor : Actors)
+			{
+				RendererConfig.AddActorToRenderer(Actor, true);
+			}
+		}
+	}
 }
 
 bool FDisplayClusterScenePreviewModule::InternalRenderImmediate(int32 RendererId, FRendererConfig& RendererConfig, FDisplayClusterMeshProjectionRenderSettings& RenderSettings, FCanvas& Canvas)
 {
+	if (!RendererConfig.Renderer.IsValid())
+	{
+		return false;
+	}
+
 	// Update this so that whoever gets the callback can immediately check whether the nDisplay preview may be out of date
 	UpdateIsRealTimePreviewEnabled();
 
+	InternalUpdateRendererConfig(RendererId, RendererConfig);
+
 	// Get the Root Actor or proxy for rendering previews.
-	ADisplayClusterRootActor* RootActor = InternalGetRendererRootActorOrProxy(RendererId, RendererConfig);
-	UWorld* World = RootActor ? RootActor->GetWorld() : nullptr;
-	if (!World)
+	ADisplayClusterRootActor* ProxyRootActor = InternalGetRendererRootActorOrProxy(RendererId, RendererConfig);
+	ADisplayClusterRootActor* SceneRootActor = InternalGetRendererRootActor(RendererId, RendererConfig);
+	UWorld* PreviewWorld = ProxyRootActor ? ProxyRootActor->GetWorld() : nullptr;
+	UWorld* SceneWorld = SceneRootActor ? SceneRootActor->GetWorld() : nullptr;
+	if (!PreviewWorld || !SceneWorld)
 	{
 		return false;
 	}
 
-	const bool bAutoUpdateLightcards = EnumHasAnyFlags(RendererConfig.PreviewFlags, EDisplayClusterScenePreviewFlags::AutoUpdateLightcards);
-	if (bAutoUpdateLightcards && RendererConfig.bIsSceneDirty)
-	{
-		AutoPopulateScene(RendererId, RendererConfig);
-	}
-
 	// Push any deferred render state updates to ensure that light card positions, preview meshes modified above, etc. are up to date
-	World->SendAllEndOfFrameUpdates();
+	PreviewWorld->SendAllEndOfFrameUpdates();
 
-	RendererConfig.Renderer->Render(&Canvas, World->Scene, RenderSettings);
+	// FDisplayClusterMeshProjectionRenderer uses references to UPrimitiveComponent.
+	// These primitives are taken from RootActorProxy in PreviewWorld and StageActors in SceneWorld.
+	// But we can't render UPrimitiveComponent belonging to different scenes at the same time,
+	// because they reference the internal data of the scene they belong to.
+	// The FDisplayClusterMeshProjectionRenderer::RenderScenes() function implements a new approach that
+	// allows us to render a UPrimitiveComponent from multiple worlds.
+	if (PreviewWorld != SceneWorld)
+	{
+		// Rendering is performed for all scenes in the array in the same order.
+		// First, the RootActorProxy geometries from PreviewWorld are drawn.
+		// Then StageActors from SceneWorld are drawn.
+		// Here is the rendering order: ClearRTT-> RootActorProxy -> StageActors.
+		RendererConfig.Renderer->RenderScenes(&Canvas, { PreviewWorld->Scene, SceneWorld->Scene }, RenderSettings);
+	}
+	else
+	{
+		RendererConfig.Renderer->Render(&Canvas, PreviewWorld->Scene, RenderSettings);
+	}
 
 	return true;
 }
@@ -379,8 +477,8 @@ void FDisplayClusterScenePreviewModule::RegisterOrUnregisterGlobalActorEvents()
 	bool bShouldBeRegistered = false;
 	for (const TPair<int32, FRendererConfig>& ConfigPair : RendererConfigs)
 	{
-		const bool bAutoUpdateLightcards = EnumHasAnyFlags(ConfigPair.Value.PreviewFlags, EDisplayClusterScenePreviewFlags::AutoUpdateLightcards);
-		if (bAutoUpdateLightcards)
+		const bool bAutoUpdateStageActors = EnumHasAnyFlags(ConfigPair.Value.PreviewFlags, EDisplayClusterScenePreviewFlags::AutoUpdateStageActors);
+		if (bAutoUpdateStageActors)
 		{
 			bShouldBeRegistered = true;
 			break;
@@ -418,7 +516,7 @@ void FDisplayClusterScenePreviewModule::RegisterOrUnregisterGlobalActorEvents()
 void FDisplayClusterScenePreviewModule::RegisterRootActorEvents(int32 RendererId, FRendererConfig& RendererConfig, bool bShouldRegister)
 {
 #if WITH_EDITOR
-	ADisplayClusterRootActor* Actor = RendererConfig.RootActor.Get();
+	ADisplayClusterRootActor* Actor = RendererConfig.GetRootActor();
 	if (!Actor)
 	{
 		return;
@@ -429,6 +527,10 @@ void FDisplayClusterScenePreviewModule::RegisterRootActorEvents(int32 RendererId
 	// Register/unregister for Blueprint events
 	if (UBlueprint* Blueprint = UBlueprint::GetBlueprintFromClass(Actor->GetClass()))
 	{
+		// Store a reference to the BP class.
+		// This reference is used to look up FRendererConfig from UBlueprint ptr in the OnBlueprintCompiled(UBlueprint* Blueprint) function.
+		RendererConfig.RootActorBlueprintClass = Blueprint;
+
 		Blueprint->OnCompiled().RemoveAll(this);
 
 		if (bShouldRegister)
@@ -439,86 +541,12 @@ void FDisplayClusterScenePreviewModule::RegisterRootActorEvents(int32 RendererId
 #endif
 }
 
-void FDisplayClusterScenePreviewModule::AutoPopulateScene(int32 RendererId, FRendererConfig& RendererConfig)
-{
-	const bool bAutoUpdateLightcards = EnumHasAnyFlags(RendererConfig.PreviewFlags, EDisplayClusterScenePreviewFlags::AutoUpdateLightcards);
-	if (bAutoUpdateLightcards)
-	{
-		RendererConfig.Renderer->ClearScene();
-		RendererConfig.AddedActors.Empty();
-		RendererConfig.AutoActors.Empty();
-	}
-	
-	// The renderer can use a proxy.
-	if (ADisplayClusterRootActor* RootActor = InternalGetRendererRootActorOrProxy(RendererId, RendererConfig))
-	{
-		TArray<FString> ProjectionMeshNames;
-
-		if (UDisplayClusterConfigurationData* Config = RootActor->GetConfigData())
-		{
-			Config->GetReferencedMeshNames(ProjectionMeshNames);
-		}
-
-		RendererConfig.Renderer->AddActor(RootActor, [&ProjectionMeshNames](const UPrimitiveComponent* PrimitiveComponent)
-		{
-			// Filter out any primitive component that isn't a projection mesh (a static mesh that has a Mesh projection configured for it) or a screen component
-			const bool bIsProjectionMesh = PrimitiveComponent->IsA<UStaticMeshComponent>() && ProjectionMeshNames.Contains(PrimitiveComponent->GetName());
-			const bool bIsScreen = PrimitiveComponent->IsA<UDisplayClusterScreenComponent>();
-			return bIsProjectionMesh || bIsScreen;
-		});
-
-		if (bAutoUpdateLightcards)
-		{
-			// Automatically add the lightcards found on this actor
-			TSet<ADisplayClusterLightCardActor*> LightCards;
-			UDisplayClusterBlueprintLib::FindLightCardsForRootActor(RootActor, LightCards);
-			
-			TSet<ADisplayClusterChromakeyCardActor*> ChromaKeyCards;
-			UDisplayClusterBlueprintLib::FindChromakeyCardsForRootActor(RootActor, ChromaKeyCards);
-			LightCards.Append(reinterpret_cast<TSet<ADisplayClusterLightCardActor*>&>(ChromaKeyCards));
-			
-			TSet<AActor*> Actors;
-			Actors.Reserve(LightCards.Num());
-			for (ADisplayClusterLightCardActor* LightCard : LightCards)
-			{
-				Actors.Add(LightCard);
-			}
-
-			// Also check for any non-lightcard actors in the world that are valid to control from ICVFX editors
-			if (UWorld* World = RootActor->GetWorld())
-			{
-				for (const TWeakObjectPtr<AActor> WeakActor : TActorRange<AActor>(World))
-				{
-					if (WeakActor.IsValid() && WeakActor->Implements<UDisplayClusterStageActor>() && !WeakActor->IsA<ADisplayClusterLightCardActor>())
-					{
-						Actors.Add(WeakActor.Get());
-					}
-				}
-			}
-
-			for (AActor* Actor : Actors)
-			{
-				if (RendererConfig.AddedActors.Contains(Actor))
-				{
-					continue;
-				}
-
-				RendererConfig.Renderer->AddActor(Actor);
-				RendererConfig.AddedActors.Add(Actor);
-				RendererConfig.AutoActors.Add(Actor);
-			}
-		}
-	}
-
-	RendererConfig.bIsSceneDirty = false;
-}
-
 bool FDisplayClusterScenePreviewModule::UpdateIsRealTimePreviewEnabled()
 {
 #if WITH_EDITOR
 	bIsRealTimePreviewEnabled = false;
 
-	if (!GEditor)
+	if (GIsEditor && !IsValid(GEditor))
 	{
 		return false;
 	}
@@ -540,8 +568,7 @@ bool FDisplayClusterScenePreviewModule::UpdateIsRealTimePreviewEnabled()
 
 bool FDisplayClusterScenePreviewModule::OnTick(float DeltaTime)
 {
-	// This function calls Tick() for a preview world with proxy root actors, which triggers rendering of previews for them.
-	ProxyManager.Tick(DeltaTime);
+	GetProxyManager().TickPreviewWorld(DeltaTime);
 
 	// This loop should break when we either run out of jobs or complete a single job
 	while (!RenderQueue.IsEmpty())
@@ -627,24 +654,22 @@ bool FDisplayClusterScenePreviewModule::OnTick(float DeltaTime)
 
 void FDisplayClusterScenePreviewModule::OnActorPropertyChanged(UObject* ObjectBeingModified, FPropertyChangedEvent& PropertyChangedEvent)
 {
-	for (TPair<int32, FRendererConfig>& ConfigPair : RendererConfigs)
+	for (TPair<int32, FRendererConfig>& RendererConfigIt : RendererConfigs)
 	{
-		FRendererConfig& Config = ConfigPair.Value;
-		const bool bAutoUpdateLightcards = EnumHasAnyFlags(Config.PreviewFlags, EDisplayClusterScenePreviewFlags::AutoUpdateLightcards);
-		if (bAutoUpdateLightcards)
+		const bool bAutoUpdateStageActors = EnumHasAnyFlags(RendererConfigIt.Value.PreviewFlags, EDisplayClusterScenePreviewFlags::AutoUpdateStageActors);
+		if (bAutoUpdateStageActors)
 		{
-			if (Config.RootActor == ObjectBeingModified)
+			ADisplayClusterRootActor* RootActorPtr = RendererConfigIt.Value.GetRootActor();
+
+			if (RootActorPtr == ObjectBeingModified)
 			{
-				Config.bIsSceneDirty = true;
-				continue;
+				EnumAddFlags(RendererConfigIt.Value.Flags, EDisplayClusterRendererConfigFlags::RootActorBeingModified);
 			}
-
-			if (UActorComponent* Component = Cast<UActorComponent>(ObjectBeingModified))
+			else if (UActorComponent* Component = Cast<UActorComponent>(ObjectBeingModified))
 			{
-				if (Component->GetOwner() == Config.RootActor)
+				if (Component->GetOwner() == RootActorPtr)
 				{
-					Config.bIsSceneDirty = true;
-					continue;
+					EnumAddFlags(RendererConfigIt.Value.Flags, EDisplayClusterRendererConfigFlags::LevelActorBeingModified);
 				}
 			}
 		}
@@ -653,12 +678,11 @@ void FDisplayClusterScenePreviewModule::OnActorPropertyChanged(UObject* ObjectBe
 
 void FDisplayClusterScenePreviewModule::OnLevelActorDeleted(AActor* Actor)
 {
-	for (TPair<int32, FRendererConfig>& ConfigPair : RendererConfigs)
+	for (TPair<int32, FRendererConfig>& RendererConfigIt : RendererConfigs)
 	{
-		FRendererConfig& Config = ConfigPair.Value;
-		if (Config.AutoActors.Contains(Actor))
+		if (RendererConfigIt.Value.AutoPopulateActors.Contains(Actor))
 		{
-			Config.bIsSceneDirty = true;
+			EnumAddFlags(RendererConfigIt.Value.Flags, EDisplayClusterRendererConfigFlags::LevelActorDeleted);
 		}
 	}
 }
@@ -671,24 +695,20 @@ void FDisplayClusterScenePreviewModule::OnLevelActorAdded(AActor* Actor)
 	}
 
 	// The actor won't be added to a root actor yet, so we can't check who it belongs to. Easier to just mark all configs as dirty.
-	for (TPair<int32, FRendererConfig>& ConfigPair : RendererConfigs)
+	for (TPair<int32, FRendererConfig>& RendererConfigIt : RendererConfigs)
 	{
-		ConfigPair.Value.bIsSceneDirty = true;
+		EnumAddFlags(RendererConfigIt.Value.Flags, EDisplayClusterRendererConfigFlags::LevelActorAdded);
 	}
 }
 
 void FDisplayClusterScenePreviewModule::OnBlueprintCompiled(UBlueprint* Blueprint)
 {
 #if WITH_EDITOR
-	for (TPair<int32, FRendererConfig>& ConfigPair : RendererConfigs)
+	for (TPair<int32, FRendererConfig>& RendererConfigIt : RendererConfigs)
 	{
-		int32 RendererId = ConfigPair.Key;
-		FRendererConfig& Config = ConfigPair.Value;
-
-		if(IsBlueprintMatchesRendererRootActor(RendererId, Config, Blueprint))
+		if(RendererConfigIt.Value.IsBlueprintMatchesRendererRootActor(Blueprint))
 		{
-			Config.bIsSceneDirty = true;
-			RegisterRootActorEvents(RendererId, Config, true);
+			EnumAddFlags(RendererConfigIt.Value.Flags, EDisplayClusterRendererConfigFlags::BlueprintCompiled);
 		}
 	}
 #endif
@@ -698,13 +718,197 @@ void FDisplayClusterScenePreviewModule::OnObjectTransacted(UObject* Object, cons
 {
 	if (TransactionObjectEvent.GetEventType() == ETransactionObjectEventType::UndoRedo)
 	{
-		for (TPair<int32, FRendererConfig>& ConfigPair : RendererConfigs)
+		for (TPair<int32, FRendererConfig>& RendererConfigIt : RendererConfigs)
 		{
-			ConfigPair.Value.bIsSceneDirty = true;
+			EnumAddFlags(RendererConfigIt.Value.Flags, EDisplayClusterRendererConfigFlags::ObjectTransacted);
 		}
 	}
 }
 
+bool FDisplayClusterScenePreviewModule::FRendererConfig::SetRootActor(ADisplayClusterRootActor* InRootActorPtr)
+{
+	if (!IsRootActorEquals(InRootActorPtr))
+	{
+		// Update flags
+		EnumAddFlags(Flags, EDisplayClusterRendererConfigFlags::UpdateRootActorRef);
+
+		// Store ref to root actor. This type of link supports BP and recompile.
+		RootActorRef.SetSceneActor(InRootActorPtr);
+
+		return true;
+	}
+
+	return false;
+}
+
+bool FDisplayClusterScenePreviewModule::FRendererConfig::IsRootActorEquals(ADisplayClusterRootActor* InRootActorPtr) const
+{
+	const TWeakObjectPtr<AActor> CurrentRootActor = RootActorRef.GetSceneActorWeakPtr().Get();
+
+	return CurrentRootActor == InRootActorPtr;
+}
+
+ADisplayClusterRootActor* FDisplayClusterScenePreviewModule::FRendererConfig::GetRootActor()
+{
+	const bool bDefinedSceneActor = RootActorRef.IsDefinedSceneActor();
+	const TWeakObjectPtr<AActor> RootActorWeakPtr = RootActorRef.GetSceneActorWeakPtr();
+	if (ADisplayClusterRootActor* RootActorPtr = Cast<ADisplayClusterRootActor>(RootActorRef.GetOrFindSceneActor()))
+	{
+		const TWeakObjectPtr<AActor> NewRootActorWeakPtr = RootActorRef.GetSceneActorWeakPtr();
+		if (NewRootActorWeakPtr != RootActorWeakPtr)
+		{
+			// The reference to the RootActor has been changed.
+			EnumAddFlags(Flags, EDisplayClusterRendererConfigFlags::UpdateRootActorRef);
+		}
+
+		if (IsValid(RootActorPtr))
+		{
+			return RootActorPtr;
+		}
+	}
+
+	if (bDefinedSceneActor)
+	{
+		// We lost the reference to the RootActor object.
+		EnumAddFlags(Flags, EDisplayClusterRendererConfigFlags::LostRootActorRef);
+	}
+
+	return nullptr;
+}
+
+bool FDisplayClusterScenePreviewModule::FRendererConfig::ClearRendererScene()
+{
+	if (Renderer.IsValid())
+	{
+		Renderer->ClearScene();
+	}
+
+	AddedActors.Empty();
+	AutoPopulateActors.Empty();
+
+	return true;
+}
+
+bool FDisplayClusterScenePreviewModule::FRendererConfig::AddActorToRenderer(AActor* Actor, bool bAutoPopulate)
+{
+	if (!Renderer.IsValid())
+	{
+		return false;
+	}
+
+	if (!AddedActors.Contains(Actor))
+	{
+		Renderer->AddActor(Actor);
+		AddedActors.Add(Actor);
+
+		if (bAutoPopulate)
+		{
+			AutoPopulateActors.Add(Actor);
+		}
+	}
+
+	return true;
+}
+
+bool FDisplayClusterScenePreviewModule::FRendererConfig::AddActorToRenderer(AActor* Actor, const TFunctionRef<bool(const UPrimitiveComponent*)>& PrimitiveFilter, bool bAutoPopulate)
+{
+	if (!Renderer.IsValid())
+	{
+		return false;
+	}
+
+	if (!AddedActors.Contains(Actor))
+	{
+		Renderer->AddActor(Actor, PrimitiveFilter);
+		AddedActors.Add(Actor);
+
+		if (bAutoPopulate)
+		{
+			AutoPopulateActors.Add(Actor);
+		}
+	}
+
+	return true;
+}
+
+bool FDisplayClusterScenePreviewModule::FRendererConfig::RemoveActorFromRenderer(AActor* Actor)
+{
+	if (!Renderer.IsValid())
+	{
+		return false;
+	}
+
+	if (!AddedActors.Contains(Actor))
+	{
+		return false;
+	}
+
+	Renderer->RemoveActor(Actor);
+	AddedActors.Remove(Actor);
+	
+	return true;
+}
+
+bool FDisplayClusterScenePreviewModule::FRendererConfig::GetActorsInRendererScene(bool bIncludeRoot, TArray<AActor*>& OutActors)
+{
+	if (bIncludeRoot)
+	{
+		if (ADisplayClusterRootActor* RootActorPtr = GetRootActor())
+		{
+			OutActors.Add(RootActorPtr);
+		}
+	}
+
+	for (const TWeakObjectPtr<AActor>& Actor : AddedActors)
+	{
+		if (!Actor.IsValid())
+		{
+			continue;
+		}
+
+		OutActors.Add(Actor.Get());
+	}
+
+	return true;
+}
+
+bool FDisplayClusterScenePreviewModule::FRendererConfig::SetRendererActorSelectedDelegate(FDisplayClusterMeshProjectionRenderer::FSelection ActorSelectedDelegate)
+{
+	if (Renderer.IsValid())
+	{
+		Renderer->ActorSelectedDelegate = ActorSelectedDelegate;
+
+		return true;
+	}
+
+	return false;
+}
+
+bool FDisplayClusterScenePreviewModule::FRendererConfig::SetRendererRenderSimpleElementsDelegate(FDisplayClusterMeshProjectionRenderer::FSimpleElementPass RenderSimpleElementsDelegate)
+{
+	if (Renderer.IsValid())
+	{
+		Renderer->RenderSimpleElementsDelegate = RenderSimpleElementsDelegate;
+
+		return true;
+	}
+
+	return false;
+}
+
+bool FDisplayClusterScenePreviewModule::FRendererConfig::IsBlueprintMatchesRendererRootActor(UBlueprint* Blueprint) const
+{
+#if WITH_EDITOR
+	UBlueprint* RootActorBlueprint = RootActorBlueprintClass.Get();
+	if (Blueprint == RootActorBlueprint)
+	{
+		return true;
+	}
+#endif
+
+	return false;
+}
+
 IMPLEMENT_MODULE(FDisplayClusterScenePreviewModule, DisplayClusterScenePreview);
 
 #undef LOCTEXT_NAMESPACE
diff --git a/Engine/Plugins/Runtime/nDisplay/Source/DisplayClusterScenePreview/Private/DisplayClusterScenePreviewModule.h b/Engine/Plugins/Runtime/nDisplay/Source/DisplayClusterScenePreview/Private/DisplayClusterScenePreviewModule.h
--- a/Engine/Plugins/Runtime/nDisplay/Source/DisplayClusterScenePreview/Private/DisplayClusterScenePreviewModule.h
+++ b/Engine/Plugins/Runtime/nDisplay/Source/DisplayClusterScenePreview/Private/DisplayClusterScenePreviewModule.h
@@ -3,7 +3,6 @@
 #pragma once
 
 #include "IDisplayClusterScenePreview.h"
-#include "DisplayClusterScenePreviewProxyManager.h"
 #include "DisplayClusterRootActorContainers.h"
 
 #include "Containers/Ticker.h"
@@ -11,9 +10,55 @@
 #include "Engine/TextureRenderTarget2D.h"
 #include "UObject/StrongObjectPtr.h"
 
+#include "Misc/DisplayClusterObjectRef.h"
+
 class ADisplayClusterRootActor;
 class ADisplayClusterLightCardActor;
 class UTextureRenderTarget2D;
+class FDisplayClusterScenePreviewProxyManager;
+
+/** Flags for the renderer config. */
+enum class EDisplayClusterRendererConfigFlags : uint16
+{
+	None = 0,
+
+	/** Updated in OnObjectTransacted(). */
+	ObjectTransacted = 1 << 0,
+
+	/** When BP is recompiled. */
+	BlueprintCompiled = 1 << 1,
+
+	/** OnLevelActorAdded() */
+	LevelActorAdded = 1 << 2,
+
+	/** OnLevelActorDeleted()*/
+	LevelActorDeleted = 1 << 3,
+
+	/** OnActorPropertyChanged() for objects owned by the RootActor.*/
+	LevelActorBeingModified = 1 << 4,
+
+	/** OnActorPropertyChanged() for the RootActor.*/
+	RootActorBeingModified = 1 << 5,
+
+	/** The reference to RootActor has been updated.  */
+	UpdateRootActorRef = 1 << 6,
+
+	/** The reference to RootActorProxy has been updated.  */
+	UpdateRootActorProxyRef = 1 << 7,
+
+	/** we have lost the reference to the RootActor, it needs to be restored in another way.*/
+	LostRootActorRef = 1 << 8,
+
+	/** RootActorPath has been changed. */
+	RootActorPathModified = 1 << 9,
+
+	/** Root actor property overrides has been changed. */
+	RootActorPropertyOverridesModified = 1 << 14,
+
+	/** RootActor preview flags has been changed. */
+	PreviewFlagsModified = 1 << 15,
+};
+ENUM_CLASS_FLAGS(EDisplayClusterRendererConfigFlags);
 
 /**
  * Module containing tools for rendering nDisplay scene previews.
@@ -32,13 +77,50 @@ public:
 	virtual bool SetRendererRootActorPath(int32 RendererId, const FString& ActorPath, const FDisplayClusterRootActorPropertyOverrides& InPropertyOverrides, const EDisplayClusterScenePreviewFlags PreviewFlags = EDisplayClusterScenePreviewFlags::None) override;
 	virtual bool SetRendererRootActor(int32 RendererId, ADisplayClusterRootActor* Actor, const FDisplayClusterRootActorPropertyOverrides& InPropertyOverrides, const EDisplayClusterScenePreviewFlags PreviewFlags = EDisplayClusterScenePreviewFlags::None) override;
 	virtual ADisplayClusterRootActor* GetRendererRootActor(int32 RendererId) override;
-	virtual bool GetActorsInRendererScene(int32 RendererId, bool bIncludeRoot, TArray<AActor*>& OutActors) override;
-	virtual bool AddActorToRenderer(int32 RendererId, AActor* Actor) override;
-	virtual bool AddActorToRenderer(int32 RendererId, AActor* Actor, const TFunctionRef<bool(const UPrimitiveComponent*)>& PrimitiveFilter) override;
-	virtual bool RemoveActorFromRenderer(int32 RendererId, AActor* Actor) override;
-	virtual bool ClearRendererScene(int32 RendererId) override;
-	virtual bool SetRendererActorSelectedDelegate(int32 RendererId, FDisplayClusterMeshProjectionRenderer::FSelection ActorSelectedDelegate) override;
-	virtual bool SetRendererRenderSimpleElementsDelegate(int32 RendererId, FDisplayClusterMeshProjectionRenderer::FSimpleElementPass RenderSimpleElementsDelegate) override;
+	virtual ADisplayClusterRootActor* GetRendererRootActorOrProxy(int32 RendererId) override;
+
+	virtual bool GetActorsInRendererScene(int32 RendererId, bool bIncludeRoot, TArray<AActor*>& OutActors) override
+	{
+		FRendererConfig* Config = RendererConfigs.Find(RendererId);
+		return Config ? Config->GetActorsInRendererScene(bIncludeRoot, OutActors) : false;
+	}
+
+	virtual bool AddActorToRenderer(int32 RendererId, AActor* Actor) override
+	{
+		FRendererConfig* Config = RendererConfigs.Find(RendererId);
+		return Config ? Config->AddActorToRenderer(Actor) : false;
+	}
+
+	virtual bool AddActorToRenderer(int32 RendererId, AActor* Actor, const TFunctionRef<bool(const UPrimitiveComponent*)>& PrimitiveFilter) override
+	{
+		FRendererConfig* Config = RendererConfigs.Find(RendererId);
+		return Config ? Config->AddActorToRenderer(Actor, PrimitiveFilter) : false;
+	}
+	
+	virtual bool RemoveActorFromRenderer(int32 RendererId, AActor* Actor) override
+	{
+		FRendererConfig* Config = RendererConfigs.Find(RendererId);
+		return Config ? Config->RemoveActorFromRenderer(Actor) : false;
+	}
+
+	virtual bool ClearRendererScene(int32 RendererId) override
+	{
+		FRendererConfig* Config = RendererConfigs.Find(RendererId);
+		return Config ? Config->ClearRendererScene() : false;
+	}
+
+	virtual bool SetRendererActorSelectedDelegate(int32 RendererId, FDisplayClusterMeshProjectionRenderer::FSelection ActorSelectedDelegate) override
+	{
+		FRendererConfig* Config = RendererConfigs.Find(RendererId);
+		return Config ? Config->SetRendererActorSelectedDelegate(ActorSelectedDelegate) : false;
+	}
+
+	virtual bool SetRendererRenderSimpleElementsDelegate(int32 RendererId, FDisplayClusterMeshProjectionRenderer::FSimpleElementPass RenderSimpleElementsDelegate) override
+	{
+		FRendererConfig* Config = RendererConfigs.Find(RendererId);
+		return Config ? Config->SetRendererRenderSimpleElementsDelegate(RenderSimpleElementsDelegate) : false;
+	}
+
 	virtual bool Render(int32 RendererId, FDisplayClusterMeshProjectionRenderSettings& RenderSettings, FCanvas& Canvas) override;
 	virtual bool RenderQueued(int32 RendererId, FDisplayClusterMeshProjectionRenderSettings& RenderSettings, const FIntPoint& Size, FRenderResultDelegate ResultDelegate) override;
 	virtual bool RenderQueued(int32 RendererId, FDisplayClusterMeshProjectionRenderSettings& RenderSettings, const TWeakPtr<FCanvas> Canvas, FRenderResultDelegate ResultDelegate) override;
@@ -49,12 +131,48 @@ private:
 	/** Holds information about an active renderer created by this module. */
 	struct FRendererConfig
 	{
+		/** Get the actors that were used in this render scene.*/
+		bool GetActorsInRendererScene(bool bIncludeRoot, TArray<AActor*>& OutActors);
+
+		/** Add actor to renderer. */
+		bool AddActorToRenderer(AActor* Actor, bool bAutoPopulate = false);
+
+		/** Add actor to renderer with primitive filter. */
+		bool AddActorToRenderer(AActor* Actor, const TFunctionRef<bool(const UPrimitiveComponent*)>& PrimitiveFilter, bool bAutoPopulate = false);
+
+		/** remove actor from renderer. */
+		bool RemoveActorFromRenderer(AActor* Actor);
+
+		/** Clear renderer scene. */
+		bool ClearRendererScene();
+
+		bool SetRendererActorSelectedDelegate(FDisplayClusterMeshProjectionRenderer::FSelection ActorSelectedDelegate);
+		bool SetRendererRenderSimpleElementsDelegate(FDisplayClusterMeshProjectionRenderer::FSimpleElementPass RenderSimpleElementsDelegate);
+
+		/** Returns true if the Blueprint class matches the RootActor class used in the renderer configuration. */
+		bool IsBlueprintMatchesRendererRootActor(UBlueprint* Blueprint) const;
+
+		/** Get root of the display cluster that this renderer is previewing. */
+		ADisplayClusterRootActor* GetRootActor();
+
+		/** Set new root actor ptr. */
+		bool SetRootActor(ADisplayClusterRootActor* InRootActorPtr);
+
+		/** Compare current root actor ptr and input. */
+		bool IsRootActorEquals(ADisplayClusterRootActor* InRootActorPtr) const;
+
+		/** return true if RootActor already defined. */
+		bool IsDefinedRootActor() const
+		{
+			return RootActorRef.IsDefinedSceneActor();
+		}
+
+		/** Reference to root actor BP. */
+		TWeakObjectPtr<UBlueprint> RootActorBlueprintClass;
+
 		/** The renderer itself. */
 		TSharedPtr<FDisplayClusterMeshProjectionRenderer> Renderer;
 
-		/** The root of the display cluster that this renderer is previewing. */
-		TWeakObjectPtr<ADisplayClusterRootActor> RootActor;
-
 		/** The path of the root actor that this renderer is previewing. If this is not empty and the root actor becomes invalid, we will attempt to find it again using this path. */
 		FString RootActorPath;
 
@@ -62,7 +180,7 @@ private:
 		TArray<TWeakObjectPtr<AActor>> AddedActors;
 
 		/** Actors that have been automatically added to the scene. */
-		TArray<TWeakObjectPtr<AActor>> AutoActors;
+		TArray<TWeakObjectPtr<AActor>> AutoPopulateActors;
 
 		/** Special flags that control the behavior of the renderer. */
 		EDisplayClusterScenePreviewFlags PreviewFlags = EDisplayClusterScenePreviewFlags::None;
@@ -70,11 +188,15 @@ private:
 		/** Container with properties to be overridden for the root actor used by this renderer. */
 		FDisplayClusterRootActorPropertyOverrides RootActorPropertyOverrides;
 
-		/** If true, the scene needs to be updated before the next render. This is only relevant if PreviewFlags.AutoUpdateLightcards is true. */
-		bool bIsSceneDirty = true;
+		/** Renderer flags. */
+		EDisplayClusterRendererConfigFlags Flags = EDisplayClusterRendererConfigFlags::None;
 
 		/** The render target to use for queued renders. */
 		TStrongObjectPtr<UTextureRenderTarget2D> RenderTarget = nullptr;
+
+	private:
+		/** Saved reference to the root actor. It is used to restore a reference to a new root actor object after BP recompilation. */
+		FDisplayClusterActorRef RootActorRef;
 	};
 
 	/** Holds information about a preview render that was queued to be completed later. */
@@ -107,6 +229,9 @@ private:
 		FRenderResultDelegate ResultDelegate;
 	};
 
+	/** Handle all changes in the renderer config. */
+	void InternalUpdateRendererConfig(int32 RendererId, FRendererConfig& RendererConfig);
+
 	/** Get the root actor for a config. If the root actor pointer is invalid but we have a path to the actor, try to reacquire a pointer using the path first. */
 	ADisplayClusterRootActor* InternalGetRendererRootActor(int32 RendererId, FRendererConfig& RendererConfig);
 
@@ -114,10 +239,7 @@ private:
 	ADisplayClusterRootActor* InternalGetRendererRootActorOrProxy(int32 RendererId, FRendererConfig& RendererConfig);
 
 	/** Set the root actor for a config, update its scene, and register events accordingly. */
-	void InternalSetRendererRootActor(int32 RendererId, FRendererConfig& RendererConfig, ADisplayClusterRootActor* Actor, const EDisplayClusterScenePreviewFlags PreviewFlags);
-
-	/** Override properties for the renderer root actor. */
-	void InternalOverridePropertiesForRendererRootActor(int32 RendererId, FRendererConfig& RendererConfig);
+	void InternalSetRendererRootActor(int32 RendererId, FRendererConfig& RendererConfig, ADisplayClusterRootActor* Actor);
 
 	/** Queue a preview to be rendered. */
 	bool InternalRenderQueued(int32 RendererId, FDisplayClusterMeshProjectionRenderSettings& RenderSettings, TWeakPtr<FCanvas> Canvas,
@@ -132,9 +254,6 @@ private:
 	/** Register/unregister to events affecting a cluster root actor. */
 	void RegisterRootActorEvents(int32 RendererId, FRendererConfig& RendererConfig, bool bShouldRegister);
 
-	/** Returns true if the Blueprint class matches the RootActor class used in the renderer configuration. */
-	bool IsBlueprintMatchesRendererRootActor(int32 RendererId, FRendererConfig& RendererConfig, UBlueprint* Blueprint);
-
 	/** Clear and re-populate a renderer's scene with the root actor and lightcards if applicable. */
 	void AutoPopulateScene(int32 RendererId, FRendererConfig& RendererConfig);
 
@@ -159,6 +278,19 @@ private:
 	/** Called when any object is transacted. */
 	void OnObjectTransacted(UObject* Object, const FTransactionObjectEvent& TransactionObjectEvent);
 
+	/** release all internals data. */
+	void Release();
+
+	/** Return ProxyManager object. */
+	FDisplayClusterScenePreviewProxyManager& GetProxyManager();
+
+private:
+#if WITH_EDITOR
+	// Callbacks
+	FDelegateHandle EditorClosedEventHandle;
+	void OnEditorClosed();
+#endif
+
 private:
 
 	/** Map from renderer ID to configuration data for that renderer. */
@@ -180,5 +312,5 @@ private:
 	bool bIsRealTimePreviewEnabled = false;
 
 	/** Manager for DCRA proxy objects. */
-	FDisplayClusterScenePreviewProxyManager ProxyManager;
+	TUniquePtr<FDisplayClusterScenePreviewProxyManager> ProxyManagerPtr;
 };
\ No newline at end of file
diff --git a/Engine/Plugins/Runtime/nDisplay/Source/DisplayClusterScenePreview/Private/DisplayClusterScenePreviewProxyManager.cpp b/Engine/Plugins/Runtime/nDisplay/Source/DisplayClusterScenePreview/Private/DisplayClusterScenePreviewProxyManager.cpp
--- a/Engine/Plugins/Runtime/nDisplay/Source/DisplayClusterScenePreview/Private/DisplayClusterScenePreviewProxyManager.cpp
+++ b/Engine/Plugins/Runtime/nDisplay/Source/DisplayClusterScenePreview/Private/DisplayClusterScenePreviewProxyManager.cpp
@@ -1,46 +1,95 @@
 // Copyright Epic Games, Inc. All Rights Reserved.
 
 #include "DisplayClusterScenePreviewProxyManager.h"
-
 #include "Components/DisplayClusterStageGeometryComponent.h"
+#include "Engine/World.h"
 
 FDisplayClusterScenePreviewProxyManager::~FDisplayClusterScenePreviewProxyManager()
 {
+	Release();
+}
+
+void FDisplayClusterScenePreviewProxyManager::Release()
+{
+	TickableGameObject.Reset();
+	RendererProxies.Empty();
+
 	DestroyPreviewWorld();
 }
 
 void FDisplayClusterScenePreviewProxyManager::CreatePreviewWorld()
 {
+	if (IsValid(PreviewWorld))
+	{
+		// Preview world already exists.
+		return;
+	}
+
 #if WITH_EDITOR
-	if (!PreviewWorld)
+	if (IsValid(GEngine))
 	{
-		PreviewWorld = UWorld::CreateWorld(EWorldType::None, false);
+		PreviewWorld = NewObject<UWorld>(GetTransientPackage(), TEXT("DisplayClusterScenePreview"), RF_NoFlags);
+		PreviewWorld->WorldType = EWorldType::EditorPreview;
 
-		FWorldContext& WorldContext = GEngine->CreateNewWorldContext(EWorldType::None);
+		FWorldContext& WorldContext = GEngine->CreateNewWorldContext(PreviewWorld->WorldType);
 		WorldContext.SetCurrentWorld(PreviewWorld);
+
+		PreviewWorld->InitializeNewWorld(UWorld::InitializationValues()
+			.AllowAudioPlayback(false)
+			.CreatePhysicsScene(false)
+			.RequiresHitProxies(true) // Only Need hit proxies in an editor scene
+			.CreateNavigation(false)
+			.CreateAISystem(false)
+			.ShouldSimulatePhysics(false)
+			.SetTransactional(false));
 	}
 #endif
 }
 
 void FDisplayClusterScenePreviewProxyManager::DestroyPreviewWorld()
 {
-#if WITH_EDITOR
-	if (PreviewWorld)
+	UWorld* LocalPreviewWorld = IsValid(PreviewWorld) ? PreviewWorld : nullptr;
+	PreviewWorld = nullptr;
+
+	if (IsValid(LocalPreviewWorld) && IsValid(GEngine))
 	{
-		GEngine->DestroyWorldContext(PreviewWorld);
-		PreviewWorld->DestroyWorld(false);
+		LocalPreviewWorld->CleanupWorld();
+		GEngine->DestroyWorldContext(LocalPreviewWorld);
+
+		// Release PhysicsScene for fixing big fbx importing bug
+		LocalPreviewWorld->ReleasePhysicsScene();
 	}
-#endif
+}
 
-	PreviewWorld = nullptr;
+void FDisplayClusterScenePreviewProxyManager::AddReferencedObjects(FReferenceCollector& Collector)
+{
+	Collector.AddReferencedObject(PreviewWorld);
 }
 
-void FDisplayClusterScenePreviewProxyManager::Tick(float DeltaTime)
+FString FDisplayClusterScenePreviewProxyManager::GetReferencerName() const
 {
+	return TEXT("FDisplayClusterScenePreviewProxyManager");
+}
+
+void FDisplayClusterScenePreviewProxyManager::TickPreviewWorld(float DeltaTime)
+{
+#if WITH_EDITOR
 	if (PreviewWorld)
 	{
 		PreviewWorld->Tick(ELevelTick::LEVELTICK_All, DeltaTime);
 	}
+#endif
+}
+
+void FDisplayClusterScenePreviewProxyManager::OnTick(float DeltaTime)
+{
+#if WITH_EDITOR
+	// Tick all proxy root actors.
+	for(const TPair<int32, FRendererProxy>& RendererProxyIt : RendererProxies)
+	{
+		RendererProxyIt.Value.TickProxyRootActor();
+	}
+#endif
 }
 
 ADisplayClusterRootActor* FDisplayClusterScenePreviewProxyManager::CreateRootActorProxy(int32 RendererId, ADisplayClusterRootActor* SceneRootActor)
@@ -50,6 +99,8 @@ ADisplayClusterRootActor* FDisplayClusterScenePreviewProxyManager::CreateRootAct
 		return nullptr;
 	}
 
+	check(IsValid(SceneRootActor));
+
 	ADisplayClusterRootActor* RootActorProxy = nullptr;
 
 #if WITH_EDITOR
@@ -64,7 +115,8 @@ ADisplayClusterRootActor* FDisplayClusterScenePreviewProxyManager::CreateRootAct
 	FObjectDuplicationParameters DupeActorParameters(SceneRootActor, PreviewWorld->GetCurrentLevel());
 	DupeActorParameters.FlagMask = RF_AllFlags & ~(RF_ArchetypeObject | RF_Transactional); // Keeps archetypes correct in config data.
 	DupeActorParameters.PortFlags = PPF_DuplicateVerbatim;
-	DupeActorParameters.DestName = FName(FString::Printf(TEXT("Preview%i_%s"), RendererId, *SceneRootActor->GetName()));
+	static int32 UniqueIndex = 0;
+	DupeActorParameters.DestName = FName(FString::Printf(TEXT("Preview-%s-%d-%d"), *SceneRootActor->GetName(), RendererId, UniqueIndex++));
 
 	RootActorProxy = CastChecked<ADisplayClusterRootActor>(StaticDuplicateObjectEx(DupeActorParameters));
 
@@ -87,6 +139,12 @@ ADisplayClusterRootActor* FDisplayClusterScenePreviewProxyManager::CreateRootAct
 	RootActorProxy->SetActorLocation(FVector::ZeroVector);
 	RootActorProxy->SetActorRotation(FRotator::ZeroRotator);
 
+	if (UDisplayClusterConfigurationData* ProxyConfig = RootActorProxy->GetConfigData())
+	{
+		// Disable lightcards so that it doesn't try to update the ones in the level instance world.
+		ProxyConfig->StageSettings.Lightcard.bEnable = false;
+	}
+
 	// Set translucency sort priority of root actor proxy primitive components so that actors that are flush with screens are rendered on top of them
 	RootActorProxy->ForEachComponent<UPrimitiveComponent>(false, [](UPrimitiveComponent* InPrimitiveComponent)
 	{
@@ -107,39 +165,56 @@ void FDisplayClusterScenePreviewProxyManager::DestroyRootActorProxy(ADisplayClus
 #endif
 }
 
-void FDisplayClusterScenePreviewProxyManager::SetSceneRootActorForRenderer(int32 RendererId, ADisplayClusterRootActor* SceneRootActor)
+void FDisplayClusterScenePreviewProxyManager::SetSceneRootActorForRenderer(int32 RendererId, ADisplayClusterRootActor* SceneRootActor, EDisplayClusterScenePreviewFlags PreviewFlags)
 {
 	if (RendererProxies.Contains(RendererId))
 	{
 		FRendererProxy& RendererProxy = RendererProxies[RendererId];
 
-		if (RendererProxy.SceneRootActor == SceneRootActor)
+		if (RendererProxy.GetProxyRootActor())
 		{
-			// Re-use existing proxy
-			return;
+			// If the proxy object already exists, check whether the root actor is the same or not.
+			if (RendererProxy.SceneRootActorWeakPtr == SceneRootActor)
+			{
+				// Re-use existing proxy, but update flags
+				RendererProxy.PreviewFlags = PreviewFlags;
+
+				return;
+			}
 		}
 
 		// The root agent has changed, destroy the proxy currently in use.
-		DestroyRootActorProxy(RendererProxy.ProxyRootActor.Get());
+		DestroyRootActorProxy(RendererProxy.GetProxyRootActor());
 
 		RendererProxies.Remove(RendererId);
 	}
 
-	if (ADisplayClusterRootActor* ProxyRootActor = CreateRootActorProxy(RendererId, SceneRootActor))
+	if (IsValid(SceneRootActor))
 	{
-		// Create new proxy
-		FRendererProxy RendererProxy;
-		RendererProxy.SceneRootActor = SceneRootActor;
-		RendererProxy.ProxyRootActor = ProxyRootActor;
+		if (ADisplayClusterRootActor* ProxyRootActor = CreateRootActorProxy(RendererId, SceneRootActor))
+		{
+			// Create new proxy
+			FRendererProxy RendererProxy;
+			RendererProxy.SceneRootActorWeakPtr = SceneRootActor;
+			RendererProxy.ProxyRootActorWeakPtr = ProxyRootActor;
+			RendererProxy.PreviewFlags = PreviewFlags;
 
-		RendererProxies.Emplace(RendererId, RendererProxy);
-	}
+			// Update RootActor proxy
+			RendererProxy.TickProxyRootActor();
 
+			RendererProxies.Emplace(RendererId, RendererProxy);
+		}
+	}
 
-	// Destroy preview world if not used
+	// Configure a tick callback for existing RendererProxies.
 	if (RendererProxies.IsEmpty())
 	{
-		DestroyPreviewWorld();
+		TickableGameObject.Reset();
+	}
+	else if (!TickableGameObject.IsValid())
+	{
+		TickableGameObject = MakeUnique<FDisplayClusterTickableGameObject >();
+		TickableGameObject->OnTick().AddRaw(this, &FDisplayClusterScenePreviewProxyManager::OnTick);
 	}
 }
 
@@ -147,8 +222,80 @@ ADisplayClusterRootActor* FDisplayClusterScenePreviewProxyManager::GetProxyRootA
 {
 	if (RendererProxies.Contains(RendererId))
 	{
-		return RendererProxies[RendererId].ProxyRootActor.Get();
+		return RendererProxies[RendererId].GetProxyRootActor();
+	}
+
+	return nullptr;
+}
+
+ADisplayClusterRootActor* FDisplayClusterScenePreviewProxyManager::GetSceneRootActor(int32 RendererId) const
+{
+	if (RendererProxies.Contains(RendererId))
+	{
+		if (ADisplayClusterRootActor* SceneRootActor = RendererProxies[RendererId].GetSceneRootActor())
+		{
+			return IsValid(SceneRootActor) ? SceneRootActor : nullptr;
+		}
+	}
+
+	return nullptr;
+}
+
+void FDisplayClusterScenePreviewProxyManager::FRendererProxy::TickProxyRootActor() const
+{
+	if (ADisplayClusterRootActor* ProxyRootActor = GetProxyRootActor())
+	{
+		if (EnumHasAnyFlags(PreviewFlags, EDisplayClusterScenePreviewFlags::ProxyFollowSceneRootActor))
+		{
+			if (ADisplayClusterRootActor* SceneRootActor = GetSceneRootActor())
+			{
+				if (ProxyRootActor != SceneRootActor)
+				{
+					// Move the RootActorProxy to the same position as in the scene to match the position of the LC in world space.
+					if (EnumHasAnyFlags(PreviewFlags, EDisplayClusterScenePreviewFlags::ProxyFollowSceneRootActor))
+					{
+						const FTransform NewTransform = SceneRootActor->GetActorTransform();
+						const FTransform OldTransform = ProxyRootActor->GetActorTransform();
+						if (!NewTransform.Equals(OldTransform, UE_KINDA_SMALL_NUMBER))
+						{
+							ProxyRootActor->SetActorTransform(NewTransform);
+						}
+					}
+				}
+			}
+		}
+
+		// Force preview renderer call for proxy root actor.
+		if (EnumHasAnyFlags(PreviewFlags, EDisplayClusterScenePreviewFlags::ProxyTickPreviewRenderer))
+		{
+			ProxyRootActor->TickPreviewRenderer();
+		}
+	}
+}
+
+ADisplayClusterRootActor* FDisplayClusterScenePreviewProxyManager::FRendererProxy::GetProxyRootActor() const
+{
+	if (ADisplayClusterRootActor* RootActor = ProxyRootActorWeakPtr.Get())
+	{
+		if (IsValid(RootActor))
+		{
+			return RootActor;
+		}
+	}
+
+	return nullptr;
+}
+
+ADisplayClusterRootActor* FDisplayClusterScenePreviewProxyManager::FRendererProxy::GetSceneRootActor() const
+{
+	if (ADisplayClusterRootActor* RootActor = SceneRootActorWeakPtr.Get())
+	{
+		if (IsValid(RootActor))
+		{
+			return RootActor;
+		}
 	}
 
 	return nullptr;
 }
+
diff --git a/Engine/Plugins/Runtime/nDisplay/Source/DisplayClusterScenePreview/Private/DisplayClusterScenePreviewProxyManager.h b/Engine/Plugins/Runtime/nDisplay/Source/DisplayClusterScenePreview/Private/DisplayClusterScenePreviewProxyManager.h
--- a/Engine/Plugins/Runtime/nDisplay/Source/DisplayClusterScenePreview/Private/DisplayClusterScenePreviewProxyManager.h
+++ b/Engine/Plugins/Runtime/nDisplay/Source/DisplayClusterScenePreview/Private/DisplayClusterScenePreviewProxyManager.h
@@ -3,14 +3,16 @@
 #pragma once
 
 #include "CoreMinimal.h"
+#include "UObject/GCObject.h"
 
+#include "DisplayClusterScenePreviewEnums.h"
 #include "DisplayClusterRootActor.h"
-#include "Engine/World.h"
+#include "Misc/DisplayClusterTickableGameObject.h"
 
 /**
  * Creates and handle DCRA proxies.
  */
-class FDisplayClusterScenePreviewProxyManager
+class FDisplayClusterScenePreviewProxyManager : public FGCObject
 {
 public:
 	~FDisplayClusterScenePreviewProxyManager();
@@ -21,25 +23,37 @@ public:
 	* 
 	* Note: If RootActor is null, it means that the renderer is no longer using the proxy.
 	*/
-	void SetSceneRootActorForRenderer(int32 RendererId, ADisplayClusterRootActor* RootActor);
+	void SetSceneRootActorForRenderer(int32 RendererId, ADisplayClusterRootActor* RootActor, EDisplayClusterScenePreviewFlags PreviewFlags);
 
-	/** Get the DCRA proxy for the renderer by renderer id.
+	/** Get the DCRA proxy for the renderer by id.
 	* 
 	* @param RendererId - ID of renderer that uses this DCRA
 	* 
-	* @return - A reference to the actor's root proxy, if it exists, or nullptr.
+	* @return - A reference to the RootActor proxy, if it exists, or nullptr.
 	* 
 	* Note: You must call SetSceneRootActorForRenderer() for the renderer and the root actor being used before using this function.
 	*/
 	ADisplayClusterRootActor* GetProxyRootActor(int32 RendererId) const;
 
-	/** This function calls Tick() for a preview world with proxy root actors, which triggers rendering of previews for them.
-	* 
-	* @param DeltaTime - has passed to PreviewWorld->Tick().
+	/** Get the DCRA in scene used to create a proxy for the renderer by id.
+	*
+	* @param RendererId - ID of renderer that uses this DCRA
+	*
+	* @return - A reference to the RootActor in scene, if it exists, or nullptr.
+	*
+	* Note: You must call SetSceneRootActorForRenderer() for the renderer and the root actor being used before using this function.
 	*/
-	void Tick(float DeltaTime);
+	ADisplayClusterRootActor* GetSceneRootActor(int32 RendererId) const;
+
+	/** Run the Tick() function for the preview world. */
+	void TickPreviewWorld(float DeltaTime);
+
+	/** Release all internal data.*/
+	void Release();
 
 private:
+	/** This function calls Tick() for a preview world with proxy root actors, which triggers rendering of previews for them. */
+	void OnTick(float DeltaTime);
 	void CreatePreviewWorld();
 	void DestroyPreviewWorld();
 
@@ -48,19 +62,41 @@ private:
 
 	void DestroyRootActorProxy(ADisplayClusterRootActor* ProxyRootActor);
 
+	//~ Begin FGCObject interface
+	virtual void AddReferencedObjects(FReferenceCollector& Collector) override;
+	virtual FString GetReferencerName() const override;
+	//~ End FGCObject interface
+
 private:
+	/** Container for RootActorProxy and related data. */
 	struct FRendererProxy
 	{
-		// RootActor proxy objects
-		TWeakObjectPtr<ADisplayClusterRootActor> ProxyRootActor;
+		/** Get value from the ProxyRootActorWeakPtr variable. */
+		ADisplayClusterRootActor* GetProxyRootActor() const;
+
+		/** Get value from the SceneRootActorWeakPtr variable. */
+		ADisplayClusterRootActor* GetSceneRootActor() const;
+
+		/** Update ProxyRootActor. */
+		void TickProxyRootActor() const;
+
+		// RootActorProxy object created based on the RootActor in the scene.
+		TWeakObjectPtr<ADisplayClusterRootActor> ProxyRootActorWeakPtr;
 
 		// RootActor in scene that used to create DCRA proxy
-		TWeakObjectPtr<ADisplayClusterRootActor> SceneRootActor;
+		TWeakObjectPtr<ADisplayClusterRootActor> SceneRootActorWeakPtr;
+
+		// Special flags that control the behavior of the renderer.
+		EDisplayClusterScenePreviewFlags PreviewFlags;
 	};
 
 	/** Registered proxies for renderers by ID. */
 	TMap<int32, FRendererProxy> RendererProxies;
 
 	/** The preview world is only used when using the DCRA proxy. */
-	TObjectPtr<UWorld> PreviewWorld;
+	TObjectPtr<class UWorld> PreviewWorld = nullptr;
+
+	// When CachedObjects is not empty, this ticking object will be created.
+	// Also, this object will be deleted when CachedObjects becomes empty.
+	TUniquePtr<FDisplayClusterTickableGameObject> TickableGameObject;
 };
diff --git a/Engine/Plugins/Runtime/nDisplay/Source/DisplayClusterScenePreview/Public/DisplayClusterScenePreviewEnums.h b/Engine/Plugins/Runtime/nDisplay/Source/DisplayClusterScenePreview/Public/DisplayClusterScenePreviewEnums.h
--- a/Engine/Plugins/Runtime/nDisplay/Source/DisplayClusterScenePreview/Public/DisplayClusterScenePreviewEnums.h
+++ b/Engine/Plugins/Runtime/nDisplay/Source/DisplayClusterScenePreview/Public/DisplayClusterScenePreviewEnums.h
@@ -11,10 +11,16 @@ enum class EDisplayClusterScenePreviewFlags : uint8
 {
 	None = 0,
 
-	/** If true, use DCRA proxy for rendering. */
+	/** Use DCRA proxy for rendering. */
 	UseRootActorProxy = 1 << 0,
 
-	/** If true, automatically update the renderer with stage actors belonging to the root actor. */
-	AutoUpdateLightcards = 1 << 1,
+	/** Automatically update the renderer with stage actors belonging to the RootActor in scene. */
+	AutoUpdateStageActors = 1 << 1,
+
+	/* Move the RootActorProxy to the same position as RootActor in the scene to match the position of the StageActors in world space. */
+	ProxyFollowSceneRootActor = 1 << 2,
+
+	/** The proxy always calls the TickPreviewRenderer() function. */
+	ProxyTickPreviewRenderer= 1 << 3,
 };
 ENUM_CLASS_FLAGS(EDisplayClusterScenePreviewFlags);
diff --git a/Engine/Plugins/Runtime/nDisplay/Source/DisplayClusterScenePreview/Public/IDisplayClusterScenePreview.h b/Engine/Plugins/Runtime/nDisplay/Source/DisplayClusterScenePreview/Public/IDisplayClusterScenePreview.h
--- a/Engine/Plugins/Runtime/nDisplay/Source/DisplayClusterScenePreview/Public/IDisplayClusterScenePreview.h
+++ b/Engine/Plugins/Runtime/nDisplay/Source/DisplayClusterScenePreview/Public/IDisplayClusterScenePreview.h
@@ -89,6 +89,13 @@ public:
 	 */
 	virtual ADisplayClusterRootActor* GetRendererRootActor(int32 RendererId) = 0;
 
+	/**
+	 * Get the root actor (or its proxy) of a renderer.
+	 *
+	 * @param RendererId The ID of the renderer as returned from CreateRenderer.
+	 */
+	virtual ADisplayClusterRootActor* GetRendererRootActorOrProxy(int32 RendererId) = 0;
+
 	/**
 	 * Get a list of all actors that have been added to a renderer's scene.
 	 * This includes any actors that the renderer automatically added to itself.
diff --git a/Engine/Plugins/VirtualProduction/EpicStageApp/Source/EpicStageApp/Private/StageAppRouteHandler.cpp b/Engine/Plugins/VirtualProduction/EpicStageApp/Source/EpicStageApp/Private/StageAppRouteHandler.cpp
--- a/Engine/Plugins/VirtualProduction/EpicStageApp/Source/EpicStageApp/Private/StageAppRouteHandler.cpp
+++ b/Engine/Plugins/VirtualProduction/EpicStageApp/Source/EpicStageApp/Private/StageAppRouteHandler.cpp
@@ -176,13 +176,13 @@ void FStageAppRouteHandler::FPerRendererData::SetPreviewSettings(const FRCWebSoc
 
 bool FStageAppRouteHandler::FPerRendererData::GetSceneViewInitOptions(FSceneViewInitOptions& ViewInitOptions, bool bApplyActorRotation)
 {
-	ADisplayClusterRootActor* RootActor = GetRootActor();
-	if (!RootActor)
+	ADisplayClusterRootActor* RootActorProxy = GetRootActorProxy();
+	if (!RootActorProxy)
 	{
 		return false;
 	}
 
-	USceneComponent* ViewOriginComponent = RootActor->GetCommonViewPoint();
+	USceneComponent* ViewOriginComponent = RootActorProxy->GetCommonViewPoint();
 	if (!ViewOriginComponent)
 	{
 		return false;
@@ -201,7 +201,7 @@ bool FStageAppRouteHandler::FPerRendererData::GetSceneViewInitOptions(FSceneView
 	{
 		if (bApplyActorRotation)
 		{
-			Rotation = FRotator(RootActor->GetActorRotation().Quaternion() * PreviewSettings.Rotation.Quaternion());
+			Rotation = FRotator(RootActorProxy->GetActorRotation().Quaternion() * PreviewSettings.Rotation.Quaternion());
 		}
 		else
 		{
@@ -397,6 +397,13 @@ void FStageAppRouteHandler::RegisterRoute(TUniquePtr<FRemoteControlWebsocketRout
 	Routes.Emplace(MoveTemp(Route));
 }
 
+ADisplayClusterRootActor* FStageAppRouteHandler::FPerRendererData::GetRootActorProxy() const
+{
+	// In this application we will always use RootActorProxy, so the function below will always return a proxy.
+	// (see EDisplayClusterScenePreviewFlags::UseRootActorProxy flag)
+	return IDisplayClusterScenePreview::Get().GetRendererRootActorOrProxy(RendererId);
+}
+
 void FStageAppRouteHandler::SetRendererRootActorPath(const int32 RendererId, const FString& RootActorPath)
 {
 	// Use custom settings for root actor.
@@ -417,8 +424,12 @@ void FStageAppRouteHandler::SetRendererRootActorPath(const int32 RendererId, con
 		// When we use preview rendering for an external request with settings different from those used in the scene,
 		// need to create a new DCRA actor and use it as a proxy for custom rendering.
 		EDisplayClusterScenePreviewFlags::UseRootActorProxy
-		// and also continue to use the bAutoUpdateLightcards flag from the previous revision.
-		| EDisplayClusterScenePreviewFlags::AutoUpdateLightcards
+		// Move the RootActorProxy to the same position as in the scene to match the position of the StageActors in world space.
+		| EDisplayClusterScenePreviewFlags::ProxyFollowSceneRootActor
+		// Also proxy must render preview each frame
+		| EDisplayClusterScenePreviewFlags::ProxyTickPreviewRenderer
+		// Automatically update the renderer with stage actors belonging to the root actor in scene.
+		| EDisplayClusterScenePreviewFlags::AutoUpdateStageActors
 	);
 }
 
diff --git a/Engine/Plugins/VirtualProduction/EpicStageApp/Source/EpicStageApp/Private/StageAppRouteHandler.h b/Engine/Plugins/VirtualProduction/EpicStageApp/Source/EpicStageApp/Private/StageAppRouteHandler.h
--- a/Engine/Plugins/VirtualProduction/EpicStageApp/Source/EpicStageApp/Private/StageAppRouteHandler.h
+++ b/Engine/Plugins/VirtualProduction/EpicStageApp/Source/EpicStageApp/Private/StageAppRouteHandler.h
@@ -62,9 +62,20 @@ private:
 		/** Get the highest sequence number received from the client during a drag operation. */
 		int64 GetDragSequenceNumber() const;
 
-		/** Get the root actor used for previews. */
+		/** Get DCRA from the scene that is used to render preview. */
 		ADisplayClusterRootActor* GetRootActor() const;
 
+		/** Get the RootActorProxy that is used for rendering to the final texture.
+		* 
+		* This actor is used because of special rendering rules that differ from the RootActor settings in the scene,
+		* which requires a separate RootActor instance to be created.
+		* For example:
+		* - DCRA in a scene may have settings that disable preview rendering (or even freeze it). But this app will always force this feature.
+		* - An app can have custom DisplayDevice settings. However, the DCRA of the scene cannot be modified on this side.
+		* - The app can have custom preview settings, such as "PreviewSceeenPercentage," to set its own image quality, etc.
+		*/
+		ADisplayClusterRootActor* GetRootActorProxy() const;
+
 	private:
 		/** Update LightCardHelper's settings based on the ones stored in PreviewSettings. */
 		void UpdateLightCardHelperSettings();
