[nDisplay] Fixed performance drop when ICVFX camera is hidden on a viewport

diff --git a/Engine/Plugins/Runtime/nDisplay/Source/DisplayCluster/Private/Render/Viewport/Configuration/DisplayClusterViewportConfiguration_ICVFX.cpp b/Engine/Plugins/Runtime/nDisplay/Source/DisplayCluster/Private/Render/Viewport/Configuration/DisplayClusterViewportConfiguration_ICVFX.cpp
--- a/Engine/Plugins/Runtime/nDisplay/Source/DisplayCluster/Private/Render/Viewport/Configuration/DisplayClusterViewportConfiguration_ICVFX.cpp
+++ b/Engine/Plugins/Runtime/nDisplay/Source/DisplayCluster/Private/Render/Viewport/Configuration/DisplayClusterViewportConfiguration_ICVFX.cpp
@@ -2,30 +2,31 @@
 
 #include "DisplayClusterViewportConfiguration_ICVFX.h"
 
-#include "DisplayClusterViewportConfiguration.h"
-#include "DisplayClusterViewportConfigurationHelpers.h"
-#include "DisplayClusterViewportConfigurationHelpers_ICVFX.h"
-#include "DisplayClusterViewportConfigurationHelpers_Visibility.h"
-
-#include "DisplayClusterRootActor.h"
+#include "DisplayClusterEnums.h"
 #include "DisplayClusterConfigurationTypes_Viewport.h"
-
+#include "DisplayClusterProjectionStrings.h"
+#include "DisplayClusterRootActor.h"
 #include "IDisplayClusterProjection.h"
-#include "Render/Projection/IDisplayClusterProjectionPolicy.h"
+#include "IPDisplayCluster.h"
+
+#include "Cluster/IPDisplayClusterClusterManager.h"
+#include "Components/DisplayClusterICVFXCameraComponent.h"
 #include "Containers/DisplayClusterProjectionCameraPolicySettings.h"
-#include "DisplayClusterProjectionStrings.h"
 
+#include "Misc/DisplayClusterLog.h"
+#include "Misc/DisplayClusterGlobals.h"
+#include "Misc/Parse.h"
+
+#include "Render/Projection/IDisplayClusterProjectionPolicy.h"
+#include "Render/Viewport/Configuration/DisplayClusterViewportConfiguration.h"
+#include "Render/Viewport/Configuration/DisplayClusterViewportConfigurationHelpers.h"
+#include "Render/Viewport/Configuration/DisplayClusterViewportConfigurationHelpers_ICVFX.h"
+#include "Render/Viewport/Configuration/DisplayClusterViewportConfigurationHelpers_Visibility.h"
 #include "Render/Viewport/DisplayClusterViewport.h"
 #include "Render/Viewport/DisplayClusterViewportManager.h"
 #include "Render/Viewport/DisplayClusterViewportStrings.h"
-#include "Render/Viewport/RenderFrame/DisplayClusterRenderFrameSettings.h"
-
 #include "Render/Viewport/LightCard/DisplayClusterViewportLightCardManager.h"
-
-#include "Components/DisplayClusterICVFXCameraComponent.h"
-
-#include "Misc/DisplayClusterLog.h"
-#include "Misc/Parse.h"
+#include "Render/Viewport/RenderFrame/DisplayClusterRenderFrameSettings.h"
 
 
 ////////////////////////////////////////////////////////////////////////////////
@@ -322,8 +323,10 @@ void FDisplayClusterViewportConfiguration_ICVFX::GetAndUpdateStageCameras(const
 					// Add this target to all cameras visible on it
 					for (FDisplayClusterViewportConfiguration_ICVFXCamera& CameraIt : StageCameras)
 					{
-						if (CameraIt.IsCameraProjectionVisibleOnViewport(TargetIt.Get())
-							&& !CameraIt.GetCameraSettings().HiddenICVFXViewports.ItemNames.Contains(TargetIt->GetId()))
+						const bool bCameraProjectionVisible  = CameraIt.IsCameraProjectionVisibleOnViewport(TargetIt.Get());
+						const bool bIsCameraHiddenOnViewport = CameraIt.GetCameraSettings().HiddenICVFXViewports.ItemNames.Contains(TargetIt->GetId());
+						
+						if (bCameraProjectionVisible && !bIsCameraHiddenOnViewport)
 						{
 							CameraIt.VisibleTargets.Add(TargetIt);
 						}
@@ -332,10 +335,19 @@ void FDisplayClusterViewportConfiguration_ICVFX::GetAndUpdateStageCameras(const
 			}
 		}
 
+		const bool bIsRunningCluster = GDisplayCluster->GetOperationMode() == EDisplayClusterOperationMode::Cluster;
+		const FString ClusterNodeId  = GDisplayCluster->GetPrivateClusterMgr()->GetNodeId();
+
 		// Create camera resources and initialize target ICVFX viewports
 		for (FDisplayClusterViewportConfiguration_ICVFXCamera& CameraIt : StageCameras)
 		{
-			if (InTargetViewports == nullptr || CameraIt.VisibleTargets.Num() > 0)
+			const FDisplayClusterConfigurationMediaICVFX& MediaSettings = CameraIt.GetCameraSettings().RenderSettings.Media;
+
+			const bool bCameraHasMedia = MediaSettings.bEnable && (
+				MediaSettings.HasAnyMediaInputAssigned(ClusterNodeId, MediaSettings.SplitType) ||
+				MediaSettings.HasAnyMediaOutputAssigned(ClusterNodeId, MediaSettings.SplitType));
+
+			if (InTargetViewports == nullptr || CameraIt.VisibleTargets.Num() > 0 || (bIsRunningCluster && bCameraHasMedia))
 			{
 				CameraIt.Update();
 			}
