nDisplay: Added a new static function that sets FEngineShowFlags according to DC rendering modes.
- Added new function IDisplayClusterViewportManager::SetupEngineShowFlags()
- Code cleanup and duplicate code removal.

diff --git a/Engine/Plugins/Runtime/nDisplay/Source/DisplayCluster/Private/Render/Viewport/DisplayClusterViewportManager.cpp b/Engine/Plugins/Runtime/nDisplay/Source/DisplayCluster/Private/Render/Viewport/DisplayClusterViewportManager.cpp
--- a/Engine/Plugins/Runtime/nDisplay/Source/DisplayCluster/Private/Render/Viewport/DisplayClusterViewportManager.cpp
+++ b/Engine/Plugins/Runtime/nDisplay/Source/DisplayCluster/Private/Render/Viewport/DisplayClusterViewportManager.cpp
@@ -145,6 +145,60 @@ TSharedRef<IDisplayClusterViewportManager, ESPMode::ThreadSafe> IDisplayClusterV
 	return ViewportManager;
 }
 
+void IDisplayClusterViewportManager::SetupEngineShowFlags(const EDisplayClusterViewportCaptureMode InMode, FEngineShowFlags& InOutEngineShowFlags)
+{
+	// Control NaniteMeshes for ChromaKey and Lightcards:
+	switch (InMode)
+	{
+	case EDisplayClusterViewportCaptureMode::Chromakey:
+		if (!GDisplayClusterChromaKeyAllowNanite)
+		{
+			InOutEngineShowFlags.SetNaniteMeshes(0);
+		}
+
+		break;
+	case EDisplayClusterViewportCaptureMode::Lightcard:
+		if (!GDisplayClusterLightcardsAllowNanite)
+		{
+			InOutEngineShowFlags.SetNaniteMeshes(0);
+		}
+		break;
+	default:
+		break;
+	}
+
+	switch (InMode)
+	{
+	case EDisplayClusterViewportCaptureMode::Chromakey:
+	case EDisplayClusterViewportCaptureMode::Lightcard:
+		// Disable postprocess for LC\CK
+		InOutEngineShowFlags.SetPostProcessing(0);
+
+		InOutEngineShowFlags.SetAtmosphere(0);
+		InOutEngineShowFlags.SetFog(0);
+		InOutEngineShowFlags.SetVolumetricFog(0);
+		InOutEngineShowFlags.SetMotionBlur(0); // motion blur doesn't work correctly with scene captures.
+		InOutEngineShowFlags.SetSeparateTranslucency(0);
+		InOutEngineShowFlags.SetHMDDistortion(0);
+		InOutEngineShowFlags.SetOnScreenDebug(0);
+		InOutEngineShowFlags.SetHair(0);
+
+		InOutEngineShowFlags.SetLumenReflections(0);
+		InOutEngineShowFlags.SetLumenGlobalIllumination(0);
+		InOutEngineShowFlags.SetGlobalIllumination(0);
+
+		InOutEngineShowFlags.SetScreenSpaceAO(0);
+		InOutEngineShowFlags.SetAmbientOcclusion(0);
+		InOutEngineShowFlags.SetDeferredLighting(0);
+		InOutEngineShowFlags.SetVirtualTexturePrimitives(0);
+		InOutEngineShowFlags.SetRectLights(0);
+		break;
+
+	default:
+		break;
+	}
+}
+
 FDisplayClusterViewportManager::FDisplayClusterViewportManager()
 	: Configuration(MakeShared<FDisplayClusterViewportConfiguration, ESPMode::ThreadSafe>())
 	, ViewportManagerPreview(MakeShared<FDisplayClusterViewportManagerPreview, ESPMode::ThreadSafe>(Configuration))
@@ -657,84 +711,47 @@ FSceneViewFamily::ConstructionValues FDisplayClusterViewportManager::CreateViewF
 
 	bool bResolveScene = true;
 
-	// Control NaniteMeshes for ChromaKey and Lightcards:
-	switch (InFrameTarget.CaptureMode)
-	{
-	case EDisplayClusterViewportCaptureMode::Chromakey:
-		if (!GDisplayClusterChromaKeyAllowNanite)
-		{
-			InEngineShowFlags.SetNaniteMeshes(0);
-		}
-
-		break;
-	case EDisplayClusterViewportCaptureMode::Lightcard:
-		if (!GDisplayClusterLightcardsAllowNanite)
-		{
-			InEngineShowFlags.SetNaniteMeshes(0);
-		}
-		break;
-	default:
-		break;
-	}
+	// Sets the engine flags corresponding to the capture mode.
+	IDisplayClusterViewportManager::SetupEngineShowFlags(InFrameTarget.CaptureMode, InEngineShowFlags);
 
 	const FDisplayClusterRenderFrameSettings& RenderFrameSettings = Configuration->GetRenderFrameSettings();
+
+	// A special case for DCRA previewing in a scene to avoid double use of PP.
 	if(RenderFrameSettings.IsPostProcessDisabled())
-		{
-			// Disable postprocess for preview
-			InEngineShowFlags.PostProcessing = 0;
-		}
+	{
+		// Disable postprocess for preview
+		InEngineShowFlags.PostProcessing = 0;
+	}
 
+	// A special use case is setting up alpha channel capture:
+	// (When using the DC viewport rendering pipeline).
 	switch (InFrameTarget.CaptureMode)
 	{
 	case EDisplayClusterViewportCaptureMode::Chromakey:
 	case EDisplayClusterViewportCaptureMode::Lightcard:
 		switch (RenderFrameSettings.AlphaChannelCaptureMode)
 		{
-			case EDisplayClusterRenderFrameAlphaChannelCaptureMode::Copy:
-				// Disable AA
-				InEngineShowFlags.SetAntiAliasing(0);
-				InEngineShowFlags.SetTemporalAA(0);
-				break;
-
-			case EDisplayClusterRenderFrameAlphaChannelCaptureMode::CopyAA:
-				// Use AA
-				InEngineShowFlags.SetAntiAliasing(1);
-				InEngineShowFlags.SetTemporalAA(0);
-				break;
-
-			case EDisplayClusterRenderFrameAlphaChannelCaptureMode::FXAA:
-				// Alpha captured without AA, own FXAA used
-				InEngineShowFlags.SetAntiAliasing(0);
-				InEngineShowFlags.SetTemporalAA(0);
-				break;
-
-			default:
-				break;
+		case EDisplayClusterRenderFrameAlphaChannelCaptureMode::Copy:
+			// Disable AA
+			InEngineShowFlags.SetAntiAliasing(0);
+			InEngineShowFlags.SetTemporalAA(0);
+			break;
+
+		case EDisplayClusterRenderFrameAlphaChannelCaptureMode::CopyAA:
+			// Use AA
+			InEngineShowFlags.SetAntiAliasing(1);
+			InEngineShowFlags.SetTemporalAA(0);
+			break;
+
+		case EDisplayClusterRenderFrameAlphaChannelCaptureMode::FXAA:
+			// Alpha captured without AA, own FXAA used
+			InEngineShowFlags.SetAntiAliasing(0);
+			InEngineShowFlags.SetTemporalAA(0);
+			break;
+
+		default:
+			break;
 		}
-
-		// Disable postprocess for LC\CK
-		InEngineShowFlags.SetPostProcessing(0);
-
-		InEngineShowFlags.SetAtmosphere(0);
-		InEngineShowFlags.SetFog(0);
-		InEngineShowFlags.SetVolumetricFog(0);
-		InEngineShowFlags.SetMotionBlur(0); // motion blur doesn't work correctly with scene captures.
-		InEngineShowFlags.SetSeparateTranslucency(0);
-		InEngineShowFlags.SetHMDDistortion(0);
-		InEngineShowFlags.SetOnScreenDebug(0);
-		InEngineShowFlags.SetHair(0);
-
-		InEngineShowFlags.SetLumenReflections(0);
-		InEngineShowFlags.SetLumenGlobalIllumination(0);
-		InEngineShowFlags.SetGlobalIllumination(0);
-
-		InEngineShowFlags.SetScreenSpaceAO(0);
-		InEngineShowFlags.SetAmbientOcclusion(0);
-		InEngineShowFlags.SetDeferredLighting(0);
-		InEngineShowFlags.SetVirtualTexturePrimitives(0);
-		InEngineShowFlags.SetRectLights(0);
-		break;
-
 	default:
 		break;
 	}
diff --git a/Engine/Plugins/Runtime/nDisplay/Source/DisplayCluster/Public/Render/Viewport/IDisplayClusterViewportManager.h b/Engine/Plugins/Runtime/nDisplay/Source/DisplayCluster/Public/Render/Viewport/IDisplayClusterViewportManager.h
--- a/Engine/Plugins/Runtime/nDisplay/Source/DisplayCluster/Public/Render/Viewport/IDisplayClusterViewportManager.h
+++ b/Engine/Plugins/Runtime/nDisplay/Source/DisplayCluster/Public/Render/Viewport/IDisplayClusterViewportManager.h
@@ -98,6 +98,9 @@ public:
 		const bool bInAdditionalViewFamily
 	) const = 0;
 
+	/** This function sets the engine flags corresponding to the rendering mode. */
+	static void SetupEngineShowFlags(const EDisplayClusterViewportCaptureMode InMode, FEngineShowFlags& InOutEngineShowFlags);
+
 	/**
 	* Initialize view family, using rules
 	* [Game thread func]
diff --git a/Engine/Plugins/Runtime/nDisplay/Source/DisplayClusterShaders/Private/Shaders/DisplayClusterShadersPreprocess_UVLightCards.cpp b/Engine/Plugins/Runtime/nDisplay/Source/DisplayClusterShaders/Private/Shaders/DisplayClusterShadersPreprocess_UVLightCards.cpp
--- a/Engine/Plugins/Runtime/nDisplay/Source/DisplayClusterShaders/Private/Shaders/DisplayClusterShadersPreprocess_UVLightCards.cpp
+++ b/Engine/Plugins/Runtime/nDisplay/Source/DisplayClusterShaders/Private/Shaders/DisplayClusterShadersPreprocess_UVLightCards.cpp
@@ -1,6 +1,7 @@
 // Copyright Epic Games, Inc. All Rights Reserved.
 
 #include "DisplayClusterShadersPreprocess_UVLightCards.h"
+#include "Render/Viewport/IDisplayClusterViewportManager.h"
 
 #include "Components/PrimitiveComponent.h"
 #include "DataDrivenShaderPlatformInfo.h"
@@ -204,29 +205,7 @@ bool FDisplayClusterShadersPreprocess_UVLightCards::RenderPreprocess_UVLightCard
 	FRDGBuilder GraphBuilder(RHICmdList);
 
 	FEngineShowFlags EngineShowFlags(ESFIM_Game);
-
-	// LightCard settings from the FDisplayClusterViewportManager::ConfigureViewFamily
-	{
-		EngineShowFlags.PostProcessing = 0;
-		EngineShowFlags.SetAtmosphere(0);
-		EngineShowFlags.SetFog(0);
-		EngineShowFlags.SetVolumetricFog(0);
-		EngineShowFlags.SetMotionBlur(0); // motion blur doesn't work correctly with scene captures.
-		EngineShowFlags.SetSeparateTranslucency(0);
-		EngineShowFlags.SetHMDDistortion(0);
-		EngineShowFlags.SetOnScreenDebug(0);
-		EngineShowFlags.SetHair(0);
-
-		EngineShowFlags.SetLumenReflections(0);
-		EngineShowFlags.SetLumenGlobalIllumination(0);
-		EngineShowFlags.SetGlobalIllumination(0);
-
-		EngineShowFlags.SetScreenSpaceAO(0);
-		EngineShowFlags.SetAmbientOcclusion(0);
-		EngineShowFlags.SetDeferredLighting(0);
-		EngineShowFlags.SetVirtualTexturePrimitives(0);
-		EngineShowFlags.SetRectLights(0);
-	}
+	IDisplayClusterViewportManager::SetupEngineShowFlags(EDisplayClusterViewportCaptureMode::Lightcard, EngineShowFlags);
 
 	FSceneViewFamilyContext ViewFamily(FSceneViewFamily::ConstructionValues(
 		InRenderTarget,
