[nDisplay] ICVFX media settings from BP

- Runtime media configuration is based on the Blueprint media settings
- All media settings are EditDefaultsOnly now

diff --git a/Engine/Plugins/Runtime/nDisplay/Source/DisplayCluster/Private/Render/Viewport/Configuration/DisplayClusterViewportConfigurationHelpers_ICVFX.cpp b/Engine/Plugins/Runtime/nDisplay/Source/DisplayCluster/Private/Render/Viewport/Configuration/DisplayClusterViewportConfigurationHelpers_ICVFX.cpp
--- a/Engine/Plugins/Runtime/nDisplay/Source/DisplayCluster/Private/Render/Viewport/Configuration/DisplayClusterViewportConfigurationHelpers_ICVFX.cpp
+++ b/Engine/Plugins/Runtime/nDisplay/Source/DisplayCluster/Private/Render/Viewport/Configuration/DisplayClusterViewportConfigurationHelpers_ICVFX.cpp
@@ -41,9 +41,11 @@
 
 #include "Components/DisplayClusterICVFXCameraComponent.h"
 
+#include "Engine/SCS_Node.h"
+#include "Engine/SimpleConstructionScript.h"
+#include "HAL/IConsoleManager.h"
 #include "Misc/DisplayClusterLog.h"
 #include "TextureResource.h"
-#include "HAL/IConsoleManager.h"
 
 ////////////////////////////////////////////////////////////////////////////////
 namespace UE::DisplayCluster::Viewport::ConfigurationHelpers_ICVFX
@@ -545,8 +547,30 @@ void FDisplayClusterViewportConfigurationHelpers_ICVFX::UpdateCameraViewportSett
 	// Set viewport buffer ratio
 	DstViewport.SetViewportBufferRatio(InCameraSettings.GetCameraBufferRatio(*StageSettings));
 
+#if true
+	// UE-211513
+	// This is a temporary workaround that allows to initialize media based on the blueprint data instead of 
+	// the instance data. There is an issue with the propagation of instanced object changes from the parent
+	// blueprints to its instances. To avoid any potential issues, we get media settings for the original nDisplay
+	// blueprint. Once the propagation issue is fixed, we'll be able use the instances again.
+
+	if (const AActor* const OwningActor = InCameraComponent.GetOwner())
+	{
+		if (const UBlueprint* const ParentBP = UBlueprint::GetBlueprintFromClass(OwningActor->GetClass()))
+		{
+			if (const USCS_Node* const CameraSCSNode = ParentBP->SimpleConstructionScript->FindSCSNode(*InCameraComponent.GetName()))
+			{
+				if (const UDisplayClusterICVFXCameraComponent* const ArchetypeCameraComponent = Cast<UDisplayClusterICVFXCameraComponent>(CameraSCSNode->ComponentTemplate))
+				{
+					FDisplayClusterViewportConfigurationHelpers_Tile::UpdateICVFXCameraViewportTileSettings(DstViewport, ArchetypeCameraComponent->CameraSettings.RenderSettings.Media);
+				}
+			}
+		}
+	}
+#else
 	// InCamera tile rendering.
 	FDisplayClusterViewportConfigurationHelpers_Tile::UpdateICVFXCameraViewportTileSettings(DstViewport, InCameraSettings.RenderSettings.Media);
+#endif
 }
 
 void FDisplayClusterViewportConfigurationHelpers_ICVFX::UpdateChromakeyViewportSettings(FDisplayClusterViewport& DstViewport, FDisplayClusterViewport& InCameraViewport, const FDisplayClusterConfigurationICVFX_CameraSettings& InCameraSettings)
diff --git a/Engine/Plugins/Runtime/nDisplay/Source/DisplayCluster/Public/Components/DisplayClusterICVFXCameraComponent.h b/Engine/Plugins/Runtime/nDisplay/Source/DisplayCluster/Public/Components/DisplayClusterICVFXCameraComponent.h
--- a/Engine/Plugins/Runtime/nDisplay/Source/DisplayCluster/Public/Components/DisplayClusterICVFXCameraComponent.h
+++ b/Engine/Plugins/Runtime/nDisplay/Source/DisplayCluster/Public/Components/DisplayClusterICVFXCameraComponent.h
@@ -196,7 +196,7 @@ private:
 	UPROPERTY(EditAnywhere, Transient, Category = "Texture Replacement", meta = (PropertyPath = "CameraSettings.RenderSettings.Replace.TextureRegion"))
 	FDisplayClusterEditorPropertyReference TextureRegionRef;
 
-	UPROPERTY(EditAnywhere, Transient, Category = "Media", meta = (PropertyPath = "CameraSettings.RenderSettings.Media"))
+	UPROPERTY(EditDefaultsOnly, Transient, Category = "Media", meta = (PropertyPath = "CameraSettings.RenderSettings.Media"))
 	FDisplayClusterEditorPropertyReference MediaRef;
 
 	UPROPERTY(EditAnywhere, Transient, Category = Configuration, meta = (PropertyPath = "CameraSettings.RenderSettings.RenderOrder"))
diff --git a/Engine/Plugins/Runtime/nDisplay/Source/DisplayCluster/Public/DisplayClusterRootActor.h b/Engine/Plugins/Runtime/nDisplay/Source/DisplayCluster/Public/DisplayClusterRootActor.h
--- a/Engine/Plugins/Runtime/nDisplay/Source/DisplayCluster/Public/DisplayClusterRootActor.h
+++ b/Engine/Plugins/Runtime/nDisplay/Source/DisplayCluster/Public/DisplayClusterRootActor.h
@@ -344,7 +344,7 @@ private:
 	FDisplayClusterEditorPropertyReference LightcardPerViewportOCIOProfilesRef;
 
 	// Media
-	UPROPERTY(EditAnywhere, Transient, Category = "Media", meta = (PropertyPath = "CurrentConfigData.MediaSettings"))
+	UPROPERTY(EditDefaultsOnly, Transient, Category = "Media", meta = (PropertyPath = "CurrentConfigData.MediaSettings"))
 	FDisplayClusterEditorPropertyReference MediaSettingsRef;
 
 #endif // WITH_EDITORONLY_DATA
diff --git a/Engine/Plugins/Runtime/nDisplay/Source/DisplayClusterConfiguration/Public/DisplayClusterConfigurationTypes.h b/Engine/Plugins/Runtime/nDisplay/Source/DisplayClusterConfiguration/Public/DisplayClusterConfigurationTypes.h
--- a/Engine/Plugins/Runtime/nDisplay/Source/DisplayClusterConfiguration/Public/DisplayClusterConfigurationTypes.h
+++ b/Engine/Plugins/Runtime/nDisplay/Source/DisplayClusterConfiguration/Public/DisplayClusterConfigurationTypes.h
@@ -358,7 +358,7 @@ public:
 	TMap<FString, FDisplayClusterConfigurationPostprocess> Postprocess;
 
 	// Media settings
-	UPROPERTY(EditAnywhere, BlueprintReadWrite, Category = "Media", meta = (DisplayName = "Media"))
+	UPROPERTY(EditDefaultsOnly, BlueprintReadWrite, Category = "Media", meta = (DisplayName = "Media"))
 	FDisplayClusterConfigurationMediaNodeBackbuffer MediaSettings;
 
 #if WITH_EDITORONLY_DATA
diff --git a/Engine/Plugins/Runtime/nDisplay/Source/DisplayClusterConfiguration/Public/DisplayClusterConfigurationTypes_ICVFX.h b/Engine/Plugins/Runtime/nDisplay/Source/DisplayClusterConfiguration/Public/DisplayClusterConfigurationTypes_ICVFX.h
--- a/Engine/Plugins/Runtime/nDisplay/Source/DisplayClusterConfiguration/Public/DisplayClusterConfigurationTypes_ICVFX.h
+++ b/Engine/Plugins/Runtime/nDisplay/Source/DisplayClusterConfiguration/Public/DisplayClusterConfigurationTypes_ICVFX.h
@@ -503,7 +503,7 @@ public:
 	FDisplayClusterConfigurationPostRender_Override Replace;
 
 	// Media settings
-	UPROPERTY(EditAnywhere, BlueprintReadWrite, Category = NDisplay, meta = (DisplayName = "Media"))
+	UPROPERTY(EditDefaultsOnly, BlueprintReadWrite, Category = NDisplay, meta = (DisplayName = "Media"))
 	FDisplayClusterConfigurationMediaICVFX Media;
 
 	UPROPERTY()
diff --git a/Engine/Plugins/Runtime/nDisplay/Source/DisplayClusterMedia/Private/DisplayClusterMediaModule.cpp b/Engine/Plugins/Runtime/nDisplay/Source/DisplayClusterMedia/Private/DisplayClusterMediaModule.cpp
--- a/Engine/Plugins/Runtime/nDisplay/Source/DisplayClusterMedia/Private/DisplayClusterMediaModule.cpp
+++ b/Engine/Plugins/Runtime/nDisplay/Source/DisplayClusterMedia/Private/DisplayClusterMediaModule.cpp
@@ -12,9 +12,12 @@
 #include "DisplayClusterRootActor.h"
 #include "DisplayClusterConfigurationTypes.h"
 
+#include "Blueprints/DisplayClusterBlueprint.h"
 #include "Components/DisplayClusterICVFXCameraComponent.h"
 #include "Config/IDisplayClusterConfigManager.h"
 #include "Cluster/IDisplayClusterClusterManager.h"
+#include "Engine/SCS_Node.h"
+#include "Engine/SimpleConstructionScript.h"
 #include "Game/IDisplayClusterGameManager.h"
 
 #include "Capture/DisplayClusterMediaCaptureCamera.h"
@@ -50,7 +53,17 @@ void FDisplayClusterMediaModule::OnPreSubmitViewFamilies(TArray<FSceneViewFamily
 	// Unsubscribe after first call. Currently, media initialization is a one time procedure. No need to receive any further callbacks.
 	IDisplayCluster::Get().GetCallbacks().OnDisplayClusterPreSubmitViewFamilies().RemoveAll(this);
 
+#if true
+	// UE-211513
+	// This is a temporary workaround that allows to initialize media based on the blueprint data instead of 
+	// the instance data. There is an issue with the propagation of instanced object changes from the parent
+	// blueprints to its instances. To avoid any potential issues, we get media settings for the original nDisplay
+	// blueprint. Once the propagation issue is fixed, we'll be able to safely return to the original InitializeMedia().
+	InitializeMediaFromCDO();
+#else
 	InitializeMedia();
+#endif
+
 	StartCapture();
 	PlayMedia();
 }
@@ -127,6 +140,132 @@ void FDisplayClusterMediaModule::InitializeMedia()
 	}
 }
 
+void FDisplayClusterMediaModule::InitializeMediaFromCDO()
+{
+	// Runtime only for now
+	if (IDisplayCluster::Get().GetOperationMode() != EDisplayClusterOperationMode::Cluster)
+	{
+		UE_LOG(LogDisplayClusterMedia, Warning, TEXT("DisplayClusterMedia is available in 'cluster' operation mode only"));
+		return;
+	}
+
+	// Check if media enabled
+	if (!CVarMediaEnabled.GetValueOnGameThread())
+	{
+		UE_LOG(LogDisplayClusterMedia, Log, TEXT("nDisplay media subsytem is disabled by a cvar"));
+		return;
+	}
+
+	// Instantiate latency queue
+	FrameQueue.Init();
+
+	// Get active DCRA instance
+	const ADisplayClusterRootActor* const RootActor = IDisplayCluster::Get().GetGameMgr()->GetRootActor();
+	if (!IsValid(RootActor))
+	{
+		UE_LOG(LogDisplayClusterMedia, Warning, TEXT("Invalid DCRA object"));
+		return;
+	}
+
+	// Get original (parent) blueprint of this DCRA instance
+	const UDisplayClusterBlueprint* DCBP = nullptr;
+	if (const UBlueprint* BP = UBlueprint::GetBlueprintFromClass(RootActor->GetClass()))
+	{
+		DCBP = Cast<UDisplayClusterBlueprint>(BP);
+	}
+
+	// Aux data for media initialization
+	const FString ClusterNodeId = IDisplayCluster::Get().GetClusterMgr()->GetNodeId();
+	const FString RootActorName = RootActor->GetName();
+
+	// It's possible that we have DCBP==nullptr. For example, when running with portable configs, a pure DCRA
+	// actor gets spawned. There is no blueprint available in this case. But we don't expect any media in this
+	// case so stop initialization procedure.
+	// !DCBP is to prevent C6011 warning (potential dereferencing nullptr)
+	if (!IsValid(DCBP) || !DCBP)
+	{
+		UE_LOG(LogDisplayClusterMedia, Log, TEXT("Skipping media initialization for a pure DCRA instance '%s'"), *RootActorName);
+		return;
+	}
+
+	// Now we need to get the configuration data.
+	const UDisplayClusterConfigurationData* const ConfigData = DCBP->GetConfig();
+	if (!IsValid(ConfigData) || !IsValid(ConfigData->Cluster))
+	{
+		UE_LOG(LogDisplayClusterMedia, Warning, TEXT("Couldn't obtain the cluster configuration data"));
+		return;
+	}
+
+	// Find a proper cluster node configuration
+	const UDisplayClusterConfigurationClusterNode* const ClusterNodeCfg = ConfigData->Cluster->GetNode(ClusterNodeId);
+	if (!IsValid(ClusterNodeCfg))
+	{
+		UE_LOG(LogDisplayClusterMedia, Warning, TEXT("Couldn't obtain any configuration data for node '%s'"), *ClusterNodeId);
+		return;
+	}
+
+	// Node backbuffer media setup
+	{
+		InitializeBackbufferOutput(ClusterNodeCfg, RootActorName, ClusterNodeId);
+	}
+
+	// Viewports media setup
+	for (const TPair<FString, TObjectPtr<UDisplayClusterConfigurationViewport>>& ViewportIt : ClusterNodeCfg->Viewports)
+	{
+		InitializeViewportInput(ViewportIt.Value, ViewportIt.Key, RootActorName, ClusterNodeId);
+		InitializeViewportOutput(ViewportIt.Value, ViewportIt.Key, RootActorName, ClusterNodeId);
+	}
+
+	// ICVFX media setup
+	{
+		// Get all ICVFX camera components
+		TArray<UDisplayClusterICVFXCameraComponent*> ICVFXCameraComponents;
+		RootActor->GetComponents(ICVFXCameraComponents);
+
+		for (const UDisplayClusterICVFXCameraComponent* const ICVFXCameraInstanceComponent : ICVFXCameraComponents)
+		{
+			const UDisplayClusterICVFXCameraComponent* ICVFXCameraComponent = nullptr;
+			const FDisplayClusterConfigurationMediaICVFX* MediaSettings = nullptr;
+
+			const FString CameraName = ICVFXCameraInstanceComponent->GetName();
+
+			// Let's see if this camera comes from the blueprint. If so, we'll be using its media settings. Otherwise
+			// the camera instance will be used.
+			const USCS_Node* const CameraSCSNode = DCBP->SimpleConstructionScript->FindSCSNode(*CameraName);
+			if (IsValid(CameraSCSNode))
+			{
+				ICVFXCameraComponent = Cast<UDisplayClusterICVFXCameraComponent>(CameraSCSNode->ComponentTemplate);
+			}
+			else
+			{
+				ICVFXCameraComponent = ICVFXCameraInstanceComponent;
+			}
+
+			// Get media settings
+			if (IsValid(ICVFXCameraComponent))
+			{
+				MediaSettings = &ICVFXCameraComponent->CameraSettings.RenderSettings.Media;
+			}
+
+			if (MediaSettings)
+			{
+				// Full frame
+				if (MediaSettings->SplitType == EDisplayClusterConfigurationMediaSplitType::FullFrame)
+				{
+					InitializeICVFXCameraFullFrameInput(ICVFXCameraComponent, RootActorName, ClusterNodeId, CameraName);
+					InitializeICVFXCameraFullFrameOutput(ICVFXCameraComponent, RootActorName, ClusterNodeId, CameraName);
+				}
+				// Uniform tiles
+				else if (MediaSettings->SplitType == EDisplayClusterConfigurationMediaSplitType::UniformTiles)
+				{
+					InitializeICVFXCameraUniformTilesInput(ICVFXCameraComponent, RootActorName, ClusterNodeId, CameraName);
+					InitializeICVFXCameraUniformTilesOutput(ICVFXCameraComponent, RootActorName, ClusterNodeId, CameraName);
+				}
+			}
+		}
+	}
+}
+
 void FDisplayClusterMediaModule::ReleaseMedia()
 {
 	StopCapture();
@@ -289,7 +428,7 @@ void FDisplayClusterMediaModule::InitializeViewportOutput(const UDisplayClusterC
 	}
 }
 
-void FDisplayClusterMediaModule::InitializeICVFXCameraFullFrameInput(const UDisplayClusterICVFXCameraComponent* ICVFXCameraComponent, const FString& RootActorName, const FString& ClusterNodeId)
+void FDisplayClusterMediaModule::InitializeICVFXCameraFullFrameInput(const UDisplayClusterICVFXCameraComponent* ICVFXCameraComponent, const FString& RootActorName, const FString& ClusterNodeId, const FString& ExplicitCameraName)
 {
 	checkSlow(ICVFXCameraComponent);
 
@@ -299,7 +438,7 @@ void FDisplayClusterMediaModule::InitializeICVFXCameraFullFrameInput(const UDisp
 
 		if (MediaSettings.bEnable && MediaSettings.SplitType == EDisplayClusterConfigurationMediaSplitType::FullFrame)
 		{
-			const FString ICVFXCameraName = ICVFXCameraComponent->GetName();
+			const FString ICVFXCameraName = ExplicitCameraName.IsEmpty() ? ICVFXCameraComponent->GetName() : ExplicitCameraName;
 			const FString ICVFXViewportId = DisplayClusterMediaHelpers::GenerateICVFXViewportName(ClusterNodeId, ICVFXCameraName);
 
 			if (UMediaSource* MediaSource = MediaSettings.GetMediaSource(ClusterNodeId))
@@ -324,7 +463,7 @@ void FDisplayClusterMediaModule::InitializeICVFXCameraFullFrameInput(const UDisp
 	}
 }
 
-void FDisplayClusterMediaModule::InitializeICVFXCameraFullFrameOutput(const UDisplayClusterICVFXCameraComponent* ICVFXCameraComponent, const FString& RootActorName, const FString& ClusterNodeId)
+void FDisplayClusterMediaModule::InitializeICVFXCameraFullFrameOutput(const UDisplayClusterICVFXCameraComponent* ICVFXCameraComponent, const FString& RootActorName, const FString& ClusterNodeId, const FString& ExplicitCameraName)
 {
 	checkSlow(ICVFXCameraComponent);
 
@@ -334,7 +473,7 @@ void FDisplayClusterMediaModule::InitializeICVFXCameraFullFrameOutput(const UDis
 
 		if (MediaSettings.bEnable && MediaSettings.SplitType == EDisplayClusterConfigurationMediaSplitType::FullFrame)
 		{
-			const FString ICVFXCameraName = ICVFXCameraComponent->GetName();
+			const FString ICVFXCameraName = ExplicitCameraName.IsEmpty() ? ICVFXCameraComponent->GetName() : ExplicitCameraName;
 			const FString ICVFXViewportId = DisplayClusterMediaHelpers::GenerateICVFXViewportName(ClusterNodeId, ICVFXCameraName);
 
 			// Media capture (full frame)
@@ -368,7 +507,7 @@ void FDisplayClusterMediaModule::InitializeICVFXCameraFullFrameOutput(const UDis
 	}
 }
 
-void FDisplayClusterMediaModule::InitializeICVFXCameraUniformTilesInput(const UDisplayClusterICVFXCameraComponent* ICVFXCameraComponent, const FString& RootActorName, const FString& ClusterNodeId)
+void FDisplayClusterMediaModule::InitializeICVFXCameraUniformTilesInput(const UDisplayClusterICVFXCameraComponent* ICVFXCameraComponent, const FString& RootActorName, const FString& ClusterNodeId, const FString& ExplicitCameraName)
 {
 	checkSlow(ICVFXCameraComponent);
 
@@ -378,7 +517,7 @@ void FDisplayClusterMediaModule::InitializeICVFXCameraUniformTilesInput(const UD
 
 		if (MediaSettings.bEnable && MediaSettings.SplitType == EDisplayClusterConfigurationMediaSplitType::UniformTiles)
 		{
-			const FString ICVFXCameraName = ICVFXCameraComponent->GetName();
+			const FString ICVFXCameraName = ExplicitCameraName.IsEmpty() ? ICVFXCameraComponent->GetName() : ExplicitCameraName;
 			const FString ICVFXViewportId = DisplayClusterMediaHelpers::GenerateICVFXViewportName(ClusterNodeId, ICVFXCameraName);
 
 			// Find corresponsing media group
@@ -420,7 +559,7 @@ void FDisplayClusterMediaModule::InitializeICVFXCameraUniformTilesInput(const UD
 	}
 }
 
-void FDisplayClusterMediaModule::InitializeICVFXCameraUniformTilesOutput(const UDisplayClusterICVFXCameraComponent* ICVFXCameraComponent, const FString& RootActorName, const FString& ClusterNodeId)
+void FDisplayClusterMediaModule::InitializeICVFXCameraUniformTilesOutput(const UDisplayClusterICVFXCameraComponent* ICVFXCameraComponent, const FString& RootActorName, const FString& ClusterNodeId, const FString& ExplicitCameraName)
 {
 	checkSlow(ICVFXCameraComponent);
 
@@ -430,7 +569,7 @@ void FDisplayClusterMediaModule::InitializeICVFXCameraUniformTilesOutput(const U
 
 		if (MediaSettings.bEnable && MediaSettings.SplitType == EDisplayClusterConfigurationMediaSplitType::UniformTiles)
 		{
-			const FString ICVFXCameraName = ICVFXCameraComponent->GetName();
+			const FString ICVFXCameraName = ExplicitCameraName.IsEmpty() ? ICVFXCameraComponent->GetName() : ExplicitCameraName;
 			const FString ICVFXViewportId = DisplayClusterMediaHelpers::GenerateICVFXViewportName(ClusterNodeId, ICVFXCameraName);
 
 			// Find corresponsing media group
diff --git a/Engine/Plugins/Runtime/nDisplay/Source/DisplayClusterMedia/Private/DisplayClusterMediaModule.h b/Engine/Plugins/Runtime/nDisplay/Source/DisplayClusterMedia/Private/DisplayClusterMediaModule.h
--- a/Engine/Plugins/Runtime/nDisplay/Source/DisplayClusterMedia/Private/DisplayClusterMediaModule.h
+++ b/Engine/Plugins/Runtime/nDisplay/Source/DisplayClusterMedia/Private/DisplayClusterMediaModule.h
@@ -30,8 +30,11 @@ public:
 	//~ End IModuleInterface Implementation
 
 protected:
-	/** Initialize all media internals */
+	/** Initialize all media internals (from instance) */
 	void InitializeMedia();
+	
+	/** Initialize all media internals (from CDO) */
+	void InitializeMediaFromCDO();
 
 	/** Release all media internals */
 	void ReleaseMedia();
@@ -59,16 +62,16 @@ private:
 	void InitializeViewportOutput(const UDisplayClusterConfigurationViewport* Viewport, const FString& ViewportId, const FString& RootActorName, const FString& ClusterNodeId);
 
 	/** ICVFX camera input initializer (full frame) */
-	void InitializeICVFXCameraFullFrameInput(const UDisplayClusterICVFXCameraComponent* ICVFXCameraComponent, const FString& RootActorName, const FString& ClusterNodeId);
+	void InitializeICVFXCameraFullFrameInput(const UDisplayClusterICVFXCameraComponent* ICVFXCameraComponent, const FString& RootActorName, const FString& ClusterNodeId, const FString& ExplicitCameraName = FString());
 
 	/** ICVFX camera output initializer (full frame) */
-	void InitializeICVFXCameraFullFrameOutput(const UDisplayClusterICVFXCameraComponent* ICVFXCameraComponent, const FString& RootActorName, const FString& ClusterNodeId);
+	void InitializeICVFXCameraFullFrameOutput(const UDisplayClusterICVFXCameraComponent* ICVFXCameraComponent, const FString& RootActorName, const FString& ClusterNodeId, const FString& ExplicitCameraName = FString());
 
 	/** ICVFX camera input initializer (uniform tiles) */
-	void InitializeICVFXCameraUniformTilesInput(const UDisplayClusterICVFXCameraComponent* ICVFXCameraComponent, const FString& RootActorName, const FString& ClusterNodeId);
+	void InitializeICVFXCameraUniformTilesInput(const UDisplayClusterICVFXCameraComponent* ICVFXCameraComponent, const FString& RootActorName, const FString& ClusterNodeId, const FString& ExplicitCameraName = FString());
 
 	/** ICVFX camera output initializer (uniform tiles) */
-	void InitializeICVFXCameraUniformTilesOutput(const UDisplayClusterICVFXCameraComponent* ICVFXCameraComponent, const FString& RootActorName, const FString& ClusterNodeId);
+	void InitializeICVFXCameraUniformTilesOutput(const UDisplayClusterICVFXCameraComponent* ICVFXCameraComponent, const FString& RootActorName, const FString& ClusterNodeId, const FString& ExplicitCameraName = FString());
 
 private:
 	/** PreSubmitViewFamilies event handler. It's used to initialize media on start. */
