[nDisplay] Fixed a CIS issue - added missing include

diff --git a/Engine/Plugins/Runtime/nDisplay/Source/DisplayCluster/Private/Render/Viewport/Configuration/DisplayClusterViewportConfigurationHelpers_ICVFX.cpp b/Engine/Plugins/Runtime/nDisplay/Source/DisplayCluster/Private/Render/Viewport/Configuration/DisplayClusterViewportConfigurationHelpers_ICVFX.cpp
--- a/Engine/Plugins/Runtime/nDisplay/Source/DisplayCluster/Private/Render/Viewport/Configuration/DisplayClusterViewportConfigurationHelpers_ICVFX.cpp
+++ b/Engine/Plugins/Runtime/nDisplay/Source/DisplayCluster/Private/Render/Viewport/Configuration/DisplayClusterViewportConfigurationHelpers_ICVFX.cpp
@@ -41,6 +41,7 @@
 
 #include "Components/DisplayClusterICVFXCameraComponent.h"
 
+#include "Engine/Blueprint.h"
 #include "Engine/SCS_Node.h"
 #include "Engine/SimpleConstructionScript.h"
 #include "HAL/IConsoleManager.h"
diff --git a/Engine/Plugins/Runtime/nDisplay/Source/DisplayClusterMedia/Private/DisplayClusterMediaModule.cpp b/Engine/Plugins/Runtime/nDisplay/Source/DisplayClusterMedia/Private/DisplayClusterMediaModule.cpp
--- a/Engine/Plugins/Runtime/nDisplay/Source/DisplayClusterMedia/Private/DisplayClusterMediaModule.cpp
+++ b/Engine/Plugins/Runtime/nDisplay/Source/DisplayClusterMedia/Private/DisplayClusterMediaModule.cpp
@@ -16,6 +16,7 @@
 #include "Components/DisplayClusterICVFXCameraComponent.h"
 #include "Config/IDisplayClusterConfigManager.h"
 #include "Cluster/IDisplayClusterClusterManager.h"
+#include "Engine/Blueprint.h"
 #include "Engine/SCS_Node.h"
 #include "Engine/SimpleConstructionScript.h"
 #include "Game/IDisplayClusterGameManager.h"
