[nDisplay] Fixed a CIS issue - wrapped some code in WITH_EDITORONLY_DATA

diff --git a/Engine/Plugins/Runtime/nDisplay/Source/DisplayCluster/Private/Render/Viewport/Configuration/DisplayClusterViewportConfigurationHelpers_ICVFX.cpp b/Engine/Plugins/Runtime/nDisplay/Source/DisplayCluster/Private/Render/Viewport/Configuration/DisplayClusterViewportConfigurationHelpers_ICVFX.cpp
--- a/Engine/Plugins/Runtime/nDisplay/Source/DisplayCluster/Private/Render/Viewport/Configuration/DisplayClusterViewportConfigurationHelpers_ICVFX.cpp
+++ b/Engine/Plugins/Runtime/nDisplay/Source/DisplayCluster/Private/Render/Viewport/Configuration/DisplayClusterViewportConfigurationHelpers_ICVFX.cpp
@@ -548,7 +548,7 @@ void FDisplayClusterViewportConfigurationHelpers_ICVFX::UpdateCameraViewportSett
 	// Set viewport buffer ratio
 	DstViewport.SetViewportBufferRatio(InCameraSettings.GetCameraBufferRatio(*StageSettings));
 
-#if true
+#if WITH_EDITORONLY_DATA
 	// UE-211513
 	// This is a temporary workaround that allows to initialize media based on the blueprint data instead of 
 	// the instance data. There is an issue with the propagation of instanced object changes from the parent
@@ -571,7 +571,7 @@ void FDisplayClusterViewportConfigurationHelpers_ICVFX::UpdateCameraViewportSett
 #else
 	// InCamera tile rendering.
 	FDisplayClusterViewportConfigurationHelpers_Tile::UpdateICVFXCameraViewportTileSettings(DstViewport, InCameraSettings.RenderSettings.Media);
-#endif
+#endif // WITH_EDITORONLY_DATA
 }
 
 void FDisplayClusterViewportConfigurationHelpers_ICVFX::UpdateChromakeyViewportSettings(FDisplayClusterViewport& DstViewport, FDisplayClusterViewport& InCameraViewport, const FDisplayClusterConfigurationICVFX_CameraSettings& InCameraSettings)
diff --git a/Engine/Plugins/Runtime/nDisplay/Source/DisplayClusterMedia/Private/DisplayClusterMediaModule.cpp b/Engine/Plugins/Runtime/nDisplay/Source/DisplayClusterMedia/Private/DisplayClusterMediaModule.cpp
--- a/Engine/Plugins/Runtime/nDisplay/Source/DisplayClusterMedia/Private/DisplayClusterMediaModule.cpp
+++ b/Engine/Plugins/Runtime/nDisplay/Source/DisplayClusterMedia/Private/DisplayClusterMediaModule.cpp
@@ -54,7 +54,7 @@ void FDisplayClusterMediaModule::OnPreSubmitViewFamilies(TArray<FSceneViewFamily
 	// Unsubscribe after first call. Currently, media initialization is a one time procedure. No need to receive any further callbacks.
 	IDisplayCluster::Get().GetCallbacks().OnDisplayClusterPreSubmitViewFamilies().RemoveAll(this);
 
-#if true
+#if WITH_EDITORONLY_DATA
 	// UE-211513
 	// This is a temporary workaround that allows to initialize media based on the blueprint data instead of 
 	// the instance data. There is an issue with the propagation of instanced object changes from the parent
@@ -141,6 +141,7 @@ void FDisplayClusterMediaModule::InitializeMedia()
 	}
 }
 
+#if WITH_EDITORONLY_DATA
 void FDisplayClusterMediaModule::InitializeMediaFromCDO()
 {
 	// Runtime only for now
@@ -266,6 +267,7 @@ void FDisplayClusterMediaModule::InitializeMediaFromCDO()
 		}
 	}
 }
+#endif // WITH_EDITORONLY_DATA
 
 void FDisplayClusterMediaModule::ReleaseMedia()
 {
diff --git a/Engine/Plugins/Runtime/nDisplay/Source/DisplayClusterMedia/Private/DisplayClusterMediaModule.h b/Engine/Plugins/Runtime/nDisplay/Source/DisplayClusterMedia/Private/DisplayClusterMediaModule.h
--- a/Engine/Plugins/Runtime/nDisplay/Source/DisplayClusterMedia/Private/DisplayClusterMediaModule.h
+++ b/Engine/Plugins/Runtime/nDisplay/Source/DisplayClusterMedia/Private/DisplayClusterMediaModule.h
@@ -32,9 +32,11 @@ public:
 protected:
 	/** Initialize all media internals (from instance) */
 	void InitializeMedia();
-	
+
+#if WITH_EDITORONLY_DATA
 	/** Initialize all media internals (from CDO) */
 	void InitializeMediaFromCDO();
+#endif // WITH_EDITORONLY_DATA
 
 	/** Release all media internals */
 	void ReleaseMedia();
