[ColorGradingEditor] [DisplayCluster] Separate color grading panel into its own plugin and split color grading/ICVFX drawers in DisplayCluster

diff --git a/Engine/Plugins/Editor/ColorGrading/Source/ColorGradingEditor/Public/SColorGradingPanel.h b/Engine/Plugins/Editor/ColorGrading/Source/ColorGradingEditor/Public/SColorGradingPanel.h
new file mode 100644
--- /dev/null
+++ b/Engine/Plugins/Editor/ColorGrading/Source/ColorGradingEditor/Public/SColorGradingPanel.h
@@ -0,0 +1,170 @@
+// Copyright Epic Games, Inc. All Rights Reserved.
+
+#pragma once
+
+#include "CoreMinimal.h"
+#include "Widgets/SCompoundWidget.h"
+#include "EditorUndoClient.h"
+
+#include "ColorGradingPanelState.h"
+#include "ColorGradingListItem.h"
+
+class FColorGradingEditorDataModel;
+class SColorGradingColorWheelPanel;
+class SColorGradingObjectList;
+class SHorizontalBox;
+class SInlineEditableTextBlock;
+class UWorld;
+
+using FColorGradingActorFilter = TFunction<bool(AActor*)>;
+
+/** Main panel of a color grading drawer widget, which displays color wheels or selected object details */
+class COLORGRADINGEDITOR_API SColorGradingPanel : public SCompoundWidget, public FEditorUndoClient
+{
+public:
+	~SColorGradingPanel();
+
+	SLATE_BEGIN_ARGS(SColorGradingPanel)
+		: _IsInDrawer(false)
+		{}
+
+		/** Indicates whether this widget is in a drawer or docked in a tab */
+		SLATE_ARGUMENT(bool, IsInDrawer)
+
+		/** The world in which to search for actors to display for editing. If not provided, the level editor's current world will be used */
+		SLATE_ATTRIBUTE(UWorld*, OverrideWorld)
+
+		/** Event invoked when the user presses the dock button */
+		SLATE_EVENT(FSimpleDelegate, OnDocked)
+
+		/** Function which, if it returns false when passed an actor, filters it and its sub-entries out of the color grading item list */
+		SLATE_ARGUMENT(FColorGradingActorFilter, ActorFilter)
+
+	SLATE_END_ARGS()
+
+	void Construct(const FArguments& InArgs);
+
+	//~ SWidget interface
+	virtual void Tick(const FGeometry& AllottedGeometry, const double InCurrentTime, const float InDeltaTime) override;
+	//~ End SWidget interface
+
+	/** Refreshes the panel's UI to match the current state of the level */
+	void Refresh();
+
+	/** Gets the state of the panel UI */
+	void GetPanelState(FColorGradingPanelState& OutPanelState) const;
+
+	/** Sets the state of the panel UI */
+	void SetPanelState(const FColorGradingPanelState& InPanelState);
+
+	/** Set the list of selected objects, updating state and data model as appropriate */
+	void SetSelectedObjects(const TArray<UObject*>& Objects);
+
+private:
+	/** Creates the button used to dock the drawer in the operator panel */
+	TSharedRef<SWidget> CreateDockInLayoutButton();
+
+	/** Get the world currently being edited */
+	UWorld* GetWorld();
+
+	/** Binds a callback to the BlueprintCompiled delegate of the specified class */
+	void BindBlueprintCompiledDelegate(const UClass* Class);
+
+	/** Unbinds a callback to the BlueprintCompiled delegate of the specified class */
+	void UnbindBlueprintCompiledDelegate(const UClass* Class);
+
+	/** Refreshes the object list, filling it with the current color gradable objects from the root actor and world */
+	void RefreshColorGradingList();
+
+	/** Fills the color grading group toolbar using the color grading data model */
+	void FillColorGradingGroupToolBar();
+
+	/** Gets the visibility state of the color grading group toolbar */
+	EVisibility GetColorGradingGroupToolBarVisibility() const;
+
+	/** Gets whether the color grading group at the specified index is currently selected */
+	ECheckBoxState IsColorGradingGroupSelected(int32 GroupIndex) const;
+
+	/** Raised when the user has selected the specified color grading group */
+	void OnColorGradingGroupCheckedChanged(ECheckBoxState State, int32 GroupIndex);
+
+	/** Gets the display name of the specified color grading group */
+	FText GetColorGradingGroupDisplayName(int32 GroupIndex) const;
+
+	/** Gets the font of the display name label of the specified color grading group */
+	FSlateFontInfo GetColorGradingGroupDisplayNameFont(int32 GroupIndex) const;
+
+	/** Gets the content for the right click menu for the color grading group */
+	TSharedRef<SWidget> GetColorGradingGroupMenuContent(int32 GroupIndex);
+
+	/** Raised when a color grading group has been deleted by the user */
+	void OnColorGradingGroupDeleted(int32 GroupIndex);
+
+	/** Raised when a rename has been requested on a color grading group */
+	void OnColorGradingGroupRequestRename(int32 GroupIndex);
+
+	/** Raised when a rename has been committed on a color grading group */
+	void OnColorGradingGroupRenamed(const FText& InText, ETextCommit::Type TextCommitType, int32 GroupIndex);
+
+	/** Raised when the editor replaces any UObjects with new instantiations, usually when actors have been recompiled from blueprints */
+	void OnObjectsReplaced(const TMap<UObject*, UObject*>& OldToNewInstanceMap);
+
+	/** Raised when an actor is added to the current level */
+	void OnLevelActorAdded(AActor* Actor);
+
+	/** Raised when an actor has been deleted from the currnent level */
+	void OnLevelActorDeleted(AActor* Actor);
+
+	/** Raised when a world has been added */
+	void OnWorldAdded(UWorld* World);
+
+	/** Raised when a world has been destroyed */
+	void OnWorldDestroyed(UWorld* World);
+
+	/** Raised when the specified blueprint has been recompiled */
+	void OnBlueprintCompiled(UBlueprint* Blueprint);
+
+	/** Raised when the color grading data model has been generated */
+	void OnColorGradingDataModelGenerated();
+
+	/** Raised when the "Dock in Layout" button has been clicked */
+	FReply DockInLayout();
+
+	/** Raised when the user has selected a new item in any of the drawer's list views */
+	void OnListSelectionChanged(TSharedRef<SColorGradingObjectList> SourceList, FColorGradingListItemRef SelectedItem, ESelectInfo::Type SelectInfo);
+
+
+private:
+	/** Box containing the color grading groups */
+	TSharedPtr<SHorizontalBox> ColorGradingGroupToolBarBox;
+
+	/** List of editable text blocks containing color grading group names */
+	TArray<TSharedPtr<SInlineEditableTextBlock>> ColorGradingGroupTextBlocks;
+
+	/** Panel containing the color wheels */
+	TSharedPtr<SColorGradingColorWheelPanel> ColorWheelPanel;
+
+	/** The world from which to retrieve actors, if one was provided */
+	TAttribute<UWorld*> OverrideWorld;
+
+	/** Color grading object list widget being displayed in the drawer's list panel */
+	TSharedPtr<SColorGradingObjectList> ColorGradingObjectListView;
+
+	/** Source list for the color grading object list widget */
+	TArray<FColorGradingListItemRef> ColorGradingItemList;
+
+	/** The color grading data model for the currently selected objects */
+	TSharedPtr<FColorGradingEditorDataModel> ColorGradingDataModel;
+
+	/** Indicates whether this widget is in a drawer or docked in a tab */
+	bool bIsInDrawer;
+
+	/** Indicates that the panel should refresh itself on the next tick */
+	bool bRefreshOnNextTick = false;
+
+	/** The function to call when the user presses the dock button */
+	FSimpleDelegate DockCallback;
+
+	/** Function used to filter actors before adding them to the object list. */
+	FColorGradingActorFilter ActorFilter;
+};
diff --git a/Engine/Plugins/Runtime/nDisplay/Source/DisplayClusterColorGrading/DisplayClusterColorGrading.Build.cs b/Engine/Plugins/Runtime/nDisplay/Source/DisplayClusterColorGrading/DisplayClusterColorGrading.Build.cs
--- a/Engine/Plugins/Runtime/nDisplay/Source/DisplayClusterColorGrading/DisplayClusterColorGrading.Build.cs
+++ b/Engine/Plugins/Runtime/nDisplay/Source/DisplayClusterColorGrading/DisplayClusterColorGrading.Build.cs
@@ -16,6 +16,7 @@ public class DisplayClusterColorGrading : ModuleRules
 				"ApplicationCore",
 				"AppFramework",
 				"ColorCorrectRegions",
+				"ColorGradingEditor",
 				"Core",
 				"CoreUObject",
 				"DetailCustomizations",
diff --git a/Engine/Plugins/Runtime/nDisplay/Source/DisplayClusterColorGrading/Private/DataModelGenerators/DisplayClusterColorGradingGenerator_RootActor.cpp b/Engine/Plugins/Runtime/nDisplay/Source/DisplayClusterColorGrading/Private/DataModelGenerators/DisplayClusterColorGradingGenerator_RootActor.cpp
--- a/Engine/Plugins/Runtime/nDisplay/Source/DisplayClusterColorGrading/Private/DataModelGenerators/DisplayClusterColorGradingGenerator_RootActor.cpp
+++ b/Engine/Plugins/Runtime/nDisplay/Source/DisplayClusterColorGrading/Private/DataModelGenerators/DisplayClusterColorGradingGenerator_RootActor.cpp
@@ -33,9 +33,9 @@
 
 #define GET_MEMBER_NAME_ARRAY_CHECKED(ClassName, MemberName, Index) FName(FString(GET_MEMBER_NAME_STRING_CHECKED(ClassName, MemberName)).Replace(TEXT(#Index), *FString::FromInt(Index), ESearchCase::CaseSensitive))
 
-FDisplayClusterColorGradingDataModel::FColorGradingGroup FDisplayClusterColorGradingGenerator_ColorGradingRenderingSettings::CreateColorGradingGroup(const TSharedPtr<IPropertyHandle>& GroupPropertyHandle)
+FColorGradingEditorDataModel::FColorGradingGroup FDisplayClusterColorGradingGenerator_ColorGradingRenderingSettings::CreateColorGradingGroup(const TSharedPtr<IPropertyHandle>& GroupPropertyHandle)
 {
-	FDisplayClusterColorGradingDataModel::FColorGradingGroup ColorGradingGroup;
+	FColorGradingEditorDataModel::FColorGradingGroup ColorGradingGroup;
 	ColorGradingGroup.DisplayName = GroupPropertyHandle->GetPropertyDisplayName();
 	ColorGradingGroup.GroupPropertyHandle = GroupPropertyHandle;
 
@@ -55,12 +55,12 @@ FDisplayClusterColorGradingDataModel::FColorGradingGroup FDisplayClusterColorGra
 	return ColorGradingGroup;
 }
 
-FDisplayClusterColorGradingDataModel::FColorGradingElement FDisplayClusterColorGradingGenerator_ColorGradingRenderingSettings::CreateColorGradingElement(
+FColorGradingEditorDataModel::FColorGradingElement FDisplayClusterColorGradingGenerator_ColorGradingRenderingSettings::CreateColorGradingElement(
 	const TSharedPtr<IPropertyHandle>& GroupPropertyHandle,
 	FName ElementPropertyName,
 	FText ElementLabel)
 {
-	FDisplayClusterColorGradingDataModel::FColorGradingElement ColorGradingElement;
+	FColorGradingEditorDataModel::FColorGradingElement ColorGradingElement;
 	ColorGradingElement.DisplayName = ElementLabel;
 
 	TSharedPtr<IPropertyHandle> ElementPropertyHandle = GroupPropertyHandle->GetChildHandle(ElementPropertyName);
@@ -150,7 +150,7 @@ TSharedPtr<IPropertyHandle> MakePropertyTransactional(const TSharedPtr<IProperty
 	return PropertyHandle;
 }
 
-TSharedRef<IDisplayClusterColorGradingDataModelGenerator> FDisplayClusterColorGradingGenerator_RootActor::MakeInstance()
+TSharedRef<IColorGradingEditorDataModelGenerator> FDisplayClusterColorGradingGenerator_RootActor::MakeInstance()
 {
 	return MakeShareable(new FDisplayClusterColorGradingGenerator_RootActor());
 }
@@ -159,10 +159,10 @@ TSharedRef<IDisplayClusterColorGradingDataModelGenerator> FDisplayClusterColorGr
  * A detail customization that picks out only the necessary properties needed to display a root actor in the color grading drawer and hides all other properties
  * Also organizes the properties into custom categories that can be easily displayed in the color grading drawer
  */
-class FRootActorColorGradingCustomization : public IDetailCustomization
+class FRootActorDetailsCustomization : public IDetailCustomization
 {
 public:
-	FRootActorColorGradingCustomization(const TSharedRef<FDisplayClusterColorGradingDataModel>& InColorGradingDataModel)
+	FRootActorDetailsCustomization(const TSharedRef<FColorGradingEditorDataModel>& InColorGradingDataModel)
 		: ColorGradingDataModel(InColorGradingDataModel)
 	{ }
 
@@ -283,14 +283,14 @@ private:
 	}
 
 private:
-	TWeakPtr<FDisplayClusterColorGradingDataModel> ColorGradingDataModel;
+	TWeakPtr<FColorGradingEditorDataModel> ColorGradingDataModel;
 };
 
 /** A property customizer that culls unneeded properties from the FDisplayClusterConfigurationViewport_EntireClusterColorGrading struct to help speed up property node tree generation */
 class FFastEntireClusterColorGradingCustomization : public IPropertyTypeCustomization
 {
 public:
-	FFastEntireClusterColorGradingCustomization(const TSharedRef<FDisplayClusterColorGradingDataModel>& InColorGradingDataModel)
+	FFastEntireClusterColorGradingCustomization(const TSharedRef<FColorGradingEditorDataModel>& InColorGradingDataModel)
 		: ColorGradingDataModel(InColorGradingDataModel)
 	{ }
 
@@ -315,14 +315,14 @@ public:
 	}
 
 private:
-	TWeakPtr<FDisplayClusterColorGradingDataModel> ColorGradingDataModel;
+	TWeakPtr<FColorGradingEditorDataModel> ColorGradingDataModel;
 };
 
 /** A property customizer that culls unneeded properties from the FDisplayClusterConfigurationViewport_PerViewportColorGrading struct to help speed up property node tree generation */
 class FFastPerViewportColorGradingCustomization : public IPropertyTypeCustomization
 {
 public:
-	FFastPerViewportColorGradingCustomization(const TSharedRef<FDisplayClusterColorGradingDataModel>& InColorGradingDataModel)
+	FFastPerViewportColorGradingCustomization(const TSharedRef<FColorGradingEditorDataModel>& InColorGradingDataModel)
 		: ColorGradingDataModel(InColorGradingDataModel)
 	{ }
 
@@ -349,10 +349,10 @@ public:
 	}
 
 private:
-	TWeakPtr<FDisplayClusterColorGradingDataModel> ColorGradingDataModel;
+	TWeakPtr<FColorGradingEditorDataModel> ColorGradingDataModel;
 };
 
-void FDisplayClusterColorGradingGenerator_RootActor::Initialize(const TSharedRef<class FDisplayClusterColorGradingDataModel>& ColorGradingDataModel, const TSharedRef<IPropertyRowGenerator>& PropertyRowGenerator)
+void FDisplayClusterColorGradingGenerator_RootActor::Initialize(const TSharedRef<class FColorGradingEditorDataModel>& ColorGradingDataModel, const TSharedRef<IPropertyRowGenerator>& PropertyRowGenerator)
 {
 	PropertyRowGenerator->RegisterInstancedCustomPropertyTypeLayout(FDisplayClusterConfigurationViewport_EntireClusterColorGrading::StaticStruct()->GetFName(), FOnGetPropertyTypeCustomizationInstance::CreateLambda([ColorGradingDataModel]
 	{
@@ -366,18 +366,18 @@ void FDisplayClusterColorGradingGenerator_RootActor::Initialize(const TSharedRef
 
 	PropertyRowGenerator->RegisterInstancedCustomPropertyLayout(ADisplayClusterRootActor::StaticClass(), FOnGetDetailCustomizationInstance::CreateLambda([ColorGradingDataModel]
 	{
-		return MakeShared<FRootActorColorGradingCustomization>(ColorGradingDataModel);
+		return MakeShared<FRootActorDetailsCustomization>(ColorGradingDataModel);
 	}));
 }
 
-void FDisplayClusterColorGradingGenerator_RootActor::Destroy(const TSharedRef<class FDisplayClusterColorGradingDataModel>& ColorGradingDataModel, const TSharedRef<IPropertyRowGenerator>& PropertyRowGenerator)
+void FDisplayClusterColorGradingGenerator_RootActor::Destroy(const TSharedRef<class FColorGradingEditorDataModel>& ColorGradingDataModel, const TSharedRef<IPropertyRowGenerator>& PropertyRowGenerator)
 {
 	PropertyRowGenerator->UnregisterInstancedCustomPropertyTypeLayout(FDisplayClusterConfigurationViewport_EntireClusterColorGrading::StaticStruct()->GetFName());
 	PropertyRowGenerator->UnregisterInstancedCustomPropertyTypeLayout(FDisplayClusterConfigurationViewport_PerViewportColorGrading::StaticStruct()->GetFName());
 	PropertyRowGenerator->UnregisterInstancedCustomPropertyLayout(ADisplayClusterRootActor::StaticClass());
 }
 
-void FDisplayClusterColorGradingGenerator_RootActor::GenerateDataModel(IPropertyRowGenerator& PropertyRowGenerator, FDisplayClusterColorGradingDataModel& OutColorGradingDataModel)
+void FDisplayClusterColorGradingGenerator_RootActor::GenerateDataModel(IPropertyRowGenerator& PropertyRowGenerator, FColorGradingEditorDataModel& OutColorGradingDataModel)
 {
 	RootActors.Empty();
 
@@ -394,7 +394,7 @@ void FDisplayClusterColorGradingGenerator_RootActor::GenerateDataModel(IProperty
 	// Add a color grading group for the root actor's "EntireClusterColorGrading" property
 	if (TSharedPtr<IPropertyHandle> EntireClusterColorGradingHandle = FindPropertyHandle(PropertyRowGenerator, CREATE_PROPERTY_PATH(UDisplayClusterConfigurationData, StageSettings.EntireClusterColorGrading)))
 	{
-		FDisplayClusterColorGradingDataModel::FColorGradingGroup EntireClusterGroup = CreateColorGradingGroup(EntireClusterColorGradingHandle);
+		FColorGradingEditorDataModel::FColorGradingGroup EntireClusterGroup = CreateColorGradingGroup(EntireClusterColorGradingHandle);
 		EntireClusterGroup.EditConditionPropertyHandle = EntireClusterColorGradingHandle->GetChildHandle(GET_MEMBER_NAME_CHECKED(FDisplayClusterConfigurationViewport_EntireClusterColorGrading, bEnableEntireClusterColorGrading));
 
 		EntireClusterGroup.GroupHeaderWidget = SNew(SHorizontalBox)
@@ -430,7 +430,7 @@ void FDisplayClusterColorGradingGenerator_RootActor::GenerateDataModel(IProperty
 			{
 				TSharedRef<IPropertyHandle> PerViewportElementHandle = PerViewportColorGradingHandle->AsArray()->GetElement(Index);
 
-				FDisplayClusterColorGradingDataModel::FColorGradingGroup PerViewportGroup = CreateColorGradingGroup(PerViewportElementHandle);
+				FColorGradingEditorDataModel::FColorGradingGroup PerViewportGroup = CreateColorGradingGroup(PerViewportElementHandle);
 				PerViewportGroup.bCanBeDeleted = true;
 				PerViewportGroup.bCanBeRenamed = true;
 				PerViewportGroup.EditConditionPropertyHandle = PerViewportElementHandle->GetChildHandle(GET_MEMBER_NAME_CHECKED(FDisplayClusterConfigurationViewport_PerViewportColorGrading, bIsEnabled));
@@ -472,7 +472,7 @@ void FDisplayClusterColorGradingGenerator_RootActor::GenerateDataModel(IProperty
 							if (NamePropertyHandle.IsValid() && NamePropertyHandle->IsValidHandle())
 							{
 								NamePropertyHandle->SetValue(InText);
-								IDisplayClusterColorGrading::Get().GetColorGradingDrawerSingleton().RefreshColorGradingDrawers(true);
+								IDisplayClusterColorGrading::Get().GetColorGradingDrawerSingleton().RefreshColorGradingDrawers();
 							}
 						})
 						.Text_Lambda([NamePropertyHandle]()
@@ -519,31 +519,6 @@ void FDisplayClusterColorGradingGenerator_RootActor::GenerateDataModel(IProperty
 
 	OutColorGradingDataModel.OnColorGradingGroupDeleted().AddSP(this, &FDisplayClusterColorGradingGenerator_RootActor::DeleteColorGradingGroup);
 	OutColorGradingDataModel.OnColorGradingGroupRenamed().AddSP(this, &FDisplayClusterColorGradingGenerator_RootActor::RenameColorGradingGroup);
-
-	{
-		FDisplayClusterColorGradingDataModel::FDetailsSection InnerFrustumDetailsSection;
-		InnerFrustumDetailsSection.DisplayName = LOCTEXT("InnerFrustumDetailsSectionLabel", "Inner Frustum");
-		InnerFrustumDetailsSection.Categories.Add(TEXT("CustomViewportsCategory"));
-		InnerFrustumDetailsSection.Categories.Add(TEXT("CustomICVFXCategory"));
-		InnerFrustumDetailsSection.EditConditionPropertyHandle = FindPropertyHandle(PropertyRowGenerator, CREATE_PROPERTY_PATH(UDisplayClusterConfigurationData, StageSettings.bEnableInnerFrustums));
-
-		OutColorGradingDataModel.DetailsSections.Add(InnerFrustumDetailsSection);
-	}
-
-	{
-		FDisplayClusterColorGradingDataModel::FDetailsSection TransformDetailsSection;
-		TransformDetailsSection.Categories.Add(TEXT("TransformCommon"));
-
-		OutColorGradingDataModel.DetailsSections.Add(TransformDetailsSection);
-	}
-
-	{
-		FDisplayClusterColorGradingDataModel::FDetailsSection ChromakeyDetailsSection;
-		ChromakeyDetailsSection.Categories.Add(TEXT("CustomViewportChromakeyCategory"));
-		ChromakeyDetailsSection.Categories.Add(TEXT("CustomViewportChromakeyMarkersCategory"));
-
-		OutColorGradingDataModel.DetailsSections.Add(ChromakeyDetailsSection);
-	}
 }
 
 FReply FDisplayClusterColorGradingGenerator_RootActor::AddColorGradingGroup()
@@ -564,7 +539,7 @@ FReply FDisplayClusterColorGradingGenerator_RootActor::AddColorGradingGroup()
 
 				ConfigData->StageSettings.PerViewportColorGrading.Add(NewColorGradingGroup);
 
-				IDisplayClusterColorGrading::Get().GetColorGradingDrawerSingleton().RefreshColorGradingDrawers(true);
+				IDisplayClusterColorGrading::Get().GetColorGradingDrawerSingleton().RefreshColorGradingDrawers();
 			}
 		}
 	}
@@ -590,7 +565,7 @@ void FDisplayClusterColorGradingGenerator_RootActor::DeleteColorGradingGroup(int
 
 					ConfigData->StageSettings.PerViewportColorGrading.RemoveAt(GroupIndex - 1);
 
-					IDisplayClusterColorGrading::Get().GetColorGradingDrawerSingleton().RefreshColorGradingDrawers(true);
+					IDisplayClusterColorGrading::Get().GetColorGradingDrawerSingleton().RefreshColorGradingDrawers();
 				}
 			}
 		}
@@ -615,7 +590,7 @@ void FDisplayClusterColorGradingGenerator_RootActor::RenameColorGradingGroup(int
 
 					ConfigData->StageSettings.PerViewportColorGrading[GroupIndex - 1].Name = NewName;
 
-					IDisplayClusterColorGrading::Get().GetColorGradingDrawerSingleton().RefreshColorGradingDrawers(true);
+					IDisplayClusterColorGrading::Get().GetColorGradingDrawerSingleton().RefreshColorGradingDrawers();
 				}
 			}
 		}
@@ -780,7 +755,7 @@ TSharedRef<SWidget> FDisplayClusterColorGradingGenerator_RootActor::GetViewportC
 	return MenuBuilder.MakeWidget();
 }
 
-TSharedRef<IDisplayClusterColorGradingDataModelGenerator> FDisplayClusterColorGradingGenerator_ICVFXCamera::MakeInstance()
+TSharedRef<IColorGradingEditorDataModelGenerator> FDisplayClusterColorGradingGenerator_ICVFXCamera::MakeInstance()
 {
 	return MakeShareable(new FDisplayClusterColorGradingGenerator_ICVFXCamera());
 }
@@ -792,7 +767,7 @@ TSharedRef<IDisplayClusterColorGradingDataModelGenerator> FDisplayClusterColorGr
 class FICVFXCameraColorGradingCustomization : public IDetailCustomization
 {
 public:
-	FICVFXCameraColorGradingCustomization(const TSharedRef<FDisplayClusterColorGradingDataModel>& InColorGradingDataModel)
+	FICVFXCameraColorGradingCustomization(const TSharedRef<FColorGradingEditorDataModel>& InColorGradingDataModel)
 		: ColorGradingDataModel(InColorGradingDataModel)
 	{ }
 
@@ -925,14 +900,14 @@ private:
 	}
 
 private:
-	TWeakPtr<FDisplayClusterColorGradingDataModel> ColorGradingDataModel;
+	TWeakPtr<FColorGradingEditorDataModel> ColorGradingDataModel;
 };
 
 /** A property customizer that culls unneeded properties from the FDisplayClusterConfigurationViewport_AllNodesColorGrading struct to help speed up property node tree generation */
 class FFastAllNodesColorGradingCustomization : public IPropertyTypeCustomization
 {
 public:
-	FFastAllNodesColorGradingCustomization(const TSharedRef<FDisplayClusterColorGradingDataModel>& InColorGradingDataModel)
+	FFastAllNodesColorGradingCustomization(const TSharedRef<FColorGradingEditorDataModel>& InColorGradingDataModel)
 		: ColorGradingDataModel(InColorGradingDataModel)
 	{ }
 
@@ -957,14 +932,14 @@ public:
 	}
 
 private:
-	TWeakPtr<FDisplayClusterColorGradingDataModel> ColorGradingDataModel;
+	TWeakPtr<FColorGradingEditorDataModel> ColorGradingDataModel;
 };
 
 /** A property customizer that culls unneeded properties from the FDisplayClusterConfigurationViewport_PerNodeColorGrading struct to help speed up property node tree generation */
 class FFastPerNodeColorGradingCustomization : public IPropertyTypeCustomization
 {
 public:
-	FFastPerNodeColorGradingCustomization(const TSharedRef<FDisplayClusterColorGradingDataModel>& InColorGradingDataModel)
+	FFastPerNodeColorGradingCustomization(const TSharedRef<FColorGradingEditorDataModel>& InColorGradingDataModel)
 		: ColorGradingDataModel(InColorGradingDataModel)
 	{ }
 
@@ -991,10 +966,10 @@ public:
 	}
 
 private:
-	TWeakPtr<FDisplayClusterColorGradingDataModel> ColorGradingDataModel;
+	TWeakPtr<FColorGradingEditorDataModel> ColorGradingDataModel;
 };
 
-void FDisplayClusterColorGradingGenerator_ICVFXCamera::Initialize(const TSharedRef<class FDisplayClusterColorGradingDataModel>& ColorGradingDataModel, const TSharedRef<IPropertyRowGenerator>& PropertyRowGenerator)
+void FDisplayClusterColorGradingGenerator_ICVFXCamera::Initialize(const TSharedRef<class FColorGradingEditorDataModel>& ColorGradingDataModel, const TSharedRef<IPropertyRowGenerator>& PropertyRowGenerator)
 {
 	PropertyRowGenerator->RegisterInstancedCustomPropertyTypeLayout(FDisplayClusterConfigurationViewport_AllNodesColorGrading::StaticStruct()->GetFName(), FOnGetPropertyTypeCustomizationInstance::CreateLambda([ColorGradingDataModel]
 	{
@@ -1012,14 +987,14 @@ void FDisplayClusterColorGradingGenerator_ICVFXCamera::Initialize(const TSharedR
 	}));
 }
 
-void FDisplayClusterColorGradingGenerator_ICVFXCamera::Destroy(const TSharedRef<class FDisplayClusterColorGradingDataModel>& ColorGradingDataModel, const TSharedRef<IPropertyRowGenerator>& PropertyRowGenerator)
+void FDisplayClusterColorGradingGenerator_ICVFXCamera::Destroy(const TSharedRef<class FColorGradingEditorDataModel>& ColorGradingDataModel, const TSharedRef<IPropertyRowGenerator>& PropertyRowGenerator)
 {
 	PropertyRowGenerator->UnregisterInstancedCustomPropertyTypeLayout(FDisplayClusterConfigurationViewport_AllNodesColorGrading::StaticStruct()->GetFName());
 	PropertyRowGenerator->UnregisterInstancedCustomPropertyTypeLayout(FDisplayClusterConfigurationViewport_PerNodeColorGrading::StaticStruct()->GetFName());
 	PropertyRowGenerator->UnregisterInstancedCustomPropertyLayout(UDisplayClusterICVFXCameraComponent::StaticClass());
 }
 
-void FDisplayClusterColorGradingGenerator_ICVFXCamera::GenerateDataModel(IPropertyRowGenerator& PropertyRowGenerator, FDisplayClusterColorGradingDataModel& OutColorGradingDataModel)
+void FDisplayClusterColorGradingGenerator_ICVFXCamera::GenerateDataModel(IPropertyRowGenerator& PropertyRowGenerator, FColorGradingEditorDataModel& OutColorGradingDataModel)
 {
 	CameraComponents.Empty();
 
@@ -1036,7 +1011,7 @@ void FDisplayClusterColorGradingGenerator_ICVFXCamera::GenerateDataModel(IProper
 	// Add a color grading group for the camera's "AllNodesColorGrading" property
 	if (TSharedPtr<IPropertyHandle> AllNodesColorGradingHandle = FindPropertyHandle(PropertyRowGenerator, CREATE_PROPERTY_PATH(UDisplayClusterICVFXCameraComponent, CameraSettings.AllNodesColorGrading)))
 	{
-		FDisplayClusterColorGradingDataModel::FColorGradingGroup AllNodesGroup = CreateColorGradingGroup(AllNodesColorGradingHandle);
+		FColorGradingEditorDataModel::FColorGradingGroup AllNodesGroup = CreateColorGradingGroup(AllNodesColorGradingHandle);
 		AllNodesGroup.EditConditionPropertyHandle = AllNodesColorGradingHandle->GetChildHandle(GET_MEMBER_NAME_CHECKED(FDisplayClusterConfigurationViewport_AllNodesColorGrading, bEnableInnerFrustumAllNodesColorGrading));
 
 		AllNodesGroup.GroupHeaderWidget = SNew(SHorizontalBox)
@@ -1072,7 +1047,7 @@ void FDisplayClusterColorGradingGenerator_ICVFXCamera::GenerateDataModel(IProper
 			{
 				TSharedRef<IPropertyHandle> PerNodeElementHandle = PerNodeColorGradingHandle->AsArray()->GetElement(Index);
 
-				FDisplayClusterColorGradingDataModel::FColorGradingGroup PerNodeGroup = CreateColorGradingGroup(PerNodeElementHandle);
+				FColorGradingEditorDataModel::FColorGradingGroup PerNodeGroup = CreateColorGradingGroup(PerNodeElementHandle);
 				PerNodeGroup.bCanBeDeleted = true;
 				PerNodeGroup.bCanBeRenamed = true;
 				PerNodeGroup.EditConditionPropertyHandle = PerNodeElementHandle->GetChildHandle(GET_MEMBER_NAME_CHECKED(FDisplayClusterConfigurationViewport_PerNodeColorGrading, bIsEnabled));
@@ -1114,7 +1089,7 @@ void FDisplayClusterColorGradingGenerator_ICVFXCamera::GenerateDataModel(IProper
 							if (NamePropertyHandle.IsValid() && NamePropertyHandle->IsValidHandle())
 							{
 								NamePropertyHandle->SetValue(InText);
-								IDisplayClusterColorGrading::Get().GetColorGradingDrawerSingleton().RefreshColorGradingDrawers(true);
+								IDisplayClusterColorGrading::Get().GetColorGradingDrawerSingleton().RefreshColorGradingDrawers();
 							}
 						})
 						.Text_Lambda([NamePropertyHandle]()
@@ -1147,7 +1122,7 @@ void FDisplayClusterColorGradingGenerator_ICVFXCamera::GenerateDataModel(IProper
 		}
 	}
 
-	OutColorGradingDataModel.bShowColorGradingGroupToolBar = true;OutColorGradingDataModel.bShowColorGradingGroupToolBar = true;
+	OutColorGradingDataModel.bShowColorGradingGroupToolBar = true;
 	OutColorGradingDataModel.ColorGradingGroupToolBarWidget = SNew(SButton)
 		.ButtonStyle(FAppStyle::Get(), "SimpleButton")
 		.OnClicked(this, &FDisplayClusterColorGradingGenerator_ICVFXCamera::AddColorGradingGroup)
@@ -1161,33 +1136,6 @@ void FDisplayClusterColorGradingGenerator_ICVFXCamera::GenerateDataModel(IProper
 
 	OutColorGradingDataModel.OnColorGradingGroupDeleted().AddSP(this, &FDisplayClusterColorGradingGenerator_ICVFXCamera::DeleteColorGradingGroup);
 	OutColorGradingDataModel.OnColorGradingGroupRenamed().AddSP(this, &FDisplayClusterColorGradingGenerator_ICVFXCamera::RenameColorGradingGroup);
-
-	{
-		FDisplayClusterColorGradingDataModel::FDetailsSection InnerFrustumDetailsSection;
-		InnerFrustumDetailsSection.DisplayName = LOCTEXT("InnerFrustumDetailsSectionLabel", "Inner Frustum");
-		InnerFrustumDetailsSection.EditConditionPropertyHandle = FindPropertyHandle(PropertyRowGenerator, CREATE_PROPERTY_PATH(UDisplayClusterICVFXCameraComponent, CameraSettings.bEnable));
-		InnerFrustumDetailsSection.Categories = { TEXT("CustomICVFXCategory"), TEXT("CustomSoftEdgeCategory"), TEXT("CustomBorderCategory") };
-
-		OutColorGradingDataModel.DetailsSections.Add(InnerFrustumDetailsSection);
-	}
-
-	{
-		FDisplayClusterColorGradingDataModel::FDetailsSection InnerFrustumOverscanDetailsSection;
-		InnerFrustumOverscanDetailsSection.DisplayName = LOCTEXT("InnerFrustumOverscanDetailsSectionLabel", "Inner Frustum Overscan");
-		InnerFrustumOverscanDetailsSection.EditConditionPropertyHandle = FindPropertyHandle(PropertyRowGenerator, CREATE_PROPERTY_PATH(UDisplayClusterICVFXCameraComponent, CameraSettings.CustomFrustum.bEnable));
-		InnerFrustumOverscanDetailsSection.Categories = { TEXT("CustomOverscanCategory") };
-
-		OutColorGradingDataModel.DetailsSections.Add(InnerFrustumOverscanDetailsSection);
-	}
-
-	{
-		FDisplayClusterColorGradingDataModel::FDetailsSection ChromakeyDetailsSection;
-		ChromakeyDetailsSection.DisplayName = LOCTEXT("ChromakeyDetailsSectionLabel", "Chromakey");
-		ChromakeyDetailsSection.EditConditionPropertyHandle = FindPropertyHandle(PropertyRowGenerator, CREATE_PROPERTY_PATH(UDisplayClusterICVFXCameraComponent, CameraSettings.Chromakey.bEnable));
-		ChromakeyDetailsSection.Categories = { TEXT("CustomChromakeyCategory"), TEXT("CustomChromakeyMarkersCategory") };
-
-		OutColorGradingDataModel.DetailsSections.Add(ChromakeyDetailsSection);
-	}
 }
 
 FReply FDisplayClusterColorGradingGenerator_ICVFXCamera::AddColorGradingGroup()
@@ -1204,7 +1152,7 @@ FReply FDisplayClusterColorGradingGenerator_ICVFXCamera::AddColorGradingGroup()
 
 			CameraComponent->CameraSettings.PerNodeColorGrading.Add(NewColorGradingGroup);
 
-			IDisplayClusterColorGrading::Get().GetColorGradingDrawerSingleton().RefreshColorGradingDrawers(true);
+			IDisplayClusterColorGrading::Get().GetColorGradingDrawerSingleton().RefreshColorGradingDrawers();
 		}
 	}
 
@@ -1225,7 +1173,7 @@ void FDisplayClusterColorGradingGenerator_ICVFXCamera::DeleteColorGradingGroup(i
 
 				CameraComponent->CameraSettings.PerNodeColorGrading.RemoveAt(GroupIndex - 1);
 
-				IDisplayClusterColorGrading::Get().GetColorGradingDrawerSingleton().RefreshColorGradingDrawers(true);
+				IDisplayClusterColorGrading::Get().GetColorGradingDrawerSingleton().RefreshColorGradingDrawers();
 			}
 		}
 	}
@@ -1245,7 +1193,7 @@ void FDisplayClusterColorGradingGenerator_ICVFXCamera::RenameColorGradingGroup(i
 
 				CameraComponent->CameraSettings.PerNodeColorGrading[GroupIndex - 1].Name = NewName;
 
-				IDisplayClusterColorGrading::Get().GetColorGradingDrawerSingleton().RefreshColorGradingDrawers(true);
+				IDisplayClusterColorGrading::Get().GetColorGradingDrawerSingleton().RefreshColorGradingDrawers();
 			}
 		}
 	}
@@ -1394,4 +1342,41 @@ TSharedRef<SWidget> FDisplayClusterColorGradingGenerator_ICVFXCamera::GetNodeCom
 	return MenuBuilder.MakeWidget();
 }
 
+TArray<TSubclassOf<AActor>> FDisplayClusterColorGradingListItemGenerator_RootActor::GetActorClassesForListItems() const
+{
+	return { ADisplayClusterRootActor::StaticClass() };
+}
+
+void FDisplayClusterColorGradingListItemGenerator_RootActor::GenerateColorGradingListItems(AActor* InActor, TArray<FColorGradingListItemRef>& OutList) const
+{
+	if (ADisplayClusterRootActor* RootActor = Cast<ADisplayClusterRootActor>(InActor))
+	{
+		FColorGradingListItemRef ListItemRef = MakeShared<FColorGradingListItem>(RootActor);
+		ListItemRef->IsItemEnabled = CREATE_IS_ENABLED_LAMBDA(RootActor, RootActor->GetConfigData()->StageSettings.EnableColorGrading);
+		ListItemRef->OnItemEnabledChanged = CREATE_ON_ENABLED_CHANGED_LAMBDA(RootActor, RootActor->GetConfigData()->StageSettings.EnableColorGrading);
+
+		OutList.Add(ListItemRef);
+	}
+}
+
+TArray<TSubclassOf<AActor>> FDisplayClusterColorGradingListItemGenerator_ICVFXCamera::GetActorClassesForListItems() const
+{
+	return { ADisplayClusterRootActor::StaticClass() };
+}
+
+void FDisplayClusterColorGradingListItemGenerator_ICVFXCamera::GenerateColorGradingListItems(AActor* InActor, TArray<FColorGradingListItemRef>& OutList) const
+{
+	if (ADisplayClusterRootActor* RootActor = Cast<ADisplayClusterRootActor>(InActor))
+	{
+		RootActor->ForEachComponent<UDisplayClusterICVFXCameraComponent>(false, [this, RootActor, &OutList](UDisplayClusterICVFXCameraComponent* ICVFXCameraComponent)
+		{
+			FColorGradingListItemRef ICVFXCameraListItemRef = MakeShared<FColorGradingListItem>(RootActor, ICVFXCameraComponent);
+			ICVFXCameraListItemRef->IsItemEnabled = CREATE_IS_ENABLED_LAMBDA(ICVFXCameraComponent, ICVFXCameraComponent->CameraSettings.EnableInnerFrustumColorGrading);
+			ICVFXCameraListItemRef->OnItemEnabledChanged = CREATE_ON_ENABLED_CHANGED_LAMBDA(ICVFXCameraComponent, ICVFXCameraComponent->CameraSettings.EnableInnerFrustumColorGrading);
+
+			OutList.Add(ICVFXCameraListItemRef);
+		});
+	}
+}
+
 #undef LOCTEXT_NAMESPACE
diff --git a/Engine/Plugins/Runtime/nDisplay/Source/DisplayClusterColorGrading/Private/DataModelGenerators/DisplayClusterColorGradingGenerator_RootActor.h b/Engine/Plugins/Runtime/nDisplay/Source/DisplayClusterColorGrading/Private/DataModelGenerators/DisplayClusterColorGradingGenerator_RootActor.h
--- a/Engine/Plugins/Runtime/nDisplay/Source/DisplayClusterColorGrading/Private/DataModelGenerators/DisplayClusterColorGradingGenerator_RootActor.h
+++ b/Engine/Plugins/Runtime/nDisplay/Source/DisplayClusterColorGrading/Private/DataModelGenerators/DisplayClusterColorGradingGenerator_RootActor.h
@@ -3,7 +3,7 @@
 #pragma once
 
 #include "CoreMinimal.h"
-#include "DisplayClusterColorGradingDataModel.h"
+#include "ColorGradingEditorDataModel.h"
 
 #include "Input/Reply.h"
 
@@ -16,14 +16,14 @@ class UDisplayClusterICVFXCameraComponent;
 struct FCachedPropertyPath;
 
 /** Base generator for any object that needs a color grading data model generated from an FColorGradingRenderingSettings struct */
-class FDisplayClusterColorGradingGenerator_ColorGradingRenderingSettings : public IDisplayClusterColorGradingDataModelGenerator
+class FDisplayClusterColorGradingGenerator_ColorGradingRenderingSettings : public IColorGradingEditorDataModelGenerator
 {
 protected:
 	/** Creates a new color grading group structure from the specified group property handle, which finds and connects the appropriate property handles for the color wheels and details view */
-	FDisplayClusterColorGradingDataModel::FColorGradingGroup CreateColorGradingGroup(const TSharedPtr<IPropertyHandle>& GroupPropertyHandle);
+	FColorGradingEditorDataModel::FColorGradingGroup CreateColorGradingGroup(const TSharedPtr<IPropertyHandle>& GroupPropertyHandle);
 
 	/** Creates a color grading element structure from the specified property handle, whose child properties are expected to be colors with the ColorGradingMode metadata set */
-	FDisplayClusterColorGradingDataModel::FColorGradingElement CreateColorGradingElement(const TSharedPtr<IPropertyHandle>& GroupPropertyHandle, FName ElementPropertyName, FText ElementLabel);
+	FColorGradingEditorDataModel::FColorGradingElement CreateColorGradingElement(const TSharedPtr<IPropertyHandle>& GroupPropertyHandle, FName ElementPropertyName, FText ElementLabel);
 
 	/** Recursively searches the detail tree hierarchy for a property detail tree node whose name matches the specified name */
 	TSharedPtr<IDetailTreeNode> FindPropertyTreeNode(const TSharedRef<IDetailTreeNode>& Node, const FCachedPropertyPath& PropertyPath);
@@ -36,12 +36,12 @@ protected:
 class FDisplayClusterColorGradingGenerator_RootActor : public FDisplayClusterColorGradingGenerator_ColorGradingRenderingSettings
 {
 public:
-	static TSharedRef<IDisplayClusterColorGradingDataModelGenerator> MakeInstance();
+	static TSharedRef<IColorGradingEditorDataModelGenerator> MakeInstance();
 
 	//~ IDisplayClusterColorGradingDataModelGenerator interface
-	virtual void Initialize(const TSharedRef<class FDisplayClusterColorGradingDataModel>& ColorGradingDataModel, const TSharedRef<IPropertyRowGenerator>& PropertyRowGenerator) override;
-	virtual void Destroy(const TSharedRef<class FDisplayClusterColorGradingDataModel>& ColorGradingDataModel, const TSharedRef<IPropertyRowGenerator>& PropertyRowGenerator) override;
-	virtual void GenerateDataModel(IPropertyRowGenerator& PropertyRowGenerator, FDisplayClusterColorGradingDataModel& OutColorGradingDataModel) override;
+	virtual void Initialize(const TSharedRef<class FColorGradingEditorDataModel>& ColorGradingDataModel, const TSharedRef<IPropertyRowGenerator>& PropertyRowGenerator) override;
+	virtual void Destroy(const TSharedRef<class FColorGradingEditorDataModel>& ColorGradingDataModel, const TSharedRef<IPropertyRowGenerator>& PropertyRowGenerator) override;
+	virtual void GenerateDataModel(IPropertyRowGenerator& PropertyRowGenerator, FColorGradingEditorDataModel& OutColorGradingDataModel) override;
 	//~ End IDisplayClusterColorGradingDataModelGenerator interface
 
 private:
@@ -72,12 +72,12 @@ private:
 class FDisplayClusterColorGradingGenerator_ICVFXCamera : public FDisplayClusterColorGradingGenerator_ColorGradingRenderingSettings
 {
 public:
-	static TSharedRef<IDisplayClusterColorGradingDataModelGenerator> MakeInstance();
+	static TSharedRef<IColorGradingEditorDataModelGenerator> MakeInstance();
 
 	//~ IDisplayClusterColorGradingDataModelGenerator interface
-	virtual void Initialize(const TSharedRef<class FDisplayClusterColorGradingDataModel>& InColorGradingDataModel, const TSharedRef<IPropertyRowGenerator>& InPropertyRowGenerator) override;
-	virtual void Destroy(const TSharedRef<class FDisplayClusterColorGradingDataModel>& InColorGradingDataModel, const TSharedRef<IPropertyRowGenerator>& InPropertyRowGenerator) override;
-	virtual void GenerateDataModel(IPropertyRowGenerator& PropertyRowGenerator, FDisplayClusterColorGradingDataModel& OutColorGradingDataModel) override;
+	virtual void Initialize(const TSharedRef<class FColorGradingEditorDataModel>& InColorGradingDataModel, const TSharedRef<IPropertyRowGenerator>& InPropertyRowGenerator) override;
+	virtual void Destroy(const TSharedRef<class FColorGradingEditorDataModel>& InColorGradingDataModel, const TSharedRef<IPropertyRowGenerator>& InPropertyRowGenerator) override;
+	virtual void GenerateDataModel(IPropertyRowGenerator& PropertyRowGenerator, FColorGradingEditorDataModel& OutColorGradingDataModel) override;
 	//~ End IDisplayClusterColorGradingDataModelGenerator interface
 
 private:
@@ -102,4 +102,22 @@ private:
 private:
 	/** A list of camera components that are being represented by the data model */
 	TArray<TWeakObjectPtr<UDisplayClusterICVFXCameraComponent>> CameraComponents;
+};
+
+/** Color grading list item generator for an nDisplay root actor */
+class FDisplayClusterColorGradingListItemGenerator_RootActor : public IColorGradingListItemGenerator
+{
+	//~ IColorGradingListItemGenerator interface
+	virtual TArray<TSubclassOf<AActor>> GetActorClassesForListItems() const override;
+	virtual void GenerateColorGradingListItems(AActor* InActor, TArray<FColorGradingListItemRef>& OutList) const override;
+	//~ End IColorGradingListItemGenerator interface
+};
+
+/** Color grading list item generator for an nDisplay ICVFX camera component */
+class FDisplayClusterColorGradingListItemGenerator_ICVFXCamera : public IColorGradingListItemGenerator
+{
+	//~ IColorGradingListItemGenerator interface
+	virtual TArray<TSubclassOf<AActor>> GetActorClassesForListItems() const override;
+	virtual void GenerateColorGradingListItems(AActor* InActor, TArray<FColorGradingListItemRef>& OutList) const override;
+	//~ End IColorGradingListItemGenerator interface
 };
\ No newline at end of file
diff --git a/Engine/Plugins/Runtime/nDisplay/Source/DisplayClusterColorGrading/Private/DisplayClusterColorGradingCommands.cpp b/Engine/Plugins/Runtime/nDisplay/Source/DisplayClusterColorGrading/Private/DisplayClusterColorGradingCommands.cpp
--- a/Engine/Plugins/Runtime/nDisplay/Source/DisplayClusterColorGrading/Private/DisplayClusterColorGradingCommands.cpp
+++ b/Engine/Plugins/Runtime/nDisplay/Source/DisplayClusterColorGrading/Private/DisplayClusterColorGradingCommands.cpp
@@ -6,14 +6,6 @@
 
 void FDisplayClusterColorGradingCommands::RegisterCommands()
 {
-	UI_COMMAND(SaturationColorWheelVisibility, "Saturation", "Show or hide the saturation color wheel", EUserInterfaceActionType::ToggleButton, FInputChord());
-	UI_COMMAND(ContrastColorWheelVisibility, "Constrast", "Show or hide the constrast color wheel", EUserInterfaceActionType::ToggleButton, FInputChord());
-	UI_COMMAND(ColorWheelSliderOrientationHorizontal, "Right", "Puts the color wheel sliders to the right of the color wheel", EUserInterfaceActionType::RadioButton, FInputChord());
-	UI_COMMAND(ColorWheelSliderOrientationVertical, "Below", "Puts the color wheel sliders below the color wheel", EUserInterfaceActionType::RadioButton, FInputChord());
-
-	UI_COMMAND(ColorGradingMode, "Color Grading Mode", "Sets the drawer to display the color grading properties of the selected items", EUserInterfaceActionType::ToggleButton, FInputChord());
-	UI_COMMAND(DetailViewMode, "Detail View Mode", "Sets the drawer to display the additional in-camera VFX properties of the selected items", EUserInterfaceActionType::ToggleButton, FInputChord());
-
 	UI_COMMAND(OpenColorGradingDrawer, "Open Color Grading Drawer", "Opens the color grading drawer from the status bar", EUserInterfaceActionType::Button, FInputChord(EModifierKey::Control, EKeys::SpaceBar));
 }
 
diff --git a/Engine/Plugins/Runtime/nDisplay/Source/DisplayClusterColorGrading/Private/DisplayClusterColorGradingCommands.h b/Engine/Plugins/Runtime/nDisplay/Source/DisplayClusterColorGrading/Private/DisplayClusterColorGradingCommands.h
--- a/Engine/Plugins/Runtime/nDisplay/Source/DisplayClusterColorGrading/Private/DisplayClusterColorGradingCommands.h
+++ b/Engine/Plugins/Runtime/nDisplay/Source/DisplayClusterColorGrading/Private/DisplayClusterColorGradingCommands.h
@@ -20,13 +20,5 @@ public:
 	virtual void RegisterCommands() override;
 
 public:
-	TSharedPtr<FUICommandInfo> SaturationColorWheelVisibility;
-	TSharedPtr<FUICommandInfo> ContrastColorWheelVisibility;
-	TSharedPtr<FUICommandInfo> ColorWheelSliderOrientationHorizontal;
-	TSharedPtr<FUICommandInfo> ColorWheelSliderOrientationVertical;
-
-	TSharedPtr<FUICommandInfo> ColorGradingMode;
-	TSharedPtr<FUICommandInfo> DetailViewMode;
-
 	TSharedPtr<FUICommandInfo> OpenColorGradingDrawer;
 };
\ No newline at end of file
diff --git a/Engine/Plugins/Runtime/nDisplay/Source/DisplayClusterColorGrading/Private/DisplayClusterColorGradingModule.cpp b/Engine/Plugins/Runtime/nDisplay/Source/DisplayClusterColorGrading/Private/DisplayClusterColorGradingModule.cpp
--- a/Engine/Plugins/Runtime/nDisplay/Source/DisplayClusterColorGrading/Private/DisplayClusterColorGradingModule.cpp
+++ b/Engine/Plugins/Runtime/nDisplay/Source/DisplayClusterColorGrading/Private/DisplayClusterColorGradingModule.cpp
@@ -4,8 +4,6 @@
 
 #include "DisplayClusterColorGradingCommands.h"
 #include "DataModelGenerators/DisplayClusterColorGradingGenerator_RootActor.h"
-#include "DataModelGenerators/DisplayClusterColorGradingGenerator_PostProcessVolume.h"
-#include "DataModelGenerators/DisplayClusterColorGradingGenerator_ColorCorrectRegion.h"
 #include "Drawer/DisplayClusterColorGradingDrawerSingleton.h"
 #include "Drawer/SDisplayClusterColorGradingDrawer.h"
 
@@ -21,17 +19,14 @@ void FDisplayClusterColorGradingModule::StartupModule()
 {
 	ColorGradingDrawerSingleton = MakeUnique<FDisplayClusterColorGradingDrawerSingleton>();
 
-	FDisplayClusterColorGradingDataModel::RegisterColorGradingDataModelGenerator<ADisplayClusterRootActor>(
-		FGetColorGradingDataModelGenerator::CreateStatic(&FDisplayClusterColorGradingGenerator_RootActor::MakeInstance));
+	FColorGradingEditorDataModel::RegisterColorGradingDataModelGenerator<ADisplayClusterRootActor>(
+		FGetDetailsDataModelGenerator::CreateStatic(&FDisplayClusterColorGradingGenerator_RootActor::MakeInstance));
 
-	FDisplayClusterColorGradingDataModel::RegisterColorGradingDataModelGenerator<UDisplayClusterICVFXCameraComponent>(
-		FGetColorGradingDataModelGenerator::CreateStatic(&FDisplayClusterColorGradingGenerator_ICVFXCamera::MakeInstance));
+	FColorGradingEditorDataModel::RegisterColorGradingDataModelGenerator<UDisplayClusterICVFXCameraComponent>(
+		FGetDetailsDataModelGenerator::CreateStatic(&FDisplayClusterColorGradingGenerator_ICVFXCamera::MakeInstance));
 
-	FDisplayClusterColorGradingDataModel::RegisterColorGradingDataModelGenerator<APostProcessVolume>(
-		FGetColorGradingDataModelGenerator::CreateStatic(&FDisplayClusterColorGradingGenerator_PostProcessVolume::MakeInstance));
-
-	FDisplayClusterColorGradingDataModel::RegisterColorGradingDataModelGenerator<AColorCorrectRegion>(
-		FGetColorGradingDataModelGenerator::CreateStatic(&FDisplayClusterColorGradingGenerator_ColorCorrectRegion::MakeInstance));
+	FColorGradingListItem::RegisterColorGradingListItemGenerator<FDisplayClusterColorGradingListItemGenerator_RootActor>();
+	FColorGradingListItem::RegisterColorGradingListItemGenerator<FDisplayClusterColorGradingListItemGenerator_ICVFXCamera>();
 
 	FDisplayClusterColorGradingCommands::Register();
 }
diff --git a/Engine/Plugins/Runtime/nDisplay/Source/DisplayClusterColorGrading/Private/DisplayClusterColorGradingStyle.h b/Engine/Plugins/Runtime/nDisplay/Source/DisplayClusterColorGrading/Private/DisplayClusterColorGradingStyle.h
--- a/Engine/Plugins/Runtime/nDisplay/Source/DisplayClusterColorGrading/Private/DisplayClusterColorGradingStyle.h
+++ b/Engine/Plugins/Runtime/nDisplay/Source/DisplayClusterColorGrading/Private/DisplayClusterColorGradingStyle.h
@@ -26,24 +26,10 @@ public:
 		SetContentRoot(FPaths::EnginePluginsDir() / TEXT("Runtime/nDisplay/Content/Icons/"));
 		SetCoreContentRoot(FPaths::EngineContentDir() / TEXT("Editor/Slate"));
 
-		Set("ColorGradingDrawer.Icon", new IMAGE_BRUSH_SVG("Components/nDisplayCamera", Icon16x16));
-		Set("DisplayClusterColorGrading.ColorGradingMode", new IMAGE_BRUSH_SVG("OperatorPanel/Colors", Icon16x16));
-		Set("DisplayClusterColorGrading.DetailViewMode", new IMAGE_BRUSH_SVG("OperatorPanel/Sliders", Icon16x16));
+		Set("ColorGradingDrawer.Icon", new IMAGE_BRUSH_SVG("OperatorPanel/Colors", Icon16x16));
 		Set("ColorGradingDrawer.Viewports", new CORE_IMAGE_BRUSH_SVG("Starship/Common/Viewports", Icon16x16));
 		Set("ColorGradingDrawer.Nodes", new IMAGE_BRUSH_SVG("Cluster/ClusterNode", Icon16x16));
 
-		// Ultra slim toolbar style with smaller icons and smaller padding for the color grading drawer toolbars
-		FToolBarStyle UltraSlimToolbarStyle = FAppStyle::Get().GetWidgetStyle<FToolBarStyle>("SlimToolBar");
-
-		UltraSlimToolbarStyle
-			.SetBackgroundPadding(FMargin(4.0f, 0.0f, 0.0f, 0.0f))
-			.SetIconSize(Icon16x16);
-
-		UltraSlimToolbarStyle.ToggleButton
-			.SetPadding(FMargin(16.0f, 2.0f));
-
-		Set("UltraSlimToolBar", UltraSlimToolbarStyle);
-
 		FSlateStyleRegistry::RegisterSlateStyle(*this);
 	}
 
diff --git a/Engine/Plugins/Runtime/nDisplay/Source/DisplayClusterColorGrading/Private/Drawer/DisplayClusterColorGradingDrawerSingleton.cpp b/Engine/Plugins/Runtime/nDisplay/Source/DisplayClusterColorGrading/Private/Drawer/DisplayClusterColorGradingDrawerSingleton.cpp
--- a/Engine/Plugins/Runtime/nDisplay/Source/DisplayClusterColorGrading/Private/Drawer/DisplayClusterColorGradingDrawerSingleton.cpp
+++ b/Engine/Plugins/Runtime/nDisplay/Source/DisplayClusterColorGrading/Private/Drawer/DisplayClusterColorGradingDrawerSingleton.cpp
@@ -31,8 +31,9 @@ const FName FDisplayClusterColorGradingDrawerSingleton::ColorGradingDrawerTab =
 FDisplayClusterColorGradingDrawerSingleton::FDisplayClusterColorGradingDrawerSingleton()
 {
 	FGlobalTabmanager::Get()->RegisterNomadTabSpawner(ColorGradingDrawerTab, FOnSpawnTab::CreateRaw(this, &FDisplayClusterColorGradingDrawerSingleton::SpawnColorGradingDrawerTab))
-		.SetDisplayName(LOCTEXT("ColorGradingDarwerTab_DisplayName", "In-Camera VFX"))
-		.SetTooltipText(LOCTEXT("ColorGradingDarwerTab_Tooltip", "Editing tools for in-camera VFX."))
+		.SetIcon(FSlateIcon(FDisplayClusterColorGradingStyle::Get().GetStyleSetName(), "ColorGradingDrawer.Icon"))
+		.SetDisplayName(LOCTEXT("ColorGradingDarwerTab_DisplayName", "Color Grading"))
+		.SetTooltipText(LOCTEXT("ColorGradingDarwerTab_Tooltip", "Color grading controls for in-camera VFX"))
 		.SetMenuType(ETabSpawnerMenuType::Hidden);
 
 	IDisplayClusterOperator::Get().OnRegisterLayoutExtensions().AddRaw(this, &FDisplayClusterColorGradingDrawerSingleton::ExtendOperatorTabLayout);
@@ -73,11 +74,11 @@ void FDisplayClusterColorGradingDrawerSingleton::DockColorGradingDrawer()
 	}
 }
 
-void FDisplayClusterColorGradingDrawerSingleton::RefreshColorGradingDrawers(bool bPreserveDrawerState)
+void FDisplayClusterColorGradingDrawerSingleton::RefreshColorGradingDrawers()
 {
 	if (ColorGradingDrawer.IsValid())
 	{
-		ColorGradingDrawer.Pin()->Refresh(bPreserveDrawerState);
+		ColorGradingDrawer.Pin()->Refresh();
 	}
 
 	if (TSharedPtr<FTabManager> OperatorPanelTabManager = IDisplayClusterOperator::Get().GetOperatorViewModel()->GetTabManager())
@@ -85,7 +86,7 @@ void FDisplayClusterColorGradingDrawerSingleton::RefreshColorGradingDrawers(bool
 		if (TSharedPtr<SDockTab> ExistingTab = OperatorPanelTabManager->FindExistingLiveTab(ColorGradingDrawerTab))
 		{
 			TSharedRef<SDisplayClusterColorGradingDrawer> DockedDrawer = StaticCastSharedRef<SDisplayClusterColorGradingDrawer>(ExistingTab->GetContent());
-			DockedDrawer->Refresh(bPreserveDrawerState);
+			DockedDrawer->Refresh();
 		}
 	}
 }
@@ -102,14 +103,14 @@ TSharedRef<SWidget> FDisplayClusterColorGradingDrawerSingleton::CreateDrawerCont
 			ColorGradingDrawer = Drawer;
 		}
 
-		if (PreviousDrawerState.IsSet())
+		if (PreviousColorGradingPanelState.IsSet())
 		{
-			ColorGradingDrawer.Pin()->SetDrawerState(PreviousDrawerState.GetValue());
-			PreviousDrawerState.Reset();
+			ColorGradingDrawer.Pin()->SetColorGradingPanelState(PreviousColorGradingPanelState.GetValue());
+			PreviousColorGradingPanelState.Reset();
 		}
 		else
 		{
-			ColorGradingDrawer.Pin()->SetDrawerStateToDefault();
+			ColorGradingDrawer.Pin()->SelectOperatorRootActor();
 		}
 
 		return Drawer.ToSharedRef();
@@ -120,7 +121,7 @@ TSharedRef<SWidget> FDisplayClusterColorGradingDrawerSingleton::CreateDrawerCont
 
 		if (bCopyStateFromActiveDrawer && ColorGradingDrawer.IsValid())
 		{
-			NewDrawer->SetDrawerState(ColorGradingDrawer.Pin()->GetDrawerState());
+			NewDrawer->SetColorGradingPanelState(ColorGradingDrawer.Pin()->GetColorGradingPanelState());
 		}
 
 		return NewDrawer;
@@ -149,7 +150,7 @@ void FDisplayClusterColorGradingDrawerSingleton::ExtendOperatorStatusBar(FDispla
 
 	ColorGradingDrawerConfig.GetDrawerContentDelegate.BindRaw(this, &FDisplayClusterColorGradingDrawerSingleton::CreateDrawerContent, true, false);
 	ColorGradingDrawerConfig.OnDrawerDismissedDelegate.BindRaw(this, &FDisplayClusterColorGradingDrawerSingleton::SaveDrawerState);
-	ColorGradingDrawerConfig.ButtonText = LOCTEXT("ColorGradingDrawer_ButtonText", "In-Camera VFX");
+	ColorGradingDrawerConfig.ButtonText = LOCTEXT("ColorGradingDrawer_ButtonText", "Color Grading");
 	ColorGradingDrawerConfig.Icon = FDisplayClusterColorGradingStyle::Get().GetBrush("ColorGradingDrawer.Icon");
 
 	StatusBarExtender.AddWidgetDrawer(ColorGradingDrawerConfig);
@@ -172,24 +173,24 @@ void FDisplayClusterColorGradingDrawerSingleton::SaveDrawerState(const TSharedPt
 {
 	if (ColorGradingDrawer.IsValid())
 	{
-		PreviousDrawerState = ColorGradingDrawer.Pin()->GetDrawerState();
+		PreviousColorGradingPanelState = ColorGradingDrawer.Pin()->GetColorGradingPanelState();
 	}
 	else
 	{
-		PreviousDrawerState.Reset();
+		PreviousColorGradingPanelState.Reset();
 	}
 }
 
 void FDisplayClusterColorGradingDrawerSingleton::OnActiveRootActorChanged(ADisplayClusterRootActor* NewRootActor)
 {
 	// Clear the previous drawer state when the active root actor is changed, since it is most likely invalid
-	PreviousDrawerState.Reset();
+	PreviousColorGradingPanelState.Reset();
 }
 
 void FDisplayClusterColorGradingDrawerSingleton::OnDetailObjectsChanged(const TArray<UObject*>& NewObjects)
 {
 	// Clear the previous drawer state when the selected detail objects have changed
-	PreviousDrawerState.Reset();
+	PreviousColorGradingPanelState.Reset();
 }
 
 #undef LOCTEXT_NAMESPACE
\ No newline at end of file
diff --git a/Engine/Plugins/Runtime/nDisplay/Source/DisplayClusterColorGrading/Private/Drawer/DisplayClusterColorGradingDrawerSingleton.h b/Engine/Plugins/Runtime/nDisplay/Source/DisplayClusterColorGrading/Private/Drawer/DisplayClusterColorGradingDrawerSingleton.h
--- a/Engine/Plugins/Runtime/nDisplay/Source/DisplayClusterColorGrading/Private/Drawer/DisplayClusterColorGradingDrawerSingleton.h
+++ b/Engine/Plugins/Runtime/nDisplay/Source/DisplayClusterColorGrading/Private/Drawer/DisplayClusterColorGradingDrawerSingleton.h
@@ -31,7 +31,7 @@ public:
 	virtual void DockColorGradingDrawer() override;
 
 	/** Refreshes the UI of any open color grading drawers */
-	virtual void RefreshColorGradingDrawers(bool bPreserveDrawerState) override;
+	virtual void RefreshColorGradingDrawers() override;
 
 private:
 	/** Creates a new drawer widget to place in a drawer or in a tab */
@@ -65,6 +65,6 @@ private:
 	/** A weak pointer to the active color grading drawer that is open */
 	TWeakPtr<SDisplayClusterColorGradingDrawer> ColorGradingDrawer;
 
-	/** The drawer state when the last instance of the color grading drawer was dismissed */
-	TOptional<FDisplayClusterColorGradingDrawerState> PreviousDrawerState;
+	/** The color grading panel state when the last instance of the color grading drawer was dismissed */
+	TOptional<FColorGradingPanelState> PreviousColorGradingPanelState;
 };
\ No newline at end of file
diff --git a/Engine/Plugins/Runtime/nDisplay/Source/DisplayClusterColorGrading/Private/Drawer/SDisplayClusterColorGradingDrawer.cpp b/Engine/Plugins/Runtime/nDisplay/Source/DisplayClusterColorGrading/Private/Drawer/SDisplayClusterColorGradingDrawer.cpp
--- a/Engine/Plugins/Runtime/nDisplay/Source/DisplayClusterColorGrading/Private/Drawer/SDisplayClusterColorGradingDrawer.cpp
+++ b/Engine/Plugins/Runtime/nDisplay/Source/DisplayClusterColorGrading/Private/Drawer/SDisplayClusterColorGradingDrawer.cpp
@@ -2,13 +2,11 @@
 
 #include "SDisplayClusterColorGradingDrawer.h"
 
+#include "ColorGradingEditorDataModel.h"
 #include "DisplayClusterColorGradingStyle.h"
 #include "Engine/Blueprint.h"
 #include "IDisplayClusterColorGrading.h"
 #include "IDisplayClusterColorGradingDrawerSingleton.h"
-#include "SDisplayClusterColorGradingObjectList.h"
-#include "SDisplayClusterColorGradingColorWheelPanel.h"
-#include "SDisplayClusterColorGradingDetailsPanel.h"
 
 #include "IDisplayClusterOperator.h"
 #include "IDisplayClusterOperatorViewModel.h"
@@ -19,12 +17,12 @@
 
 #include "ColorCorrectRegion.h"
 #include "ColorCorrectWindow.h"
-#include "DisplayClusterColorGradingCommands.h"
 #include "Engine/PostProcessVolume.h"
+#include "Framework/Commands/GenericCommands.h"
 #include "Framework/MultiBox/MultiBoxBuilder.h"
 #include "Modules/ModuleManager.h"
 #include "ScopedTransaction.h"
-#include "Framework/Commands/GenericCommands.h"
+#include "SColorGradingPanel.h"
 #include "Styling/AppStyle.h"
 #include "Widgets/Input/SButton.h"
 #include "Widgets/Layout/SBorder.h"
@@ -44,285 +42,54 @@ SDisplayClusterColorGradingDrawer::~SDisplayClusterColorGradingDrawer()
 	OperatorViewModel = IDisplayClusterOperator::Get().GetOperatorViewModel();
 	OperatorViewModel->OnActiveRootActorChanged().RemoveAll(this);
 	OperatorViewModel->OnDetailObjectsChanged().RemoveAll(this);
-
-	FCoreUObjectDelegates::OnObjectsReplaced.RemoveAll(this);
-	GEngine->OnLevelActorAdded().RemoveAll(this);
-	GEngine->OnLevelActorDeleted().RemoveAll(this);
-
-	GEditor->UnregisterForUndo(this);
-
-	for (const TArray<FDisplayClusterColorGradingListItemRef>& ColorGradingItemList : ColorGradingItemLists)
-	{
-		for (const FDisplayClusterColorGradingListItemRef& ColorGradingItem : ColorGradingItemList)
-		{
-			if (ColorGradingItem->Component.IsValid())
-			{
-				UnbindBlueprintCompiledDelegate(ColorGradingItem->Component->GetClass());
-			}
-
-			if (ColorGradingItem->Actor.IsValid())
-			{
-				UnbindBlueprintCompiledDelegate(ColorGradingItem->Actor->GetClass());
-			}
-		}
-	}
 }
 
 void SDisplayClusterColorGradingDrawer::Construct(const FArguments& InArgs, bool bInIsInDrawer)
 {
-	ColorGradingDataModel = MakeShared<FDisplayClusterColorGradingDataModel>();
-	ColorGradingDataModel->OnDataModelGenerated().AddSP(this, &SDisplayClusterColorGradingDrawer::OnColorGradingDataModelGenerated);
-
-	CommandList = MakeShared<FUICommandList>();
-	BindCommands();
-
 	bIsInDrawer = bInIsInDrawer;
 	OperatorViewModel = IDisplayClusterOperator::Get().GetOperatorViewModel();
 	OperatorViewModel->OnActiveRootActorChanged().AddSP(this, &SDisplayClusterColorGradingDrawer::OnActiveRootActorChanged);
-	OperatorViewModel->OnDetailObjectsChanged().AddSP(this, &SDisplayClusterColorGradingDrawer::OnDetailObjectsChanged);
-
-	FCoreUObjectDelegates::OnObjectsReplaced.AddSP(this, &SDisplayClusterColorGradingDrawer::OnObjectsReplaced);
-	GEngine->OnLevelActorAdded().AddSP(this, &SDisplayClusterColorGradingDrawer::OnLevelActorAdded);
-	GEngine->OnLevelActorDeleted().AddSP(this, &SDisplayClusterColorGradingDrawer::OnLevelActorDeleted);
-
-	GEditor->RegisterForUndo(this);
-
-	RefreshColorGradingLists();
 	
 	TSharedRef<SVerticalBox> ColorGradingObjectListBox = SNew(SVerticalBox);
 
-	for (int32 Index = 0; Index < ColorGradingItemLists.Num(); ++Index)
-	{
-		TSharedPtr<SDisplayClusterColorGradingObjectList> NewListView;
-		ColorGradingObjectListBox->AddSlot()
-			.AutoHeight()
-			.Padding(0.0f, 2.0f, 0.0f, 0.0f)
-			[
-				SNew(SExpandableArea)
-				.BorderImage(FAppStyle::Get().GetBrush("Brushes.Header"))
-				.BodyBorderImage(FAppStyle::Get().GetBrush("Brushes.Recessed"))
-				.HeaderPadding(FMargin(4.0f, 2.0f))
-				.InitiallyCollapsed(false)
-				.AllowAnimatedTransition(false)
-				.Visibility_Lambda([this, Index]() { return ColorGradingItemLists[Index].Num() ? EVisibility::Visible : EVisibility::Collapsed; })
-				.HeaderContent()
-				[
-					SNew(SBox)
-					.HeightOverride(24.0f)
-					.VAlign(VAlign_Center)
-					[
-						SNew(STextBlock)
-						.Text(this, &SDisplayClusterColorGradingDrawer::GetColorGradingListName, Index)
-						.TextStyle(FAppStyle::Get(), "ButtonText")
-						.Font(FAppStyle::Get().GetFontStyle("NormalFontBold"))
-					]
-				]
-				.BodyContent()
-				[
-					SAssignNew(NewListView, SDisplayClusterColorGradingObjectList)
-					.ColorGradingItemsSource(&ColorGradingItemLists[Index])
-					.OnSelectionChanged(this, &SDisplayClusterColorGradingDrawer::OnListSelectionChanged)
-				]
-			];
-
-		ColorGradingObjectListViews.Add(NewListView);
-	}
+	TSharedPtr<SWidget> SharedThis = AsShared();
 
 	ChildSlot
 	[
-		SNew(SBorder)
-		.BorderImage(FAppStyle::Get().GetBrush("Brushes.Panel"))
-		.HAlign(HAlign_Fill)
-		.VAlign(VAlign_Fill)
-		.Padding(FMargin(0.0f, 0.0f))
-		[
-			// Splitter to divide the object list and the color panel
-			SNew(SSplitter)
-			.Orientation(Orient_Horizontal)
-			.PhysicalSplitterHandleSize(2.0f)
-
-			+SSplitter::Slot()
-			.Value(0.12f)
-			[
-				SNew(SVerticalBox)
-
-				// Toolbar slot for the object list 
-				+ SVerticalBox::Slot()
-				.AutoHeight()
-				.Padding(0, 0, 0, 0)
-				[
-					SNew(SBorder)
-					.Padding(FMargin(3))
-					.BorderImage(bIsInDrawer ? FStyleDefaults::GetNoBrush() : FAppStyle::Get().GetBrush("Brushes.Panel"))
-					[
-						SNew(SBox)
-						.HeightOverride(28.0f)
-						[
-							CreateListPanelToolbar()
-						]
-					]
-				]
-
-				+SVerticalBox::Slot()
-				.AutoHeight()
-				[
-					SNew(SSeparator)
-					.Thickness(2.0f)
-				]
-
-				// Slot for the object lists, wrapped in expandable areas
-				+ SVerticalBox::Slot()
-				.FillHeight(1.0f)
-				[
-					SNew(SBox)
-					.Padding(FMargin(4.f))
-					[
-						SNew(SBorder)
-						.Padding(FMargin(0.0f))
-						.BorderImage(FAppStyle::GetBrush("Brushes.Recessed"))
-						[
-							SNew(SScrollBox)
-							+ SScrollBox::Slot()
-							[
-								ColorGradingObjectListBox
-							]
-						]
-					]
-				]
-			]
-
-			+SSplitter::Slot()
-			.Value(0.88f)
-			[
-				SNew(SVerticalBox)
-
-				// Toolbar slot for the main drawer toolbar
-				+ SVerticalBox::Slot()
-				.AutoHeight()
-				.Padding(0, 0, 0, 0)
-				[
-					SNew(SBorder)
-					.Padding(FMargin(3))
-					.BorderImage(bIsInDrawer ? FStyleDefaults::GetNoBrush() : FAppStyle::Get().GetBrush("Brushes.Panel"))
-					[
-						SNew(SBox)
-						.HeightOverride(28.0f)
-						[
-							SNew(SHorizontalBox)
-
-							// Slot for the color grading group toolbar
-							+ SHorizontalBox::Slot()
-							.AutoWidth()
-							.HAlign(HAlign_Left)
-							.VAlign(VAlign_Center)
-							[
-								SAssignNew(ColorGradingGroupToolBarBox, SHorizontalBox)
-								.Visibility(this, &SDisplayClusterColorGradingDrawer::GetColorGradingGroupToolBarVisibility)
-							]
-							
-
-							+ SHorizontalBox::Slot()
-							.FillWidth(1.0f)
-							[
-								SNew(SSpacer)
-							]
-
-							// Slot for the "Dock in Layout" button
-							+ SHorizontalBox::Slot()
-							.AutoWidth()
-							.HAlign(HAlign_Right)
-							.VAlign(VAlign_Center)
-							[
-								CreateDockInLayoutButton()
-							]
-						]
-					]
-				]
-
-				+SVerticalBox::Slot()
-				.AutoHeight()
-				[
-					SNew(SSeparator)
-					.Thickness(2.0f)
-				]
-
-				// Slot for the color panel
-				+SVerticalBox::Slot()
-				[
-					SNew(SBorder)
-					.BorderImage(FAppStyle::Get().GetBrush("Brushes.Panel"))
-					.Padding(FMargin(2.0f, 2.0f, 2.0f, 0.0f))
-					.Visibility(this, &SDisplayClusterColorGradingDrawer::GetDrawerModeVisibility, EDisplayClusterColorGradingDrawerMode::ColorGrading)
-					[
-						SAssignNew(ColorWheelPanel, SDisplayClusterColorGradingColorWheelPanel)
-						.ColorGradingDataModelSource(ColorGradingDataModel)
-					]
-				]
+		SAssignNew(MainPanel, SColorGradingPanel)
+		.OverrideWorld(this, &SDisplayClusterColorGradingDrawer::GetOperatorWorld)
+		.IsInDrawer(bInIsInDrawer)
+		.OnDocked_Lambda([]() {
+			IDisplayClusterColorGrading::Get().GetColorGradingDrawerSingleton().DockColorGradingDrawer();
+		})
+		.ActorFilter([SharedThis, this](AActor* Actor) -> bool {
+			if (SharedThis.IsValid() && OperatorViewModel.IsValid() && Actor && Actor->IsA<ADisplayClusterRootActor>())
+			{
+				return Actor == OperatorViewModel->GetRootActor();
+			}
 
-				// Slot for the details views
-				+SVerticalBox::Slot()
-				[
-					SNew(SBorder)
-					.BorderImage(FAppStyle::Get().GetBrush("Brushes.Panel"))
-					.Padding(FMargin(2.0f, 2.0f, 2.0f, 0.0f))
-					.Visibility(this, &SDisplayClusterColorGradingDrawer::GetDrawerModeVisibility, EDisplayClusterColorGradingDrawerMode::DetailsView)
-					[
-						SAssignNew(DetailsPanel, SDisplayClusterColorGradingDetailsPanel)
-						.ColorGradingDataModelSource(ColorGradingDataModel)
-					]
-				]
-			]
-		]
+			return true;
+		})
 	];
 }
 
-void SDisplayClusterColorGradingDrawer::Refresh(bool bPreserveDrawerState)
+void SDisplayClusterColorGradingDrawer::Refresh()
 {
-	FDisplayClusterColorGradingDrawerState DrawerState = GetDrawerState();
+	FColorGradingPanelState PanelState = GetColorGradingPanelState();
 
-	ColorGradingDataModel->Reset();
-
-	RefreshColorGradingLists();
-
-	if (ColorWheelPanel.IsValid())
+	if (MainPanel.IsValid())
 	{
-		ColorWheelPanel->Refresh();
+		MainPanel->Refresh();
 	}
 
-	if (DetailsPanel.IsValid())
-	{
-		DetailsPanel->Refresh();
-	}
-
-	if (bPreserveDrawerState)
-	{
-		SetDrawerState(DrawerState);
-	}
-	else
-	{
-		SetDrawerStateToDefault();
-	}
-}
-
-void SDisplayClusterColorGradingDrawer::Tick(const FGeometry& AllottedGeometry, const double InCurrentTime, const float InDeltaTime)
-{
-	SCompoundWidget::Tick(AllottedGeometry, InCurrentTime, InDeltaTime);
-
-	if (bRefreshOnNextTick)
-	{
-		const bool bPreserveDrawerState = true;
-		Refresh(bPreserveDrawerState);
-
-		bRefreshOnNextTick = false;
-	}
+	SetColorGradingPanelState(PanelState);
 }
 
 void SDisplayClusterColorGradingDrawer::PostUndo(bool bSuccess)
 {
 	if (bSuccess)
 	{
-		const bool bPreserveDrawerState = true;
-		Refresh(bPreserveDrawerState);
+		Refresh();
 	}
 }
 
@@ -330,230 +97,44 @@ void SDisplayClusterColorGradingDrawer::PostRedo(bool bSuccess)
 {
 	if (bSuccess)
 	{
-		const bool bPreserveDrawerState = true;
-		Refresh(bPreserveDrawerState);
+		Refresh();
 	}
 }
 
-FDisplayClusterColorGradingDrawerState SDisplayClusterColorGradingDrawer::GetDrawerState() const
+FColorGradingPanelState SDisplayClusterColorGradingDrawer::GetColorGradingPanelState() const
 {
-	FDisplayClusterColorGradingDrawerState DrawerState;
-	DrawerState.DrawerMode = CurrentDrawerMode;
+	FColorGradingPanelState PanelState;
 
-	ColorGradingDataModel->GetDrawerState(DrawerState);
-
-	if (ColorWheelPanel.IsValid())
+	if (MainPanel.IsValid())
 	{
-		ColorWheelPanel->GetDrawerState(DrawerState);
+		MainPanel->GetPanelState(PanelState);
 	}
 
-	if (DetailsPanel.IsValid())
-	{
-		DetailsPanel->GetDrawerState(DrawerState);
-	}
-
-	for (const TSharedPtr<SDisplayClusterColorGradingObjectList>& ListView : ColorGradingObjectListViews)
-	{
-		if (ListView.IsValid())
-		{
-			TArray<FDisplayClusterColorGradingListItemRef> SelectedItems = ListView->GetSelectedItems();
-
-			for (const FDisplayClusterColorGradingListItemRef& SelectedItem : SelectedItems)
-			{
-				if (SelectedItem.IsValid())
-				{
-					if (SelectedItem->Component.IsValid())
-					{
-						DrawerState.SelectedObjects.Add(SelectedItem->Component);
-					}
-					else if (SelectedItem->Actor.IsValid())
-					{
-						DrawerState.SelectedObjects.Add(SelectedItem->Actor);
-					}
-				}
-			}
-		}
-	}
-
-	return DrawerState;
+	return PanelState;
 }
 
-void SDisplayClusterColorGradingDrawer::SetDrawerState(const FDisplayClusterColorGradingDrawerState& InDrawerState)
+void SDisplayClusterColorGradingDrawer::SetColorGradingPanelState(const FColorGradingPanelState& InPanelState)
 {
-	TArray<TArray<FDisplayClusterColorGradingListItemRef>> ItemsToSelect;
-	ItemsToSelect.AddDefaulted(ColorGradingItemLists.Num());
-
-	for (const TWeakObjectPtr<UObject>& SelectedObject : InDrawerState.SelectedObjects)
+	if (MainPanel.IsValid())
 	{
-		if (SelectedObject.IsValid())
-		{
-			auto FindColorGradingItem = [&SelectedObject](const FDisplayClusterColorGradingListItemRef& ColorGradingItem)
-			{
-				return ColorGradingItem->Actor == SelectedObject || ColorGradingItem->Component == SelectedObject;
-			};
-
-			for (int32 Index = 0; Index < ColorGradingItemLists.Num(); ++Index)
-			{
-				if (FDisplayClusterColorGradingListItemRef* FoundItem = ColorGradingItemLists[Index].FindByPredicate(FindColorGradingItem))
-				{
-					ItemsToSelect[Index].Add(*FoundItem);
-					break;
-				}
-			}
-		}
-	}
-
-	// TODO: For now, since we don't support multiple color grading items selected at once, ensure either a level item or a root actor item is selected, not both
-	for (int32 Index = 0; Index < ColorGradingItemLists.Num(); ++Index)
-	{
-		if (ItemsToSelect[Index].Num())
-		{
-			if (ColorGradingObjectListViews.Num() > Index && ColorGradingObjectListViews[Index].IsValid())
-			{
-				ColorGradingObjectListViews[Index]->SetSelectedItems(ItemsToSelect[Index]);
-			}
-		}
-	}
-
-	ColorGradingDataModel->SetDrawerState(InDrawerState);
-
-	CurrentDrawerMode = InDrawerState.DrawerMode;
-
-	if (ColorWheelPanel.IsValid())
-	{
-		ColorWheelPanel->SetDrawerState(InDrawerState);
-	}
-
-	if (DetailsPanel.IsValid())
-	{
-		DetailsPanel->SetDrawerState(InDrawerState);
+		MainPanel->SetPanelState(InPanelState);
 	}
 }
 
-void SDisplayClusterColorGradingDrawer::SetDrawerStateToDefault()
+void SDisplayClusterColorGradingDrawer::SelectOperatorRootActor()
 {
-	constexpr int32 RootActorListIndex = 1;
-	if (ColorGradingItemLists.Num() > RootActorListIndex && ColorGradingObjectListViews.Num() > RootActorListIndex)
+	if (!MainPanel.IsValid())
 	{
-		TSharedPtr<SDisplayClusterColorGradingObjectList>& RootActorList = ColorGradingObjectListViews[RootActorListIndex];
-		TArray<FDisplayClusterColorGradingListItemRef>& RootActorColorGradingItems = ColorGradingItemLists[RootActorListIndex];
-
-		if (RootActorList.IsValid() && RootActorColorGradingItems.Num() > 0)
-		{
-			auto FindCCW = [](const TWeakObjectPtr<UObject>& Object)
-			{
-				return Object.IsValid() && Object->IsA<AColorCorrectionWindow>();
-			};
-
-			const bool bIsCCWSelected = OperatorViewModel->GetDetailObjects().ContainsByPredicate(FindCCW);
-
-			// If there is a CCW in the currently selected detail objects, automatically select the last (most recently selected) one
-			if (bIsCCWSelected)
-			{
-				int32 CCWIndex = OperatorViewModel->GetDetailObjects().FindLastByPredicate(FindCCW);
-				if (CCWIndex > INDEX_NONE)
-				{
-					AColorCorrectionWindow* CCW = Cast<AColorCorrectionWindow>(OperatorViewModel->GetDetailObjects()[CCWIndex]);
-					FDisplayClusterColorGradingListItemRef* ListItemPtr = RootActorColorGradingItems.FindByPredicate([CCW](const FDisplayClusterColorGradingListItemRef& ListItem)
-					{
-						return ListItem.IsValid() && ListItem->Actor.Get() == CCW;
-					});
-
-					if (ListItemPtr)
-					{
-						RootActorList->SetSelectedItems({ *ListItemPtr });
-						SetColorGradingDataModelObjects({ CCW });
-					}
-				}
-			}
-			else
-			{
-				// The nDisplay stage actor is always the first item in the root actor color grading items list, so set that as the currently selected item
-				RootActorList->SetSelectedItems({ RootActorColorGradingItems[0] });
-				SetColorGradingDataModelObjects({ RootActorColorGradingItems[0]->Actor.Get() });
-			}
-		}
+		return;
 	}
-}
-
-void SDisplayClusterColorGradingDrawer::BindCommands()
-{
-	CommandList->MapAction(
-		FDisplayClusterColorGradingCommands::Get().ColorGradingMode,
-		FExecuteAction::CreateSP(this, &SDisplayClusterColorGradingDrawer::OnDrawerModeSelected, EDisplayClusterColorGradingDrawerMode::ColorGrading),
-		FCanExecuteAction(),
-		FIsActionChecked::CreateSP(this, &SDisplayClusterColorGradingDrawer::IsDrawerModeSelected, EDisplayClusterColorGradingDrawerMode::ColorGrading)
-	);
-
-	CommandList->MapAction(
-		FDisplayClusterColorGradingCommands::Get().DetailViewMode,
-		FExecuteAction::CreateSP(this, &SDisplayClusterColorGradingDrawer::OnDrawerModeSelected, EDisplayClusterColorGradingDrawerMode::DetailsView),
-		FCanExecuteAction(),
-		FIsActionChecked::CreateSP(this, &SDisplayClusterColorGradingDrawer::IsDrawerModeSelected, EDisplayClusterColorGradingDrawerMode::DetailsView)
-	);
-}
 
-TSharedRef<SWidget> SDisplayClusterColorGradingDrawer::CreateDockInLayoutButton()
-{
-	if (bIsInDrawer)
+	if (OperatorViewModel->HasRootActor())
 	{
-		return SNew(SButton)
-			.ButtonStyle(FAppStyle::Get(), "SimpleButton")
-			.ToolTipText(LOCTEXT("DockInLayout_Tooltip", "Docks this panel in the current operator window, copying all settings from the drawer.\nThe drawer will still be usable."))
-			.OnClicked(this, &SDisplayClusterColorGradingDrawer::DockInLayout)
-			[
-				SNew(SHorizontalBox)
-
-				+ SHorizontalBox::Slot()
-				.AutoWidth()
-				.VAlign(VAlign_Center)
-				.Padding(4.0, 0.0f)
-				[
-					SNew(SImage)
-					.ColorAndOpacity(FSlateColor::UseForeground())
-					.Image(FAppStyle::Get().GetBrush("Icons.Layout"))
-				]
-
-				+ SHorizontalBox::Slot()
-				.VAlign(VAlign_Center)
-				.Padding(4.0, 0.0f)
-				[
-					SNew(STextBlock)
-					.Text(LOCTEXT("DockInLayout", "Dock in Layout"))
-					.ColorAndOpacity(FSlateColor::UseForeground())
-				]
-			];
+		MainPanel->SetSelectedObjects({ OperatorViewModel->GetRootActor() });
 	}
-	
-	return SNullWidget::NullWidget;
-}
-
-TSharedRef<SWidget> SDisplayClusterColorGradingDrawer::CreateListPanelToolbar() const
-{
-	FSlimHorizontalToolBarBuilder Toolbar(CommandList, FMultiBoxCustomization::None);
-
-	Toolbar.AddToolBarButton(FDisplayClusterColorGradingCommands::Get().ColorGradingMode, NAME_None, FText::GetEmpty());
-	Toolbar.AddToolBarButton(FDisplayClusterColorGradingCommands::Get().DetailViewMode, NAME_None, FText::GetEmpty());
-	Toolbar.SetStyle(&FDisplayClusterColorGradingStyle::Get(), "UltraSlimToolBar");
-
-	return Toolbar.MakeWidget();
-}
-
-FText SDisplayClusterColorGradingDrawer::GetColorGradingListName(int32 ListIndex) const
-{
-	switch (ListIndex)
+	else
 	{
-	case 0:
-		return GetCurrentLevelName();
-		
-	case 1:
-		return GetCurrentRootActorName();
-
-	case 2:
-		return LOCTEXT("ColorCorrectionRegionsListLabel", "Color Correction Regions");
-
-	default:
-		return FText::GetEmpty();
+		MainPanel->SetSelectedObjects({});
 	}
 }
 
@@ -580,633 +161,23 @@ FText SDisplayClusterColorGradingDrawer::GetCurrentRootActorName() const
 	return FText::GetEmpty();
 }
 
-void SDisplayClusterColorGradingDrawer::BindBlueprintCompiledDelegate(const UClass* Class)
-{
-	if (UBlueprint* Blueprint = UBlueprint::GetBlueprintFromClass(Class))
-	{
-		if (!Blueprint->OnCompiled().IsBoundToObject(this))
-		{
-			Blueprint->OnCompiled().AddSP(this, &SDisplayClusterColorGradingDrawer::OnBlueprintCompiled);
-		}
-	}
-}
-
-void SDisplayClusterColorGradingDrawer::UnbindBlueprintCompiledDelegate(const UClass* Class)
-{
-	if (UBlueprint* Blueprint = UBlueprint::GetBlueprintFromClass(Class))
-	{
-		Blueprint->OnCompiled().RemoveAll(this);
-	}
-}
-
-// Macros to create the lambdas used for a specified UObject's color grading enabled property in the color grading object list
-#define CREATE_IS_ENABLED_LAMBDA(Object, IsEnabledProperty) TAttribute<bool>::CreateLambda([Object]() { return (bool)IsEnabledProperty; })
-#define CREATE_ON_ENABLED_CHANGED_LAMBDA(Object, IsEnabledProperty) FOnColorGradingItemEnabledChanged::CreateLambda([Object](FDisplayClusterColorGradingListItemRef ListItem, bool bIsEnabled) \
-	{ \
-		FScopedTransaction Transaction(LOCTEXT("ColorGradingToggledTransaction", "Color Grading Toggled")); \
-		Object->Modify(!Object->IsA<ABrush>()); \
-		IsEnabledProperty = bIsEnabled; \
-	})
-
-void SDisplayClusterColorGradingDrawer::RefreshColorGradingLists()
-{
-	constexpr int32 NumLists = 3;
-	if (ColorGradingItemLists.Num() < NumLists)
-	{
-		ColorGradingItemLists.AddDefaulted(NumLists - ColorGradingItemLists.Num());
-	}
-
-	for (int32 Index = 0; Index < NumLists; ++Index)
-	{
-		RefreshColorGradingList(Index);
-	}
-}
-
-void SDisplayClusterColorGradingDrawer::RefreshColorGradingList(int32 Index)
-{
-	TArray<FDisplayClusterColorGradingListItemRef>& ItemList = ColorGradingItemLists[Index];
-	for (const FDisplayClusterColorGradingListItemRef& Item : ItemList)
-	{
-		if (Item->Component.IsValid())
-		{
-			UnbindBlueprintCompiledDelegate(Item->Component->GetClass());
-		}
-
-		if (Item->Actor.IsValid())
-		{
-			UnbindBlueprintCompiledDelegate(Item->Actor->GetClass());
-		}
-	}
-
-	ItemList.Empty();
-
-	switch (Index)
-	{
-	case 0:
-		FillLevelColorGradingList(ItemList);
-		break;
-
-	case 1:
-		FillRootActorColorGradingList(ItemList);
-		break;
-
-	case 2:
-		FillColorCorrectionRegionColorGradingList(ItemList);
-		break;
-	}
-
-	if (ColorGradingObjectListViews.Num() > Index && ColorGradingObjectListViews[Index].IsValid())
-	{
-		ColorGradingObjectListViews[Index]->RefreshList();
-	}
-}
-
-void SDisplayClusterColorGradingDrawer::FillLevelColorGradingList(TArray<FDisplayClusterColorGradingListItemRef>& List)
-{
-	if (ADisplayClusterRootActor* RootActor = OperatorViewModel->GetRootActor())
-	{
-		if (UWorld* World = RootActor->GetWorld())
-		{
-			// Sorter that sorts the list items alphabetically by their display name
-			auto AlphabeticalSort = [](const FDisplayClusterColorGradingListItemRef& A, const FDisplayClusterColorGradingListItemRef& B)
-			{
-				if (A.IsValid() && B.IsValid())
-				{
-					return *A < *B;
-				}
-				else
-				{
-					return false;
-				}
-			};
-
-			// Add all post process volumes that are in the same world as the stage actor
-			{
-				TArray<FDisplayClusterColorGradingListItemRef> SortedPPVs;
-				for (TActorIterator<APostProcessVolume> PPVIter(World); PPVIter; ++PPVIter)
-				{
-					APostProcessVolume* PostProcessVolume = *PPVIter;
-					BindBlueprintCompiledDelegate(PostProcessVolume->GetClass());
-
-					FDisplayClusterColorGradingListItemRef PPVListItemRef = MakeShared<FDisplayClusterColorGradingListItem>(PostProcessVolume);
-					PPVListItemRef->IsItemEnabled = CREATE_IS_ENABLED_LAMBDA(PostProcessVolume, PostProcessVolume->bEnabled);
-					PPVListItemRef->OnItemEnabledChanged = CREATE_ON_ENABLED_CHANGED_LAMBDA(PostProcessVolume, PostProcessVolume->bEnabled);
-
-					SortedPPVs.Add(PPVListItemRef);
-				}
-
-				SortedPPVs.Sort(AlphabeticalSort);
-				List.Append(SortedPPVs);
-			}
-		}
-	}
-}
-
-void SDisplayClusterColorGradingDrawer::FillRootActorColorGradingList(TArray<FDisplayClusterColorGradingListItemRef>& List)
-{
-	if (ADisplayClusterRootActor* RootActor = OperatorViewModel->GetRootActor())
-	{
-		BindBlueprintCompiledDelegate(RootActor->GetClass());
-
-		FDisplayClusterColorGradingListItemRef RootActorListItemRef = MakeShared<FDisplayClusterColorGradingListItem>(RootActor);
-		RootActorListItemRef->IsItemEnabled = CREATE_IS_ENABLED_LAMBDA(RootActor, RootActor->GetConfigData()->StageSettings.EnableColorGrading);
-		RootActorListItemRef->OnItemEnabledChanged = CREATE_ON_ENABLED_CHANGED_LAMBDA(RootActor, RootActor->GetConfigData()->StageSettings.EnableColorGrading);
-
-		List.Add(RootActorListItemRef);
-
-		auto AlphabeticalSort = [](const FDisplayClusterColorGradingListItemRef& A, const FDisplayClusterColorGradingListItemRef& B)
-		{
-			if (A.IsValid() && B.IsValid())
-			{
-				return *A < *B;
-			}
-			else
-			{
-				return false;
-			}
-		};
-
-		// Add any ICVFX camera component the root actor has to the color grading list
-		{
-			TArray<FDisplayClusterColorGradingListItemRef> SortedICVFXCameras;
-			RootActor->ForEachComponent<UDisplayClusterICVFXCameraComponent>(false, [this, RootActor, &SortedICVFXCameras](UDisplayClusterICVFXCameraComponent* ICVFXCameraComponent)
-			{
-				BindBlueprintCompiledDelegate(ICVFXCameraComponent->GetClass());
-
-				FDisplayClusterColorGradingListItemRef ICVFXCameraListItemRef = MakeShared<FDisplayClusterColorGradingListItem>(RootActor, ICVFXCameraComponent);
-				ICVFXCameraListItemRef->IsItemEnabled = CREATE_IS_ENABLED_LAMBDA(ICVFXCameraComponent, ICVFXCameraComponent->CameraSettings.EnableInnerFrustumColorGrading);
-				ICVFXCameraListItemRef->OnItemEnabledChanged = CREATE_ON_ENABLED_CHANGED_LAMBDA(ICVFXCameraComponent, ICVFXCameraComponent->CameraSettings.EnableInnerFrustumColorGrading);
-
-				SortedICVFXCameras.Add(ICVFXCameraListItemRef);
-			});
-
-			SortedICVFXCameras.Sort(AlphabeticalSort);
-			List.Append(SortedICVFXCameras);
-		}
-
-		// Add any color correction window that is currently selected in the operator's details panel to the list
-		{
-			const TArray<TWeakObjectPtr<UObject>>& DetailObjects = OperatorViewModel->GetDetailObjects();
-
-			TArray<FDisplayClusterColorGradingListItemRef> SortedCCWs;
-			for (const TWeakObjectPtr<UObject>& Object : DetailObjects)
-			{
-				if (Object.IsValid())
-				{
-					if (AColorCorrectionWindow* CCW = Cast<AColorCorrectionWindow>(Object.Get()))
-					{
-						BindBlueprintCompiledDelegate(CCW->GetClass());
-
-						FDisplayClusterColorGradingListItemRef ICVFXCameraListItemRef = MakeShared<FDisplayClusterColorGradingListItem>(CCW);
-						ICVFXCameraListItemRef->IsItemEnabled = CREATE_IS_ENABLED_LAMBDA(CCW, CCW->Enabled);
-						ICVFXCameraListItemRef->OnItemEnabledChanged = CREATE_ON_ENABLED_CHANGED_LAMBDA(CCW, CCW->Enabled);
-
-						SortedCCWs.Add(ICVFXCameraListItemRef);
-					}
-				}
-			}
-
-			SortedCCWs.Sort(AlphabeticalSort);
-			List.Append(SortedCCWs);
-		}
-	}
-}
-
-void SDisplayClusterColorGradingDrawer::FillColorCorrectionRegionColorGradingList(TArray<FDisplayClusterColorGradingListItemRef>& List)
-{
-	if (ADisplayClusterRootActor* RootActor = OperatorViewModel->GetRootActor())
-	{
-		if (UWorld* World = RootActor->GetWorld())
-		{
-			// Sorter that sorts the list items alphabetically by their display name
-			auto AlphabeticalSort = [](const FDisplayClusterColorGradingListItemRef& A, const FDisplayClusterColorGradingListItemRef& B)
-			{
-				if (A.IsValid() && B.IsValid())
-				{
-					return *A < *B;
-				}
-				else
-				{
-					return false;
-				}
-			};
-
-			// Add all color correction regions that are in the same world as the stage actor
-			{
-				TArray<FDisplayClusterColorGradingListItemRef> SortedCCRs;
-				for (TActorIterator<AColorCorrectRegion> CCRIter(World); CCRIter; ++CCRIter)
-				{
-					AColorCorrectRegion* ColorCorrectRegion = *CCRIter;
-
-					// Skip color correct windows, since those are stage managed actors and will be handled in a separate way
-					if (ColorCorrectRegion->IsA<AColorCorrectionWindow>())
-					{
-						continue;
-					}
-
-					BindBlueprintCompiledDelegate(ColorCorrectRegion->GetClass());
-
-					FDisplayClusterColorGradingListItemRef CCRListItemRef = MakeShared<FDisplayClusterColorGradingListItem>(ColorCorrectRegion);
-					CCRListItemRef->IsItemEnabled = CREATE_IS_ENABLED_LAMBDA(ColorCorrectRegion, ColorCorrectRegion->Enabled);
-					CCRListItemRef->OnItemEnabledChanged = CREATE_ON_ENABLED_CHANGED_LAMBDA(ColorCorrectRegion, ColorCorrectRegion->Enabled);
-
-					SortedCCRs.Add(CCRListItemRef);
-				}
-
-				SortedCCRs.Sort(AlphabeticalSort);
-				List.Append(SortedCCRs);
-			}
-		}
-	}
-}
-
-void SDisplayClusterColorGradingDrawer::SetColorGradingDataModelObjects(const TArray<UObject*>& Objects)
-{
-	ColorGradingDataModel->SetObjects(Objects);
-
-	// If the current drawer mode is the details view but the newly selected objects don't have any details sections to display,
-	// switch to the color grading drawer mode
-	if (CurrentDrawerMode == EDisplayClusterColorGradingDrawerMode::DetailsView && ColorGradingDataModel->DetailsSections.Num() == 0)
-	{
-		CurrentDrawerMode = EDisplayClusterColorGradingDrawerMode::ColorGrading;
-	}
-}
-
-void SDisplayClusterColorGradingDrawer::FillColorGradingGroupToolBar()
-{
-	if (ColorGradingGroupToolBarBox.IsValid())
-	{
-		ColorGradingGroupToolBarBox->ClearChildren();
-		ColorGradingGroupTextBlocks.Empty(ColorGradingDataModel->ColorGradingGroups.Num());
-
-		for (int32 Index = 0; Index < ColorGradingDataModel->ColorGradingGroups.Num(); ++Index)
-		{
-			TSharedPtr<SInlineEditableTextBlock> TextBlock = nullptr;
-
-			ColorGradingGroupToolBarBox->AddSlot()
-			.AutoWidth()
-			.Padding(2.0f, 0.0f, 2.0f, 0.0f)
-			[
-				SNew(SCheckBox)
-				.Style(FAppStyle::Get(), "DetailsView.SectionButton")
-				.OnCheckStateChanged(this, &SDisplayClusterColorGradingDrawer::OnColorGradingGroupCheckedChanged, Index)
-				.IsChecked(this, &SDisplayClusterColorGradingDrawer::IsColorGradingGroupSelected, Index)
-				.OnGetMenuContent(this, &SDisplayClusterColorGradingDrawer::GetColorGradingGroupMenuContent, Index)
-				[
-					SNew(SBox)
-					.HeightOverride(20.0)
-					[
-						SAssignNew(TextBlock, SInlineEditableTextBlock)
-						.Text(this, &SDisplayClusterColorGradingDrawer::GetColorGradingGroupDisplayName, Index)
-						.Font(this, &SDisplayClusterColorGradingDrawer::GetColorGradingGroupDisplayNameFont, Index)
-						.OnTextCommitted(this, &SDisplayClusterColorGradingDrawer::OnColorGradingGroupRenamed, Index)
-					]
-				]
-			];
-
-			ColorGradingGroupTextBlocks.Add(TextBlock);
-		}
-
-		if (ColorGradingDataModel->ColorGradingGroupToolBarWidget.IsValid())
-		{
-			ColorGradingGroupToolBarBox->AddSlot()
-			.AutoWidth()
-			.Padding(2.0f, 0.0f, 2.0f, 0.0f)
-			[
-				ColorGradingDataModel->ColorGradingGroupToolBarWidget.ToSharedRef()
-			];
-		}
-	}
-}
-
-bool SDisplayClusterColorGradingDrawer::IsDrawerModeSelected(EDisplayClusterColorGradingDrawerMode InDrawerMode) const
-{
-	return CurrentDrawerMode == InDrawerMode;
-}
-
-EVisibility SDisplayClusterColorGradingDrawer::GetDrawerModeVisibility(EDisplayClusterColorGradingDrawerMode InDrawerMode) const
-{
-	return CurrentDrawerMode == InDrawerMode ? EVisibility::Visible : EVisibility::Collapsed;
-}
-
-void SDisplayClusterColorGradingDrawer::OnDrawerModeSelected(EDisplayClusterColorGradingDrawerMode InDrawerMode)
-{
-	CurrentDrawerMode = InDrawerMode;
-}
-
-EVisibility SDisplayClusterColorGradingDrawer::GetColorGradingGroupToolBarVisibility() const
-{
-	if (ColorGradingDataModel->bShowColorGradingGroupToolBar && CurrentDrawerMode == EDisplayClusterColorGradingDrawerMode::ColorGrading)
-	{
-		return EVisibility::Visible;
-	}
-
-	return EVisibility::Hidden;
-}
-
-ECheckBoxState SDisplayClusterColorGradingDrawer::IsColorGradingGroupSelected(int32 GroupIndex) const
-{
-	return ColorGradingDataModel->GetSelectedColorGradingGroupIndex() == GroupIndex ? ECheckBoxState::Checked : ECheckBoxState::Unchecked;
-}
-
-void SDisplayClusterColorGradingDrawer::OnColorGradingGroupCheckedChanged(ECheckBoxState State, int32 GroupIndex)
-{
-	if (State == ECheckBoxState::Checked)
-	{
-		ColorGradingDataModel->SetSelectedColorGradingGroup(GroupIndex);
-	}
-}
-
-FText SDisplayClusterColorGradingDrawer::GetColorGradingGroupDisplayName(int32 GroupIndex) const
-{
-	if (ColorGradingDataModel->ColorGradingGroups.IsValidIndex(GroupIndex))
-	{
-		FText DisplayName = ColorGradingDataModel->ColorGradingGroups[GroupIndex].DisplayName;
-		if (DisplayName.IsEmpty())
-		{
-			return LOCTEXT("ColorGradingGroupEmptyNameLabel", "Unnamed");
-		}
-
-		return DisplayName;
-	}
-
-	return FText::GetEmpty();
-}
-
-FSlateFontInfo SDisplayClusterColorGradingDrawer::GetColorGradingGroupDisplayNameFont(int32 GroupIndex) const
-{
-	if (ColorGradingDataModel->ColorGradingGroups.IsValidIndex(GroupIndex))
-	{
-		FText DisplayName = ColorGradingDataModel->ColorGradingGroups[GroupIndex].DisplayName;
-		if (DisplayName.IsEmpty())
-		{
-			return FAppStyle::Get().GetFontStyle("NormalFontItalic");
-		}
-	}
-
-	return FAppStyle::Get().GetFontStyle("NormalFont");
-}
-
-TSharedRef<SWidget> SDisplayClusterColorGradingDrawer::GetColorGradingGroupMenuContent(int32 GroupIndex)
-{
-	if (ColorGradingDataModel->ColorGradingGroups.IsValidIndex(GroupIndex))
-	{
-		const FDisplayClusterColorGradingDataModel::FColorGradingGroup& Group = ColorGradingDataModel->ColorGradingGroups[GroupIndex];
-
-		FMenuBuilder MenuBuilder(true, nullptr);
-
-		const FGenericCommands& GenericCommands = FGenericCommands::Get();
-
-		if (Group.bCanBeRenamed)
-		{
-			MenuBuilder.AddMenuEntry(
-				GenericCommands.Rename->GetLabel(),
-				GenericCommands.Rename->GetDescription(),
-				GenericCommands.Rename->GetIcon(),
-				FUIAction(FExecuteAction::CreateSP(this, &SDisplayClusterColorGradingDrawer::OnColorGradingGroupRequestRename, GroupIndex))
-			);
-		}
-
-		if (Group.bCanBeDeleted)
-		{
-			MenuBuilder.AddMenuEntry(
-				GenericCommands.Delete->GetLabel(),
-				GenericCommands.Delete->GetDescription(),
-				GenericCommands.Delete->GetIcon(),
-				FUIAction(FExecuteAction::CreateSP(this, &SDisplayClusterColorGradingDrawer::OnColorGradingGroupDeleted, GroupIndex))
-			);
-		}
-
-		return MenuBuilder.MakeWidget();
-	}
-
-	return SNullWidget::NullWidget;
-}
-
-void SDisplayClusterColorGradingDrawer::OnColorGradingGroupDeleted(int32 GroupIndex)
-{
-	// If the group being deleted is in front of the currently selected one, we want to make sure
-	// that the same group is selected even after the deletion, so preemptively adjust the
-	// currently selected group index
-	int32 SelectedGroupIndex = ColorGradingDataModel->GetSelectedColorGradingGroupIndex();
-	if (SelectedGroupIndex > GroupIndex)
-	{
-		ColorGradingDataModel->SetSelectedColorGradingGroup(SelectedGroupIndex - 1);
-	}
-
-	ColorGradingDataModel->OnColorGradingGroupDeleted().Broadcast(GroupIndex);
-}
-
-void SDisplayClusterColorGradingDrawer::OnColorGradingGroupRequestRename(int32 GroupIndex)
-{
-	if (ColorGradingGroupTextBlocks.IsValidIndex(GroupIndex) && ColorGradingGroupTextBlocks[GroupIndex].IsValid())
-	{
-		ColorGradingGroupTextBlocks[GroupIndex]->EnterEditingMode();
-	}
-}
-
-void SDisplayClusterColorGradingDrawer::OnColorGradingGroupRenamed(const FText& InText, ETextCommit::Type TextCommitType, int32 GroupIndex)
-{
-	ColorGradingDataModel->OnColorGradingGroupRenamed().Broadcast(GroupIndex, InText);
-}
-
-void SDisplayClusterColorGradingDrawer::OnObjectsReplaced(const TMap<UObject*, UObject*>& OldToNewInstanceMap)
-{
-	bool bNeedsFullRefresh = false;
-	bool bNeedsListRefresh = false;
-
-	const TArray<TWeakObjectPtr<UObject>> SelectedObjects = ColorGradingDataModel->GetObjects();
-
-	for (const TPair<UObject*, UObject*>& Pair : OldToNewInstanceMap)
-	{
-		if (Pair.Key && Pair.Value)
-		{
-			FDisplayClusterColorGradingListItemRef* FoundColorGradingItemPtr = nullptr;
-
-			constexpr int32 RootActorListIndex = 1;
-			if (ColorGradingItemLists.Num() > RootActorListIndex)
-			{
-				// Must use GetEvenIfUnreachable on the weak pointers here because most of the time, the objects being replaced have already been marked for GC, and TWeakObjectPtr
-				// will return nullptr from Get on GC-marked objects
-				FoundColorGradingItemPtr = ColorGradingItemLists[RootActorListIndex].FindByPredicate([&Pair](const FDisplayClusterColorGradingListItemRef& ColorGradingItem)
-				{
-					return ColorGradingItem->Actor.GetEvenIfUnreachable() == Pair.Key || ColorGradingItem->Component.GetEvenIfUnreachable() == Pair.Key;
-				});
-			}
-
-			if (FoundColorGradingItemPtr)
-			{
-				FDisplayClusterColorGradingListItemRef FoundColorGradingItem = *FoundColorGradingItemPtr;
-				if (FoundColorGradingItem->Actor.GetEvenIfUnreachable() == Pair.Key)
-				{
-					FoundColorGradingItem->Actor = Cast<AActor>(Pair.Value);
-				}
-				else if (FoundColorGradingItem->Component.GetEvenIfUnreachable() == Pair.Key)
-				{
-					FoundColorGradingItem->Component = Cast<UActorComponent>(Pair.Value);
-				}
-
-				bNeedsListRefresh = true;
-			}
-
-			if (SelectedObjects.Contains(Pair.Key))
-			{
-				bNeedsFullRefresh = true;
-			}
-		}
-	}
-
-	if (bNeedsFullRefresh)
-	{
-		const bool bPreserveDrawerState = true;
-		Refresh(bPreserveDrawerState);
-	}
-	else if (bNeedsListRefresh)
-	{
-		for (const TSharedPtr<SDisplayClusterColorGradingObjectList>& ListView : ColorGradingObjectListViews)
-		{
-			if (ListView.IsValid())
-			{
-				ListView->RefreshList();
-			}
-		}
-	}
-}
-
-void SDisplayClusterColorGradingDrawer::OnLevelActorAdded(AActor* Actor)
-{
-	// Only refresh when the actor being added is being added to the root actor's world and is of a type this drawer cares about
-	if (OperatorViewModel->HasRootActor())
-	{
-		if (UWorld* World = OperatorViewModel->GetRootActor()->GetWorld())
-		{
-			if (World == Actor->GetWorld())
-			{
-				if (Actor->IsA<ADisplayClusterRootActor>() || Actor->IsA<APostProcessVolume>() || Actor->IsA<AColorCorrectRegion>())
-				{
-					// Wait to refresh, as this event can be fired off for several actors in a row in certain cases, such as when the root actor is recompiled after a property change
-					bRefreshOnNextTick = true;
-				}
-			}
-		}
-	}
-}
-
-void SDisplayClusterColorGradingDrawer::OnLevelActorDeleted(AActor* Actor)
-{
-	auto ContainsActorRef = [Actor](const FDisplayClusterColorGradingListItemRef& ColorGradingItem)
-	{
-		return ColorGradingItem->Actor.GetEvenIfUnreachable() == Actor;
-	};
-
-	for (const TArray<FDisplayClusterColorGradingListItemRef>& ItemList : ColorGradingItemLists)
-	{
-		if (ItemList.ContainsByPredicate(ContainsActorRef))
-		{
-			// Must wait for next tick to refresh because the actor has not actually been removed from the level at this point
-			bRefreshOnNextTick = true;
-			break;
-		}
-	}
-}
-
-void SDisplayClusterColorGradingDrawer::OnBlueprintCompiled(UBlueprint* Blueprint)
-{
-	const bool bPreserveDrawerState = true;
-	Refresh(bPreserveDrawerState);
-}
-
 void SDisplayClusterColorGradingDrawer::OnActiveRootActorChanged(ADisplayClusterRootActor* NewRootActor)
 {
-	const bool bPreserveDrawerState = false;
-	Refresh(bPreserveDrawerState);
+	Refresh();
+	SelectOperatorRootActor();
 }
 
-void SDisplayClusterColorGradingDrawer::OnDetailObjectsChanged(const TArray<UObject*>& Objects)
+UWorld* SDisplayClusterColorGradingDrawer::GetOperatorWorld() const
 {
-	constexpr int32 RootActorListIndex = 1;
-	if (ColorGradingItemLists.Num() > RootActorListIndex && ColorGradingObjectListViews.Num() > RootActorListIndex)
+	if (OperatorViewModel.IsValid())
 	{
-		TSharedPtr<SDisplayClusterColorGradingObjectList>& RootActorList = ColorGradingObjectListViews[RootActorListIndex];
-		TArray<FDisplayClusterColorGradingListItemRef>& RootActorColorGradingItems = ColorGradingItemLists[RootActorListIndex];
-		if (RootActorList.IsValid())
+		if (OperatorViewModel->HasRootActor())
 		{
-			const bool bIsCCWBeingSelected = Objects.ContainsByPredicate([](UObject* Object) { return Object && Object->IsA<AColorCorrectionWindow>(); });
-			const bool bAreCCWsInList = RootActorColorGradingItems.ContainsByPredicate([](const FDisplayClusterColorGradingListItemRef& ListItem)
-			{
-				return ListItem.IsValid() && ListItem->Actor.IsValid() && ListItem->Actor->IsA<AColorCorrectionWindow>();
-			});
-
-			// Only need to refresh when CCWs are being selected or unselected
-			if (bIsCCWBeingSelected || bAreCCWsInList)
-			{
-				
-				RefreshColorGradingList(RootActorListIndex);
-
-				SetDrawerStateToDefault();
-			}
+			return OperatorViewModel->GetRootActor()->GetWorld();
 		}
 	}
-}
-
-void SDisplayClusterColorGradingDrawer::OnColorGradingDataModelGenerated()
-{
-	FillColorGradingGroupToolBar();
-
-	if (ColorWheelPanel.IsValid())
-	{
-		ColorWheelPanel->Refresh();
-	}
 
-	if (DetailsPanel.IsValid())
-	{
-		DetailsPanel->Refresh();
-	}
-}
-
-void SDisplayClusterColorGradingDrawer::OnListSelectionChanged(TSharedRef<SDisplayClusterColorGradingObjectList> SourceList, FDisplayClusterColorGradingListItemRef SelectedItem, ESelectInfo::Type SelectInfo)
-{
-	TArray<FDisplayClusterColorGradingListItemRef> SelectedObjects = SourceList->GetSelectedItems();
-	if (SelectedObjects.Num())
-	{
-		// When an item on one list is selected, clear the selection of the other list
-		for (const TSharedPtr<SDisplayClusterColorGradingObjectList>& List : ColorGradingObjectListViews)
-		{
-			if (List == SourceList)
-			{
-				continue;
-			}
-
-			if (List->GetSelectedItems().Num())
-			{
-				List->SetSelectedItems(TArray<FDisplayClusterColorGradingListItemRef>());
-			}
-		}
-	}
-
-	if (SelectInfo != ESelectInfo::Direct)
-	{
-		TArray<UObject*> ObjectsToColorGrade;
-		for (const FDisplayClusterColorGradingListItemRef& SelectedObject : SelectedObjects)
-		{
-			if (SelectedObject->Component.IsValid())
-			{
-				ObjectsToColorGrade.Add(SelectedObject->Component.Get());
-			}
-			else if (SelectedObject->Actor.IsValid())
-			{
-				ObjectsToColorGrade.Add(SelectedObject->Actor.Get());
-			}
-		}
-
-		SetColorGradingDataModelObjects(ObjectsToColorGrade);
-	}
-}
-
-FReply SDisplayClusterColorGradingDrawer::DockInLayout()
-{
-	IDisplayClusterColorGrading::Get().GetColorGradingDrawerSingleton().DockColorGradingDrawer();
-	return FReply::Handled();
+	return nullptr;
 }
 
 #undef LOCTEXT_NAMESPACE
\ No newline at end of file
