[ColorGradingEditor] [DisplayCluster] Separate color grading panel into its own plugin and split color grading/ICVFX drawers in DisplayCluster

diff --git a/Engine/Plugins/Runtime/nDisplay/Source/DisplayClusterColorGrading/Private/Drawer/SDisplayClusterColorGradingDrawer.h b/Engine/Plugins/Runtime/nDisplay/Source/DisplayClusterColorGrading/Private/Drawer/SDisplayClusterColorGradingDrawer.h
--- a/Engine/Plugins/Runtime/nDisplay/Source/DisplayClusterColorGrading/Private/Drawer/SDisplayClusterColorGradingDrawer.h
+++ b/Engine/Plugins/Runtime/nDisplay/Source/DisplayClusterColorGrading/Private/Drawer/SDisplayClusterColorGradingDrawer.h
@@ -2,25 +2,22 @@
 
 #pragma once
 
+#include "ColorGradingPanelState.h"
 #include "CoreMinimal.h"
 #include "EditorUndoClient.h"
 #include "Styling/SlateTypes.h"
 #include "Widgets/DeclarativeSyntaxSupport.h"
 #include "Widgets/SCompoundWidget.h"
 
-#include "DisplayClusterColorGradingDrawerState.h"
-#include "SDisplayClusterColorGradingObjectList.h"
-#include "SDisplayClusterColorGradingColorWheel.h"
-
-class SInlineEditableTextBlock;
 class ADisplayClusterRootActor;
+class FColorGradingEditorDataModel;
 class FDisplayClusterOperatorStatusBarExtender;
-class FDisplayClusterColorGradingDataModel;
 class IDisplayClusterOperatorViewModel;
 class IPropertyRowGenerator;
-class SHorizontalBox;
-class SDisplayClusterColorGradingColorWheelPanel;
+class SColorGradingPanel;
 class SDisplayClusterColorGradingDetailsPanel;
+class SHorizontalBox;
+class SInlineEditableTextBlock;
 
 /** Color grading drawer widget, which displays a list of color gradable items, and the color wheel panel */
 class SDisplayClusterColorGradingDrawer : public SCompoundWidget, public FEditorUndoClient
@@ -34,159 +31,42 @@ public:
 
 	void Construct(const FArguments& InArgs, bool bInIsInDrawer);
 
-	/** Refreshes the drawer's UI to match the current state of the level and active root actor, optionally preserving UI state */
-	void Refresh(bool bPreserveDrawerState = false);
-
-	//~ SWidget interface
-	virtual void Tick(const FGeometry& AllottedGeometry, const double InCurrentTime, const float InDeltaTime) override;
-	//~ End SWidget interface
+	/** Refreshes the drawer's UI to match the current state of the level and active root actor */
+	void Refresh();
 
 	//~ FEditorUndoClient interface
 	virtual void PostUndo(bool bSuccess) override;
 	virtual void PostRedo(bool bSuccess) override;
 	//~ End FEditorUndoClient interface
 
-	/** Gets the state of the drawer UI */
-	FDisplayClusterColorGradingDrawerState GetDrawerState() const;
+	/** Gets the state of the color grading panel UI */
+	FColorGradingPanelState GetColorGradingPanelState() const;
 
-	/** Sets the state of the drawer UI */
-	void SetDrawerState(const FDisplayClusterColorGradingDrawerState& InDrawerState);
+	/** Sets the state of the color grading panel UI */
+	void SetColorGradingPanelState(const FColorGradingPanelState& InPanelState);
 
-	/** Sets the state of the drawer UI to its default value, which is to have the nDisplay stage actor selected */
-	void SetDrawerStateToDefault();
+	/** Sets the color grading panel's selected object to the operator panel's selected root actor */
+	void SelectOperatorRootActor();
 
 private:
-	/** Binds commands to the drawer's command list */
-	void BindCommands();
-
-	/** Creates the button used to dock the drawer in the operator panel */
-	TSharedRef<SWidget> CreateDockInLayoutButton();
-
-	/** Creates the toolbar used in the object list panel of the color grading drawer */
-	TSharedRef<SWidget> CreateListPanelToolbar() const;
-
-	/** Gets the display name of the color grading list at the specified index */
-	FText GetColorGradingListName(int32 ListIndex) const;
-
 	/** Gets the name of the current level the active root actor is in */
 	FText GetCurrentLevelName() const;
 
 	/** Gets the name of the active root actor */
 	FText GetCurrentRootActorName() const;
 
-	/** Binds a callback to the BlueprintCompiled delegate of the specified class */
-	void BindBlueprintCompiledDelegate(const UClass* Class);
-
-	/** Unbinds a callback to the BlueprintCompiled delegate of the specified class */
-	void UnbindBlueprintCompiledDelegate(const UClass* Class);
-
-	/** Refreshes all of the lists in the drawer, filling them with the current color gradable objects from the root actor and world  */
-	void RefreshColorGradingLists();
-
-	/** Refreshes the specified list, filling it with the current color gradable objects from the relevant source */
-	void RefreshColorGradingList(int32 Index);
-
-	/** Fills the specified list with any color gradable items found in the currently loaded level */
-	void FillLevelColorGradingList(TArray<FDisplayClusterColorGradingListItemRef>& InOutList);
-
-	/** Fills the specified list with any color gradable components of the active root actor */
-	void FillRootActorColorGradingList(TArray<FDisplayClusterColorGradingListItemRef>& InOutList);
-
-	/** Fills the specified list with any color gradable color correction regions in the currently loaded level */
-	void FillColorCorrectionRegionColorGradingList(TArray<FDisplayClusterColorGradingListItemRef>& InOutList);
-
-	/** Updates the color grading data model with the specified list of objects */
-	void SetColorGradingDataModelObjects(const TArray<UObject*>& Objects);
-
-	/** Fills the color grading group toolbar using the color grading data model */
-	void FillColorGradingGroupToolBar();
-
-	/** Gets whether the specified drawer mode is currently selected */
-	bool IsDrawerModeSelected(EDisplayClusterColorGradingDrawerMode InDrawerMode) const;
-
-	/** Gets the visibility state for the specified drawer mode */
-	EVisibility GetDrawerModeVisibility(EDisplayClusterColorGradingDrawerMode InDrawerMode) const;
-
-	/** Raised when the specified drawer mode checked state is changed */
-	void OnDrawerModeSelected(EDisplayClusterColorGradingDrawerMode InDrawerMode);
-
-	/** Gets the visibility state of the color grading group toolbar */
-	EVisibility GetColorGradingGroupToolBarVisibility() const;
-
-	/** Gets whether the color grading group at the specified index is currently selected */
-	ECheckBoxState IsColorGradingGroupSelected(int32 GroupIndex) const;
-
-	/** Raised when the user has selected the specified color grading group */
-	void OnColorGradingGroupCheckedChanged(ECheckBoxState State, int32 GroupIndex);
-
-	/** Gets the display name of the specified color grading group */
-	FText GetColorGradingGroupDisplayName(int32 GroupIndex) const;
-
-	/** Gets the font of the display name label of the specified color grading group */
-	FSlateFontInfo GetColorGradingGroupDisplayNameFont(int32 GroupIndex) const;
-
-	/** Gets the content for the right click menu for the color grading group */
-	TSharedRef<SWidget> GetColorGradingGroupMenuContent(int32 GroupIndex);
-
-	/** Raised when a color grading group has been deleted by the user */
-	void OnColorGradingGroupDeleted(int32 GroupIndex);
-
-	/** Raised when a rename has been requested on a color grading group */
-	void OnColorGradingGroupRequestRename(int32 GroupIndex);
-
-	/** Raised when a rename has been committed on a color grading group */
-	void OnColorGradingGroupRenamed(const FText& InText, ETextCommit::Type TextCommitType, int32 GroupIndex);
-
-	/** Raised when the editor replaces any UObjects with new instantiations, usually when actors have been recompiled from blueprints */
-	void OnObjectsReplaced(const TMap<UObject*, UObject*>& OldToNewInstanceMap);
-
-	/** Raised when an actor is added to the current level */
-	void OnLevelActorAdded(AActor* Actor);
-
-	/** Raised when an actor has been deleted from the currnent level */
-	void OnLevelActorDeleted(AActor* Actor);
-
-	/** Raised when the specified blueprint has been recompiled */
-	void OnBlueprintCompiled(UBlueprint* Blueprint);
-
 	/** Raised when the user has changed the active root actor selected in the nDisplay operator panel */
 	void OnActiveRootActorChanged(ADisplayClusterRootActor* NewRootActor);
-	
-	/** Raised when the objects being displayed in the operator's details panel have changed */
-	void OnDetailObjectsChanged(const TArray<UObject*>& Objects);
-
-	/** Raised when the color grading data model has been generated */
-	void OnColorGradingDataModelGenerated();
-
-	/** Raised when the user has selected a new item in any of the drawer's list views */
-	void OnListSelectionChanged(TSharedRef<SDisplayClusterColorGradingObjectList> SourceList, FDisplayClusterColorGradingListItemRef SelectedItem, ESelectInfo::Type SelectInfo);
 
-	/** Raised when the "Dock in Layout" button has been clicked */
-	FReply DockInLayout();
+	/** Get the world of the current root actor in the ICVFX panel */
+	UWorld* GetOperatorWorld() const;
 
 private:
 	/** The operator panel's view model */
 	TSharedPtr<IDisplayClusterOperatorViewModel> OperatorViewModel;
 
-	/** A list of all color grading object list widgets being displayed in the drawer's list panel */
-	TArray<TSharedPtr<SDisplayClusterColorGradingObjectList>> ColorGradingObjectListViews;
-
-	/** A list of source lists for the color grading object list widgets */
-	TArray<TArray<FDisplayClusterColorGradingListItemRef>> ColorGradingItemLists;
-
-	TSharedPtr<SHorizontalBox> ColorGradingGroupToolBarBox;
-	TSharedPtr<SDisplayClusterColorGradingColorWheelPanel> ColorWheelPanel;
-	TSharedPtr<SDisplayClusterColorGradingDetailsPanel> DetailsPanel;
-	TArray<TSharedPtr<SInlineEditableTextBlock>> ColorGradingGroupTextBlocks;
-
-	/** Commands used by the drawer */
-	TSharedPtr<FUICommandList> CommandList;
-
-	/** The color grading data model for the currently selected objects */
-	TSharedPtr<FDisplayClusterColorGradingDataModel> ColorGradingDataModel;
-
-	/** The current mode that the drawer is in */
-	EDisplayClusterColorGradingDrawerMode CurrentDrawerMode = EDisplayClusterColorGradingDrawerMode::ColorGrading;
+	/** The panel containing the color grading wheels or object details */
+	TSharedPtr<SColorGradingPanel> MainPanel;
 
 	/** Gets whether this widget is in a drawer or docked in a tab */
 	bool bIsInDrawer = false;
diff --git a/Engine/Plugins/Runtime/nDisplay/Source/DisplayClusterColorGrading/Public/IDisplayClusterColorGradingDrawerSingleton.h b/Engine/Plugins/Runtime/nDisplay/Source/DisplayClusterColorGrading/Public/IDisplayClusterColorGradingDrawerSingleton.h
--- a/Engine/Plugins/Runtime/nDisplay/Source/DisplayClusterColorGrading/Public/IDisplayClusterColorGradingDrawerSingleton.h
+++ b/Engine/Plugins/Runtime/nDisplay/Source/DisplayClusterColorGrading/Public/IDisplayClusterColorGradingDrawerSingleton.h
@@ -12,5 +12,5 @@ public:
 	virtual void DockColorGradingDrawer() = 0;
 
 	/** Refreshes the UI of any open color grading drawers */
-	virtual void RefreshColorGradingDrawers(bool bPreserveDrawerState) = 0;
+	virtual void RefreshColorGradingDrawers() = 0;
 };
\ No newline at end of file
diff --git a/Engine/Plugins/Runtime/nDisplay/nDisplay.uplugin b/Engine/Plugins/Runtime/nDisplay/nDisplay.uplugin
--- a/Engine/Plugins/Runtime/nDisplay/nDisplay.uplugin
+++ b/Engine/Plugins/Runtime/nDisplay/nDisplay.uplugin
@@ -144,6 +144,15 @@
 				"Linux"
 			]
 		},
+		{
+			"Name": "DisplayClusterDetails",
+			"Type": "Editor",
+			"LoadingPhase": "Default",
+			"PlatformAllowList": [
+				"Win64",
+				"Linux"
+			]
+		},
 		{
 			"Name": "DisplayClusterStageMonitoring",
 			"Type": "Runtime",
@@ -247,6 +256,10 @@
 			"Name": "ColorCorrectRegions",
 			"Enabled": true
 		},
+		{
+			"Name": "ColorGrading",
+			"Enabled": true
+		},
 		{
 			"Name": "Composure",
 			"Enabled": true
