[nDisplay] Late OCIO to properly support per-node OCIO for ICVFX cameras streamed in-cluster via media

- Sending nodes don't apply OCIO before exporting an image
- An image is taken directly from the rendering pipeline, in its native format
- Receiving nodes apply their own OCIO based on per-node OCIO settings/profiles

diff --git a/Engine/Plugins/Runtime/nDisplay/Source/DisplayCluster/Private/DisplayClusterCallbacks.h b/Engine/Plugins/Runtime/nDisplay/Source/DisplayCluster/Private/DisplayClusterCallbacks.h
--- a/Engine/Plugins/Runtime/nDisplay/Source/DisplayCluster/Private/DisplayClusterCallbacks.h
+++ b/Engine/Plugins/Runtime/nDisplay/Source/DisplayCluster/Private/DisplayClusterCallbacks.h
@@ -84,12 +84,17 @@ public:
 
 	virtual FDisplayClusterFailoverNodeDown& OnDisplayClusterFailoverNodeDown() override
 	{
-		return DisplayClusterFailoverNodeDown;
+		return DisplayClusterFailoverNodeDownEvent;
+	}
+
+	virtual FDisplayClusterPostTonemapPass_RenderThread& OnDisplayClusterPostTonemapPass_RenderThread() override
+	{
+		return DisplayClusterPostTonemapPassEvent;
 	}
 
 	virtual FDisplayClusterPostRenderViewFamily_RenderThread& OnDisplayClusterPostRenderViewFamily_RenderThread() override
 	{
-		return DisplayClusterPostRenderViewFamily;
+		return DisplayClusterPostRenderViewFamilyEvent;
 	}
 
 	virtual FDisplayClusterPreWarp_RenderThread& OnDisplayClusterPreWarp_RenderThread() override
@@ -159,13 +164,15 @@ private:
 	FDisplayClusterEndSceneEvent             DisplayClusterEndSceneEvent;
 	FDisplayClusterPreSubmitViewFamilies     DisplayClusterPreSubmitViewFamiliesEvent;
 	FDisplayClusterCustomPresentSetEvent     DisplayClusterCustomPresentSetEvent;
-	FDisplayClusterFailoverNodeDown          DisplayClusterFailoverNodeDown;
+	FDisplayClusterUpdateViewportMediaState  DisplayClusterUpdateViewportMediaStateEvent;
+	FDisplayClusterFailoverNodeDown          DisplayClusterFailoverNodeDownEvent;
 
 	FDisplayClusterPresentationPreSynchronization_RHIThread  DisplayClusterPresentationPreSynchronizationEvent;
 	FDisplayClusterPresentationPostSynchronization_RHIThread DisplayClusterPresentationPostSynchronizationEvent;
 	FDisplayClusterFramePresentated_RHIThread                DisplayClusterFramePresentedEvent;
 
-	FDisplayClusterPostRenderViewFamily_RenderThread DisplayClusterPostRenderViewFamily;
+	FDisplayClusterPostTonemapPass_RenderThread      DisplayClusterPostTonemapPassEvent;
+	FDisplayClusterPostRenderViewFamily_RenderThread DisplayClusterPostRenderViewFamilyEvent;
 	FDisplayClusterPreWarp_RenderThread              DisplayClusterPreWarpEvent;
 	FDisplayClusterPreWarpViewport_RenderThread      DisplayClusterPreWarpViewportEvent;
 	FDisplayClusterPostWarp_RenderThread             DisplayClusterPostWarpEvent;
@@ -176,6 +183,4 @@ private:
 	FDisplayClusterPostBackbufferUpdate_RenderThread DisplayClusterPostBackbufferUpdateEvent;
 	FDisplayClusterPostBackbufferUpdated_RenderThread DisplayClusterPostBackbufferUpdatedEvent;
 	FDisplayClusterPreProcessIcvfx_RenderThread      FDisplayClusterPreProcessIcvfxEvent;
-
-	FDisplayClusterUpdateViewportMediaState DisplayClusterUpdateViewportMediaStateEvent;
 };
diff --git a/Engine/Plugins/Runtime/nDisplay/Source/DisplayCluster/Private/Render/Viewport/DisplayClusterViewport.cpp b/Engine/Plugins/Runtime/nDisplay/Source/DisplayCluster/Private/Render/Viewport/DisplayClusterViewport.cpp
--- a/Engine/Plugins/Runtime/nDisplay/Source/DisplayCluster/Private/Render/Viewport/DisplayClusterViewport.cpp
+++ b/Engine/Plugins/Runtime/nDisplay/Source/DisplayCluster/Private/Render/Viewport/DisplayClusterViewport.cpp
@@ -756,3 +756,15 @@ bool FDisplayClusterViewport::UpdateFrameContexts(const uint32 InStereoViewIndex
 
 	return true;
 }
+
+bool FDisplayClusterViewport::GetOCIOConversionSettings(FOpenColorIOColorConversionSettings& OutOCIOConversionSettings) const
+{
+	// Return OCIO conversion settings if available
+	if (OpenColorIO)
+	{
+		OutOCIOConversionSettings = OpenColorIO->GetConversionSettings();
+		return true;
+	}
+
+	return false;
+}
diff --git a/Engine/Plugins/Runtime/nDisplay/Source/DisplayCluster/Private/Render/Viewport/DisplayClusterViewport.h b/Engine/Plugins/Runtime/nDisplay/Source/DisplayCluster/Private/Render/Viewport/DisplayClusterViewport.h
--- a/Engine/Plugins/Runtime/nDisplay/Source/DisplayCluster/Private/Render/Viewport/DisplayClusterViewport.h
+++ b/Engine/Plugins/Runtime/nDisplay/Source/DisplayCluster/Private/Render/Viewport/DisplayClusterViewport.h
@@ -109,6 +109,9 @@ public:
 
 	virtual FVector2D GetClippingPlanes() const override;
 
+	/** Provides current OCIO settings */
+	virtual bool GetOCIOConversionSettings(FOpenColorIOColorConversionSettings& OutOCIOConversionSettings) const override;
+
 	virtual void CalculateProjectionMatrix(const uint32 InContextNum, float Left, float Right, float Top, float Bottom, float ZNear, float ZFar, bool bIsAnglesInput) override;
 
 	virtual bool    CalculateView(const uint32 InContextNum, FVector& InOutViewLocation, FRotator& InOutViewRotation, const float WorldToMeters) override;
diff --git a/Engine/Plugins/Runtime/nDisplay/Source/DisplayCluster/Private/Render/Viewport/DisplayClusterViewportProxy.cpp b/Engine/Plugins/Runtime/nDisplay/Source/DisplayCluster/Private/Render/Viewport/DisplayClusterViewportProxy.cpp
--- a/Engine/Plugins/Runtime/nDisplay/Source/DisplayCluster/Private/Render/Viewport/DisplayClusterViewportProxy.cpp
+++ b/Engine/Plugins/Runtime/nDisplay/Source/DisplayCluster/Private/Render/Viewport/DisplayClusterViewportProxy.cpp
@@ -141,12 +141,6 @@ EDisplayClusterViewportOpenColorIOMode FDisplayClusterViewportProxy::GetOpenColo
 			// Rendering without post-processing, OCIO is applied last, to the RTT texture of the viewport
 			return EDisplayClusterViewportOpenColorIOMode::Resolved;
 		}
-		else if (RenderSettings.HasAnyMediaStates(EDisplayClusterViewportMediaState::Capture_ForceLateOCIOPass))
-		{
-			// When capturing a viewport, it's possible that it's going to be shared within a cluster via the media pipeline.
-			// In this case we should postpone the OCIO step so every node can apply its own OCIO settings.
-			return EDisplayClusterViewportOpenColorIOMode::Resolved;
-		}
 
 		// By default, viewports render with a postprocess, OCIO must be done in between.
 		return EDisplayClusterViewportOpenColorIOMode::PostProcess;
@@ -421,9 +415,13 @@ FScreenPassTexture FDisplayClusterViewportProxy::OnPostProcessPassAfterFXAA_Rend
 
 FScreenPassTexture FDisplayClusterViewportProxy::OnPostProcessPassAfterTonemap_RenderThread(FRDGBuilder& GraphBuilder, const FSceneView& View, const FPostProcessMaterialInputs& Inputs, const uint32 ContextNum)
 {
+	// Broadcast PassTonemap event
+	IDisplayCluster::Get().GetCallbacks().OnDisplayClusterPostTonemapPass_RenderThread().Broadcast(GraphBuilder, this, View, Inputs, ContextNum);
+
 	// Perform OCIO rendering after the tonemapper
 	if (GetOpenColorIOMode() == EDisplayClusterViewportOpenColorIOMode::PostProcess)
 	{
+		// Add OCIO pass
 		return OpenColorIO->PostProcessPassAfterTonemap_RenderThread(GraphBuilder, GetContexts_RenderThread()[ContextNum], View, Inputs);
 	}
 
diff --git a/Engine/Plugins/Runtime/nDisplay/Source/DisplayCluster/Public/IDisplayClusterCallbacks.h b/Engine/Plugins/Runtime/nDisplay/Source/DisplayCluster/Public/IDisplayClusterCallbacks.h
--- a/Engine/Plugins/Runtime/nDisplay/Source/DisplayCluster/Public/IDisplayClusterCallbacks.h
+++ b/Engine/Plugins/Runtime/nDisplay/Source/DisplayCluster/Public/IDisplayClusterCallbacks.h
@@ -11,6 +11,7 @@
 class FDisplayClusterShaderParameters_ICVFX;
 class FRDGBuilder;
 class FRHICommandListImmediate;
+class FSceneView;
 class FSceneViewFamily;
 class FSceneViewFamilyContext;
 class FViewport;
@@ -18,6 +19,7 @@ class IDisplayClusterViewport;
 class IDisplayClusterViewportProxy;
 class IDisplayClusterViewportManagerProxy;
 struct FDisplayClusterShaderParameters_WarpBlend;
+struct FPostProcessMaterialInputs;
 
 
 /**
@@ -89,6 +91,10 @@ public:
 	DECLARE_EVENT_OneParam(IDisplayClusterCallbacks, FDisplayClusterFailoverNodeDown, const FString&);
 	virtual FDisplayClusterFailoverNodeDown& OnDisplayClusterFailoverNodeDown() = 0;
 
+	/** Called after tonemap pass (before applying OCIO transformation) **/
+	DECLARE_EVENT_FiveParams(IDisplayClusterCallbacks, FDisplayClusterPostTonemapPass_RenderThread, FRDGBuilder&, const IDisplayClusterViewportProxy*, const FSceneView&, const FPostProcessMaterialInputs&, const uint32);
+	virtual FDisplayClusterPostTonemapPass_RenderThread& OnDisplayClusterPostTonemapPass_RenderThread() = 0;
+
 	/** Called once the ViewFamily of this viewport is rendered **/
 	DECLARE_EVENT_ThreeParams(IDisplayClusterCallbacks, FDisplayClusterPostRenderViewFamily_RenderThread, FRDGBuilder&, const FSceneViewFamily&, const IDisplayClusterViewportProxy*);
 	virtual FDisplayClusterPostRenderViewFamily_RenderThread& OnDisplayClusterPostRenderViewFamily_RenderThread() = 0;
diff --git a/Engine/Plugins/Runtime/nDisplay/Source/DisplayCluster/Public/Render/Viewport/Containers/DisplayClusterViewport_Enums.h b/Engine/Plugins/Runtime/nDisplay/Source/DisplayCluster/Public/Render/Viewport/Containers/DisplayClusterViewport_Enums.h
--- a/Engine/Plugins/Runtime/nDisplay/Source/DisplayCluster/Public/Render/Viewport/Containers/DisplayClusterViewport_Enums.h
+++ b/Engine/Plugins/Runtime/nDisplay/Source/DisplayCluster/Public/Render/Viewport/Containers/DisplayClusterViewport_Enums.h
@@ -189,15 +189,10 @@ enum class EDisplayClusterViewportMediaState : uint8
 	// This viewport does not use media.
 	None = 0,
 
-	// This viewport will be captured by the media device.
+	// This viewport will be captured by a media device.
 	Capture = 1 << 0,
-	Capture_ForceLateOCIOPass = 1 << 1,
 
-	// This viewport is overridden by the media device.
-	Input = 1 << 4,
-	Input_ForceLateOCIOPass = 1 << 5,
-
-	ForceLateOCIOPass = Capture_ForceLateOCIOPass | Input_ForceLateOCIOPass,
+	// This viewport is overridden by a media device.
+	Input = 1 << 1,
 };
 ENUM_CLASS_FLAGS(EDisplayClusterViewportMediaState);
-
diff --git a/Engine/Plugins/Runtime/nDisplay/Source/DisplayCluster/Public/Render/Viewport/IDisplayClusterViewport.h b/Engine/Plugins/Runtime/nDisplay/Source/DisplayCluster/Public/Render/Viewport/IDisplayClusterViewport.h
--- a/Engine/Plugins/Runtime/nDisplay/Source/DisplayCluster/Public/Render/Viewport/IDisplayClusterViewport.h
+++ b/Engine/Plugins/Runtime/nDisplay/Source/DisplayCluster/Public/Render/Viewport/IDisplayClusterViewport.h
@@ -19,6 +19,8 @@
 class UCameraComponent;
 class UWorld;
 struct FMinimalViewInfo;
+struct FOpenColorIOColorConversionSettings;
+
 
 /**
  * nDisplay: Viewport (interface for GameThread)
@@ -135,6 +137,13 @@ public:
 	 */
 	virtual FVector2D GetClippingPlanes() const = 0;
 
+	/**
+	 * Returns OCIO configuration if any set
+	 * 
+	 * @return true if OCIO is configured
+	 */
+	virtual bool GetOCIOConversionSettings(FOpenColorIOColorConversionSettings& OutOCIOConversionSettings) const = 0;
+
 	/**
 	* [const] Get custom postprocess settings for this viewport.
 	*/
diff --git a/Engine/Plugins/Runtime/nDisplay/Source/DisplayClusterConfiguration/Public/DisplayClusterConfigurationTypes_Media.h b/Engine/Plugins/Runtime/nDisplay/Source/DisplayClusterConfiguration/Public/DisplayClusterConfigurationTypes_Media.h
--- a/Engine/Plugins/Runtime/nDisplay/Source/DisplayClusterConfiguration/Public/DisplayClusterConfigurationTypes_Media.h
+++ b/Engine/Plugins/Runtime/nDisplay/Source/DisplayClusterConfiguration/Public/DisplayClusterConfigurationTypes_Media.h
@@ -300,8 +300,8 @@ public:
 	UPROPERTY(EditAnywhere, BlueprintReadWrite, Category = "Media", meta = (DisplayName = "Media Output Groups"))
 	TArray<FDisplayClusterConfigurationMediaTiledOutputGroup> TiledMediaOutputGroups;
 
-	/** Force late OCIO pass */
-	UPROPERTY()
+	/** Allows the receviers to apply their own OCIO transformations (per node OCIO override). Requires the media to support FloatRGBA.*/
+	UPROPERTY(EditAnywhere, BlueprintReadWrite, Category = "Media", meta = (DisplayName = "Receivers apply OCIO"))
 	bool bLateOCIOPass = false;
 
 public:
diff --git a/Engine/Plugins/Runtime/nDisplay/Source/DisplayClusterMedia/DisplayClusterMedia.Build.cs b/Engine/Plugins/Runtime/nDisplay/Source/DisplayClusterMedia/DisplayClusterMedia.Build.cs
--- a/Engine/Plugins/Runtime/nDisplay/Source/DisplayClusterMedia/DisplayClusterMedia.Build.cs
+++ b/Engine/Plugins/Runtime/nDisplay/Source/DisplayClusterMedia/DisplayClusterMedia.Build.cs
@@ -6,18 +6,12 @@ public class DisplayClusterMedia : ModuleRules
 {
 	public DisplayClusterMedia(ReadOnlyTargetRules ROTargetRules) : base(ROTargetRules)
 	{
-		PublicIncludePathModuleNames.AddRange(
-			new string[] {
-				"DisplayClusterShaders",
-			});
-
 		PublicDependencyModuleNames.AddRange(
 			new string[] {
 				"DisplayClusterConfiguration",
 				"Media",
 				"MediaAssets",
 				"MediaIOCore",
-				"SharedMemoryMedia",
 			});
 
 		PrivateDependencyModuleNames.AddRange(
@@ -26,6 +20,7 @@ public class DisplayClusterMedia : ModuleRules
 				"CoreUObject",
 				"DisplayCluster",
 				"Engine",
+				"OpenColorIO",
 				"Renderer",
 				"RenderCore",
 				"RHI",
diff --git a/Engine/Plugins/Runtime/nDisplay/Source/DisplayClusterMedia/Private/Capture/DisplayClusterMediaCaptureBase.cpp b/Engine/Plugins/Runtime/nDisplay/Source/DisplayClusterMedia/Private/Capture/DisplayClusterMediaCaptureBase.cpp
--- a/Engine/Plugins/Runtime/nDisplay/Source/DisplayClusterMedia/Private/Capture/DisplayClusterMediaCaptureBase.cpp
+++ b/Engine/Plugins/Runtime/nDisplay/Source/DisplayClusterMedia/Private/Capture/DisplayClusterMediaCaptureBase.cpp
@@ -13,13 +13,22 @@
 #include "MediaOutput.h"
 
 #include "RenderGraphBuilder.h"
+#include "RenderGraphUtils.h"
+#include "RHICommandList.h"
+#include "RHIUtilities.h"
 
 #include "UObject/UObjectGlobals.h"
 #include "UObject/Package.h"
 
 
-FDisplayClusterMediaCaptureBase::FDisplayClusterMediaCaptureBase(const FString& InMediaId, const FString& InClusterNodeId, UMediaOutput* InMediaOutput, UDisplayClusterMediaOutputSynchronizationPolicy* InSyncPolicy)
-	: FDisplayClusterMediaBase(InMediaId, InClusterNodeId)
+FDisplayClusterMediaCaptureBase::FDisplayClusterMediaCaptureBase(
+	const FString& InMediaId,
+	const FString& InClusterNodeId,
+	UMediaOutput* InMediaOutput,
+	UDisplayClusterMediaOutputSynchronizationPolicy* InSyncPolicy,
+	bool bInLateOCIO
+)
+	: FDisplayClusterMediaBase(InMediaId, InClusterNodeId, bInLateOCIO)
 	, SyncPolicy(InSyncPolicy)
 {
 	checkSlow(InMediaOutput);
@@ -115,19 +124,27 @@ void FDisplayClusterMediaCaptureBase::StopCapture()
 	}
 }
 
-void FDisplayClusterMediaCaptureBase::ExportMediaData(FRDGBuilder& GraphBuilder, const FMediaTextureInfo& TextureInfo)
+void FDisplayClusterMediaCaptureBase::ExportMediaData_RenderThread(FRDGBuilder& GraphBuilder, const FMediaOutputTextureInfo& TextureInfo)
 {
-	FRHITexture* const SrcTexture = TextureInfo.Texture;
-
-	if (SrcTexture)
+	// Check if source texture is valid
+	if (!TextureInfo.Texture)
 	{
-		UE_LOG(LogDisplayClusterMedia, Verbose, TEXT("MediaCapture '%s': exporting texture on RT frame '%lu'..."), *GetMediaId(), GFrameCounterRenderThread);
+		UE_LOG(LogDisplayClusterMedia, Warning, TEXT("MediaCapture '%s': invalid source texture on RT frame %lu"), *GetMediaId(), GFrameCounterRenderThread);
+		return;
+	}
 
-		MediaCapture->SetValidSourceGPUMask(GraphBuilder.RHICmdList.GetGPUMask());
+	MediaCapture->SetValidSourceGPUMask(GraphBuilder.RHICmdList.GetGPUMask());
 
-		LastSrcRegionSize = FIntSize(TextureInfo.Region.Size());
-		MediaCapture->CaptureImmediate_RenderThread(GraphBuilder, SrcTexture);
-	}
+	const FIntPoint SrcTextureSize = TextureInfo.Texture->Desc.Extent;
+	const FIntPoint SrcRegionSize  = TextureInfo.Region.Size();
+
+	LastSrcRegionSize = FIntSize(SrcRegionSize);
+
+	UE_LOG(LogDisplayClusterMedia, VeryVerbose, TEXT("MediaCapture '%s': exporting texture [size=%dx%d, rect=%dx%d] on RT frame '%lu'..."),
+		*GetMediaId(), SrcTextureSize.X, SrcTextureSize.Y, SrcRegionSize.X, SrcRegionSize.Y, GFrameCounterRenderThread);
+
+	// Capture
+	MediaCapture->CaptureImmediate_RenderThread(GraphBuilder, TextureInfo.Texture, TextureInfo.Region);
 }
 
 void FDisplayClusterMediaCaptureBase::OnPostClusterTick()
diff --git a/Engine/Plugins/Runtime/nDisplay/Source/DisplayClusterMedia/Private/Capture/DisplayClusterMediaCaptureBase.h b/Engine/Plugins/Runtime/nDisplay/Source/DisplayClusterMedia/Private/Capture/DisplayClusterMediaCaptureBase.h
--- a/Engine/Plugins/Runtime/nDisplay/Source/DisplayClusterMedia/Private/Capture/DisplayClusterMediaCaptureBase.h
+++ b/Engine/Plugins/Runtime/nDisplay/Source/DisplayClusterMedia/Private/Capture/DisplayClusterMediaCaptureBase.h
@@ -6,24 +6,32 @@
 #include "DisplayClusterMediaBase.h"
 #include "UObject/GCObject.h"
 
+#include "RenderGraphFwd.h"
+
 #include "RHI.h"
 #include "RHIResources.h"
 
-class FRDGBuilder;
 class UMediaCapture;
 class UMediaOutput;
 class UDisplayClusterMediaOutputSynchronizationPolicy;
 class IDisplayClusterMediaOutputSynchronizationPolicyHandler;
 
 /**
- * Base media capture class
+ * Base media capture adapter class
  */
 class FDisplayClusterMediaCaptureBase
 	: public FDisplayClusterMediaBase
 	, public FGCObject
 {
 public:
-	FDisplayClusterMediaCaptureBase(const FString& MediaId, const FString& ClusterNodeId, UMediaOutput* MediaOutput, UDisplayClusterMediaOutputSynchronizationPolicy* SyncPolicy = nullptr);
+	FDisplayClusterMediaCaptureBase(
+		const FString& MediaId,
+		const FString& ClusterNodeId,
+		UMediaOutput* MediaOutput,
+		UDisplayClusterMediaOutputSynchronizationPolicy* SyncPolicy = nullptr,
+		bool bInLateOCIO = false
+	);
+
 	virtual ~FDisplayClusterMediaCaptureBase();
 
 public:
@@ -36,19 +44,43 @@ public:
 	//~ End FGCObject interface
 
 public:
+
+	/** Start capturing */
 	virtual bool StartCapture();
+
+	/** Stop capturing */
 	virtual void StopCapture();
 
+	/** Returns current media capture device */
 	UMediaCapture* GetMediaCapture() const
 	{
 		return MediaCapture;
 	}
 
 protected:
-	void ExportMediaData(FRDGBuilder& GraphBuilder, const FMediaTextureInfo& TextureInfo);
+
+	/** Media capture data */
+	struct FMediaOutputTextureInfo
+	{
+		/** Texture to capture by a media capture device */
+		FRDGTextureRef Texture = nullptr;
+
+		/** Subregion to capture */
+		FIntRect Region = { FIntPoint::ZeroValue, FIntPoint::ZeroValue };
+	};
+
+protected:
+
+	/** PostClusterTick event handler. It's used to restart capturing if needed */
 	void OnPostClusterTick();
+
+	/** Re-starts media capturing after failure */
 	bool StartMediaCapture();
 
+	/** Passes capture data request to the capture device */
+	void ExportMediaData_RenderThread(FRDGBuilder& GraphBuilder, const FMediaOutputTextureInfo& TextureInfo);
+
+	/** Returns capture size (main thread) */
 	virtual FIntPoint GetCaptureSize() const = 0;
 
 private:
diff --git a/Engine/Plugins/Runtime/nDisplay/Source/DisplayClusterMedia/Private/Capture/DisplayClusterMediaCaptureCamera.cpp b/Engine/Plugins/Runtime/nDisplay/Source/DisplayClusterMedia/Private/Capture/DisplayClusterMediaCaptureCamera.cpp
--- a/Engine/Plugins/Runtime/nDisplay/Source/DisplayClusterMedia/Private/Capture/DisplayClusterMediaCaptureCamera.cpp
+++ b/Engine/Plugins/Runtime/nDisplay/Source/DisplayClusterMedia/Private/Capture/DisplayClusterMediaCaptureCamera.cpp
@@ -9,8 +9,16 @@
 #include "Game/IDisplayClusterGameManager.h"
 
 
-FDisplayClusterMediaCaptureCamera::FDisplayClusterMediaCaptureCamera(const FString& InMediaId, const FString& InClusterNodeId, const FString& InCameraId, const FString& InViewportId, UMediaOutput* InMediaOutput, UDisplayClusterMediaOutputSynchronizationPolicy* SyncPolicy)
-	: FDisplayClusterMediaCaptureViewport(InMediaId, InClusterNodeId, InViewportId, InMediaOutput, SyncPolicy)
+FDisplayClusterMediaCaptureCamera::FDisplayClusterMediaCaptureCamera(
+	const FString& InMediaId,
+	const FString& InClusterNodeId,
+	const FString& InCameraId,
+	const FString& InViewportId,
+	UMediaOutput* InMediaOutput,
+	UDisplayClusterMediaOutputSynchronizationPolicy* SyncPolicy,
+	bool bInLateOCIO
+)
+	: FDisplayClusterMediaCaptureViewport(InMediaId, InClusterNodeId, InViewportId, InMediaOutput, SyncPolicy, bInLateOCIO)
 	, CameraId(InCameraId)
 {
 }
diff --git a/Engine/Plugins/Runtime/nDisplay/Source/DisplayClusterMedia/Private/Capture/DisplayClusterMediaCaptureCamera.h b/Engine/Plugins/Runtime/nDisplay/Source/DisplayClusterMedia/Private/Capture/DisplayClusterMediaCaptureCamera.h
--- a/Engine/Plugins/Runtime/nDisplay/Source/DisplayClusterMedia/Private/Capture/DisplayClusterMediaCaptureCamera.h
+++ b/Engine/Plugins/Runtime/nDisplay/Source/DisplayClusterMedia/Private/Capture/DisplayClusterMediaCaptureCamera.h
@@ -12,7 +12,15 @@ class FDisplayClusterMediaCaptureCamera
 	: public FDisplayClusterMediaCaptureViewport
 {
 public:
-	FDisplayClusterMediaCaptureCamera(const FString& MediaId, const FString& ClusterNodeId, const FString& CameraId, const FString& ViewportId, UMediaOutput* MediaOutput, UDisplayClusterMediaOutputSynchronizationPolicy* SyncPolicy = nullptr);
+	FDisplayClusterMediaCaptureCamera(
+		const FString& MediaId,
+		const FString& ClusterNodeId,
+		const FString& CameraId,
+		const FString& ViewportId,
+		UMediaOutput* MediaOutput,
+		UDisplayClusterMediaOutputSynchronizationPolicy* SyncPolicy = nullptr,
+		bool bInLateOCIO = false
+	);
 
 protected:
 	//~ Begin FDisplayClusterMediaCaptureViewport
diff --git a/Engine/Plugins/Runtime/nDisplay/Source/DisplayClusterMedia/Private/Capture/DisplayClusterMediaCaptureNode.cpp b/Engine/Plugins/Runtime/nDisplay/Source/DisplayClusterMedia/Private/Capture/DisplayClusterMediaCaptureNode.cpp
--- a/Engine/Plugins/Runtime/nDisplay/Source/DisplayClusterMedia/Private/Capture/DisplayClusterMediaCaptureNode.cpp
+++ b/Engine/Plugins/Runtime/nDisplay/Source/DisplayClusterMedia/Private/Capture/DisplayClusterMediaCaptureNode.cpp
@@ -10,6 +10,7 @@
 #include "IDisplayClusterCallbacks.h"
 
 #include "RenderGraphBuilder.h"
+#include "RenderGraphUtils.h"
 #include "Render/Viewport/IDisplayClusterViewportManagerProxy.h"
 #include "RHICommandList.h"
 #include "RHIResources.h"
@@ -19,7 +20,12 @@
 #include "UnrealClient.h"
 
 
-FDisplayClusterMediaCaptureNode::FDisplayClusterMediaCaptureNode(const FString& InMediaId, const FString& InClusterNodeId, UMediaOutput* InMediaOutput, UDisplayClusterMediaOutputSynchronizationPolicy* SyncPolicy)
+FDisplayClusterMediaCaptureNode::FDisplayClusterMediaCaptureNode(
+	const FString& InMediaId,
+	const FString& InClusterNodeId,
+	UMediaOutput* InMediaOutput,
+	UDisplayClusterMediaOutputSynchronizationPolicy* SyncPolicy
+)
 	: FDisplayClusterMediaCaptureBase(InMediaId, InClusterNodeId, InMediaOutput, SyncPolicy)
 {
 }
@@ -68,11 +74,16 @@ void FDisplayClusterMediaCaptureNode::OnPostBackbufferUpdated_RenderThread(FRHIC
 	{
 		if (FRHITexture* const BackbufferTexture = Viewport->GetRenderTargetTexture())
 		{
+			FRDGBuilder GraphBuilder(RHICmdList);
+
+			FRDGTextureRef BackbufferTextureRef = RegisterExternalTexture(GraphBuilder, BackbufferTexture, TEXT("DCMediaOutBackbufferTex"));
+			const FIntRect TextureRegion = { FIntPoint::ZeroValue, BackbufferTextureRef->Desc.Extent };
+			FMediaOutputTextureInfo TextureInfo{ BackbufferTextureRef, TextureRegion };
+
 			// Capture backbuffer
-			FMediaTextureInfo TextureInfo{ BackbufferTexture, FIntRect(FIntPoint::ZeroValue, BackbufferTexture->GetDesc().Extent) };
-			FRDGBuilder Builder(RHICmdList);
-			ExportMediaData(Builder, TextureInfo);
-			Builder.Execute();
+			ExportMediaData_RenderThread(GraphBuilder, TextureInfo);
+
+			GraphBuilder.Execute();
 		}
 	}
 }
diff --git a/Engine/Plugins/Runtime/nDisplay/Source/DisplayClusterMedia/Private/Capture/DisplayClusterMediaCaptureNode.h b/Engine/Plugins/Runtime/nDisplay/Source/DisplayClusterMedia/Private/Capture/DisplayClusterMediaCaptureNode.h
--- a/Engine/Plugins/Runtime/nDisplay/Source/DisplayClusterMedia/Private/Capture/DisplayClusterMediaCaptureNode.h
+++ b/Engine/Plugins/Runtime/nDisplay/Source/DisplayClusterMedia/Private/Capture/DisplayClusterMediaCaptureNode.h
@@ -18,15 +18,25 @@ class FDisplayClusterMediaCaptureNode
 	: public FDisplayClusterMediaCaptureBase
 {
 public:
-	FDisplayClusterMediaCaptureNode(const FString& MediaId, const FString& ClusterNodeId, UMediaOutput* MediaOutput, UDisplayClusterMediaOutputSynchronizationPolicy* SyncPolicy = nullptr);
+	FDisplayClusterMediaCaptureNode(
+		const FString& MediaId,
+		const FString& ClusterNodeId,
+		UMediaOutput* MediaOutput,
+		UDisplayClusterMediaOutputSynchronizationPolicy* SyncPolicy = nullptr
+	);
 
 public:
+	/** Start backbuffer capture */
 	virtual bool StartCapture() override;
+
+	/** Stop backbuffer capture */
 	virtual void StopCapture() override;
 
 protected:
+	/** Returns backbuffer size */
 	virtual FIntPoint GetCaptureSize() const override;
 
 private:
+	/** PostBackbufferUpdated event handler */
 	void OnPostBackbufferUpdated_RenderThread(FRHICommandListImmediate& RHICmdList, FViewport* Viewport);
 };
diff --git a/Engine/Plugins/Runtime/nDisplay/Source/DisplayClusterMedia/Private/Capture/DisplayClusterMediaCaptureTile.cpp b/Engine/Plugins/Runtime/nDisplay/Source/DisplayClusterMedia/Private/Capture/DisplayClusterMediaCaptureTile.cpp
--- a/Engine/Plugins/Runtime/nDisplay/Source/DisplayClusterMedia/Private/Capture/DisplayClusterMediaCaptureTile.cpp
+++ b/Engine/Plugins/Runtime/nDisplay/Source/DisplayClusterMedia/Private/Capture/DisplayClusterMediaCaptureTile.cpp
@@ -3,8 +3,15 @@
 #include "Capture/DisplayClusterMediaCaptureTile.h"
 
 
-FDisplayClusterMediaCaptureTile::FDisplayClusterMediaCaptureTile(const FString& InMediaId, const FString& InClusterNodeId, const FString& InViewportId, UMediaOutput* InMediaOutput, UDisplayClusterMediaOutputSynchronizationPolicy* SyncPolicy)
-	: FDisplayClusterMediaCaptureViewport(InMediaId, InClusterNodeId, InViewportId, InMediaOutput, SyncPolicy)
+FDisplayClusterMediaCaptureTile::FDisplayClusterMediaCaptureTile(
+	const FString& InMediaId,
+	const FString& InClusterNodeId,
+	const FString& InViewportId,
+	UMediaOutput* InMediaOutput,
+	UDisplayClusterMediaOutputSynchronizationPolicy* SyncPolicy,
+	bool bInLateOCIO
+)
+	: FDisplayClusterMediaCaptureViewport(InMediaId, InClusterNodeId, InViewportId, InMediaOutput, SyncPolicy, bInLateOCIO)
 {
 }
 
diff --git a/Engine/Plugins/Runtime/nDisplay/Source/DisplayClusterMedia/Private/Capture/DisplayClusterMediaCaptureTile.h b/Engine/Plugins/Runtime/nDisplay/Source/DisplayClusterMedia/Private/Capture/DisplayClusterMediaCaptureTile.h
--- a/Engine/Plugins/Runtime/nDisplay/Source/DisplayClusterMedia/Private/Capture/DisplayClusterMediaCaptureTile.h
+++ b/Engine/Plugins/Runtime/nDisplay/Source/DisplayClusterMedia/Private/Capture/DisplayClusterMediaCaptureTile.h
@@ -12,7 +12,14 @@ class FDisplayClusterMediaCaptureTile
 	: public FDisplayClusterMediaCaptureViewport
 {
 public:
-	FDisplayClusterMediaCaptureTile(const FString& MediaId, const FString& ClusterNodeId, const FString& ViewportId, UMediaOutput* MediaOutput, UDisplayClusterMediaOutputSynchronizationPolicy* SyncPolicy = nullptr);
+	FDisplayClusterMediaCaptureTile(
+		const FString& MediaId,
+		const FString& ClusterNodeId,
+		const FString& ViewportId,
+		UMediaOutput* MediaOutput,
+		UDisplayClusterMediaOutputSynchronizationPolicy* SyncPolicy = nullptr,
+		bool bInLateOCIO = false
+	);
 
 protected:
 
diff --git a/Engine/Plugins/Runtime/nDisplay/Source/DisplayClusterMedia/Private/Capture/DisplayClusterMediaCaptureViewport.cpp b/Engine/Plugins/Runtime/nDisplay/Source/DisplayClusterMedia/Private/Capture/DisplayClusterMediaCaptureViewport.cpp
--- a/Engine/Plugins/Runtime/nDisplay/Source/DisplayClusterMedia/Private/Capture/DisplayClusterMediaCaptureViewport.cpp
+++ b/Engine/Plugins/Runtime/nDisplay/Source/DisplayClusterMedia/Private/Capture/DisplayClusterMediaCaptureViewport.cpp
@@ -6,24 +6,35 @@
 #include "IDisplayClusterCallbacks.h"
 
 #include "DisplayClusterConfigurationTypes_Viewport.h"
+#include "DisplayClusterMediaCVars.h"
 #include "DisplayClusterMediaLog.h"
 #include "DisplayClusterRootActor.h"
 
 #include "Config/IDisplayClusterConfigManager.h"
 #include "Game/IDisplayClusterGameManager.h"
 
+#include "PostProcess/PostProcessMaterialInputs.h"
+
 #include "Render/IDisplayClusterRenderManager.h"
 #include "Render/Viewport/IDisplayClusterViewport.h"
 #include "Render/Viewport/IDisplayClusterViewportManager.h"
 #include "Render/Viewport/IDisplayClusterViewportManagerProxy.h"
 #include "Render/Viewport/Containers/DisplayClusterViewport_Context.h"
 #include "Render/Viewport/Containers/DisplayClusterViewport_RenderSettings.h"
+
 #include "RHICommandList.h"
 #include "RHIResources.h"
 
 
-FDisplayClusterMediaCaptureViewport::FDisplayClusterMediaCaptureViewport(const FString& InMediaId, const FString& InClusterNodeId, const FString& InViewportId, UMediaOutput* InMediaOutput, UDisplayClusterMediaOutputSynchronizationPolicy* SyncPolicy)
-	: FDisplayClusterMediaCaptureBase(InMediaId, InClusterNodeId, InMediaOutput, SyncPolicy)
+FDisplayClusterMediaCaptureViewport::FDisplayClusterMediaCaptureViewport(
+	const FString& InMediaId,
+	const FString& InClusterNodeId,
+	const FString& InViewportId,
+	UMediaOutput* InMediaOutput,
+	UDisplayClusterMediaOutputSynchronizationPolicy* SyncPolicy,
+	bool bInLateOCIO
+)
+	: FDisplayClusterMediaCaptureBase(InMediaId, InClusterNodeId, InMediaOutput, SyncPolicy, bInLateOCIO)
 	, ViewportId(InViewportId)
 {
 }
@@ -32,9 +43,18 @@ FDisplayClusterMediaCaptureViewport::FDisplayClusterMediaCaptureViewport(const F
 bool FDisplayClusterMediaCaptureViewport::StartCapture()
 {
 	// Subscribe for events
-	IDisplayCluster::Get().GetCallbacks().OnDisplayClusterPostRenderViewFamily_RenderThread().AddRaw(this, &FDisplayClusterMediaCaptureViewport::OnPostRenderViewFamily_RenderThread);
 	IDisplayCluster::Get().GetCallbacks().OnDisplayClusterUpdateViewportMediaState().AddRaw(this, &FDisplayClusterMediaCaptureViewport::OnUpdateViewportMediaState);
 
+	// Depending on late OCIO configuration, grab the image in different places
+	if (IsLateOCIO())
+	{
+		IDisplayCluster::Get().GetCallbacks().OnDisplayClusterPostTonemapPass_RenderThread().AddRaw(this, &FDisplayClusterMediaCaptureViewport::OnPostTonemapPass_RenderThread);
+	}
+	else
+	{
+		IDisplayCluster::Get().GetCallbacks().OnDisplayClusterPostRenderViewFamily_RenderThread().AddRaw(this, &FDisplayClusterMediaCaptureViewport::OnPostRenderViewFamily_RenderThread);
+	}
+
 	// Start capture
 	const bool bStarted = FDisplayClusterMediaCaptureBase::StartCapture();
 
@@ -43,8 +63,9 @@ bool FDisplayClusterMediaCaptureViewport::StartCapture()
 
 void FDisplayClusterMediaCaptureViewport::StopCapture()
 {
-	// Unsubscribe from events
+	// Unsubscribe from external events/callbacks
 	IDisplayCluster::Get().GetCallbacks().OnDisplayClusterPostRenderViewFamily_RenderThread().RemoveAll(this);
+	IDisplayCluster::Get().GetCallbacks().OnDisplayClusterPostTonemapPass_RenderThread().RemoveAll(this);
 	IDisplayCluster::Get().GetCallbacks().OnDisplayClusterUpdateViewportMediaState().RemoveAll(this);
 
 	// Stop capturing
@@ -53,22 +74,20 @@ void FDisplayClusterMediaCaptureViewport::StopCapture()
 
 void FDisplayClusterMediaCaptureViewport::OnUpdateViewportMediaState(IDisplayClusterViewport* InViewport, EDisplayClusterViewportMediaState& InOutMediaState)
 {
-	// Note: Media currently supports only one DCRA.
-	// In the future, after the media redesign, the DCRA name will also need to be checked here.
+	// Set capture flag for the matching viewport
 	if (InViewport && InViewport->GetId().Equals(GetViewportId(), ESearchCase::IgnoreCase))
 	{
 		// Raise flags that this viewport will be captured by media.
 		InOutMediaState |= EDisplayClusterViewportMediaState::Capture;
-
-		if (bForceLateOCIOPass)
-		{
-			// Raise flags that this capture requires ForceLateOCIOPass.
-			InOutMediaState |= EDisplayClusterViewportMediaState::Capture_ForceLateOCIOPass;
-		}
 	}
 }
 
 FIntPoint FDisplayClusterMediaCaptureViewport::GetCaptureSize() const
+{
+	return GetViewportSize();
+}
+
+FIntPoint FDisplayClusterMediaCaptureViewport::GetViewportSize() const
 {
 	FIntPoint CaptureSize{ FIntPoint::ZeroValue };
 
@@ -130,22 +149,65 @@ bool FDisplayClusterMediaCaptureViewport::GetCaptureSizeFromGameProxy(FIntPoint&
 	return false;
 }
 
+void FDisplayClusterMediaCaptureViewport::OnPostTonemapPass_RenderThread(FRDGBuilder& GraphBuilder, const IDisplayClusterViewportProxy* ViewportProxy, const FSceneView& View, const FPostProcessMaterialInputs& Inputs, const uint32 ContextNum)
+{
+	checkSlow(ViewportProxy);
+
+	// Media subsystem does not support stereo, therefore we process context 0 only
+	if (ContextNum != 0)
+	{
+		return;
+	}
+
+	// Check if proxy object is valid
+	if (!ViewportProxy)
+	{
+		return;
+	}
+
+	// Make sure this is our viewport
+	const bool bMatchingViewport = ViewportProxy->GetId().Equals(GetViewportId(), ESearchCase::IgnoreCase);
+	if (!bMatchingViewport)
+	{
+		return;
+	}
+
+	// Get current SceneColor texture
+	const FScreenPassTexture& SceneColor = FScreenPassTexture::CopyFromSlice(GraphBuilder, Inputs.GetInput(EPostProcessMaterialInput::SceneColor));
+
+	// Pass it to the media capture pipeline
+	if (SceneColor.IsValid())
+	{
+		FMediaOutputTextureInfo TextureInfo{ SceneColor.Texture, SceneColor.ViewRect };
+		ExportMediaData_RenderThread(GraphBuilder, TextureInfo);
+	}
+}
+
 void FDisplayClusterMediaCaptureViewport::OnPostRenderViewFamily_RenderThread(FRDGBuilder& GraphBuilder, const FSceneViewFamily& ViewFamily, const IDisplayClusterViewportProxy* ViewportProxy)
 {
-	ensure(ViewportProxy);
+	checkSlow(ViewportProxy);
 
+	// Nothing to do if late OCIO is required. The texture has been exported already on PostTonemap callback.
+	if (IsLateOCIO())
+	{
+		return;
+	}
+
+	// Otherwise, find our viewport and export its texture
 	if (ViewportProxy && ViewportProxy->GetId().Equals(GetViewportId(), ESearchCase::IgnoreCase))
 	{
 		TArray<FRHITexture*> Textures;
 		TArray<FIntRect>     Regions;
 
-		// Get RHI texture
+		// Get RHI texture and pass it to the media capture pipeline
 		if (ViewportProxy->GetResourcesWithRects_RenderThread(EDisplayClusterViewportResourceType::InternalRenderTargetResource, Textures, Regions))
 		{
-			if (Textures.Num() > 0 && Regions.Num() > 0)
+			if (Textures.Num() > 0 && Regions.Num() > 0 && Textures[0])
 			{
-				FMediaTextureInfo TextureInfo{ Textures[0], Regions[0] };
-				ExportMediaData(GraphBuilder, TextureInfo);
+				FRDGTextureRef SrcTextureRef = RegisterExternalTexture(GraphBuilder, Textures[0], TEXT("DCMediaOutViewportTex"));
+
+				FMediaOutputTextureInfo TextureInfo{ SrcTextureRef, Regions[0] };
+				ExportMediaData_RenderThread(GraphBuilder, TextureInfo);
 			}
 		}
 	}
diff --git a/Engine/Plugins/Runtime/nDisplay/Source/DisplayClusterMedia/Private/Capture/DisplayClusterMediaCaptureViewport.h b/Engine/Plugins/Runtime/nDisplay/Source/DisplayClusterMedia/Private/Capture/DisplayClusterMediaCaptureViewport.h
--- a/Engine/Plugins/Runtime/nDisplay/Source/DisplayClusterMedia/Private/Capture/DisplayClusterMediaCaptureViewport.h
+++ b/Engine/Plugins/Runtime/nDisplay/Source/DisplayClusterMedia/Private/Capture/DisplayClusterMediaCaptureViewport.h
@@ -7,9 +7,12 @@
 #include "Render/Viewport/Containers/DisplayClusterViewport_Enums.h"
 
 class FRDGBuilder;
+class FSceneView;
 class FSceneViewFamily;
 class IDisplayClusterViewport;
 class IDisplayClusterViewportProxy;
+struct FPostProcessMaterialInputs;
+struct FScreenPassTexture;
 
 
 /**
@@ -19,7 +22,14 @@ class FDisplayClusterMediaCaptureViewport
 	: public FDisplayClusterMediaCaptureBase
 {
 public:
-	FDisplayClusterMediaCaptureViewport(const FString& MediaId, const FString& ClusterNodeId, const FString& ViewportId, UMediaOutput* MediaOutput, UDisplayClusterMediaOutputSynchronizationPolicy* SyncPolicy = nullptr);
+	FDisplayClusterMediaCaptureViewport(
+		const FString& MediaId,
+		const FString& ClusterNodeId,
+		const FString& ViewportId,
+		UMediaOutput* MediaOutput,
+		UDisplayClusterMediaOutputSynchronizationPolicy* SyncPolicy = nullptr,
+		bool bInLateOCIO = false
+	);
 
 public:
 	/** Start capturing */
@@ -34,8 +44,7 @@ public:
 		return ViewportId;
 	}
 
-protected:
-	/** Returns texture size of a viewport assigned to capture */
+	/** Returns texture size of a viewport assigned to capture (main thread) */
 	virtual FIntPoint GetCaptureSize() const override;
 
 	/** Provides default texture size from config */
@@ -45,15 +54,20 @@ protected:
 	bool GetCaptureSizeFromGameProxy(FIntPoint& OutSize) const;
 
 private:
-	/** PostRenderViewFamily callback handler where data is captured */
-	void OnPostRenderViewFamily_RenderThread(FRDGBuilder& GraphBuilder, const FSceneViewFamily& ViewFamily, const IDisplayClusterViewportProxy* ViewportProxy);
 
 	/** UpdateViewportMediaState callback to configure media state for a viewoprt */
 	void OnUpdateViewportMediaState(IDisplayClusterViewport* InViewport, EDisplayClusterViewportMediaState& InOutMediaState);
 
-public:
-	/** Force late OCIO pass */
-	bool bForceLateOCIOPass = false;
+	/** PostRenderViewFamily callback handler where data is captured (no late OCIO) */
+	void OnPostRenderViewFamily_RenderThread(FRDGBuilder& GraphBuilder, const FSceneViewFamily& ViewFamily, const IDisplayClusterViewportProxy* ViewportProxy);
+
+	/** PostTonemapPass callback handler (late OCIO) */
+	void OnPostTonemapPass_RenderThread(FRDGBuilder& GraphBuilder, const IDisplayClusterViewportProxy* ViewportProxy, const FSceneView& View, const FPostProcessMaterialInputs& Inputs, const uint32 ContextNum);
+
+private:
+
+	/** Returns size of the viewport bound to this media */
+	virtual FIntPoint GetViewportSize() const;
 
 private:
 	/** Viewport ID assigned to capture */
diff --git a/Engine/Plugins/Runtime/nDisplay/Source/DisplayClusterMedia/Private/DisplayClusterMediaBase.h b/Engine/Plugins/Runtime/nDisplay/Source/DisplayClusterMedia/Private/DisplayClusterMediaBase.h
--- a/Engine/Plugins/Runtime/nDisplay/Source/DisplayClusterMedia/Private/DisplayClusterMediaBase.h
+++ b/Engine/Plugins/Runtime/nDisplay/Source/DisplayClusterMedia/Private/DisplayClusterMediaBase.h
@@ -4,41 +4,51 @@
 
 #include "CoreMinimal.h"
 
-class FRHITexture;
+#include "OpenColorIOColorSpace.h"
 
 
 /**
- * Base media class
+ * Base media adapter class
  */
 class FDisplayClusterMediaBase
 {
 public:
-	FDisplayClusterMediaBase(const FString& InMediaId, const FString& InClusterNodeId)
+	FDisplayClusterMediaBase(const FString& InMediaId, const FString& InClusterNodeId, bool bInLateOCIO = false)
 		: MediaId(InMediaId)
 		, ClusterNodeId(InClusterNodeId)
+		, bLateOCIO(bInLateOCIO)
 	{ }
 
 	virtual ~FDisplayClusterMediaBase() = default;
 
 public:
+
+	/** Returns ID of this media adapter */
 	const FString& GetMediaId() const
 	{
 		return MediaId;
 	}
 
+	/** Returns current cluster node ID */
 	const FString& GetClusterNodeId() const
 	{
 		return ClusterNodeId;
 	}
 
-protected:
-	struct FMediaTextureInfo
+	/** Returns current late OCIO configuration/state */
+	bool IsLateOCIO() const
 	{
-		FRHITexture* Texture = nullptr;
-		FIntRect     Region;
-	};
+		return bLateOCIO;
+	}
 
 private:
+
+	/** ID of this media adapter */
 	const FString MediaId;
+
+	/** Cluster node ID we're running on */
 	const FString ClusterNodeId;
+
+	/** Is OCIO expected to be applied late on the receiver side */
+	const bool bLateOCIO = false;
 };
diff --git a/Engine/Plugins/Runtime/nDisplay/Source/DisplayClusterMedia/Private/DisplayClusterMediaCVars.h b/Engine/Plugins/Runtime/nDisplay/Source/DisplayClusterMedia/Private/DisplayClusterMediaCVars.h
--- a/Engine/Plugins/Runtime/nDisplay/Source/DisplayClusterMedia/Private/DisplayClusterMediaCVars.h
+++ b/Engine/Plugins/Runtime/nDisplay/Source/DisplayClusterMedia/Private/DisplayClusterMediaCVars.h
@@ -4,6 +4,9 @@
 
 #include "HAL/IConsoleManager.h"
 
-// Enable/disable media.
-// May be useful if media needs to be disabled without doing any changes to the nDisplay configuration asset.
+/**
+ * Enable/disable media.
+ * 
+ * May be useful if media needs to be disabled without doing any changes to the nDisplay configuration asset.
+ */
 extern TAutoConsoleVariable<bool> CVarMediaEnabled;
diff --git a/Engine/Plugins/Runtime/nDisplay/Source/DisplayClusterMedia/Private/DisplayClusterMediaHelpers.h b/Engine/Plugins/Runtime/nDisplay/Source/DisplayClusterMedia/Private/DisplayClusterMediaHelpers.h
--- a/Engine/Plugins/Runtime/nDisplay/Source/DisplayClusterMedia/Private/DisplayClusterMediaHelpers.h
+++ b/Engine/Plugins/Runtime/nDisplay/Source/DisplayClusterMedia/Private/DisplayClusterMediaHelpers.h
@@ -12,14 +12,14 @@ namespace DisplayClusterMediaHelpers
 {
 	namespace MediaId
 	{
-		// Helper for media ID generation
+		/** Helper for media ID generation */
 		enum class EMediaDeviceType : uint8
 		{
 			Input,
 			Output
 		};
 
-		// Helper for media ID generation
+		/** Helper for media ID generation */
 		enum class EMediaOwnerType : uint8
 		{
 			Backbuffer,
@@ -27,16 +27,16 @@ namespace DisplayClusterMediaHelpers
 			ICVFXCamera
 		};
 
-		// Generate media ID for a specific entity
+		/** Generate media ID for a specific entity */
 		FString GenerateMediaId(EMediaDeviceType DeviceType, EMediaOwnerType OwnerType, const FString& NodeId, const FString& DCRAName, const FString& OwnerName, uint8 Index, const FIntPoint* InTilePos = nullptr);
 	}
 
-	// Generates internal ICVFX viewport IDs
+	/** Generates internal ICVFX viewport IDs */
 	FString GenerateICVFXViewportName(const FString& ClusterNodeId, const FString& ICVFXCameraName);
 
-	// Generates viewport tile IDs
+	/** Generates viewport tile IDs */
 	FString GenerateTileViewportName(const FString& InViewportId, const FIntPoint& InTilePos);
 
-	// Copy and resize an RHI texture
+	/** Copy and resize an RHI texture */
 	void ResampleTexture_RenderThread(FRHICommandListImmediate& RHICmdList, FRHITexture* SrcTexture, FRHITexture* DstTexture, const FIntRect& SrcRect, const FIntRect& DstRect);
 }
diff --git a/Engine/Plugins/Runtime/nDisplay/Source/DisplayClusterMedia/Private/DisplayClusterMediaModule.cpp b/Engine/Plugins/Runtime/nDisplay/Source/DisplayClusterMedia/Private/DisplayClusterMediaModule.cpp
--- a/Engine/Plugins/Runtime/nDisplay/Source/DisplayClusterMedia/Private/DisplayClusterMediaModule.cpp
+++ b/Engine/Plugins/Runtime/nDisplay/Source/DisplayClusterMedia/Private/DisplayClusterMediaModule.cpp
@@ -456,9 +456,8 @@ void FDisplayClusterMediaModule::InitializeICVFXCameraFullFrameInput(const UDisp
 				TSharedPtr<FDisplayClusterMediaInputViewport> NewICVFXInput = MakeShared<FDisplayClusterMediaInputViewport>(
 					MediaInputId, ClusterNodeId,
 					ICVFXViewportId,
-					MediaSource);
-
-				NewICVFXInput->bForceLateOCIOPass = MediaSettings.bLateOCIOPass;
+					MediaSource,
+					MediaSettings.bLateOCIOPass);
 
 				InputViewports.Emplace(MediaInputId, MoveTemp(NewICVFXInput));
 			}
@@ -497,9 +496,8 @@ void FDisplayClusterMediaModule::InitializeICVFXCameraFullFrameOutput(const UDis
 						MediaCaptureId, ClusterNodeId,
 						ICVFXCameraName, ICVFXViewportId,
 						MediaOutputItem.MediaOutput,
-						MediaOutputItem.OutputSyncPolicy);
-
-					NewICVFXCapture->bForceLateOCIOPass = MediaSettings.bLateOCIOPass;
+						MediaOutputItem.OutputSyncPolicy,
+						MediaSettings.bLateOCIOPass);
 
 					CaptureViewports.Emplace(MediaCaptureId, MoveTemp(NewICVFXCapture));
 				}
@@ -548,9 +546,8 @@ void FDisplayClusterMediaModule::InitializeICVFXCameraUniformTilesInput(const UD
 						TSharedPtr<FDisplayClusterMediaInputViewport> NewICVFXTileInput = MakeShared<FDisplayClusterMediaInputViewport>(
 							MediaInputId, ClusterNodeId,
 							ICVFXViewportTileId,
-							MediaInputTile.MediaSource);
-
-						NewICVFXTileInput->bForceLateOCIOPass = MediaSettings.bLateOCIOPass;
+							MediaInputTile.MediaSource,
+							MediaSettings.bLateOCIOPass);
 
 						InputViewports.Emplace(MediaInputId, MoveTemp(NewICVFXTileInput));
 					}
@@ -600,9 +597,9 @@ void FDisplayClusterMediaModule::InitializeICVFXCameraUniformTilesOutput(const U
 						TSharedPtr<FDisplayClusterMediaCaptureTile> NewICVFXTileOutput = MakeShared<FDisplayClusterMediaCaptureTile>(
 							MediaOutputId, ClusterNodeId,
 							ICVFXViewportTileId,
-							MediaOutputTile.MediaOutput);
-
-						NewICVFXTileOutput->bForceLateOCIOPass = MediaSettings.bLateOCIOPass;
+							MediaOutputTile.MediaOutput,
+							nullptr, // no sync for tiles as it's basically for in-cluster use
+							MediaSettings.bLateOCIOPass);
 
 						CaptureViewports.Emplace(MediaOutputId, MoveTemp(NewICVFXTileOutput));
 					}
@@ -614,5 +611,4 @@ void FDisplayClusterMediaModule::InitializeICVFXCameraUniformTilesOutput(const U
 	}
 }
 
-
 IMPLEMENT_MODULE(FDisplayClusterMediaModule, DisplayClusterMedia);
diff --git a/Engine/Plugins/Runtime/nDisplay/Source/DisplayClusterMedia/Private/Input/DisplayClusterMediaInputBase.cpp b/Engine/Plugins/Runtime/nDisplay/Source/DisplayClusterMedia/Private/Input/DisplayClusterMediaInputBase.cpp
--- a/Engine/Plugins/Runtime/nDisplay/Source/DisplayClusterMedia/Private/Input/DisplayClusterMediaInputBase.cpp
+++ b/Engine/Plugins/Runtime/nDisplay/Source/DisplayClusterMedia/Private/Input/DisplayClusterMediaInputBase.cpp
@@ -11,6 +11,14 @@
 #include "TextureResource.h"
 #include "RHIUtilities.h"
 
+#include "OpenColorIOColorSpace.h"
+#include "OpenColorIORendering.h"
+
+#include "RenderGraphBuilder.h"
+#include "ScreenPass.h"
+
+#include "Engine/Engine.h"
+
 #include "UObject/UObjectGlobals.h"
 #include "UObject/Package.h"
 
@@ -40,8 +48,13 @@ TAutoConsoleVariable<int32> CVarTempRivermaxExtraPixelsRemove(
 	ECVF_RenderThreadSafe
 );
 
-FDisplayClusterMediaInputBase::FDisplayClusterMediaInputBase(const FString& InMediaId, const FString& InClusterNodeId, UMediaSource* InMediaSource)
-	: FDisplayClusterMediaBase(InMediaId, InClusterNodeId)
+FDisplayClusterMediaInputBase::FDisplayClusterMediaInputBase(
+	const FString& InMediaId,
+	const FString& InClusterNodeId,
+	UMediaSource* InMediaSource,
+	bool bInLateOCIO
+)
+	: FDisplayClusterMediaBase(InMediaId, InClusterNodeId, bInLateOCIO)
 {
 	checkSlow(InMediaSource);
 	MediaSource = DuplicateObject(InMediaSource, GetTransientPackage());
@@ -101,10 +114,13 @@ void FDisplayClusterMediaInputBase::Stop()
 		MediaPlayer->OnMediaEvent().RemoveAll(this);
 	}
 
+	// Release internals
+	ReleaseInternals();
+
 	bRunningRivermaxMedia = false;
 }
 
-void FDisplayClusterMediaInputBase::OverrideTextureRegions_RenderThread(FRHITexture* const SrcTexture, FRHITexture* const DstTexture, FIntRect& InOutSrcRect, FIntRect& InOutDstRect) const
+void FDisplayClusterMediaInputBase::OverrideTextureRegions_RenderThread(FIntRect& InOutSrcRect, FIntRect& InOutDstRect) const
 {
 	const FIntPoint SrcSize = InOutSrcRect.Size();
 	const FIntPoint DstSize = InOutDstRect.Size();
@@ -139,20 +155,41 @@ void FDisplayClusterMediaInputBase::OverrideTextureRegions_RenderThread(FRHIText
 	}
 }
 
-void FDisplayClusterMediaInputBase::ImportMediaData(FRHICommandListImmediate& RHICmdList, const FMediaTextureInfo& TextureInfo)
+void FDisplayClusterMediaInputBase::ReleaseInternals()
+{
+	OCIOAppliedTexture.SafeRelease();
+}
+
+void FDisplayClusterMediaInputBase::ImportMediaData_RenderThread(FRDGBuilder& GraphBuilder, const FMediaInputTextureInfo& TextureInfo)
 {
-	UE_LOG(LogDisplayClusterMedia, Verbose, TEXT("MediaInput '%s': importing texture on RT frame '%lu'..."), *GetMediaId(), GFrameCounterRenderThread);
+	UE_LOG(LogDisplayClusterMedia, Verbose, TEXT("MediaInput '%s': importing texture on RT frame '%llu'..."), *GetMediaId(), GFrameCounterRenderThread);
 
 	MediaTexture->JustInTimeRender();
 
-	FRHITexture* const SrcTexture = MediaTexture->GetResource() ? MediaTexture->GetResource()->GetTextureRHI() : nullptr;
+	FRHITexture* SrcTexture = MediaTexture->GetResource() ? MediaTexture->GetResource()->GetTextureRHI() : nullptr;
 	FRHITexture* const DstTexture = TextureInfo.Texture;
 
+	if (!SrcTexture || !DstTexture)
+	{
+		return;
+	}
+
+	// Apply OCIO if needed
+	if (IsLateOCIO())
+	{
+		const bool bApplied = ProcessLateOCIO(GraphBuilder, SrcTexture, TextureInfo.OCIOPassResources);
+		if (bApplied)
+		{
+			// Redirect SrcTexture to the intermediate OCIO texture so we'll be importing this new one
+			SrcTexture = OCIOAppliedTexture;
+		}
+	}
+
 	if (SrcTexture && DstTexture)
 	{
 		FIntRect SrcRect(FIntPoint(0, 0), SrcTexture->GetDesc().Extent);
 		FIntRect DstRect(TextureInfo.Region);
-		OverrideTextureRegions_RenderThread(SrcTexture, DstTexture, SrcRect, DstRect);
+		OverrideTextureRegions_RenderThread(SrcRect, DstRect);
 
 		const bool bSrcSrgb = EnumHasAnyFlags(SrcTexture->GetFlags(), TexCreate_SRGB);
 		const bool bDstSrgb = EnumHasAnyFlags(DstTexture->GetFlags(), TexCreate_SRGB);
@@ -166,15 +203,106 @@ void FDisplayClusterMediaInputBase::ImportMediaData(FRHICommandListImmediate& RH
 			CopyInfo.SourcePosition = FIntVector(SrcRect.Min.X, SrcRect.Min.Y, 0);
 			CopyInfo.DestPosition = FIntVector(DstRect.Min.X, DstRect.Min.Y, 0);
 			CopyInfo.Size = FIntVector(DstRect.Size().X, DstRect.Size().Y, 0);
-			TransitionAndCopyTexture(RHICmdList, SrcTexture, DstTexture, CopyInfo);
+
+			TransitionAndCopyTexture(GraphBuilder.RHICmdList, SrcTexture, DstTexture, CopyInfo);
 		}
 		else
 		{
-			DisplayClusterMediaHelpers::ResampleTexture_RenderThread(RHICmdList, SrcTexture, DstTexture, SrcRect, DstRect);
+			DisplayClusterMediaHelpers::ResampleTexture_RenderThread(GraphBuilder.RHICmdList, SrcTexture, DstTexture, SrcRect, DstRect);
 		}
 	}
 }
 
+bool FDisplayClusterMediaInputBase::ProcessLateOCIO(FRDGBuilder& GraphBuilder, FRHITexture* SrcTexture, const FOpenColorIORenderPassResources& OCIORenderPassResources)
+{
+	checkSlow(SrcTexture);
+
+	if (!SrcTexture)
+	{
+		return false;
+	}
+
+	if (!OCIORenderPassResources.IsValid())
+	{
+		return false;
+	}
+
+	// Create an intermediate texture if not exists
+	if (!OCIOAppliedTexture)
+	{
+		OCIOAppliedTexture = CreateTexture(SrcTexture);
+	}
+	// Re-create it if parameters don't match
+	else
+	{
+		const bool bSameFormat = (OCIOAppliedTexture->GetDesc().Format == SrcTexture->GetDesc().Format);
+		const bool bSameSize   = (OCIOAppliedTexture->GetDesc().Extent == SrcTexture->GetDesc().Extent);
+
+		if (!bSameFormat || !bSameSize)
+		{
+			OCIOAppliedTexture = CreateTexture(SrcTexture);
+		}
+	}
+
+	if (!OCIOAppliedTexture)
+	{
+		return false;
+	}
+
+	FRDGTextureRef InputTexture = RegisterExternalTexture(GraphBuilder, SrcTexture, TEXT("DCMediaLateOCIOTexIn"));
+	FRDGTextureRef OutputTexture = RegisterExternalTexture(GraphBuilder, OCIOAppliedTexture, TEXT("DCMediaLateOCIOTexOut"));
+
+	const FIntPoint OutputResolution = OCIOAppliedTexture->GetDesc().Extent;
+	const FIntRect  OutputRect = FIntRect(FIntPoint::ZeroValue, OutputResolution);
+
+	FScreenPassTexture Input = FScreenPassTexture(InputTexture);
+	FScreenPassRenderTarget Output = FScreenPassRenderTarget(OutputTexture, OutputRect, ERenderTargetLoadAction::EClear);
+
+	FOpenColorIORendering::AddPass_RenderThread(
+		GraphBuilder,
+		FScreenPassViewInfo(),
+		GEngine->GetDefaultWorldFeatureLevel(),
+		Input,
+		Output,
+		OCIORenderPassResources,
+		1.0f,
+		EOpenColorIOTransformAlpha::None
+	);
+
+	return true;
+}
+
+FTextureRHIRef FDisplayClusterMediaInputBase::CreateTexture(const FRHITexture* ReferenceTexture)
+{
+	if (!ReferenceTexture)
+	{
+		return nullptr;
+	}
+
+	// Use original format and size
+	const int32 SizeX = ReferenceTexture->GetDesc().Extent.X;
+	const int32 SizeY = ReferenceTexture->GetDesc().Extent.Y;
+	const EPixelFormat Format = ReferenceTexture->GetFormat();
+
+	// Prepare description
+	FRHITextureCreateDesc Desc =
+		FRHITextureCreateDesc::Create2D(TEXT("DisplayClusterFrameQueueCacheTexture"), SizeX, SizeY, Format)
+		.SetClearValue(FClearValueBinding::Black)
+		.SetNumMips(1)
+		.SetFlags(ETextureCreateFlags::Dynamic)
+		.AddFlags(ETextureCreateFlags::MultiGPUGraphIgnore)
+		.SetInitialState(ERHIAccess::SRVMask);
+
+	// Leave original flags, but make sure it's ResolveTargetable but not RenderTargetable
+	ETextureCreateFlags Flags = ReferenceTexture->GetFlags();
+	Flags &= ~ETextureCreateFlags::RenderTargetable;
+	Flags |= ETextureCreateFlags::ResolveTargetable;
+	Desc.SetFlags(Flags);
+
+	// Create texture
+	return RHICreateTexture(Desc);
+}
+
 void FDisplayClusterMediaInputBase::OnMediaEvent(EMediaEvent MediaEvent)
 {
 	switch (MediaEvent)
diff --git a/Engine/Plugins/Runtime/nDisplay/Source/DisplayClusterMedia/Private/Input/DisplayClusterMediaInputBase.h b/Engine/Plugins/Runtime/nDisplay/Source/DisplayClusterMedia/Private/Input/DisplayClusterMediaInputBase.h
--- a/Engine/Plugins/Runtime/nDisplay/Source/DisplayClusterMedia/Private/Input/DisplayClusterMediaInputBase.h
+++ b/Engine/Plugins/Runtime/nDisplay/Source/DisplayClusterMedia/Private/Input/DisplayClusterMediaInputBase.h
@@ -5,9 +5,11 @@
 #include "CoreMinimal.h"
 #include "DisplayClusterMediaBase.h"
 #include "IMediaEventSink.h"
+#include "OpenColorIORendering.h"
 #include "UObject/GCObject.h"
 
 class FRHICommandListImmediate;
+class FRHITexture;
 
 class UMediaSource;
 class UMediaPlayer;
@@ -22,7 +24,12 @@ class FDisplayClusterMediaInputBase
 	, public FGCObject
 {
 public:
-	FDisplayClusterMediaInputBase(const FString& MediaId, const FString& ClusterNodeId, UMediaSource* MediaSource);
+	FDisplayClusterMediaInputBase(
+		const FString& MediaId,
+		const FString& ClusterNodeId,
+		UMediaSource* MediaSource,
+		bool bInLateOCIO = false
+	);
 
 public:
 	//~ Begin FGCObject interface
@@ -34,40 +41,83 @@ public:
 	//~ End FGCObject interface
 
 public:
+
+	/** Start playback */
 	virtual bool Play();
+
+	/** Stop playback */
 	virtual void Stop();
 
 public:
+
+	/** Returns current media source */
 	UMediaSource* GetMediaSource() const
 	{
 		return MediaSource;
 	}
 
+	/** Returns current media player */
 	UMediaPlayer* GetMediaPlayer() const
 	{
 		return MediaPlayer;
 	}
 
+	/** Returns current media texture */
 	UMediaTexture* GetMediaTexture() const
 	{
 		return MediaTexture;
 	}
 
 protected:
-	void ImportMediaData(FRHICommandListImmediate& RHICmdList, const FMediaTextureInfo& TextureInfo);
+
+	/** Media playback data */
+	struct FMediaInputTextureInfo
+	{
+		/** Target texture for media input */
+		FRHITexture* Texture = nullptr;
+
+		/** Target subregion */
+		FIntRect Region = { FIntPoint::ZeroValue, FIntPoint::ZeroValue };
+
+		/** OpenColorIO render pass parameters */
+		FOpenColorIORenderPassResources OCIOPassResources;
+	};
+
+protected:
+
+	/** Imports texture from a media source */
+	void ImportMediaData_RenderThread(FRDGBuilder& GraphBuilder, const FMediaInputTextureInfo& TextureInfo);
 
 private:
+
+	/** Applies OCIO transformation to the SRC texture, and stores result into DST */
+	bool ProcessLateOCIO(FRDGBuilder& GraphBuilder, FRHITexture* SrcTexture, const FOpenColorIORenderPassResources& OCIORenderPassResources);
+
+	/** Creates an internal intermediate texture */
+	FTextureRHIRef CreateTexture(const FRHITexture* ReferenceTexture);
+
+private:
+
+	/** Media events root handler */
 	void OnMediaEvent(EMediaEvent MediaEvent);
+
+	/** Start playback. Used to restart playback after failure */
 	bool StartPlayer();
+
+	/** Media event handler. Called when media source is closed. */
 	void OnPlayerClosed();
 
-	void OverrideTextureRegions_RenderThread(FRHITexture* const SrcTexture, FRHITexture* const DstTexture, FIntRect& InOutSrcRect, FIntRect& InOutDstRect) const;
+	/** [TEMP] A temporary workaround to cut off extra pixels for Rivermax input streams */
+	void OverrideTextureRegions_RenderThread(FIntRect& InOutSrcRect, FIntRect& InOutDstRect) const;
+
+	/** Releases internal resources */
+	void ReleaseInternals();
 
 private:
 	//~ Begin GC by AddReferencedObjects
-	TObjectPtr<UMediaSource>  MediaSource = nullptr;
-	TObjectPtr<UMediaPlayer>  MediaPlayer = nullptr;
-	TObjectPtr<UMediaTexture> MediaTexture = nullptr;
+	TObjectPtr<UMediaSource>  MediaSource;
+	TObjectPtr<UMediaPlayer>  MediaPlayer;
+	TObjectPtr<UMediaTexture> MediaTexture;
 	//~ End GC by AddReferencedObjects
 
 	// Used to restart media player in the case it falls in error
@@ -76,6 +126,9 @@ private:
 	// Used to control the rate at which we try to restart the player
 	double LastRestartTimestamp = 0;
 
+	// Intermediate texture for late OCIO pass
+	FTextureRHIRef OCIOAppliedTexture;
+
 	// [Temp workaround] Whether current media is Rivermax
 	bool bRunningRivermaxMedia = false;
 };
diff --git a/Engine/Plugins/Runtime/nDisplay/Source/DisplayClusterMedia/Private/Input/DisplayClusterMediaInputViewport.cpp b/Engine/Plugins/Runtime/nDisplay/Source/DisplayClusterMedia/Private/Input/DisplayClusterMediaInputViewport.cpp
--- a/Engine/Plugins/Runtime/nDisplay/Source/DisplayClusterMedia/Private/Input/DisplayClusterMediaInputViewport.cpp
+++ b/Engine/Plugins/Runtime/nDisplay/Source/DisplayClusterMedia/Private/Input/DisplayClusterMediaInputViewport.cpp
@@ -1,6 +1,12 @@
 // Copyright Epic Games, Inc. All Rights Reserved.
 
 #include "Input/DisplayClusterMediaInputViewport.h"
+
+#include "Engine/Engine.h"
+#include "Engine/World.h"
+
+#include "Render/IDisplayClusterRenderManager.h"
+#include "Render/Viewport/IDisplayClusterViewportManager.h"
 #include "Render/Viewport/IDisplayClusterViewportManagerProxy.h"
 #include "Render/Viewport/IDisplayClusterViewport.h"
 
@@ -9,12 +15,23 @@
 
 #include "MediaTexture.h"
 
+#include "OpenColorIOColorSpace.h"
+
+#include "RenderGraphUtils.h"
+
 #include "RHICommandList.h"
+#include "RHIFeatureLevel.h"
 #include "RHIResources.h"
 
 
-FDisplayClusterMediaInputViewport::FDisplayClusterMediaInputViewport(const FString& InMediaId, const FString& InClusterNodeId, const FString& InViewportId, UMediaSource* InMediaSource)
-	: FDisplayClusterMediaInputBase(InMediaId, InClusterNodeId, InMediaSource)
+FDisplayClusterMediaInputViewport::FDisplayClusterMediaInputViewport(
+	const FString& InMediaId,
+	const FString& InClusterNodeId,
+	const FString& InViewportId,
+	UMediaSource* InMediaSource,
+	bool bInLateOCIO
+)
+	: FDisplayClusterMediaInputBase(InMediaId, InClusterNodeId, InMediaSource, bInLateOCIO)
 	, ViewportId(InViewportId)
 {
 }
@@ -25,7 +42,8 @@ bool FDisplayClusterMediaInputViewport::Play()
 	// If playback has started successfully, subscribe for rendering callbacks
 	if (FDisplayClusterMediaInputBase::Play())
 	{
-		IDisplayCluster::Get().GetCallbacks().OnDisplayClusterPostCrossGpuTransfer_RenderThread().AddRaw(this, &FDisplayClusterMediaInputViewport::PostCrossGpuTransfer_RenderThread);
+		IDisplayCluster::Get().GetCallbacks().OnDisplayClusterPreSubmitViewFamilies().AddRaw(this, &FDisplayClusterMediaInputViewport::OnPreSubmitViewFamilies);
+		IDisplayCluster::Get().GetCallbacks().OnDisplayClusterPostCrossGpuTransfer_RenderThread().AddRaw(this, &FDisplayClusterMediaInputViewport::OnPostCrossGpuTransfer_RenderThread);
 		IDisplayCluster::Get().GetCallbacks().OnDisplayClusterUpdateViewportMediaState().AddRaw(this, &FDisplayClusterMediaInputViewport::OnUpdateViewportMediaState);
 
 		return true;
@@ -36,16 +54,47 @@ bool FDisplayClusterMediaInputViewport::Play()
 
 void FDisplayClusterMediaInputViewport::Stop()
 {
-	// Stop receiving notifications
+	// Unsubscribe from external events/callbacks
 	IDisplayCluster::Get().GetCallbacks().OnDisplayClusterPostCrossGpuTransfer_RenderThread().RemoveAll(this);
-
-	// Stop raising media flags for the viewport.
 	IDisplayCluster::Get().GetCallbacks().OnDisplayClusterUpdateViewportMediaState().RemoveAll(this);
 
 	// Stop playing
 	FDisplayClusterMediaInputBase::Stop();
 }
 
+void FDisplayClusterMediaInputViewport::OnPreSubmitViewFamilies(TArray<FSceneViewFamilyContext*>&)
+{
+	// Get OCIO settings if there are any
+	if (const IDisplayClusterViewportManager* const ViewportMgr = IDisplayCluster::Get().GetRenderMgr()->GetViewportManager())
+	{
+		if (const IDisplayClusterViewport* const Viewport = ViewportMgr->FindViewport(ViewportId))
+		{
+			FOpenColorIOColorConversionSettings OCIOConversionSettings;
+
+			const bool bSettingsAvailable = Viewport->GetOCIOConversionSettings(OCIOConversionSettings);
+			if (bSettingsAvailable)
+			{
+				const UWorld* const CurrentWorld = ViewportMgr->GetConfiguration().GetCurrentWorld();
+
+				const ERHIFeatureLevel::Type FeatureLevel = (CurrentWorld ?
+					CurrentWorld->GetFeatureLevel() :
+					GEngine->GetDefaultWorldFeatureLevel());
+
+				// Get OCIO render pass resources
+				FOpenColorIORenderPassResources OCIOPassResources = FOpenColorIORendering::GetRenderPassResources(OCIOConversionSettings, FeatureLevel);
+
+				// And push it to the rendering thread
+				ENQUEUE_RENDER_COMMAND(DCMediaInputUpdateOCIOResources)(
+					[this, InOCIOPassResources = MoveTemp(OCIOPassResources)](FRHICommandListImmediate& RHICmdList)
+					{
+						OCIOPassResources_RT = InOCIOPassResources;
+					}
+				);
+			}
+		}
+	}
+}
+
 void FDisplayClusterMediaInputViewport::OnUpdateViewportMediaState(IDisplayClusterViewport* InViewport, EDisplayClusterViewportMediaState& InOutMediaState)
 {
 	// Note: Media currently supports only one DCRA.
@@ -54,16 +103,10 @@ void FDisplayClusterMediaInputViewport::OnUpdateViewportMediaState(IDisplayClust
 	{
 		// Raise flags that this viewport texture will be overridden by media.
 		InOutMediaState |= EDisplayClusterViewportMediaState::Input;
-
-		if (bForceLateOCIOPass)
-		{
-			// Raise flags that this viewport requires ForceLateOCIOPass.
-			InOutMediaState |= EDisplayClusterViewportMediaState::Input_ForceLateOCIOPass;
-		}
 	}
 }
 
-void FDisplayClusterMediaInputViewport::PostCrossGpuTransfer_RenderThread(FRHICommandListImmediate& RHICmdList, const IDisplayClusterViewportManagerProxy* ViewportManagerProxy, FViewport* Viewport)
+void FDisplayClusterMediaInputViewport::OnPostCrossGpuTransfer_RenderThread(FRHICommandListImmediate& RHICmdList, const IDisplayClusterViewportManagerProxy* ViewportManagerProxy, FViewport* Viewport)
 {
 	checkSlow(ViewportManagerProxy);
 
@@ -78,10 +121,17 @@ void FDisplayClusterMediaInputViewport::PostCrossGpuTransfer_RenderThread(FRHICo
 
 			if (PlaybackViewport->GetResourcesWithRects_RenderThread(EDisplayClusterViewportResourceType::InternalRenderTargetResource, Textures, Regions))
 			{
-				if (Textures.Num() > 0 && Regions.Num() > 0)
+				if (Textures.Num() > 0 && Regions.Num() > 0 && Textures[0])
 				{
-					FMediaTextureInfo TextureInfo{ Textures[0], Regions[0] };
-					ImportMediaData(RHICmdList, TextureInfo);
+					FRDGBuilder GraphBuilder(RHICmdList);
+
+					// Prepare request data
+					FMediaInputTextureInfo TextureInfo{ Textures[0], Regions[0], MoveTemp(OCIOPassResources_RT) };
+
+					// Import texture from media input
+					ImportMediaData_RenderThread(GraphBuilder, TextureInfo);
+
+					GraphBuilder.Execute();
 				}
 			}
 		}
diff --git a/Engine/Plugins/Runtime/nDisplay/Source/DisplayClusterMedia/Private/Input/DisplayClusterMediaInputViewport.h b/Engine/Plugins/Runtime/nDisplay/Source/DisplayClusterMedia/Private/Input/DisplayClusterMediaInputViewport.h
--- a/Engine/Plugins/Runtime/nDisplay/Source/DisplayClusterMedia/Private/Input/DisplayClusterMediaInputViewport.h
+++ b/Engine/Plugins/Runtime/nDisplay/Source/DisplayClusterMedia/Private/Input/DisplayClusterMediaInputViewport.h
@@ -5,20 +5,29 @@
 #include "Input/DisplayClusterMediaInputBase.h"
 #include "Render/Viewport/Containers/DisplayClusterViewport_Enums.h"
 
+#include "OpenColorIORendering.h"
+
 class FRHICommandListImmediate;
 class FViewport;
+class FSceneViewFamilyContext;
 class IDisplayClusterViewport;
 class IDisplayClusterViewportManagerProxy;
 
 
 /**
- * Viewport media input
+ * Viewport media input adapter
  */
 class FDisplayClusterMediaInputViewport
 	: public FDisplayClusterMediaInputBase
 {
 public:
-	FDisplayClusterMediaInputViewport(const FString& MediaId, const FString& ClusterNodeId, const FString& ViewportId, UMediaSource* MediaSource);
+	FDisplayClusterMediaInputViewport(
+		const FString& MediaId,
+		const FString& ClusterNodeId,
+		const FString& ViewportId,
+		UMediaSource* MediaSource,
+		bool bInLateOCIO = false
+	);
 
 public:
 	/** Start playback */
@@ -34,17 +43,21 @@ public:
 	}
 
 private:
+
+	/** PreSubmitViewFamilies event handler. It's used to initialize media on start. */
+	void OnPreSubmitViewFamilies(TArray<FSceneViewFamilyContext*>&);
+
 	/** PostCrossGpuTransfer callback handler where media data is pushed into nDisplay internal buffers */
-	void PostCrossGpuTransfer_RenderThread(FRHICommandListImmediate& RHICmdList, const IDisplayClusterViewportManagerProxy* ViewportManagerProxy, FViewport* Viewport);
+	void OnPostCrossGpuTransfer_RenderThread(FRHICommandListImmediate& RHICmdList, const IDisplayClusterViewportManagerProxy* ViewportManagerProxy, FViewport* Viewport);
 
 	/** UpdateViewportMediaState callback to configure media state for a viewoprt */
 	void OnUpdateViewportMediaState(IDisplayClusterViewport* InViewport, EDisplayClusterViewportMediaState& InOutMediaState);
 
-public:
-	/** Force late OCIO pass */
-	bool bForceLateOCIOPass = false;
-
 private:
+
 	/** Viewport ID assigned for this media input */
 	const FString ViewportId;
+
+	/** OCIO conversion pass resources (render thread data) */
+	FOpenColorIORenderPassResources OCIOPassResources_RT;
 };
