nDisplay: Fixed a crash that occurs when closing a UE application with the ICVFX panel open.

diff --git a/Engine/Plugins/Runtime/nDisplay/Source/DisplayCluster/Private/Render/Viewport/Preview/DisplayClusterViewportManagerPreview.cpp b/Engine/Plugins/Runtime/nDisplay/Source/DisplayCluster/Private/Render/Viewport/Preview/DisplayClusterViewportManagerPreview.cpp
--- a/Engine/Plugins/Runtime/nDisplay/Source/DisplayCluster/Private/Render/Viewport/Preview/DisplayClusterViewportManagerPreview.cpp
+++ b/Engine/Plugins/Runtime/nDisplay/Source/DisplayCluster/Private/Render/Viewport/Preview/DisplayClusterViewportManagerPreview.cpp
@@ -84,7 +84,7 @@ void FDisplayClusterViewportManagerPreview::OnPreviewRenderTick()
 
 	// Special case for RootActorProxy object (proxy always refers to external DCRA in the scene)
 	ADisplayClusterRootActor* RootActorProxy = Configuration->GetRootActor(EDisplayClusterRootActorType::Preview);
-	if (RootActorProxy != SceneRootActor)
+	if (RootActorProxy && RootActorProxy != SceneRootActor)
 	{
 		// Since this DCRA is a proxy and the DCRA on the scene is a separate object,
 		// it does not pass component positions and properties from the scene to the proxy.
