nDisplay: Force media settings to use archetype values.

diff --git a/Engine/Plugins/Runtime/nDisplay/Source/DisplayCluster/Private/Game/EngineClasses/Scene/DisplayClusterRootActorEditor.cpp b/Engine/Plugins/Runtime/nDisplay/Source/DisplayCluster/Private/Game/EngineClasses/Scene/DisplayClusterRootActorEditor.cpp
--- a/Engine/Plugins/Runtime/nDisplay/Source/DisplayCluster/Private/Game/EngineClasses/Scene/DisplayClusterRootActorEditor.cpp
+++ b/Engine/Plugins/Runtime/nDisplay/Source/DisplayCluster/Private/Game/EngineClasses/Scene/DisplayClusterRootActorEditor.cpp
@@ -9,28 +9,18 @@
 #include "Components/DisplayClusterSceneComponentSyncParent.h"
 #include "Components/DisplayClusterScreenComponent.h"
 #include "Components/DisplayClusterStageGeometryComponent.h"
-#include "Components/MeshComponent.h"
-#include "Components/SceneComponent.h"
-
+#include "Components/DisplayClusterStageIsosphereComponent.h"
 #include "Config/IPDisplayClusterConfigManager.h"
 #include "DisplayClusterConfigurationStrings.h"
-
 #include "DisplayClusterConfigurationTypes.h"
+#include "DisplayClusterPlayerInput.h"
 #include "DisplayClusterProjectionStrings.h"
 #include "IDisplayClusterConfiguration.h"
-
-#include "DisplayClusterPlayerInput.h"
-
-#include "Engine/Engine.h"
-#include "Engine/World.h"
-#include "GameFramework/PlayerController.h"
-
 #include "Misc/DisplayClusterGlobals.h"
 #include "Misc/DisplayClusterHelpers.h"
 #include "Misc/DisplayClusterLog.h"
 #include "Misc/DisplayClusterStrings.h"
-
-#include "Misc/TransactionObjectEvent.h"
+#include "Render/Viewport/DisplayClusterViewportManager.h"
 #include "Render/Viewport/DisplayClusterViewportStrings.h"
 #include "Render/Viewport/IDisplayClusterViewport.h"
 #include "Render/Viewport/IDisplayClusterViewportManager.h"
@@ -38,13 +28,19 @@
 #include "Render/Viewport/RenderFrame/DisplayClusterRenderFrame.h"
 #include "Render/Viewport/RenderFrame/DisplayClusterRenderFrameSettings.h"
 
+#include "Components/MeshComponent.h"
+#include "Components/SceneComponent.h"
+#include "Engine/BlueprintGeneratedClass.h"
+#include "Engine/Engine.h"
 #include "Engine/TextureRenderTarget2D.h"
-
+#include "Engine/World.h"
+#include "GameFramework/PlayerController.h"
 #include "Materials/Material.h"
-#include "Render/Viewport/DisplayClusterViewportManager.h"
-#include "Render/Viewport/IDisplayClusterViewport.h"
+#include "Misc/TransactionObjectEvent.h"
 #include "TextureResource.h"
-#include "Components/DisplayClusterStageIsosphereComponent.h"
+#include "UObject/UnrealType.h"
+
+
 
 //////////////////////////////////////////////////////////////////////////////////////////////
 // IN-EDITOR STUFF
@@ -258,7 +254,172 @@ namespace UE::DisplayCluster::RootActor_Editor
 			PropagateDefaultMapToInstancedMap_Editor(ViewportsMapProperty, ClusterKeyVal.Value, DestinationValue);
 		}
 	}
+
+
+	/**
+	 * Syncs default media settings to instance.
+	 *
+	 * @param InstanceDCRA An instance of the nDisplay root actor.
+	 */
+	static void PropagateMediaSettingsFromDefault_Editor(ADisplayClusterRootActor* InstanceDCRA)
+	{
+		if (!InstanceDCRA)
+		{
+			return;
+		}
+
+		// Propagate ICVFX Camera Component Media property
+		//
+		{
+			TArray<UDisplayClusterICVFXCameraComponent*> InstanceCameras;
+			InstanceDCRA->GetComponents<UDisplayClusterICVFXCameraComponent>(InstanceCameras);
+
+			for (UDisplayClusterICVFXCameraComponent* InstanceCamera : InstanceCameras)
+			{
+				if (!InstanceCamera)
+				{
+					continue;
+				}
+
+				const UDisplayClusterICVFXCameraComponent* CameraArchetype = Cast<UDisplayClusterICVFXCameraComponent>(InstanceCamera->GetArchetype());
+
+				if (!CameraArchetype)
+				{
+					continue;
+				}
+
+				// We do not expect to ever pass an archetype instance to this function.
+				check(InstanceCamera != CameraArchetype);
+
+				// Deep-copy the media settings from the ICVFX Camera archetype to the instance.
+
+				const FProperty* MediaProperty = FindFProperty<FProperty>(
+					FDisplayClusterConfigurationICVFX_CameraRenderSettings::StaticStruct(), 
+					GET_MEMBER_NAME_CHECKED(FDisplayClusterConfigurationICVFX_CameraRenderSettings, Media)
+				);
+
+				check(MediaProperty);
+
+				void* MediaInstancePtr = &InstanceCamera->CameraSettings.RenderSettings.Media;
+				const void* MediaArchetypePtr = &CameraArchetype->CameraSettings.RenderSettings.Media;
+
+				// Perform a deep copy of the Media property
+				MediaProperty->CopyCompleteValue(MediaInstancePtr /*Dest*/, MediaArchetypePtr /*Src*/);
+			}
+		}
+
+		// Propagate Backbuffer and Viewports MediaSettings property
+		//
+
+		const UDisplayClusterBlueprint* DCBP = Cast<UDisplayClusterBlueprint>(UBlueprint::GetBlueprintFromClass(InstanceDCRA->GetClass()));
+
+		if (!DCBP)
+		{
+			return;
+		}
+
+		const UDisplayClusterConfigurationData* const ArchetypeConfig = DCBP->GetConfig();
+		const UDisplayClusterConfigurationData* const InstanceConfig = InstanceDCRA->GetConfigData();
+
+		if (!IsValid(ArchetypeConfig)
+			|| !IsValid(InstanceConfig)
+			|| !IsValid(ArchetypeConfig->Cluster)
+			|| !IsValid(InstanceConfig->Cluster))
+		{
+			return;
+		}
+
+		// Iterate over instance nodes, find the archetype, and deep-copy the Media setting.
+		for (auto ItNode = InstanceConfig->Cluster->Nodes.CreateIterator(); ItNode; ++ItNode)
+		{
+			const FString& NodeName = ItNode.Key();
+			TObjectPtr<UDisplayClusterConfigurationClusterNode> InstanceNode = ItNode.Value();
+
+			if (!InstanceNode)
+			{
+				continue;
+			}
+
+			TObjectPtr<UDisplayClusterConfigurationClusterNode>* ArchetypeNodePtr = ArchetypeConfig->Cluster->Nodes.Find(NodeName);
+
+			if (!ArchetypeNodePtr)
+			{
+				continue;
+			}
+
+			TObjectPtr<UDisplayClusterConfigurationClusterNode> ArchetypeNode = *ArchetypeNodePtr;
+
+			if (!ArchetypeNode)
+			{
+				continue;
+			}
+
+			// Get the property so that we can use its CopyCompleteValue function for the deep copy.
+
+			const FProperty* NodeMediaProperty = FindFProperty<FProperty>(
+				UDisplayClusterConfigurationClusterNode::StaticClass(),
+				GET_MEMBER_NAME_CHECKED(UDisplayClusterConfigurationClusterNode, MediaSettings)
+			);
+
+			check(NodeMediaProperty);
+
+			// Get the pointers to the property values on instance and archetype
+			void* NodeMediaInstancePtr = &InstanceNode->MediaSettings;
+			const void* NodeMediaArchetypePtr = &ArchetypeNode->MediaSettings;
+
+			// Perform a deep copy of the Media property
+			NodeMediaProperty->CopyCompleteValue(NodeMediaInstancePtr /*Dest*/, NodeMediaArchetypePtr /*Src*/);
+
+			// Viewports Media property
+			//
+
+			// Iterate over instance viewports, find the archetype, and deep-copy the Media setting.
+			for (auto ItViewport = InstanceNode->Viewports.CreateIterator(); ItViewport; ++ItViewport)
+			{
+				const FString& ViewportName = ItViewport.Key();
+				TObjectPtr<UDisplayClusterConfigurationViewport> InstanceViewport = ItViewport.Value();
+
+				if (!InstanceViewport)
+				{
+					continue;
+				}
+
+				TObjectPtr<UDisplayClusterConfigurationViewport>* ArchetypeViewportPtr = ArchetypeNode->Viewports.Find(ViewportName);
+
+				if (!ArchetypeViewportPtr)
+				{
+					continue;
+				}
+
+				TObjectPtr<UDisplayClusterConfigurationViewport> ArchetypeViewport = *ArchetypeViewportPtr;
+
+				if (!ArchetypeViewport)
+				{
+					continue;
+				}
+				
+				// Get the property so that we can use its CopyCompleteValue function for the deep copy.
+				
+				const FProperty* ViewportMediaProperty = FindFProperty<FProperty>(
+					FDisplayClusterConfigurationViewport_RenderSettings::StaticStruct(),
+					GET_MEMBER_NAME_CHECKED(FDisplayClusterConfigurationViewport_RenderSettings, Media)
+				);
+
+				check(ViewportMediaProperty);
+
+				// Get the pointers to the property values on instance and archetype
+				void* ViewportMediaInstancePtr = &InstanceViewport->RenderSettings.Media;
+				const void* ViewportMediaArchetypePtr = &ArchetypeViewport->RenderSettings.Media;
+
+				// Perform a deep copy of the Media property
+				ViewportMediaProperty->CopyCompleteValue(ViewportMediaInstancePtr /*Dest*/, ViewportMediaArchetypePtr /*Src*/);
+			}
+		}
+
+	}
 };
+
+
 //////////////////////////////////////////////////////////////////////////////////////////////
 // ADisplayClusterRootActor
 //////////////////////////////////////////////////////////////////////////////////////////////
@@ -282,6 +443,12 @@ void ADisplayClusterRootActor::RerunConstructionScripts()
 				UDisplayClusterConfigurationData* DefaultData = CDO->GetConfigData();
 				PropagateDataFromDefaultConfig_Editor(DefaultData, CurrentData);
 			}
+
+			// We propagate the media settings to the instances to keep them in sync.
+			// This is not strictly necessary since the custom serializers make sure the intances always
+			// load archetype defaults but this makes it so they are always in sync even when
+			// the archetype is being edited, since the instances get reconstructed but not re-deserialized.
+			PropagateMediaSettingsFromDefault_Editor(this);
 		}
 		RerunConstructionScripts_Editor();
 	}
diff --git a/Engine/Plugins/Runtime/nDisplay/Source/DisplayCluster/Private/Render/Viewport/Configuration/DisplayClusterViewportConfigurationHelpers_ICVFX.cpp b/Engine/Plugins/Runtime/nDisplay/Source/DisplayCluster/Private/Render/Viewport/Configuration/DisplayClusterViewportConfigurationHelpers_ICVFX.cpp
--- a/Engine/Plugins/Runtime/nDisplay/Source/DisplayCluster/Private/Render/Viewport/Configuration/DisplayClusterViewportConfigurationHelpers_ICVFX.cpp
+++ b/Engine/Plugins/Runtime/nDisplay/Source/DisplayCluster/Private/Render/Viewport/Configuration/DisplayClusterViewportConfigurationHelpers_ICVFX.cpp
@@ -41,9 +41,6 @@
 
 #include "Components/DisplayClusterICVFXCameraComponent.h"
 
-#include "Engine/Blueprint.h"
-#include "Engine/SCS_Node.h"
-#include "Engine/SimpleConstructionScript.h"
 #include "HAL/IConsoleManager.h"
 #include "Misc/DisplayClusterLog.h"
 #include "TextureResource.h"
@@ -548,30 +545,8 @@ void FDisplayClusterViewportConfigurationHelpers_ICVFX::UpdateCameraViewportSett
 	// Set viewport buffer ratio
 	DstViewport.SetViewportBufferRatio(InCameraSettings.GetCameraBufferRatio(*StageSettings));
 
-#if WITH_EDITORONLY_DATA
-	// UE-211513
-	// This is a temporary workaround that allows to initialize media based on the blueprint data instead of 
-	// the instance data. There is an issue with the propagation of instanced object changes from the parent
-	// blueprints to its instances. To avoid any potential issues, we get media settings for the original nDisplay
-	// blueprint. Once the propagation issue is fixed, we'll be able use the instances again.
-
-	if (const AActor* const OwningActor = InCameraComponent.GetOwner())
-	{
-		if (const UBlueprint* const ParentBP = UBlueprint::GetBlueprintFromClass(OwningActor->GetClass()))
-		{
-			if (const USCS_Node* const CameraSCSNode = ParentBP->SimpleConstructionScript->FindSCSNode(*InCameraComponent.GetName()))
-			{
-				if (const UDisplayClusterICVFXCameraComponent* const ArchetypeCameraComponent = Cast<UDisplayClusterICVFXCameraComponent>(CameraSCSNode->ComponentTemplate))
-				{
-					FDisplayClusterViewportConfigurationHelpers_Tile::UpdateICVFXCameraViewportTileSettings(DstViewport, ArchetypeCameraComponent->CameraSettings.RenderSettings.Media);
-				}
-			}
-		}
-	}
-#else
 	// InCamera tile rendering.
 	FDisplayClusterViewportConfigurationHelpers_Tile::UpdateICVFXCameraViewportTileSettings(DstViewport, InCameraSettings.RenderSettings.Media);
-#endif // WITH_EDITORONLY_DATA
 }
 
 void FDisplayClusterViewportConfigurationHelpers_ICVFX::UpdateChromakeyViewportSettings(FDisplayClusterViewport& DstViewport, FDisplayClusterViewport& InCameraViewport, const FDisplayClusterConfigurationICVFX_CameraSettings& InCameraSettings)
diff --git a/Engine/Plugins/Runtime/nDisplay/Source/DisplayClusterConfiguration/Private/DisplayClusterConfigurationTypes.cpp b/Engine/Plugins/Runtime/nDisplay/Source/DisplayClusterConfiguration/Private/DisplayClusterConfigurationTypes.cpp
--- a/Engine/Plugins/Runtime/nDisplay/Source/DisplayClusterConfiguration/Private/DisplayClusterConfigurationTypes.cpp
+++ b/Engine/Plugins/Runtime/nDisplay/Source/DisplayClusterConfiguration/Private/DisplayClusterConfigurationTypes.cpp
@@ -2,14 +2,14 @@
 
 #include "DisplayClusterConfigurationTypes.h"
 
+#include "DisplayClusterConfigurationLog.h"
+#include "DisplayClusterConfigurationStrings.h"
 #include "DisplayClusterConfigurationTypes_Base.h"
 #include "DisplayClusterConfigurationTypes_Viewport.h"
 #include "DisplayClusterConfigurationTypes_ICVFX.h"
 #include "DisplayClusterConfigurationTypes_PostRender.h"
-
-#include "DisplayClusterConfigurationLog.h"
-#include "DisplayClusterConfigurationStrings.h"
 #include "DisplayClusterProjectionStrings.h"
+#include "DisplayClusterPropertySkipperArchive.h"
 
 #include "Engine/StaticMesh.h"
 #include "UObject/Package.h"
@@ -30,7 +30,7 @@
 
 #define SAVE_MAP(Map) \
 	SAVE_MAP_TO_ARRAY(Map, OutObjects); \
-	
+
 
 FIntRect FDisplayClusterReplaceTextureCropRectangle::ToRect() const
 {
@@ -182,7 +182,6 @@ FDisplayClusterConfigurationProjection::FDisplayClusterConfigurationProjection()
 	Type = TEXT("simple");
 }
 
-
 constexpr float UDisplayClusterConfigurationViewport::ViewportMinimumSize = 1.0f;
 constexpr float UDisplayClusterConfigurationViewport::ViewportMaximumSize = 15360.0f;
 
@@ -337,6 +336,16 @@ void UDisplayClusterConfigurationClusterNode::PostLoad()
 #endif // WITH_EDITOR
 }
 
+void UDisplayClusterConfigurationClusterNode::Serialize(FArchive& Ar)
+{
+	// Use our custom archive to skip serializing Media except for archetypes.
+
+	FDisplayClusterPropertySkipperArchive PropertySkipperAr = FDisplayClusterPropertySkipperArchive(Ar);
+	check(PropertySkipperAr.AddPropertyToSkip(StaticClass(), GET_MEMBER_NAME_CHECKED(UDisplayClusterConfigurationClusterNode, MediaSettings)));
+
+	return Super::Serialize(PropertySkipperAr);
+}
+
 UDisplayClusterConfigurationClusterNode::UDisplayClusterConfigurationClusterNode()
 	: bIsSoundEnabled(false)
 #if WITH_EDITORONLY_DATA
diff --git a/Engine/Plugins/Runtime/nDisplay/Source/DisplayClusterConfiguration/Private/DisplayClusterConfigurationTypes_ICVFX.cpp b/Engine/Plugins/Runtime/nDisplay/Source/DisplayClusterConfiguration/Private/DisplayClusterConfigurationTypes_ICVFX.cpp
--- a/Engine/Plugins/Runtime/nDisplay/Source/DisplayClusterConfiguration/Private/DisplayClusterConfigurationTypes_ICVFX.cpp
+++ b/Engine/Plugins/Runtime/nDisplay/Source/DisplayClusterConfiguration/Private/DisplayClusterConfigurationTypes_ICVFX.cpp
@@ -2,6 +2,7 @@
 
 #include "DisplayClusterConfigurationTypes_ICVFX.h"
 #include "DisplayClusterConfigurationTypes.h"
+#include "DisplayClusterPropertySkipperArchive.h"
 #include "IDisplayCluster.h"
 #include "Camera/CameraTypes.h"
 #include "CineCameraComponent.h"
@@ -20,6 +21,7 @@ namespace UE::DisplayClusterConfiguration::ICVFX
 		return FMath::Clamp(InValue, -InMax, InMax);
 	}
 };
+
 using namespace UE::DisplayClusterConfiguration::ICVFX;
 
 int32 GDisplayClusterICVFXCameraAdoptResolution = 1;
@@ -59,6 +61,21 @@ void FDisplayClusterConfigurationICVFX_CameraRenderSettings::SetupViewInfo(const
 	}
 }
 
+
+bool FDisplayClusterConfigurationICVFX_CameraRenderSettings::Serialize(FArchive& Ar)
+{	
+	// Use our custom archive to skip serializing Media except for archetypes.
+
+	UScriptStruct& Struct = *StaticStruct();
+
+	FDisplayClusterPropertySkipperArchive PropertySkipperAr = FDisplayClusterPropertySkipperArchive(Ar);
+	check(PropertySkipperAr.AddPropertyToSkip(&Struct, GET_MEMBER_NAME_CHECKED(FDisplayClusterConfigurationICVFX_CameraRenderSettings, Media)));
+
+	Struct.SerializeTaggedProperties(PropertySkipperAr, (uint8*)this, &Struct, nullptr);
+
+	return true;
+}
+
 ///////////////////////////////////////////////////////////////////////////////////////
 // FDisplayClusterConfigurationICVFX_LightcardCustomOCIO
 ///////////////////////////////////////////////////////////////////////////////////////
diff --git a/Engine/Plugins/Runtime/nDisplay/Source/DisplayClusterConfiguration/Private/DisplayClusterConfigurationTypes_Viewport.cpp b/Engine/Plugins/Runtime/nDisplay/Source/DisplayClusterConfiguration/Private/DisplayClusterConfigurationTypes_Viewport.cpp
--- a/Engine/Plugins/Runtime/nDisplay/Source/DisplayClusterConfiguration/Private/DisplayClusterConfigurationTypes_Viewport.cpp
+++ b/Engine/Plugins/Runtime/nDisplay/Source/DisplayClusterConfiguration/Private/DisplayClusterConfigurationTypes_Viewport.cpp
@@ -2,6 +2,7 @@
 
 #include "DisplayClusterConfigurationTypes_Viewport.h"
 #include "DisplayClusterConfigurationTypes_ICVFX.h"
+#include "DisplayClusterPropertySkipperArchive.h"
 
 ///////////////////////////////////////////////////////////////////////////////////////
 // FDisplayClusterConfigurationViewport_ICVFX
@@ -92,3 +93,19 @@ EDisplayClusterViewportICVFXFlags FDisplayClusterConfigurationViewport_ICVFX::Ge
 
 	return OutFlags;
 }
+
+
+bool FDisplayClusterConfigurationViewport_RenderSettings::Serialize(FArchive& Ar)
+{
+	// Use our custom archive to skip serializing Media except for archetypes.
+
+	UScriptStruct& Struct = *StaticStruct();
+
+	FDisplayClusterPropertySkipperArchive PropertySkipperAr = FDisplayClusterPropertySkipperArchive(Ar);
+	check(PropertySkipperAr.AddPropertyToSkip(&Struct, GET_MEMBER_NAME_CHECKED(FDisplayClusterConfigurationViewport_RenderSettings, Media)));
+
+	Struct.SerializeTaggedProperties(PropertySkipperAr, (uint8*)this, &Struct, nullptr);
+
+	return true;
+}
+
diff --git a/Engine/Plugins/Runtime/nDisplay/Source/DisplayClusterConfiguration/Private/DisplayClusterPropertySkipperArchive.h b/Engine/Plugins/Runtime/nDisplay/Source/DisplayClusterConfiguration/Private/DisplayClusterPropertySkipperArchive.h
new file mode 100644
--- /dev/null
+++ b/Engine/Plugins/Runtime/nDisplay/Source/DisplayClusterConfiguration/Private/DisplayClusterPropertySkipperArchive.h
@@ -0,0 +1,79 @@
+// Copyright Epic Games, Inc. All Rights Reserved.
+
+#pragma once
+
+#include "Serialization/ArchiveProxy.h"
+#include "UObject/Object.h"
+#include "UObject/UObjectThreadContext.h"
+#include "UObject/UnrealType.h"
+
+
+/**
+* This is a custom archive to address the issues around instanced property propagation (UE-211513).
+* It skips loading the given property by name unless it is the archetype, such that the instances always get the 
+* archetype default values.
+*/
+class FDisplayClusterPropertySkipperArchive : public FArchiveProxy
+{
+
+public:
+
+	FDisplayClusterPropertySkipperArchive(FArchive& InInnerArchive)
+		: FArchiveProxy(InInnerArchive)
+	{
+	}
+
+	//~ Begin FArchive interface
+	virtual bool ShouldSkipProperty(const FProperty* Property) const override
+	{
+		// Skip loading the property unless it is the archetype.
+
+		bool bShouldSkip = false;
+
+		if (IsLoading() && PropertiesToSkip.Contains(MakePropertyData(Property)))
+		{
+			const FUObjectSerializeContext* SerializeContext = FUObjectThreadContext::Get().GetSerializeContext();
+			const UObject* SerializedObject = SerializeContext ? SerializeContext->SerializedObject : nullptr;
+
+			// We assume that if the SerializedObject is null, it is indicative of a template.
+			const bool bIsArchetype = SerializedObject ? SerializedObject->IsTemplate() : true;
+
+			bShouldSkip |= !bIsArchetype;
+		}
+
+		return bShouldSkip || FArchiveProxy::ShouldSkipProperty(Property);
+	}
+	//~ End FArchive interface
+
+	/** Adds a new property to skip. */
+	bool AddPropertyToSkip(UStruct* Struct, FName PropertyName)
+	{
+		if (const FProperty* Property = FindFProperty<FProperty>(Struct, PropertyName)) // FindFProperty handles null Structs.
+		{
+			PropertiesToSkip.Add(MakePropertyData(Property));
+			return true;
+		}
+
+		return false;
+	}
+
+private:
+
+	/** Creates the data needed for a property when matching properties to be skipped */
+	TPair<FName, FName> MakePropertyData(const FProperty* Property) const
+	{
+		if (!Property)
+		{
+			return TPair<FName, FName>(NAME_None, NAME_None);
+		}
+
+		const FName PropertyOwnerName = Property->GetOwnerUField() ? Property->GetOwnerUField()->GetFName() : NAME_None;
+
+		return TPair<FName, FName>(PropertyOwnerName, Property->GetName());
+	}
+
+private:
+
+	/** Names of properties to skip except for archetypes */
+	TSet<TPair<FName,FName>> PropertiesToSkip;
+};
\ No newline at end of file
diff --git a/Engine/Plugins/Runtime/nDisplay/Source/DisplayClusterConfiguration/Public/DisplayClusterConfigurationTypes.h b/Engine/Plugins/Runtime/nDisplay/Source/DisplayClusterConfiguration/Public/DisplayClusterConfigurationTypes.h
--- a/Engine/Plugins/Runtime/nDisplay/Source/DisplayClusterConfiguration/Public/DisplayClusterConfigurationTypes.h
+++ b/Engine/Plugins/Runtime/nDisplay/Source/DisplayClusterConfiguration/Public/DisplayClusterConfigurationTypes.h
@@ -306,6 +306,7 @@ public:
 	//~ Begin UObject interface
 public:
 	virtual void PostLoad() override;
+	virtual void Serialize(FArchive& Ar) override;
 
 #if WITH_EDITOR
 	virtual void PostEditChangeChainProperty(FPropertyChangedChainEvent& PropertyChangedEvent) override;
diff --git a/Engine/Plugins/Runtime/nDisplay/Source/DisplayClusterConfiguration/Public/DisplayClusterConfigurationTypes_ICVFX.h b/Engine/Plugins/Runtime/nDisplay/Source/DisplayClusterConfiguration/Public/DisplayClusterConfigurationTypes_ICVFX.h
--- a/Engine/Plugins/Runtime/nDisplay/Source/DisplayClusterConfiguration/Public/DisplayClusterConfigurationTypes_ICVFX.h
+++ b/Engine/Plugins/Runtime/nDisplay/Source/DisplayClusterConfiguration/Public/DisplayClusterConfigurationTypes_ICVFX.h
@@ -520,6 +520,16 @@ public:
 public:
 	/** Propagates general render related settings to the view info. */
 	void SetupViewInfo(const FDisplayClusterConfigurationICVFX_StageSettings& InStageSettings, FMinimalViewInfo& InOutViewInfo) const;
+
+	/** Custom serialization to skip Media on instances */
+	bool Serialize(FArchive& Ar);
+};
+
+/** Makes the given struct use the custom serializer */
+template<> struct TStructOpsTypeTraits<FDisplayClusterConfigurationICVFX_CameraRenderSettings> 
+	: public TStructOpsTypeTraitsBase2<FDisplayClusterConfigurationICVFX_CameraRenderSettings>
+{
+	enum { WithSerializer = true };
 };
 
 USTRUCT(BlueprintType)
diff --git a/Engine/Plugins/Runtime/nDisplay/Source/DisplayClusterConfiguration/Public/DisplayClusterConfigurationTypes_Viewport.h b/Engine/Plugins/Runtime/nDisplay/Source/DisplayClusterConfiguration/Public/DisplayClusterConfigurationTypes_Viewport.h
--- a/Engine/Plugins/Runtime/nDisplay/Source/DisplayClusterConfiguration/Public/DisplayClusterConfigurationTypes_Viewport.h
+++ b/Engine/Plugins/Runtime/nDisplay/Source/DisplayClusterConfiguration/Public/DisplayClusterConfigurationTypes_Viewport.h
@@ -112,8 +112,21 @@ public:
 	// [not implemented yet]
 	UPROPERTY()
 	int RenderFamilyGroup = -1;
+
+public:
+	/** Custom serialization to skip Media on instances */
+	bool Serialize(FArchive& Ar);
 };
 
+
+/** Makes the given struct use the custom serializer */
+template<> struct TStructOpsTypeTraits<FDisplayClusterConfigurationViewport_RenderSettings>
+: public TStructOpsTypeTraitsBase2<FDisplayClusterConfigurationViewport_RenderSettings>
+{
+	enum { WithSerializer = true };
+};
+
+
 UCLASS(Blueprintable)
 class DISPLAYCLUSTERCONFIGURATION_API UDisplayClusterConfigurationViewport
 	: public UDisplayClusterConfigurationData_Base
diff --git a/Engine/Plugins/Runtime/nDisplay/Source/DisplayClusterMedia/Private/DisplayClusterMediaModule.cpp b/Engine/Plugins/Runtime/nDisplay/Source/DisplayClusterMedia/Private/DisplayClusterMediaModule.cpp
--- a/Engine/Plugins/Runtime/nDisplay/Source/DisplayClusterMedia/Private/DisplayClusterMediaModule.cpp
+++ b/Engine/Plugins/Runtime/nDisplay/Source/DisplayClusterMedia/Private/DisplayClusterMediaModule.cpp
@@ -12,13 +12,9 @@
 #include "DisplayClusterRootActor.h"
 #include "DisplayClusterConfigurationTypes.h"
 
-#include "Blueprints/DisplayClusterBlueprint.h"
 #include "Components/DisplayClusterICVFXCameraComponent.h"
 #include "Config/IDisplayClusterConfigManager.h"
 #include "Cluster/IDisplayClusterClusterManager.h"
-#include "Engine/Blueprint.h"
-#include "Engine/SCS_Node.h"
-#include "Engine/SimpleConstructionScript.h"
 #include "Game/IDisplayClusterGameManager.h"
 
 #include "Capture/DisplayClusterMediaCaptureCamera.h"
@@ -54,16 +50,7 @@ void FDisplayClusterMediaModule::OnPreSubmitViewFamilies(TArray<FSceneViewFamily
 	// Unsubscribe after first call. Currently, media initialization is a one time procedure. No need to receive any further callbacks.
 	IDisplayCluster::Get().GetCallbacks().OnDisplayClusterPreSubmitViewFamilies().RemoveAll(this);
 
-#if WITH_EDITORONLY_DATA
-	// UE-211513
-	// This is a temporary workaround that allows to initialize media based on the blueprint data instead of 
-	// the instance data. There is an issue with the propagation of instanced object changes from the parent
-	// blueprints to its instances. To avoid any potential issues, we get media settings for the original nDisplay
-	// blueprint. Once the propagation issue is fixed, we'll be able to safely return to the original InitializeMedia().
-	InitializeMediaFromCDO();
-#else
 	InitializeMedia();
-#endif
 
 	StartCapture();
 	PlayMedia();
@@ -141,134 +128,6 @@ void FDisplayClusterMediaModule::InitializeMedia()
 	}
 }
 
-#if WITH_EDITORONLY_DATA
-void FDisplayClusterMediaModule::InitializeMediaFromCDO()
-{
-	// Runtime only for now
-	if (IDisplayCluster::Get().GetOperationMode() != EDisplayClusterOperationMode::Cluster)
-	{
-		UE_LOG(LogDisplayClusterMedia, Warning, TEXT("DisplayClusterMedia is available in 'cluster' operation mode only"));
-		return;
-	}
-
-	// Check if media enabled
-	if (!CVarMediaEnabled.GetValueOnGameThread())
-	{
-		UE_LOG(LogDisplayClusterMedia, Log, TEXT("nDisplay media subsytem is disabled by a cvar"));
-		return;
-	}
-
-	// Instantiate latency queue
-	FrameQueue.Init();
-
-	// Get active DCRA instance
-	const ADisplayClusterRootActor* const RootActor = IDisplayCluster::Get().GetGameMgr()->GetRootActor();
-	if (!IsValid(RootActor))
-	{
-		UE_LOG(LogDisplayClusterMedia, Warning, TEXT("Invalid DCRA object"));
-		return;
-	}
-
-	// Get original (parent) blueprint of this DCRA instance
-	const UDisplayClusterBlueprint* DCBP = nullptr;
-	if (const UBlueprint* BP = UBlueprint::GetBlueprintFromClass(RootActor->GetClass()))
-	{
-		DCBP = Cast<UDisplayClusterBlueprint>(BP);
-	}
-
-	// Aux data for media initialization
-	const FString ClusterNodeId = IDisplayCluster::Get().GetClusterMgr()->GetNodeId();
-	const FString RootActorName = RootActor->GetName();
-
-	// It's possible that we have DCBP==nullptr. For example, when running with portable configs, a pure DCRA
-	// actor gets spawned. There is no blueprint available in this case. But we don't expect any media in this
-	// case so stop initialization procedure.
-	// !DCBP is to prevent C6011 warning (potential dereferencing nullptr)
-	if (!IsValid(DCBP) || !DCBP)
-	{
-		UE_LOG(LogDisplayClusterMedia, Log, TEXT("Skipping media initialization for a pure DCRA instance '%s'"), *RootActorName);
-		return;
-	}
-
-	// Now we need to get the configuration data.
-	const UDisplayClusterConfigurationData* const ConfigData = DCBP->GetConfig();
-	if (!IsValid(ConfigData) || !IsValid(ConfigData->Cluster))
-	{
-		UE_LOG(LogDisplayClusterMedia, Warning, TEXT("Couldn't obtain the cluster configuration data"));
-		return;
-	}
-
-	// Find a proper cluster node configuration
-	const UDisplayClusterConfigurationClusterNode* const ClusterNodeCfg = ConfigData->Cluster->GetNode(ClusterNodeId);
-	if (!IsValid(ClusterNodeCfg))
-	{
-		UE_LOG(LogDisplayClusterMedia, Warning, TEXT("Couldn't obtain any configuration data for node '%s'"), *ClusterNodeId);
-		return;
-	}
-
-	// Node backbuffer media setup
-	{
-		InitializeBackbufferOutput(ClusterNodeCfg, RootActorName, ClusterNodeId);
-	}
-
-	// Viewports media setup
-	for (const TPair<FString, TObjectPtr<UDisplayClusterConfigurationViewport>>& ViewportIt : ClusterNodeCfg->Viewports)
-	{
-		InitializeViewportInput(ViewportIt.Value, ViewportIt.Key, RootActorName, ClusterNodeId);
-		InitializeViewportOutput(ViewportIt.Value, ViewportIt.Key, RootActorName, ClusterNodeId);
-	}
-
-	// ICVFX media setup
-	{
-		// Get all ICVFX camera components
-		TArray<UDisplayClusterICVFXCameraComponent*> ICVFXCameraComponents;
-		RootActor->GetComponents(ICVFXCameraComponents);
-
-		for (const UDisplayClusterICVFXCameraComponent* const ICVFXCameraInstanceComponent : ICVFXCameraComponents)
-		{
-			const UDisplayClusterICVFXCameraComponent* ICVFXCameraComponent = nullptr;
-			const FDisplayClusterConfigurationMediaICVFX* MediaSettings = nullptr;
-
-			const FString CameraName = ICVFXCameraInstanceComponent->GetName();
-
-			// Let's see if this camera comes from the blueprint. If so, we'll be using its media settings. Otherwise
-			// the camera instance will be used.
-			const USCS_Node* const CameraSCSNode = DCBP->SimpleConstructionScript->FindSCSNode(*CameraName);
-			if (IsValid(CameraSCSNode))
-			{
-				ICVFXCameraComponent = Cast<UDisplayClusterICVFXCameraComponent>(CameraSCSNode->ComponentTemplate);
-			}
-			else
-			{
-				ICVFXCameraComponent = ICVFXCameraInstanceComponent;
-			}
-
-			// Get media settings
-			if (IsValid(ICVFXCameraComponent))
-			{
-				MediaSettings = &ICVFXCameraComponent->CameraSettings.RenderSettings.Media;
-			}
-
-			if (MediaSettings)
-			{
-				// Full frame
-				if (MediaSettings->SplitType == EDisplayClusterConfigurationMediaSplitType::FullFrame)
-				{
-					InitializeICVFXCameraFullFrameInput(ICVFXCameraComponent, RootActorName, ClusterNodeId, CameraName);
-					InitializeICVFXCameraFullFrameOutput(ICVFXCameraComponent, RootActorName, ClusterNodeId, CameraName);
-				}
-				// Uniform tiles
-				else if (MediaSettings->SplitType == EDisplayClusterConfigurationMediaSplitType::UniformTiles)
-				{
-					InitializeICVFXCameraUniformTilesInput(ICVFXCameraComponent, RootActorName, ClusterNodeId, CameraName);
-					InitializeICVFXCameraUniformTilesOutput(ICVFXCameraComponent, RootActorName, ClusterNodeId, CameraName);
-				}
-			}
-		}
-	}
-}
-#endif // WITH_EDITORONLY_DATA
-
 void FDisplayClusterMediaModule::ReleaseMedia()
 {
 	StopCapture();
diff --git a/Engine/Plugins/Runtime/nDisplay/Source/DisplayClusterMedia/Private/DisplayClusterMediaModule.h b/Engine/Plugins/Runtime/nDisplay/Source/DisplayClusterMedia/Private/DisplayClusterMediaModule.h
--- a/Engine/Plugins/Runtime/nDisplay/Source/DisplayClusterMedia/Private/DisplayClusterMediaModule.h
+++ b/Engine/Plugins/Runtime/nDisplay/Source/DisplayClusterMedia/Private/DisplayClusterMediaModule.h
@@ -33,11 +33,6 @@ protected:
 	/** Initialize all media internals (from instance) */
 	void InitializeMedia();
 
-#if WITH_EDITORONLY_DATA
-	/** Initialize all media internals (from CDO) */
-	void InitializeMediaFromCDO();
-#endif // WITH_EDITORONLY_DATA
-
 	/** Release all media internals */
 	void ReleaseMedia();
 
