[nDisplay] Fixed NDC compilation crash after changing any cluster properties

diff --git a/Engine/Plugins/Runtime/nDisplay/Source/DisplayClusterConfigurator/Private/Views/Details/Cluster/DisplayClusterConfiguratorClusterDetailsCustomization.cpp b/Engine/Plugins/Runtime/nDisplay/Source/DisplayClusterConfigurator/Private/Views/Details/Cluster/DisplayClusterConfiguratorClusterDetailsCustomization.cpp
--- a/Engine/Plugins/Runtime/nDisplay/Source/DisplayClusterConfigurator/Private/Views/Details/Cluster/DisplayClusterConfiguratorClusterDetailsCustomization.cpp
+++ b/Engine/Plugins/Runtime/nDisplay/Source/DisplayClusterConfigurator/Private/Views/Details/Cluster/DisplayClusterConfiguratorClusterDetailsCustomization.cpp
@@ -15,19 +15,23 @@ void FDisplayClusterConfiguratorClusterDetailsCustomization::CustomizeDetails(ID
 	if (ToolkitPtr.IsValid())
 	{
 		check (RootActorPtr.IsValid());
+		if (RootActorPtr->GetConfigData())
 		{
-			const TSharedPtr<IPropertyHandle> PropertyHandle = InLayoutBuilder.AddObjectPropertyData(
-				{ RootActorPtr->GetConfigData() },
-				GET_MEMBER_NAME_CHECKED(UDisplayClusterConfigurationData, Info));
-			check(PropertyHandle->IsValidHandle());
-			InLayoutBuilder.AddPropertyToCategory(PropertyHandle);
-		}
-		{
-			const TSharedPtr<IPropertyHandle> PropertyHandle = InLayoutBuilder.AddObjectPropertyData(
-				{ RootActorPtr->GetConfigData() },
-				GET_MEMBER_NAME_CHECKED(UDisplayClusterConfigurationData, Diagnostics));
-			check(PropertyHandle->IsValidHandle());
-			InLayoutBuilder.AddPropertyToCategory(PropertyHandle);
+			{
+				const TSharedPtr<IPropertyHandle> PropertyHandle = InLayoutBuilder.AddObjectPropertyData(
+					{ RootActorPtr->GetConfigData() },
+					GET_MEMBER_NAME_CHECKED(UDisplayClusterConfigurationData, Info));
+				check(PropertyHandle->IsValidHandle());
+				InLayoutBuilder.AddPropertyToCategory(PropertyHandle);
+			}
+
+			{
+				const TSharedPtr<IPropertyHandle> PropertyHandle = InLayoutBuilder.AddObjectPropertyData(
+					{ RootActorPtr->GetConfigData() },
+					GET_MEMBER_NAME_CHECKED(UDisplayClusterConfigurationData, Diagnostics));
+				check(PropertyHandle->IsValidHandle());
+				InLayoutBuilder.AddPropertyToCategory(PropertyHandle);
+			}
 		}
 	}
 }
