nDisplay: Fix serialization crash when running packaged (BootTest).

Changed to a strategy of copying back the defaults to the Media settings instead of skipping loading because the untagged serialization logic when packaged was running into binary incompatibility issues.

#rnx

diff --git a/Engine/Plugins/Runtime/nDisplay/Source/DisplayClusterConfiguration/Private/DisplayClusterConfigurationTypes.cpp b/Engine/Plugins/Runtime/nDisplay/Source/DisplayClusterConfiguration/Private/DisplayClusterConfigurationTypes.cpp
--- a/Engine/Plugins/Runtime/nDisplay/Source/DisplayClusterConfiguration/Private/DisplayClusterConfigurationTypes.cpp
+++ b/Engine/Plugins/Runtime/nDisplay/Source/DisplayClusterConfiguration/Private/DisplayClusterConfigurationTypes.cpp
@@ -9,7 +9,6 @@
 #include "DisplayClusterConfigurationTypes_ICVFX.h"
 #include "DisplayClusterConfigurationTypes_PostRender.h"
 #include "DisplayClusterProjectionStrings.h"
-#include "DisplayClusterPropertySkipperArchive.h"
 
 #include "Engine/StaticMesh.h"
 #include "UObject/Package.h"
@@ -338,12 +337,17 @@ void UDisplayClusterConfigurationClusterNode::PostLoad()
 
 void UDisplayClusterConfigurationClusterNode::Serialize(FArchive& Ar)
 {
-	// Use our custom archive to skip serializing Media except for archetypes.
-
-	FDisplayClusterPropertySkipperArchive PropertySkipperAr = FDisplayClusterPropertySkipperArchive(Ar);
-	check(PropertySkipperAr.AddPropertyToSkip(StaticClass(), GET_MEMBER_NAME_CHECKED(UDisplayClusterConfigurationClusterNode, MediaSettings)));
-
-	return Super::Serialize(PropertySkipperAr);
+	// When loading, overwrite Media settings with defaults unless this is the archetype
+	if (Ar.IsLoading() && !IsTemplate())
+	{
+		const FDisplayClusterConfigurationMediaNodeBackbuffer MediaSettingsOriginal = MediaSettings;
+		Super::Serialize(Ar);
+		MediaSettings = MediaSettingsOriginal;
+	}
+	else
+	{
+		Super::Serialize(Ar);
+	}
 }
 
 UDisplayClusterConfigurationClusterNode::UDisplayClusterConfigurationClusterNode()
diff --git a/Engine/Plugins/Runtime/nDisplay/Source/DisplayClusterConfiguration/Private/DisplayClusterConfigurationTypes_ICVFX.cpp b/Engine/Plugins/Runtime/nDisplay/Source/DisplayClusterConfiguration/Private/DisplayClusterConfigurationTypes_ICVFX.cpp
--- a/Engine/Plugins/Runtime/nDisplay/Source/DisplayClusterConfiguration/Private/DisplayClusterConfigurationTypes_ICVFX.cpp
+++ b/Engine/Plugins/Runtime/nDisplay/Source/DisplayClusterConfiguration/Private/DisplayClusterConfigurationTypes_ICVFX.cpp
@@ -2,7 +2,7 @@
 
 #include "DisplayClusterConfigurationTypes_ICVFX.h"
 #include "DisplayClusterConfigurationTypes.h"
-#include "DisplayClusterPropertySkipperArchive.h"
+#include "DisplayClusterConfigurationUtils.h"
 #include "IDisplayCluster.h"
 #include "Camera/CameraTypes.h"
 #include "CineCameraComponent.h"
@@ -64,14 +64,25 @@ void FDisplayClusterConfigurationICVFX_CameraRenderSettings::SetupViewInfo(const
 
 bool FDisplayClusterConfigurationICVFX_CameraRenderSettings::Serialize(FArchive& Ar)
 {	
-	// Use our custom archive to skip serializing Media except for archetypes.
+	// When loading, overwrite Media settings with defaults unless this is the archetype
 
-	UScriptStruct& Struct = *StaticStruct();
+	if (!Ar.IsLoading() && !Ar.IsSaving())
+	{
+		return true;
+	}
 
-	FDisplayClusterPropertySkipperArchive PropertySkipperAr = FDisplayClusterPropertySkipperArchive(Ar);
-	check(PropertySkipperAr.AddPropertyToSkip(&Struct, GET_MEMBER_NAME_CHECKED(FDisplayClusterConfigurationICVFX_CameraRenderSettings, Media)));
+	UScriptStruct& Struct = *StaticStruct();
 
-	Struct.SerializeTaggedProperties(PropertySkipperAr, (uint8*)this, &Struct, nullptr);
+	if (Ar.IsLoading() && !FDisplayClusterConfigurationUtils::IsSerializingTemplate(Ar))
+	{
+		const FDisplayClusterConfigurationMediaICVFX MediaOriginal = Media;
+		Struct.SerializeTaggedProperties(Ar, (uint8*)this, &Struct, nullptr);
+		Media = MediaOriginal;
+	}
+	else
+	{
+		Struct.SerializeTaggedProperties(Ar, (uint8*)this, &Struct, nullptr);
+	}
 
 	return true;
 }
diff --git a/Engine/Plugins/Runtime/nDisplay/Source/DisplayClusterConfiguration/Private/DisplayClusterConfigurationTypes_Viewport.cpp b/Engine/Plugins/Runtime/nDisplay/Source/DisplayClusterConfiguration/Private/DisplayClusterConfigurationTypes_Viewport.cpp
--- a/Engine/Plugins/Runtime/nDisplay/Source/DisplayClusterConfiguration/Private/DisplayClusterConfigurationTypes_Viewport.cpp
+++ b/Engine/Plugins/Runtime/nDisplay/Source/DisplayClusterConfiguration/Private/DisplayClusterConfigurationTypes_Viewport.cpp
@@ -2,7 +2,7 @@
 
 #include "DisplayClusterConfigurationTypes_Viewport.h"
 #include "DisplayClusterConfigurationTypes_ICVFX.h"
-#include "DisplayClusterPropertySkipperArchive.h"
+#include "DisplayClusterConfigurationUtils.h"
 
 ///////////////////////////////////////////////////////////////////////////////////////
 // FDisplayClusterConfigurationViewport_ICVFX
@@ -97,14 +97,25 @@ EDisplayClusterViewportICVFXFlags FDisplayClusterConfigurationViewport_ICVFX::Ge
 
 bool FDisplayClusterConfigurationViewport_RenderSettings::Serialize(FArchive& Ar)
 {
-	// Use our custom archive to skip serializing Media except for archetypes.
+	// When loading, overwrite Media settings with defaults unless this is the archetype
 
-	UScriptStruct& Struct = *StaticStruct();
+	if (!Ar.IsLoading() && !Ar.IsSaving())
+	{
+		return true;
+	}
 
-	FDisplayClusterPropertySkipperArchive PropertySkipperAr = FDisplayClusterPropertySkipperArchive(Ar);
-	check(PropertySkipperAr.AddPropertyToSkip(&Struct, GET_MEMBER_NAME_CHECKED(FDisplayClusterConfigurationViewport_RenderSettings, Media)));
+	UScriptStruct& Struct = *StaticStruct();
 
-	Struct.SerializeTaggedProperties(PropertySkipperAr, (uint8*)this, &Struct, nullptr);
+	if (Ar.IsLoading() && !FDisplayClusterConfigurationUtils::IsSerializingTemplate(Ar))
+	{
+		const FDisplayClusterConfigurationMediaViewport MediaOriginal = Media;
+		Struct.SerializeTaggedProperties(Ar, (uint8*)this, &Struct, nullptr);
+		Media = MediaOriginal;
+	}
+	else
+	{
+		Struct.SerializeTaggedProperties(Ar, (uint8*)this, &Struct, nullptr);
+	}
 
 	return true;
 }
diff --git a/Engine/Plugins/Runtime/nDisplay/Source/DisplayClusterConfiguration/Private/DisplayClusterConfigurationUtils.cpp b/Engine/Plugins/Runtime/nDisplay/Source/DisplayClusterConfiguration/Private/DisplayClusterConfigurationUtils.cpp
new file mode 100644
--- /dev/null
+++ b/Engine/Plugins/Runtime/nDisplay/Source/DisplayClusterConfiguration/Private/DisplayClusterConfigurationUtils.cpp
@@ -0,0 +1,16 @@
+// Copyright Epic Games, Inc. All Rights Reserved.
+
+#include "DisplayClusterConfigurationUtils.h"
+
+#include "UObject/Object.h"
+#include "UObject/UObjectThreadContext.h"
+
+
+bool FDisplayClusterConfigurationUtils::IsSerializingTemplate(const FArchive& Ar)
+{
+	const FUObjectSerializeContext* SerializeContext = FUObjectThreadContext::Get().GetSerializeContext();
+	const UObject* SerializedObject = SerializeContext ? SerializeContext->SerializedObject : nullptr;
+
+	// We assume that if the SerializedObject is null, it is indicative of a template.
+	return SerializedObject ? SerializedObject->IsTemplate() : true;
+}
diff --git a/Engine/Plugins/Runtime/nDisplay/Source/DisplayClusterConfiguration/Private/DisplayClusterConfigurationUtils.h b/Engine/Plugins/Runtime/nDisplay/Source/DisplayClusterConfiguration/Private/DisplayClusterConfigurationUtils.h
new file mode 100644
--- /dev/null
+++ b/Engine/Plugins/Runtime/nDisplay/Source/DisplayClusterConfiguration/Private/DisplayClusterConfigurationUtils.h
@@ -0,0 +1,21 @@
+// Copyright Epic Games, Inc. All Rights Reserved.
+
+#pragma once
+
+
+class FArchive;
+
+
+/**
+ * Utility functions
+ */
+class FDisplayClusterConfigurationUtils
+{
+private:
+	FDisplayClusterConfigurationUtils() = default;
+
+public:
+	// Returns true if we're serializing a template or archetype. False otherwise.
+	static bool IsSerializingTemplate(const FArchive& Ar);
+
+};
diff --git a/Engine/Plugins/Runtime/nDisplay/Source/DisplayClusterConfiguration/Private/DisplayClusterPropertySkipperArchive.h b/Engine/Plugins/Runtime/nDisplay/Source/DisplayClusterConfiguration/Private/DisplayClusterPropertySkipperArchive.h
deleted file mode 100644
