[nDisplay] Late OCIO fix
- PostToneMap event is always triggered when the viewport is being captured

#rnx

diff --git a/Engine/Plugins/Runtime/nDisplay/Source/DisplayCluster/Private/Render/Viewport/DisplayClusterViewportManagerViewExtension.cpp b/Engine/Plugins/Runtime/nDisplay/Source/DisplayCluster/Private/Render/Viewport/DisplayClusterViewportManagerViewExtension.cpp
--- a/Engine/Plugins/Runtime/nDisplay/Source/DisplayCluster/Private/Render/Viewport/DisplayClusterViewportManagerViewExtension.cpp
+++ b/Engine/Plugins/Runtime/nDisplay/Source/DisplayCluster/Private/Render/Viewport/DisplayClusterViewportManagerViewExtension.cpp
@@ -171,6 +171,9 @@ FScreenPassTexture FDisplayClusterViewportManagerViewExtension::PostProcessPassA
 		uint32 ContextNum = 0;
 		if (FDisplayClusterViewportProxy* ViewportProxyPtr = ViewportManagerProxy->ImplFindViewportProxy_RenderThread(View.StereoViewIndex, &ContextNum))
 		{
+			// Broadcast PassTonemap event
+			IDisplayCluster::Get().GetCallbacks().OnDisplayClusterPostTonemapPass_RenderThread().Broadcast(GraphBuilder, ViewportProxyPtr, View, Inputs, ContextNum);
+
 			if (ViewportProxyPtr->ShouldUsePostProcessPassTonemap())
 			{
 				return ViewportProxyPtr->OnPostProcessPassAfterTonemap_RenderThread(GraphBuilder, View, Inputs, ContextNum);
diff --git a/Engine/Plugins/Runtime/nDisplay/Source/DisplayCluster/Private/Render/Viewport/DisplayClusterViewportProxy.cpp b/Engine/Plugins/Runtime/nDisplay/Source/DisplayCluster/Private/Render/Viewport/DisplayClusterViewportProxy.cpp
--- a/Engine/Plugins/Runtime/nDisplay/Source/DisplayCluster/Private/Render/Viewport/DisplayClusterViewportProxy.cpp
+++ b/Engine/Plugins/Runtime/nDisplay/Source/DisplayCluster/Private/Render/Viewport/DisplayClusterViewportProxy.cpp
@@ -415,9 +415,6 @@ FScreenPassTexture FDisplayClusterViewportProxy::OnPostProcessPassAfterFXAA_Rend
 
 FScreenPassTexture FDisplayClusterViewportProxy::OnPostProcessPassAfterTonemap_RenderThread(FRDGBuilder& GraphBuilder, const FSceneView& View, const FPostProcessMaterialInputs& Inputs, const uint32 ContextNum)
 {
-	// Broadcast PassTonemap event
-	IDisplayCluster::Get().GetCallbacks().OnDisplayClusterPostTonemapPass_RenderThread().Broadcast(GraphBuilder, this, View, Inputs, ContextNum);
-
 	// Perform OCIO rendering after the tonemapper
 	if (GetOpenColorIOMode() == EDisplayClusterViewportOpenColorIOMode::PostProcess)
 	{
diff --git a/Engine/Plugins/Runtime/nDisplay/Source/DisplayCluster/Private/Render/Viewport/DisplayClusterViewportProxy_Rules.cpp b/Engine/Plugins/Runtime/nDisplay/Source/DisplayCluster/Private/Render/Viewport/DisplayClusterViewportProxy_Rules.cpp
--- a/Engine/Plugins/Runtime/nDisplay/Source/DisplayCluster/Private/Render/Viewport/DisplayClusterViewportProxy_Rules.cpp
+++ b/Engine/Plugins/Runtime/nDisplay/Source/DisplayCluster/Private/Render/Viewport/DisplayClusterViewportProxy_Rules.cpp
@@ -255,6 +255,11 @@ bool FDisplayClusterViewportProxy::ShouldUsePostProcessPassTonemap() const
 	{
 		return true;
 	}
+	// When capturing, Late OCIO requires PostTonemap
+	else if (RenderSettings.HasAnyMediaStates(EDisplayClusterViewportMediaState::Capture))
+	{
+		return true;
+	}
 
 	return false;
 }
