ndisplay Rivermax: Fixed a crash caused by the issue with RDG Graph Builder being created twice in the same stack.

diff --git a/Engine/Plugins/Runtime/nDisplay/Source/DisplayClusterMedia/Private/Input/DisplayClusterMediaInputBase.cpp b/Engine/Plugins/Runtime/nDisplay/Source/DisplayClusterMedia/Private/Input/DisplayClusterMediaInputBase.cpp
--- a/Engine/Plugins/Runtime/nDisplay/Source/DisplayClusterMedia/Private/Input/DisplayClusterMediaInputBase.cpp
+++ b/Engine/Plugins/Runtime/nDisplay/Source/DisplayClusterMedia/Private/Input/DisplayClusterMediaInputBase.cpp
@@ -160,12 +160,15 @@ void FDisplayClusterMediaInputBase::ReleaseInternals()
 	OCIOAppliedTexture.SafeRelease();
 }
 
-void FDisplayClusterMediaInputBase::ImportMediaData_RenderThread(FRDGBuilder& GraphBuilder, const FMediaInputTextureInfo& TextureInfo)
+void FDisplayClusterMediaInputBase::ImportMediaData_RenderThread(FRHICommandListImmediate& RHICmdList, const FMediaInputTextureInfo& TextureInfo)
 {
 	UE_LOG(LogDisplayClusterMedia, Verbose, TEXT("MediaInput '%s': importing texture on RT frame '%llu'..."), *GetMediaId(), GFrameCounterRenderThread);
 
 	MediaTexture->JustInTimeRender();
 
+	// Media texture relies on RDGBuilder so calling this before MediaTexture causes a crash.
+	FRDGBuilder GraphBuilder(RHICmdList);
+
 	FRHITexture* SrcTexture = MediaTexture->GetResource() ? MediaTexture->GetResource()->GetTextureRHI() : nullptr;
 	FRHITexture* const DstTexture = TextureInfo.Texture;
 
@@ -211,6 +214,9 @@ void FDisplayClusterMediaInputBase::ImportMediaData_RenderThread(FRDGBuilder& Gr
 			DisplayClusterMediaHelpers::ResampleTexture_RenderThread(GraphBuilder.RHICmdList, SrcTexture, DstTexture, SrcRect, DstRect);
 		}
 	}
+
+	GraphBuilder.Execute();
+
 }
 
 bool FDisplayClusterMediaInputBase::ProcessLateOCIO(FRDGBuilder& GraphBuilder, FRHITexture* SrcTexture, const FOpenColorIORenderPassResources& OCIORenderPassResources)
diff --git a/Engine/Plugins/Runtime/nDisplay/Source/DisplayClusterMedia/Private/Input/DisplayClusterMediaInputBase.h b/Engine/Plugins/Runtime/nDisplay/Source/DisplayClusterMedia/Private/Input/DisplayClusterMediaInputBase.h
--- a/Engine/Plugins/Runtime/nDisplay/Source/DisplayClusterMedia/Private/Input/DisplayClusterMediaInputBase.h
+++ b/Engine/Plugins/Runtime/nDisplay/Source/DisplayClusterMedia/Private/Input/DisplayClusterMediaInputBase.h
@@ -86,7 +86,7 @@ protected:
 protected:
 
 	/** Imports texture from a media source */
-	void ImportMediaData_RenderThread(FRDGBuilder& GraphBuilder, const FMediaInputTextureInfo& TextureInfo);
+	void ImportMediaData_RenderThread(FRHICommandListImmediate& RHICmdList, const FMediaInputTextureInfo& TextureInfo);
 
 private:
 
diff --git a/Engine/Plugins/Runtime/nDisplay/Source/DisplayClusterMedia/Private/Input/DisplayClusterMediaInputViewport.cpp b/Engine/Plugins/Runtime/nDisplay/Source/DisplayClusterMedia/Private/Input/DisplayClusterMediaInputViewport.cpp
--- a/Engine/Plugins/Runtime/nDisplay/Source/DisplayClusterMedia/Private/Input/DisplayClusterMediaInputViewport.cpp
+++ b/Engine/Plugins/Runtime/nDisplay/Source/DisplayClusterMedia/Private/Input/DisplayClusterMediaInputViewport.cpp
@@ -123,15 +123,11 @@ void FDisplayClusterMediaInputViewport::OnPostCrossGpuTransfer_RenderThread(FRHI
 			{
 				if (Textures.Num() > 0 && Regions.Num() > 0 && Textures[0])
 				{
-					FRDGBuilder GraphBuilder(RHICmdList);
-
 					// Prepare request data
 					FMediaInputTextureInfo TextureInfo{ Textures[0], Regions[0], MoveTemp(OCIOPassResources_RT) };
 
 					// Import texture from media input
-					ImportMediaData_RenderThread(GraphBuilder, TextureInfo);
-
-					GraphBuilder.Execute();
+					ImportMediaData_RenderThread(RHICmdList, TextureInfo);
 				}
 			}
 		}
