[nDisplay] Late OCIO fix + fixed crash when nDisplay is using media capture
- temporary workaround for ReturnUntouchedSceneColorForPostProcessing crash (UE-219768)
- minor updates to the late OCIO

#rnx

diff --git a/Engine/Plugins/Runtime/nDisplay/Source/DisplayCluster/Private/Render/Viewport/DisplayClusterViewportManagerViewExtension.cpp b/Engine/Plugins/Runtime/nDisplay/Source/DisplayCluster/Private/Render/Viewport/DisplayClusterViewportManagerViewExtension.cpp
--- a/Engine/Plugins/Runtime/nDisplay/Source/DisplayCluster/Private/Render/Viewport/DisplayClusterViewportManagerViewExtension.cpp
+++ b/Engine/Plugins/Runtime/nDisplay/Source/DisplayCluster/Private/Render/Viewport/DisplayClusterViewportManagerViewExtension.cpp
@@ -174,14 +174,34 @@ FScreenPassTexture FDisplayClusterViewportManagerViewExtension::PostProcessPassA
 			// Broadcast PassTonemap event
 			IDisplayCluster::Get().GetCallbacks().OnDisplayClusterPostTonemapPass_RenderThread().Broadcast(GraphBuilder, ViewportProxyPtr, View, Inputs, ContextNum);
 
-			if (ViewportProxyPtr->ShouldUsePostProcessPassTonemap())
+			if (ViewportProxyPtr->GetOpenColorIOMode() == EDisplayClusterViewportOpenColorIOMode::PostProcess)
 			{
 				return ViewportProxyPtr->OnPostProcessPassAfterTonemap_RenderThread(GraphBuilder, View, Inputs, ContextNum);
 			}
 		}
 	}
 
+// UE-219768
+// This is a temporary workaround that prevents crash in ReturnUntouchedSceneColorForPostProcessing.
+// Once the jira is fixed, we can remove this and leave original code from the #else-block
+#if true
+	const FScreenPassTexture& SceneColor = FScreenPassTexture::CopyFromSlice(GraphBuilder, Inputs.GetInput(EPostProcessMaterialInput::SceneColor));
+	check(SceneColor.IsValid());
+
+	FScreenPassRenderTarget Output = Inputs.OverrideOutput;
+
+	// If the override output is provided, it means that this is the last pass in post processing.
+	if (!Output.IsValid())
+	{
+		Output = FScreenPassRenderTarget::CreateFromInput(GraphBuilder, SceneColor, View.GetOverwriteLoadAction(), TEXT("FinalSceneColor"));
+	}
+
+	AddDrawTexturePass(GraphBuilder, View, SceneColor, Output);
+
+	return Output;
+#else
 	return Inputs.ReturnUntouchedSceneColorForPostProcessing(GraphBuilder);
+#endif
 }
 
 bool FDisplayClusterViewportManagerViewExtension::IsActiveThisFrame_Internal(const FSceneViewExtensionContext& Context) const
diff --git a/Engine/Plugins/Runtime/nDisplay/Source/DisplayClusterMedia/Private/Capture/DisplayClusterMediaCaptureViewport.cpp b/Engine/Plugins/Runtime/nDisplay/Source/DisplayClusterMedia/Private/Capture/DisplayClusterMediaCaptureViewport.cpp
--- a/Engine/Plugins/Runtime/nDisplay/Source/DisplayClusterMedia/Private/Capture/DisplayClusterMediaCaptureViewport.cpp
+++ b/Engine/Plugins/Runtime/nDisplay/Source/DisplayClusterMedia/Private/Capture/DisplayClusterMediaCaptureViewport.cpp
@@ -153,6 +153,11 @@ void FDisplayClusterMediaCaptureViewport::OnPostTonemapPass_RenderThread(FRDGBui
 {
 	checkSlow(ViewportProxy);
 
+	if (!IsLateOCIO())
+	{
+		return;
+	}
+
 	// Media subsystem does not support stereo, therefore we process context 0 only
 	if (ContextNum != 0)
 	{
