[nDisplay] Another few fixes to the late OCIO
- Fixed extra frame of latency on the receiving nodes that apply custom OCIO (late OCIO pass)
- Fixed ViewFamily setup when capturing for late OCIO, but no local OCIO transformation expected

#rnx

diff --git a/Engine/Plugins/Runtime/nDisplay/Source/DisplayCluster/Private/Render/Viewport/DisplayClusterViewport.cpp b/Engine/Plugins/Runtime/nDisplay/Source/DisplayCluster/Private/Render/Viewport/DisplayClusterViewport.cpp
--- a/Engine/Plugins/Runtime/nDisplay/Source/DisplayCluster/Private/Render/Viewport/DisplayClusterViewport.cpp
+++ b/Engine/Plugins/Runtime/nDisplay/Source/DisplayCluster/Private/Render/Viewport/DisplayClusterViewport.cpp
@@ -29,6 +29,7 @@
 #include "Render/DisplayDevice/Components/DisplayClusterDisplayDeviceBaseComponent.h"
 
 #include "EngineUtils.h"
+#include "OpenColorIORendering.h"
 #include "SceneManagement.h"
 #include "SceneView.h"
 #include "UnrealClient.h"
@@ -324,6 +325,11 @@ void FDisplayClusterViewport::SetupSceneView(uint32 ContextNum, class UWorld* Wo
 	{
 		OpenColorIO->SetupSceneView(InOutViewFamily, InOutView);
 	}
+	// When capturing with late OCIO, we also need to modify some OCIO specific parameters
+	else if (RenderSettings.HasAnyMediaStates(EDisplayClusterViewportMediaState::CaptureLateOCIO))
+	{
+		FOpenColorIORendering::PrepareView(InOutViewFamily, InOutView);
+	}
 
 	if(Contexts[ContextNum].GPUIndex >= 0)
 	{
diff --git a/Engine/Plugins/Runtime/nDisplay/Source/DisplayCluster/Private/Render/Viewport/DisplayClusterViewportProxy_Rules.cpp b/Engine/Plugins/Runtime/nDisplay/Source/DisplayCluster/Private/Render/Viewport/DisplayClusterViewportProxy_Rules.cpp
--- a/Engine/Plugins/Runtime/nDisplay/Source/DisplayCluster/Private/Render/Viewport/DisplayClusterViewportProxy_Rules.cpp
+++ b/Engine/Plugins/Runtime/nDisplay/Source/DisplayCluster/Private/Render/Viewport/DisplayClusterViewportProxy_Rules.cpp
@@ -255,8 +255,8 @@ bool FDisplayClusterViewportProxy::ShouldUsePostProcessPassTonemap() const
 	{
 		return true;
 	}
-	// When capturing, Late OCIO requires PostTonemap
-	else if (RenderSettings.HasAnyMediaStates(EDisplayClusterViewportMediaState::Capture))
+	// When capturing with late OCIO pass, PostTonemap is required
+	else if (RenderSettings.HasAnyMediaStates(EDisplayClusterViewportMediaState::CaptureLateOCIO))
 	{
 		return true;
 	}
diff --git a/Engine/Plugins/Runtime/nDisplay/Source/DisplayCluster/Public/Render/Viewport/Containers/DisplayClusterViewport_Enums.h b/Engine/Plugins/Runtime/nDisplay/Source/DisplayCluster/Public/Render/Viewport/Containers/DisplayClusterViewport_Enums.h
--- a/Engine/Plugins/Runtime/nDisplay/Source/DisplayCluster/Public/Render/Viewport/Containers/DisplayClusterViewport_Enums.h
+++ b/Engine/Plugins/Runtime/nDisplay/Source/DisplayCluster/Public/Render/Viewport/Containers/DisplayClusterViewport_Enums.h
@@ -192,7 +192,13 @@ enum class EDisplayClusterViewportMediaState : uint8
 	// This viewport will be captured by a media device.
 	Capture = 1 << 0,
 
+	// Custom OCIO transformation is expected on the receiving side
+	CaptureLateOCIO = 1 << 1,
+
 	// This viewport is overridden by a media device.
-	Input = 1 << 1,
+	Input = 1 << 2,
+
+	// Custom OCIO transformation is expected on receiving
+	InputLateOCIO = 1 << 3,
 };
 ENUM_CLASS_FLAGS(EDisplayClusterViewportMediaState);
diff --git a/Engine/Plugins/Runtime/nDisplay/Source/DisplayClusterMedia/Private/Capture/DisplayClusterMediaCaptureViewport.cpp b/Engine/Plugins/Runtime/nDisplay/Source/DisplayClusterMedia/Private/Capture/DisplayClusterMediaCaptureViewport.cpp
--- a/Engine/Plugins/Runtime/nDisplay/Source/DisplayClusterMedia/Private/Capture/DisplayClusterMediaCaptureViewport.cpp
+++ b/Engine/Plugins/Runtime/nDisplay/Source/DisplayClusterMedia/Private/Capture/DisplayClusterMediaCaptureViewport.cpp
@@ -79,6 +79,12 @@ void FDisplayClusterMediaCaptureViewport::OnUpdateViewportMediaState(IDisplayClu
 	{
 		// Raise flags that this viewport will be captured by media.
 		InOutMediaState |= EDisplayClusterViewportMediaState::Capture;
+
+		// Late OCIO flag
+		if (IsLateOCIO())
+		{
+			InOutMediaState |= EDisplayClusterViewportMediaState::CaptureLateOCIO;
+		}
 	}
 }
 
diff --git a/Engine/Plugins/Runtime/nDisplay/Source/DisplayClusterMedia/Private/Input/DisplayClusterMediaInputBase.cpp b/Engine/Plugins/Runtime/nDisplay/Source/DisplayClusterMedia/Private/Input/DisplayClusterMediaInputBase.cpp
--- a/Engine/Plugins/Runtime/nDisplay/Source/DisplayClusterMedia/Private/Input/DisplayClusterMediaInputBase.cpp
+++ b/Engine/Plugins/Runtime/nDisplay/Source/DisplayClusterMedia/Private/Input/DisplayClusterMediaInputBase.cpp
@@ -166,9 +166,6 @@ void FDisplayClusterMediaInputBase::ImportMediaData_RenderThread(FRHICommandList
 
 	MediaTexture->JustInTimeRender();
 
-	// Media texture relies on RDGBuilder so calling this before MediaTexture causes a crash.
-	FRDGBuilder GraphBuilder(RHICmdList);
-
 	FRHITexture* SrcTexture = MediaTexture->GetResource() ? MediaTexture->GetResource()->GetTextureRHI() : nullptr;
 	FRHITexture* const DstTexture = TextureInfo.Texture;
 
@@ -180,7 +177,7 @@ void FDisplayClusterMediaInputBase::ImportMediaData_RenderThread(FRHICommandList
 	// Apply OCIO if needed
 	if (IsLateOCIO())
 	{
-		const bool bApplied = ProcessLateOCIO(GraphBuilder, SrcTexture, TextureInfo.OCIOPassResources);
+		const bool bApplied = ProcessLateOCIO(RHICmdList, SrcTexture, TextureInfo.OCIOPassResources);
 		if (bApplied)
 		{
 			// Redirect SrcTexture to the intermediate OCIO texture so we'll be importing this new one
@@ -207,19 +204,16 @@ void FDisplayClusterMediaInputBase::ImportMediaData_RenderThread(FRHICommandList
 			CopyInfo.DestPosition = FIntVector(DstRect.Min.X, DstRect.Min.Y, 0);
 			CopyInfo.Size = FIntVector(DstRect.Size().X, DstRect.Size().Y, 0);
 
-			TransitionAndCopyTexture(GraphBuilder.RHICmdList, SrcTexture, DstTexture, CopyInfo);
+			TransitionAndCopyTexture(RHICmdList, SrcTexture, DstTexture, CopyInfo);
 		}
 		else
 		{
-			DisplayClusterMediaHelpers::ResampleTexture_RenderThread(GraphBuilder.RHICmdList, SrcTexture, DstTexture, SrcRect, DstRect);
+			DisplayClusterMediaHelpers::ResampleTexture_RenderThread(RHICmdList, SrcTexture, DstTexture, SrcRect, DstRect);
 		}
 	}
-
-	GraphBuilder.Execute();
-
 }
 
-bool FDisplayClusterMediaInputBase::ProcessLateOCIO(FRDGBuilder& GraphBuilder, FRHITexture* SrcTexture, const FOpenColorIORenderPassResources& OCIORenderPassResources)
+bool FDisplayClusterMediaInputBase::ProcessLateOCIO(FRHICommandListImmediate& RHICmdList, FRHITexture* SrcTexture, const FOpenColorIORenderPassResources& OCIORenderPassResources)
 {
 	checkSlow(SrcTexture);
 
@@ -255,6 +249,8 @@ bool FDisplayClusterMediaInputBase::ProcessLateOCIO(FRDGBuilder& GraphBuilder, F
 		return false;
 	}
 
+	FRDGBuilder GraphBuilder(RHICmdList);
+
 	FRDGTextureRef InputTexture = RegisterExternalTexture(GraphBuilder, SrcTexture, TEXT("DCMediaLateOCIOTexIn"));
 	FRDGTextureRef OutputTexture = RegisterExternalTexture(GraphBuilder, OCIOAppliedTexture, TEXT("DCMediaLateOCIOTexOut"));
 
@@ -275,6 +271,8 @@ bool FDisplayClusterMediaInputBase::ProcessLateOCIO(FRDGBuilder& GraphBuilder, F
 		EOpenColorIOTransformAlpha::None
 	);
 
+	GraphBuilder.Execute();
+
 	return true;
 }
 
diff --git a/Engine/Plugins/Runtime/nDisplay/Source/DisplayClusterMedia/Private/Input/DisplayClusterMediaInputBase.h b/Engine/Plugins/Runtime/nDisplay/Source/DisplayClusterMedia/Private/Input/DisplayClusterMediaInputBase.h
--- a/Engine/Plugins/Runtime/nDisplay/Source/DisplayClusterMedia/Private/Input/DisplayClusterMediaInputBase.h
+++ b/Engine/Plugins/Runtime/nDisplay/Source/DisplayClusterMedia/Private/Input/DisplayClusterMediaInputBase.h
@@ -91,7 +91,7 @@ protected:
 private:
 
 	/** Applies OCIO transformation to the SRC texture, and stores result into DST */
-	bool ProcessLateOCIO(FRDGBuilder& GraphBuilder, FRHITexture* SrcTexture, const FOpenColorIORenderPassResources& OCIORenderPassResources);
+	bool ProcessLateOCIO(FRHICommandListImmediate& RHICmdList, FRHITexture* SrcTexture, const FOpenColorIORenderPassResources& OCIORenderPassResources);
 
 	/** Creates an internal intermediate texture */
 	FTextureRHIRef CreateTexture(const FRHITexture* ReferenceTexture);
@@ -120,15 +120,15 @@ private:
 	TObjectPtr<UMediaTexture> MediaTexture;
 	//~ End GC by AddReferencedObjects
 
-	// Used to restart media player in the case it falls in error
+	/** Used to restart media player in the case it falls in error */
 	bool bWasPlayerStarted = false;
 
-	// Used to control the rate at which we try to restart the player
+	/** Used to control the rate at which we try to restart the player */
 	double LastRestartTimestamp = 0;
 
-	// Intermediate texture for late OCIO pass
+	/** Intermediate texture for late OCIO pass */
 	FTextureRHIRef OCIOAppliedTexture;
 
-	// [Temp workaround] Whether current media is Rivermax
+	/** [Temp workaround] Whether current media is Rivermax */
 	bool bRunningRivermaxMedia = false;
 };
diff --git a/Engine/Plugins/Runtime/nDisplay/Source/DisplayClusterMedia/Private/Input/DisplayClusterMediaInputViewport.cpp b/Engine/Plugins/Runtime/nDisplay/Source/DisplayClusterMedia/Private/Input/DisplayClusterMediaInputViewport.cpp
--- a/Engine/Plugins/Runtime/nDisplay/Source/DisplayClusterMedia/Private/Input/DisplayClusterMediaInputViewport.cpp
+++ b/Engine/Plugins/Runtime/nDisplay/Source/DisplayClusterMedia/Private/Input/DisplayClusterMediaInputViewport.cpp
@@ -103,6 +103,12 @@ void FDisplayClusterMediaInputViewport::OnUpdateViewportMediaState(IDisplayClust
 	{
 		// Raise flags that this viewport texture will be overridden by media.
 		InOutMediaState |= EDisplayClusterViewportMediaState::Input;
+
+		// Late OCIO flag
+		if (IsLateOCIO())
+		{
+			InOutMediaState |= EDisplayClusterViewportMediaState::InputLateOCIO;
+		}
 	}
 }
 
