nDisplay: Support DoF and other PostProcess in Outer Viewports.
- Implemented a new feature that propagates post processes from a custom camera defined in a ViewPoint component to viewports.
- Fixed DoF parameters in post-processing settings for Outer and InCamera viewports.
- Reworked logic related to PostProcess.
- Added UI customization for the UDisplayClusterCameraComponent.
- Code cleanup.

Note: this CL is based on this archived CL 30875388. It has been moved from version 5.4 to 5.5.

diff --git a/Engine/Plugins/Runtime/nDisplay/Source/DisplayCluster/Private/Game/EngineClasses/Scene/DisplayClusterCameraComponent.cpp b/Engine/Plugins/Runtime/nDisplay/Source/DisplayCluster/Private/Game/EngineClasses/Scene/DisplayClusterCameraComponent.cpp
--- a/Engine/Plugins/Runtime/nDisplay/Source/DisplayCluster/Private/Game/EngineClasses/Scene/DisplayClusterCameraComponent.cpp
+++ b/Engine/Plugins/Runtime/nDisplay/Source/DisplayCluster/Private/Game/EngineClasses/Scene/DisplayClusterCameraComponent.cpp
@@ -2,6 +2,11 @@
 
 #include "Components/DisplayClusterCameraComponent.h"
 #include "Components/BillboardComponent.h"
+#include "Components/DisplayClusterICVFXCameraComponent.h"
+#include "Camera/CameraComponent.h"
+#include "Render/Viewport/IDisplayClusterViewport.h"
+#include "Render/Viewport/Misc/DisplayClusterViewportHelpers.h"
+#include "Render/Viewport/Configuration/DisplayClusterViewportConfigurationHelpers_Postprocess.h"
 
 #include "DisplayClusterRootActor.h"
 
@@ -32,11 +37,132 @@ UDisplayClusterCameraComponent::UDisplayClusterCameraComponent(const FObjectInit
 #endif
 }
 
+void UDisplayClusterCameraComponent::ApplyViewPointComponentPostProcessesToViewport(IDisplayClusterViewport* InViewport)
+{
+	check(InViewport && !EnumHasAnyFlags(InViewport->GetRenderSettingsICVFX().RuntimeFlags, EDisplayClusterViewportRuntimeICVFXFlags::InternalResource));
+
+	// Viewports that use the ViewPoint component get post - processing and more from the referenced camera component.
+	// As we can see, we will have to use up to 3 different classes as different sources of these settings :
+	// UCameraComponent->UCineCameraComponent->UDisplayClusterICVFXCameraComponent
+	// Thus, we have to use all available rendering settings in our component's class.
+	// Override viewport PP from camera, except internal ICVFX viewports
+		
+	using namespace UE::DisplayClusterViewportHelpers;
+	// Get the same component from DCRA that is used as the configuration source. Then this component can also be used as a configuration data source.
+	const UDisplayClusterCameraComponent& CfgCameraComponent = GetMatchingComponentFromRootActor(InViewport->GetConfiguration(), EDisplayClusterRootActorType::Configuration, *this);
+
+	// Setup Outer Viewport postprocessing
+	if (CfgCameraComponent.bEnableOuterViewportCamera)
+	{
+		const EDisplayClusterViewportCameraPostProcessFlags CameraPostProcessingFlags = CfgCameraComponent.OuterViewportPostProcessSettings.GetCameraPostProcessFlags();
+
+		// Also, if we are referencing the ICVFXCamera component, use the special ICVFX PostProcess from it.
+		if (UDisplayClusterICVFXCameraComponent* SceneICVFXCameraComponent = GetRootActorComponentByName<UDisplayClusterICVFXCameraComponent>(InViewport->GetConfiguration(), EDisplayClusterRootActorType::Scene, CfgCameraComponent.OuterViewportCameraName))
+		{
+			// Use PostProcess from the ICVFXCamera
+			// This function also uses PostProcess from the parent CineCamera class.
+			SceneICVFXCameraComponent->ApplyICVFXCameraPostProcessesToViewport(InViewport, CameraPostProcessingFlags);
+		}
+		else
+		{
+			// Use post-processing settings from Camera/CineCamera or from the active game camera.
+			FMinimalViewInfo CustomViewInfo;
+			if (CfgCameraComponent.GetOuterViewportCameraDesiredViewInternal(InViewport->GetConfiguration(), CustomViewInfo))
+			{
+				// Applies a filter to the post-processing settings.
+				FDisplayClusterViewportConfigurationHelpers_Postprocess::FilterPostProcessSettings(CustomViewInfo.PostProcessSettings, CameraPostProcessingFlags);
+
+				// Send camera postprocess to override
+				InViewport->GetViewport_CustomPostProcessSettings().AddCustomPostProcess(IDisplayClusterViewport_CustomPostProcessSettings::ERenderPass::Override, CustomViewInfo.PostProcessSettings, CustomViewInfo.PostProcessBlendWeight, true);
+			}
+		}
+	}
+}
+
+UCameraComponent* UDisplayClusterCameraComponent::GetOuterViewportCameraComponent(const IDisplayClusterViewportConfiguration& InViewportConfiguration) const
+{
+	using namespace UE::DisplayClusterViewportHelpers;
+
+	// Get the same component from DCRA that is used as the configuration source. Then this component can also be used as a configuration data source.
+	const UDisplayClusterCameraComponent& CfgCameraComponent = GetMatchingComponentFromRootActor(InViewportConfiguration, EDisplayClusterRootActorType::Configuration, *this);
+	if (CfgCameraComponent.bEnableOuterViewportCamera)
+	{
+		if(UCameraComponent* SceneCameraComponent = GetRootActorComponentByName<UCameraComponent>(InViewportConfiguration, EDisplayClusterRootActorType::Scene, CfgCameraComponent.OuterViewportCameraName))
+		{
+			// If we use the ICVFX camera component, we must use GetActualCineCameraComponent() to get the actual camera.
+			if(SceneCameraComponent->IsA<UDisplayClusterICVFXCameraComponent>())
+			{
+				if (UDisplayClusterICVFXCameraComponent* ICVFXCameraComponent = Cast<UDisplayClusterICVFXCameraComponent>(SceneCameraComponent))
+				{
+					if (UCineCameraComponent* ExtCineCameraComponent = ICVFXCameraComponent->GetActualCineCameraComponent())
+					{
+						// Use referenced camera as the source of Camera PP and CineCamera CustomNearClippingPlane
+						return ExtCineCameraComponent;
+					}
+					}
+				}
+
+			return SceneCameraComponent;
+		}
+	}
+
+	return nullptr;
+}
+
+bool UDisplayClusterCameraComponent::GetOuterViewportCameraDesiredViewInternal(const IDisplayClusterViewportConfiguration& InViewportConfiguration, FMinimalViewInfo& InOutViewInfo, float* OutCustomNCP) const
+{
+	using namespace UE::DisplayClusterViewportHelpers;
+
+	bool bViewInfoFound = false;
+
+	// Get the same component from DCRA that is used as the configuration source. Then this component can also be used as a configuration data source.
+	const UDisplayClusterCameraComponent& CfgCameraComponent = GetMatchingComponentFromRootActor(InViewportConfiguration, EDisplayClusterRootActorType::Configuration, *this);
+	if (CfgCameraComponent.bEnableOuterViewportCamera)
+	{
+		const EDisplayClusterViewportCameraPostProcessFlags CameraPostProcessingFlags = CfgCameraComponent.OuterViewportPostProcessSettings.GetCameraPostProcessFlags();
+
+		const bool bUseCameraPostprocess = EnumHasAnyFlags(CameraPostProcessingFlags, EDisplayClusterViewportCameraPostProcessFlags::EnablePostProcess);
+
+		float* OutCustomNearClippingPlane = OutCustomNCP;
+		if (!EnumHasAnyFlags(CameraPostProcessingFlags, EDisplayClusterViewportCameraPostProcessFlags::EnableNearClippingPlane))
+		{
+			// Ignore NCP from the custom camera.
+			OutCustomNearClippingPlane = nullptr;
+		}
+
+		if (UCameraComponent* SceneCameraComponent = GetOuterViewportCameraComponent(InViewportConfiguration))
+		{
+			if (IDisplayClusterViewport::GetCameraComponentView(SceneCameraComponent, InViewportConfiguration.GetRootActorWorldDeltaSeconds(), bUseCameraPostprocess, InOutViewInfo, OutCustomNearClippingPlane))
+			{
+				bViewInfoFound = true;
+			}
+		}
+		// Get PostProcess from the Game camera.
+		else if (IDisplayClusterViewport::GetPlayerCameraView(InViewportConfiguration.GetCurrentWorld(), bUseCameraPostprocess, InOutViewInfo))
+		{
+			
+			bViewInfoFound = true;
+		}
+
+		if (bViewInfoFound && !CfgCameraComponent.bFollowOuterViewportCamera)
+		{
+			// Use this component as a camera
+			InOutViewInfo.Location = GetComponentLocation();
+			InOutViewInfo.Rotation = GetComponentRotation();
+		}
+
+		// The default camera is not found, so we can't use the custom camera view.
+	}
+
+	return bViewInfoFound;
+}
+
 void UDisplayClusterCameraComponent::GetDesiredView(IDisplayClusterViewportConfiguration& InViewportConfiguration, FMinimalViewInfo& InOutViewInfo, float* OutCustomNearClippingPlane)
 {
-	// Use this component as a camera
-	InOutViewInfo.Location = GetComponentLocation();
-	InOutViewInfo.Rotation = GetComponentRotation();
+	if (GetOuterViewportCameraDesiredViewInternal(InViewportConfiguration, InOutViewInfo, OutCustomNearClippingPlane))
+	{
+		return;
+	}
 
 	// Ignore PP, because this component has no such settings
 	InOutViewInfo.PostProcessBlendWeight = 0.f;
@@ -46,6 +172,28 @@ void UDisplayClusterCameraComponent::GetDesiredView(IDisplayClusterViewportConfi
 		// Value less than zero means: don't override the NCP value
 		*OutCustomNearClippingPlane = -1.f;
 	}
+
+	// By default this component is used as ViewPoint
+	// Use this component as a camera
+	InOutViewInfo.Location = GetComponentLocation();
+	InOutViewInfo.Rotation = GetComponentRotation();
+}
+
+void UDisplayClusterCameraComponent::GetEyePosition(const IDisplayClusterViewportConfiguration& InViewportConfiguration, FVector& OutViewLocation, FRotator& OutViewRotation)
+{
+	FMinimalViewInfo ViewInfo;
+	if (GetOuterViewportCameraDesiredViewInternal(InViewportConfiguration, ViewInfo))
+	{
+		OutViewLocation = ViewInfo.Location;
+		OutViewRotation = ViewInfo.Rotation;
+	}
+	else
+	{
+		// By default this component is used as ViewPoint
+		// Use this component as a camera
+		OutViewLocation = GetComponentLocation();
+		OutViewRotation = GetComponentRotation();
+	}
 }
 
 #if WITH_EDITOR
@@ -115,3 +263,41 @@ void UDisplayClusterCameraComponent::RefreshVisualRepresentation()
 	}
 }
 #endif
+
+EDisplayClusterViewportCameraPostProcessFlags FDisplayClusterCameraComponent_OuterViewportPostProcessSettings::GetCameraPostProcessFlags() const
+{
+	EDisplayClusterViewportCameraPostProcessFlags OutPostProcessFlags = EDisplayClusterViewportCameraPostProcessFlags::None;
+
+	if (bEnablePostProcess)
+	{
+		EnumAddFlags(OutPostProcessFlags, EDisplayClusterViewportCameraPostProcessFlags::EnablePostProcess);
+	}
+
+	if (bEnableDepthOfField)
+	{
+		EnumAddFlags(OutPostProcessFlags, EDisplayClusterViewportCameraPostProcessFlags::EnableDepthOfField);
+	}
+
+	if (bEnableNearClippingPlane)
+	{
+		EnumAddFlags(OutPostProcessFlags, EDisplayClusterViewportCameraPostProcessFlags::EnableNearClippingPlane);
+	}
+
+	if (bEnableICVFXColorGrading)
+	{
+		EnumAddFlags(OutPostProcessFlags, EDisplayClusterViewportCameraPostProcessFlags::EnableICVFXColorGrading);
+	}
+
+	if (bEnableICVFXMotionBlur)
+	{
+		EnumAddFlags(OutPostProcessFlags, EDisplayClusterViewportCameraPostProcessFlags::EnableICVFXMotionBlur);
+	}
+
+	if (bEnableICVFXDepthOfFieldCompensation)
+	{
+		EnumAddFlags(OutPostProcessFlags, EDisplayClusterViewportCameraPostProcessFlags::EnableICVFXDepthOfFieldCompensation);
+	}
+
+	return OutPostProcessFlags;
+}
+
diff --git a/Engine/Plugins/Runtime/nDisplay/Source/DisplayCluster/Private/Game/EngineClasses/Scene/DisplayClusterICVFXCameraComponent.cpp b/Engine/Plugins/Runtime/nDisplay/Source/DisplayCluster/Private/Game/EngineClasses/Scene/DisplayClusterICVFXCameraComponent.cpp
--- a/Engine/Plugins/Runtime/nDisplay/Source/DisplayCluster/Private/Game/EngineClasses/Scene/DisplayClusterICVFXCameraComponent.cpp
+++ b/Engine/Plugins/Runtime/nDisplay/Source/DisplayCluster/Private/Game/EngineClasses/Scene/DisplayClusterICVFXCameraComponent.cpp
@@ -5,9 +5,13 @@
 
 #include "Cluster/IPDisplayClusterClusterManager.h"
 
+#include "Render/Viewport/IDisplayClusterViewport.h"
+
 #include "Render/Viewport/Containers/DisplayClusterViewport_CameraMotionBlur.h"
 #include "Render/Viewport/Containers/DisplayClusterViewport_CustomFrustumRuntimeSettings.h"
 #include "Render/Viewport/Configuration/DisplayClusterViewportConfigurationHelpers_ICVFX.h"
+#include "Render/Viewport/Configuration/DisplayClusterViewportConfigurationHelpers_Postprocess.h"
+#include "Render/Viewport/Misc/DisplayClusterViewportHelpers.h"
 
 #include "Components/DisplayClusterCameraComponent.h"
 #include "DisplayClusterRootActor.h"
@@ -74,8 +78,6 @@ void UDisplayClusterICVFXCameraComponent::GetCameraView(float DeltaTime, FMinima
 		return;
 	}
 
-	const FDisplayClusterConfigurationICVFX_StageSettings& StageSettings = RootActor->GetStageSettings();
-
 	if (CameraSettings.ExternalCameraActor.IsValid())
 	{
 		// Get ViewInfo from external CineCamera
@@ -87,7 +89,7 @@ void UDisplayClusterICVFXCameraComponent::GetCameraView(float DeltaTime, FMinima
 		UCineCameraComponent::GetCameraView(DeltaTime, InOutViewInfo);
 	}
 
-	CameraSettings.SetupViewInfo(StageSettings, InOutViewInfo);
+	CameraSettings.SetupViewInfo(RootActor->GetStageSettings(), InOutViewInfo);
 }
 
 UCineCameraComponent* UDisplayClusterICVFXCameraComponent::GetActualCineCameraComponent()
@@ -158,6 +160,23 @@ void UDisplayClusterICVFXCameraComponent::TickComponent(float DeltaTime, ELevelT
 	}
 }
 
+const FDisplayClusterConfigurationICVFX_CameraSettings& UDisplayClusterICVFXCameraComponent::GetCameraSettingsICVFX() const
+{
+	return CameraSettings;
+}
+
+void UDisplayClusterICVFXCameraComponent::ApplyICVFXCameraPostProcessesToViewport(IDisplayClusterViewport* InViewport, const EDisplayClusterViewportCameraPostProcessFlags InPostProcessingFlags)
+{
+	if (InViewport)
+	{
+		using namespace UE::DisplayClusterViewportHelpers;
+		// Get the same component from DCRA that is used as the configuration source. Then this component can also be used as a configuration data source.
+		const UDisplayClusterICVFXCameraComponent& CfgICVFXCameraComponent = GetMatchingComponentFromRootActor(InViewport->GetConfiguration(), EDisplayClusterRootActorType::Configuration, *this);
+		
+		FDisplayClusterViewportConfigurationHelpers_Postprocess::ImplApplyICVFXCameraPostProcessesToViewport(InViewport->ToSharedRef().Get(), *this, CfgICVFXCameraComponent.GetCameraSettingsICVFX(), InPostProcessingFlags);
+	}
+}
+
 void UDisplayClusterICVFXCameraComponent::UpdateOverscanEstimatedFrameSize()
 {
 	const ADisplayClusterRootActor* RootActor = Cast<ADisplayClusterRootActor>(GetOwner());
@@ -224,52 +243,6 @@ void UDisplayClusterICVFXCameraComponent::UpdateOverscanEstimatedFrameSize()
 	CameraSettings.CustomFrustum.OverscanPixelsIncrease = ((float)(EstimatedPixel) / (float)(BasePixels));
 }
 
-FDisplayClusterViewport_CameraMotionBlur UDisplayClusterICVFXCameraComponent::GetMotionBlurParameters()
-{
-	FDisplayClusterViewport_CameraMotionBlur OutParameters;
-	OutParameters.Mode = EDisplayClusterViewport_CameraMotionBlur::Undefined;
-
-	switch (CameraSettings.CameraMotionBlur.MotionBlurMode)
-	{
-	case EDisplayClusterConfigurationCameraMotionBlurMode::Off:
-		OutParameters.Mode = EDisplayClusterViewport_CameraMotionBlur::Off;
-		break;
-
-	case EDisplayClusterConfigurationCameraMotionBlurMode::On:
-		OutParameters.Mode = EDisplayClusterViewport_CameraMotionBlur::On;
-		break;
-
-	case EDisplayClusterConfigurationCameraMotionBlurMode::Override:
-		ADisplayClusterRootActor* RootActor = static_cast<ADisplayClusterRootActor*>(GetOwner());
-		if (RootActor)
-		{
-			UDisplayClusterCameraComponent* OuterCamera = RootActor->GetDefaultCamera();
-			if (OuterCamera)
-			{
-				OutParameters.CameraLocation   = OuterCamera->GetComponentLocation();
-				OutParameters.CameraRotation   = OuterCamera->GetComponentRotation();
-				OutParameters.TranslationScale = CameraSettings.CameraMotionBlur.TranslationScale;
-				OutParameters.Mode             = EDisplayClusterViewport_CameraMotionBlur::Override;
-			}
-		}
-		break;
-	}
-
-	return OutParameters;
-}
-
-FDisplayClusterViewport_CameraDepthOfField UDisplayClusterICVFXCameraComponent::GetDepthOfFieldParameters()
-{
-	FDisplayClusterViewport_CameraDepthOfField OutParameters;
-
-	OutParameters.bEnableDepthOfFieldCompensation = CameraSettings.CameraDepthOfField.bEnableDepthOfFieldCompensation;
-	OutParameters.DistanceToWall = CameraSettings.CameraDepthOfField.DistanceToWall;
-	OutParameters.DistanceToWallOffset = CameraSettings.CameraDepthOfField.DistanceToWallOffset;
-	OutParameters.CompensationLUT = CameraSettings.CameraDepthOfField.DynamicCompensationLUT ? ToRawPtr(CameraSettings.CameraDepthOfField.DynamicCompensationLUT) : CameraSettings.CameraDepthOfField.CompensationLUT.Get();
-
-	return OutParameters;
-}
-
 void UDisplayClusterICVFXCameraComponent::OnRegister()
 {
 	Super::OnRegister();
diff --git a/Engine/Plugins/Runtime/nDisplay/Source/DisplayCluster/Private/Render/Device/DisplayClusterDeviceBase.cpp b/Engine/Plugins/Runtime/nDisplay/Source/DisplayCluster/Private/Render/Device/DisplayClusterDeviceBase.cpp
--- a/Engine/Plugins/Runtime/nDisplay/Source/DisplayCluster/Private/Render/Device/DisplayClusterDeviceBase.cpp
+++ b/Engine/Plugins/Runtime/nDisplay/Source/DisplayCluster/Private/Render/Device/DisplayClusterDeviceBase.cpp
@@ -485,67 +485,52 @@ bool FDisplayClusterDeviceBase::NeedReAllocateViewportRenderTarget(const class F
 //////////////////////////////////////////////////////////////////////////////////////////////
 // FDisplayClusterDeviceBase
 //////////////////////////////////////////////////////////////////////////////////////////////
-void FDisplayClusterDeviceBase::StartFinalPostprocessSettings(struct FPostProcessSettings* StartPostProcessingSettings, const enum EStereoscopicPass StereoPassType, const int32 StereoViewIndex)
+void FDisplayClusterDeviceBase::StartFinalPostprocessSettings(FPostProcessSettings* StartPostProcessingSettings, const enum EStereoscopicPass StereoPassType, const int32 StereoViewIndex)
 {
 	check(IsInGameThread());
 
 	// eSSP_FULL pass reserved for UE internal render
-	IDisplayClusterViewportManager* ViewportManager = GetViewportManager();
-	if (StereoPassType != EStereoscopicPass::eSSP_FULL && ViewportManager)
+	if (StereoPassType != EStereoscopicPass::eSSP_FULL && StartPostProcessingSettings)
 	{
-		IDisplayClusterViewport* ViewportPtr = ViewportManager->FindViewport(StereoViewIndex);
-		if (ViewportPtr)
+		IDisplayClusterViewportManager* ViewportManager = GetViewportManager();
+		uint32 ContextNum = 0;
+		if (IDisplayClusterViewport* Viewport = ViewportManager ? ViewportManager->FindViewport(StereoViewIndex, &ContextNum) : nullptr)
 		{
-			ViewportPtr->GetViewport_CustomPostProcessSettings().DoPostProcess(IDisplayClusterViewport_CustomPostProcessSettings::ERenderPass::Start, StartPostProcessingSettings);
+			Viewport->GetViewport_CustomPostProcessSettings().ApplyCustomPostProcess(Viewport, ContextNum, IDisplayClusterViewport_CustomPostProcessSettings ::ERenderPass::Start, *StartPostProcessingSettings);
 		}
 	}
 }
 
-bool FDisplayClusterDeviceBase::OverrideFinalPostprocessSettings(struct FPostProcessSettings* OverridePostProcessingSettings, const enum EStereoscopicPass StereoPassType, const int32 StereoViewIndex, float& BlendWeight)
+bool FDisplayClusterDeviceBase::OverrideFinalPostprocessSettings(FPostProcessSettings* OverridePostProcessingSettings, const enum EStereoscopicPass StereoPassType, const int32 StereoViewIndex, float& BlendWeight)
 {
 	check(IsInGameThread());
 
 	// eSSP_FULL pass reserved for UE internal render
-	IDisplayClusterViewportManager* ViewportManager = GetViewportManager();
-	if (StereoPassType != EStereoscopicPass::eSSP_FULL && ViewportManager)
+	if (StereoPassType != EStereoscopicPass::eSSP_FULL)
 	{
-		IDisplayClusterViewport* ViewportPtr = ViewportManager->FindViewport(StereoViewIndex);
-		if (ViewportPtr)
+		IDisplayClusterViewportManager* ViewportManager = GetViewportManager();
+		uint32 ContextNum = 0;
+		if (IDisplayClusterViewport* Viewport = ViewportManager ? ViewportManager->FindViewport(StereoViewIndex, &ContextNum) : nullptr)
 		{
-			return ViewportPtr->GetViewport_CustomPostProcessSettings().DoPostProcess(IDisplayClusterViewport_CustomPostProcessSettings::ERenderPass::Override, OverridePostProcessingSettings, &BlendWeight);
+			return Viewport->GetViewport_CustomPostProcessSettings().ApplyCustomPostProcess(Viewport, ContextNum, IDisplayClusterViewport_CustomPostProcessSettings::ERenderPass::Override, *OverridePostProcessingSettings, &BlendWeight);
 		}
 	}
 
 	return false;
 }
 
-void FDisplayClusterDeviceBase::EndFinalPostprocessSettings(struct FPostProcessSettings* FinalPostProcessingSettings, const enum EStereoscopicPass StereoPassType, const int32 StereoViewIndex)
+void FDisplayClusterDeviceBase::EndFinalPostprocessSettings(FPostProcessSettings* FinalPostProcessingSettings, const enum EStereoscopicPass StereoPassType, const int32 StereoViewIndex)
 {
 	check(IsInGameThread());
 
 	// eSSP_FULL pass reserved for UE internal render
-	IDisplayClusterViewportManager* ViewportManager = GetViewportManager();
-	if (StereoPassType != EStereoscopicPass::eSSP_FULL && ViewportManager && FinalPostProcessingSettings != nullptr)
+	if (StereoPassType != EStereoscopicPass::eSSP_FULL)
 	{
-		IDisplayClusterViewport* ViewportPtr = ViewportManager->FindViewport(StereoViewIndex);
-		if (ViewportPtr)
+		IDisplayClusterViewportManager* ViewportManager = GetViewportManager();
+		uint32 ContextNum = 0;
+		if (IDisplayClusterViewport* Viewport = ViewportManager ? ViewportManager->FindViewport(StereoViewIndex, &ContextNum) : nullptr)
 		{
-			ViewportPtr->GetViewport_CustomPostProcessSettings().DoPostProcess(IDisplayClusterViewport_CustomPostProcessSettings::ERenderPass::Final, FinalPostProcessingSettings);
-
-			FPostProcessSettings RequestedFinalPerViewportPPS;
-			// Get the final overall cluster + per-viewport PPS from nDisplay
-			if (ViewportPtr->GetViewport_CustomPostProcessSettings().DoPostProcess(IDisplayClusterViewport_CustomPostProcessSettings::ERenderPass::FinalPerViewport, &RequestedFinalPerViewportPPS))
-			{
-				FDisplayClusterConfigurationViewport_ColorGradingRenderingSettings InPPSnDisplay;
-				FDisplayClusterViewportConfigurationHelpers_Postprocess::CopyPPSStructConditional(&InPPSnDisplay, &RequestedFinalPerViewportPPS);
-
-				// Get the passed-in cumulative PPS from the game/viewport (includes all PPVs affecting this viewport)
-				FDisplayClusterConfigurationViewport_ColorGradingRenderingSettings InPPSCumulative;
-				FDisplayClusterViewportConfigurationHelpers_Postprocess::CopyPPSStruct(&InPPSCumulative, FinalPostProcessingSettings);
-
-				// Blend both together with our custom math instead of the default PPS blending
-				FDisplayClusterViewportConfigurationHelpers_Postprocess::BlendPostProcessSettings(*FinalPostProcessingSettings, InPPSCumulative, InPPSnDisplay);
-			}
+			Viewport->GetViewport_CustomPostProcessSettings().ApplyCustomPostProcess(Viewport, ContextNum, IDisplayClusterViewport_CustomPostProcessSettings::ERenderPass::Final, *FinalPostProcessingSettings);
 		}
 	}
 }
diff --git a/Engine/Plugins/Runtime/nDisplay/Source/DisplayCluster/Private/Render/Viewport/Configuration/DisplayClusterViewportConfiguration.cpp b/Engine/Plugins/Runtime/nDisplay/Source/DisplayCluster/Private/Render/Viewport/Configuration/DisplayClusterViewportConfiguration.cpp
--- a/Engine/Plugins/Runtime/nDisplay/Source/DisplayCluster/Private/Render/Viewport/Configuration/DisplayClusterViewportConfiguration.cpp
+++ b/Engine/Plugins/Runtime/nDisplay/Source/DisplayCluster/Private/Render/Viewport/Configuration/DisplayClusterViewportConfiguration.cpp
@@ -217,6 +217,16 @@ UWorld* FDisplayClusterViewportConfiguration::GetCurrentWorld() const
 	return CurrentWorldRef.Get();
 }
 
+float FDisplayClusterViewportConfiguration::GetRootActorWorldDeltaSeconds(const EDisplayClusterRootActorType InRootActorType) const
+{
+	if (ADisplayClusterRootActor* RootActor = GetRootActor(InRootActorType))
+	{
+		return RootActor->GetWorldDeltaSeconds();
+	}
+
+	return 0.0f;
+}
+
 const UDisplayClusterConfigurationData* FDisplayClusterViewportConfiguration::GetConfigurationData() const
 {
 	ADisplayClusterRootActor* ConfigurationRootActor = GetRootActor(EDisplayClusterRootActorType::Configuration);
diff --git a/Engine/Plugins/Runtime/nDisplay/Source/DisplayCluster/Private/Render/Viewport/Configuration/DisplayClusterViewportConfiguration.h b/Engine/Plugins/Runtime/nDisplay/Source/DisplayCluster/Private/Render/Viewport/Configuration/DisplayClusterViewportConfiguration.h
--- a/Engine/Plugins/Runtime/nDisplay/Source/DisplayCluster/Private/Render/Viewport/Configuration/DisplayClusterViewportConfiguration.h
+++ b/Engine/Plugins/Runtime/nDisplay/Source/DisplayCluster/Private/Render/Viewport/Configuration/DisplayClusterViewportConfiguration.h
@@ -56,6 +56,7 @@ public:
 	virtual void ReleaseConfiguration() override;
 
 	virtual UWorld* GetCurrentWorld() const override;
+	virtual float GetRootActorWorldDeltaSeconds(const EDisplayClusterRootActorType InRootActorType = EDisplayClusterRootActorType::Scene) const override;
 	virtual ADisplayClusterRootActor* GetRootActor(const EDisplayClusterRootActorType InRootActorType) const override;
 	virtual IDisplayClusterViewportManager* GetViewportManager() const override;
 	virtual const UDisplayClusterConfigurationData* GetConfigurationData() const override;
diff --git a/Engine/Plugins/Runtime/nDisplay/Source/DisplayCluster/Private/Render/Viewport/Configuration/DisplayClusterViewportConfigurationHelpers_ICVFX.cpp b/Engine/Plugins/Runtime/nDisplay/Source/DisplayCluster/Private/Render/Viewport/Configuration/DisplayClusterViewportConfigurationHelpers_ICVFX.cpp
--- a/Engine/Plugins/Runtime/nDisplay/Source/DisplayCluster/Private/Render/Viewport/Configuration/DisplayClusterViewportConfigurationHelpers_ICVFX.cpp
+++ b/Engine/Plugins/Runtime/nDisplay/Source/DisplayCluster/Private/Render/Viewport/Configuration/DisplayClusterViewportConfigurationHelpers_ICVFX.cpp
@@ -503,12 +503,6 @@ void FDisplayClusterViewportConfigurationHelpers_ICVFX::UpdateCameraViewportSett
 	// Update OCIO for Camera Viewport
 	FDisplayClusterViewportConfigurationHelpers_OpenColorIO::UpdateCameraViewportOCIO(DstViewport, InCameraSettings);
 
-	// Motion blur:
-	DstViewport.UpdateConfiguration_CameraMotionBlur(InCameraComponent.GetMotionBlurParameters());
-
-	// Depth of field
-	DstViewport.UpdateConfiguration_CameraDepthOfField(InCameraComponent.GetDepthOfFieldParameters());
-
 	// FDisplayClusterConfigurationICVFX_CameraSettings
 	InOutRenderSettings.CameraId.Empty();
 
@@ -520,8 +514,8 @@ void FDisplayClusterViewportConfigurationHelpers_ICVFX::UpdateCameraViewportSett
 
 	InOutRenderSettings.Rect = FDisplayClusterViewportHelpers::GetValidViewportRect(FIntRect(FIntPoint(0, 0), DesiredSize), DstViewport.GetId(), TEXT("Configuration Camera Frame Size"));
 
-	// Apply postprocess for camera
-	FDisplayClusterViewportConfigurationHelpers_Postprocess::UpdateCameraPostProcessSettings(DstViewport, InCameraComponent, InCameraSettings);
+	// Apply postprocess for ICVFX camera
+	InCameraComponent.ApplyICVFXCameraPostProcessesToViewport(&DstViewport, EDisplayClusterViewportCameraPostProcessFlags::All);
 
 	DstViewport.UpdateConfiguration_PostRenderOverride(InCameraSettings.RenderSettings.Replace);
 	DstViewport.UpdateConfiguration_PostRenderBlur(InCameraSettings.RenderSettings.PostprocessBlur);
diff --git a/Engine/Plugins/Runtime/nDisplay/Source/DisplayCluster/Private/Render/Viewport/Configuration/DisplayClusterViewportConfigurationHelpers_Postprocess.cpp b/Engine/Plugins/Runtime/nDisplay/Source/DisplayCluster/Private/Render/Viewport/Configuration/DisplayClusterViewportConfigurationHelpers_Postprocess.cpp
--- a/Engine/Plugins/Runtime/nDisplay/Source/DisplayCluster/Private/Render/Viewport/Configuration/DisplayClusterViewportConfigurationHelpers_Postprocess.cpp
+++ b/Engine/Plugins/Runtime/nDisplay/Source/DisplayCluster/Private/Render/Viewport/Configuration/DisplayClusterViewportConfigurationHelpers_Postprocess.cpp
@@ -8,6 +8,7 @@
 
 #include "DisplayClusterRootActor.h"
 #include "Components/DisplayClusterICVFXCameraComponent.h"
+#include "Components/DisplayClusterCameraComponent.h"
 
 namespace UE::DisplayCluster::Configuration::PostprocessHelpers
 {
@@ -28,158 +29,6 @@ namespace UE::DisplayCluster::Configuration::PostprocessHelpers
 		DstViewport.GetViewport_CustomPostProcessSettings().RemoveCustomPostProcess(RenderPass);
 	}
 
-	static inline bool ImplUpdatePerViewportColorGrading(FDisplayClusterViewport& DstViewport, const FDisplayClusterConfigurationViewport_PerViewportColorGrading PerViewportColorGrading)
-	{
-		const FDisplayClusterConfigurationICVFX_StageSettings* StageSettings = DstViewport.Configuration->GetStageSettings();
-		if (!StageSettings)
-		{
-			return false;
-		}
-
-		// enable entire cluster only when global settings is on
-		const bool bUseEntireClusterPostProcess = StageSettings->EntireClusterColorGrading.bEnableEntireClusterColorGrading && PerViewportColorGrading.bIsEntireClusterEnabled;
-
-		FDisplayClusterConfigurationViewport_CustomPostprocessSettings FinalPerViewportColorGrading;
-		FinalPerViewportColorGrading.bIsEnabled = true;
-		FinalPerViewportColorGrading.bIsOneFrame = true;
-		FinalPerViewportColorGrading.BlendWeight = 1;
-
-		// blend with entire cluster
-		if (bUseEntireClusterPostProcess)
-		{
-			FinalPerViewportColorGrading.BlendWeight = StageSettings->EntireClusterColorGrading.ColorGradingSettings.BlendWeight;
-
-			FDisplayClusterViewportConfigurationHelpers_Postprocess::BlendPostProcessSettings(FinalPerViewportColorGrading.PostProcessSettings, StageSettings->EntireClusterColorGrading.ColorGradingSettings, PerViewportColorGrading.ColorGradingSettings);
-			FinalPerViewportColorGrading.BlendWeight *= PerViewportColorGrading.ColorGradingSettings.BlendWeight;
-
-			ImplUpdateCustomPostprocess(DstViewport, FinalPerViewportColorGrading.bIsEnabled, FinalPerViewportColorGrading, IDisplayClusterViewport_CustomPostProcessSettings::ERenderPass::FinalPerViewport);
-			return true;
-		}
-		else
-		{
-			// pass color grading without blending with entire cluster
-			FDisplayClusterViewportConfigurationHelpers_Postprocess::CopyBlendPostProcessSettings(FinalPerViewportColorGrading.PostProcessSettings, PerViewportColorGrading.ColorGradingSettings);
-			FinalPerViewportColorGrading.BlendWeight *= PerViewportColorGrading.ColorGradingSettings.BlendWeight;
-
-			ImplUpdateCustomPostprocess(DstViewport, FinalPerViewportColorGrading.bIsEnabled, FinalPerViewportColorGrading, IDisplayClusterViewport_CustomPostProcessSettings::ERenderPass::FinalPerViewport);
-			return true;
-		}
-
-		return false;
-	}
-
-	static inline bool ImplUpdateEntireClusterColorGrading(FDisplayClusterViewport& DstViewport)
-	{
-		const FDisplayClusterConfigurationICVFX_StageSettings* StageSettings = DstViewport.Configuration->GetStageSettings();
-		if (StageSettings && StageSettings->EntireClusterColorGrading.bEnableEntireClusterColorGrading)
-		{
-			FDisplayClusterConfigurationViewport_CustomPostprocessSettings FinalEntireClusterColorGrading;
-			FinalEntireClusterColorGrading.bIsEnabled = true;
-			FinalEntireClusterColorGrading.bIsOneFrame = true;
-			FinalEntireClusterColorGrading.BlendWeight = 1;
-
-			FDisplayClusterViewportConfigurationHelpers_Postprocess::CopyBlendPostProcessSettings(FinalEntireClusterColorGrading.PostProcessSettings, StageSettings->EntireClusterColorGrading.ColorGradingSettings);
-			ImplUpdateCustomPostprocess(DstViewport, FinalEntireClusterColorGrading.bIsEnabled, FinalEntireClusterColorGrading, IDisplayClusterViewport_CustomPostProcessSettings::ERenderPass::FinalPerViewport);
-			return true;
-		}
-
-		return false;
-	}
-
-	static inline bool ImplUpdateIncameraPerNodeColorGrading(FDisplayClusterViewport& DstViewport, const FDisplayClusterConfigurationViewport_AllNodesColorGrading AllNodesColorGrading, const FDisplayClusterConfigurationViewport_PerNodeColorGrading PerNodeColorGrading)
-	{
-		const FDisplayClusterConfigurationICVFX_StageSettings* StageSettings = DstViewport.Configuration->GetStageSettings();
-		if (!StageSettings)
-		{
-			return false;
-		}
-
-		FDisplayClusterConfigurationViewport_CustomPostprocessSettings FinalPerNodeColorGrading;
-		FinalPerNodeColorGrading.bIsEnabled = true;
-		FinalPerNodeColorGrading.bIsOneFrame = true;
-		FinalPerNodeColorGrading.BlendWeight = 1;
-
-		const bool bIncludeUseEntireClusterPostProcess = StageSettings->EntireClusterColorGrading.bEnableEntireClusterColorGrading && PerNodeColorGrading.bEntireClusterColorGrading;
-		const bool bIncludeAllNodesColorGrading = AllNodesColorGrading.bEnableInnerFrustumAllNodesColorGrading && PerNodeColorGrading.bAllNodesColorGrading;
-
-		if (bIncludeUseEntireClusterPostProcess)
-		{
-			if (bIncludeAllNodesColorGrading)
-			{
-				// all three options are enabled - cluster + all + node
-				FDisplayClusterViewportConfigurationHelpers_Postprocess::PerNodeBlendPostProcessSettings(FinalPerNodeColorGrading.PostProcessSettings, StageSettings->EntireClusterColorGrading.ColorGradingSettings, AllNodesColorGrading.ColorGradingSettings, PerNodeColorGrading.ColorGradingSettings);
-				FinalPerNodeColorGrading.BlendWeight *= AllNodesColorGrading.ColorGradingSettings.BlendWeight * PerNodeColorGrading.ColorGradingSettings.BlendWeight;
-
-				ImplUpdateCustomPostprocess(DstViewport, FinalPerNodeColorGrading.bIsEnabled, FinalPerNodeColorGrading, IDisplayClusterViewport_CustomPostProcessSettings::ERenderPass::FinalPerViewport);
-			}
-			else
-			{
-				// only cluster + node
-				FDisplayClusterViewportConfigurationHelpers_Postprocess::BlendPostProcessSettings(FinalPerNodeColorGrading.PostProcessSettings, StageSettings->EntireClusterColorGrading.ColorGradingSettings, PerNodeColorGrading.ColorGradingSettings);
-				FinalPerNodeColorGrading.BlendWeight *= PerNodeColorGrading.ColorGradingSettings.BlendWeight;
-				ImplUpdateCustomPostprocess(DstViewport, FinalPerNodeColorGrading.bIsEnabled, FinalPerNodeColorGrading, IDisplayClusterViewport_CustomPostProcessSettings::ERenderPass::FinalPerViewport);
-			}
-
-			return true;
-		}
-		else
-		{
-			// entire cluster settings disabled, only all nodes cases
-			FinalPerNodeColorGrading.BlendWeight = AllNodesColorGrading.ColorGradingSettings.BlendWeight;
-
-			if (bIncludeAllNodesColorGrading)
-			{
-				// all nodes + node
-				FDisplayClusterViewportConfigurationHelpers_Postprocess::BlendPostProcessSettings(FinalPerNodeColorGrading.PostProcessSettings, AllNodesColorGrading.ColorGradingSettings, PerNodeColorGrading.ColorGradingSettings);
-				FinalPerNodeColorGrading.BlendWeight *= PerNodeColorGrading.ColorGradingSettings.BlendWeight;
-			}
-			else
-			{
-				// node only
-				FDisplayClusterViewportConfigurationHelpers_Postprocess::CopyBlendPostProcessSettings(FinalPerNodeColorGrading.PostProcessSettings, PerNodeColorGrading.ColorGradingSettings);
-			}
-
-			ImplUpdateCustomPostprocess(DstViewport, FinalPerNodeColorGrading.bIsEnabled, FinalPerNodeColorGrading, IDisplayClusterViewport_CustomPostProcessSettings::ERenderPass::FinalPerViewport);
-			return true;
-		}
-
-		return false;
-	}
-
-	static inline bool ImplUpdateIncameraAllNodesColorGrading(FDisplayClusterViewport& DstViewport, const FDisplayClusterConfigurationViewport_AllNodesColorGrading AllNodesColorGrading)
-	{
-		const FDisplayClusterConfigurationICVFX_StageSettings* StageSettings = DstViewport.Configuration->GetStageSettings();
-		if (!StageSettings)
-		{
-			return false;
-		}
-
-		FDisplayClusterConfigurationViewport_CustomPostprocessSettings FinalAllNodesColorGrading;
-		FinalAllNodesColorGrading.bIsEnabled = true;
-		FinalAllNodesColorGrading.bIsOneFrame = true;
-		FinalAllNodesColorGrading.BlendWeight = 1;
-
-
-		// enable entire cluster only when global settings is on
-		const bool bEnableEntireClusterColorGrading = StageSettings->EntireClusterColorGrading.bEnableEntireClusterColorGrading && AllNodesColorGrading.bEnableEntireClusterColorGrading;
-
-		if (bEnableEntireClusterColorGrading)
-		{
-			FinalAllNodesColorGrading.BlendWeight = StageSettings->EntireClusterColorGrading.ColorGradingSettings.BlendWeight;
-
-			FDisplayClusterViewportConfigurationHelpers_Postprocess::BlendPostProcessSettings(FinalAllNodesColorGrading.PostProcessSettings, StageSettings->EntireClusterColorGrading.ColorGradingSettings, AllNodesColorGrading.ColorGradingSettings);
-			FinalAllNodesColorGrading.BlendWeight *= AllNodesColorGrading.ColorGradingSettings.BlendWeight;
-		}
-		else
-		{
-			FDisplayClusterViewportConfigurationHelpers_Postprocess::CopyBlendPostProcessSettings(FinalAllNodesColorGrading.PostProcessSettings, AllNodesColorGrading.ColorGradingSettings);
-		}
-
-		ImplUpdateCustomPostprocess(DstViewport, FinalAllNodesColorGrading.bIsEnabled, FinalAllNodesColorGrading, IDisplayClusterViewport_CustomPostProcessSettings::ERenderPass::FinalPerViewport);
-
-		return true;
-	}
-
 // Note that skipped parameters in macro definitions will just evaluate to nothing
 // This is intentional to get around the inconsistent naming in the color grading fields in FPostProcessSettings
 #define PP_CONDITIONAL_BLEND(BLENDOP, COLOR, OUTGROUP, INGROUP, NAME, OFFSETOP, OFFSETVALUE) \
@@ -187,8 +36,14 @@ namespace UE::DisplayCluster::Configuration::PostprocessHelpers
 		bool bOverridePPSettings0 = PPSettings0.INGROUP bOverride_##NAME; \
 		bool bOverridePPSettings1 = (PPSettings1 != nullptr) && PPSettings1->INGROUP bOverride_##NAME; \
 		bool bOverridePPSettings2 = (PPSettings2 != nullptr) && PPSettings2->INGROUP bOverride_##NAME; \
+		bool bOverridePPSettings3 = (PPSettings3 != nullptr) && PPSettings3->INGROUP bOverride_##NAME; \
 		 \
-		if (bOverridePPSettings0 && bOverridePPSettings1 && bOverridePPSettings2) \
+		if (bOverridePPSettings0 && bOverridePPSettings1 && bOverridePPSettings2 && bOverridePPSettings3) \
+		{ \
+			OutputPP.COLOR##NAME##OUTGROUP = PPSettings0.INGROUP NAME BLENDOP PPSettings1->INGROUP NAME BLENDOP PPSettings2->INGROUP NAME BLENDOP PPSettings3->INGROUP NAME OFFSETOP OFFSETVALUE OFFSETOP OFFSETVALUE; \
+			OutputPP.bOverride_##COLOR##NAME##OUTGROUP = true; \
+		} \
+		else if (bOverridePPSettings0 && bOverridePPSettings1 && bOverridePPSettings2) \
 		{ \
 			OutputPP.COLOR##NAME##OUTGROUP = PPSettings0.INGROUP NAME BLENDOP PPSettings1->INGROUP NAME BLENDOP PPSettings2->INGROUP NAME OFFSETOP OFFSETVALUE OFFSETOP OFFSETVALUE; \
 			OutputPP.bOverride_##COLOR##NAME##OUTGROUP = true; \
@@ -236,6 +91,7 @@ namespace UE::DisplayCluster::Configuration::PostprocessHelpers
 		bool bOverridePPSettings0 = PPSettings0.INGROUP bOverride_##NAME; \
 		bool bOverridePPSettings1 = PPSettings1 && PPSettings1->INGROUP bOverride_##NAME; \
 		bool bOverridePPSettings2 = PPSettings2 && PPSettings2->INGROUP bOverride_##NAME; \
+		bool bOverridePPSettings3 = PPSettings3 && PPSettings3->INGROUP bOverride_##NAME; \
 		if (bOverridePPSettings0) \
 		{ \
 			OutputPP.COLOR##NAME##OUTGROUP = PPSettings0.INGROUP NAME; \
@@ -251,9 +107,19 @@ namespace UE::DisplayCluster::Configuration::PostprocessHelpers
 			OutputPP.COLOR##NAME##OUTGROUP = PPSettings2->INGROUP NAME; \
 			OutputPP.bOverride_##COLOR##NAME##OUTGROUP = true; \
 		} \
+		if (bOverridePPSettings3) \
+		{ \
+			OutputPP.COLOR##NAME##OUTGROUP = PPSettings3->INGROUP NAME; \
+			OutputPP.bOverride_##COLOR##NAME##OUTGROUP = true; \
+		} \
 	} \
 
-	static inline void ImplBlendPostProcessSettings(FPostProcessSettings& OutputPP, const FDisplayClusterConfigurationViewport_ColorGradingRenderingSettings& PPSettings0, const FDisplayClusterConfigurationViewport_ColorGradingRenderingSettings* PPSettings1, const FDisplayClusterConfigurationViewport_ColorGradingRenderingSettings* PPSettings2)
+	static inline void ImplBlendPostProcessSettings(
+		FPostProcessSettings& OutputPP,
+		const FDisplayClusterConfigurationViewport_ColorGradingRenderingSettings& PPSettings0,
+		const FDisplayClusterConfigurationViewport_ColorGradingRenderingSettings* PPSettings1,
+		const FDisplayClusterConfigurationViewport_ColorGradingRenderingSettings* PPSettings2,
+		const FDisplayClusterConfigurationViewport_ColorGradingRenderingSettings* PPSettings3)
 	{
 		PP_CONDITIONAL_BLEND(+, , , , AutoExposureBias, , );
 		PP_CONDITIONAL_BLEND(+, , , , ColorCorrectionHighlightsMin, , );
@@ -385,30 +251,59 @@ bool FDisplayClusterViewportConfigurationHelpers_Postprocess::ImplUpdateInnerFru
 	const FString& ClusterNodeId = DstViewport.GetClusterNodeId();
 	check(!ClusterNodeId.IsEmpty());
 
-	for (const FDisplayClusterConfigurationViewport_PerNodeColorGrading& ColorGradingProfileIt : InCameraSettings.PerNodeColorGrading)
+	const FDisplayClusterConfigurationICVFX_StageSettings* StageSettings = DstViewport.Configuration->GetStageSettings();
+	if (!StageSettings)
 	{
-		// Only allowed profiles
-		if (ColorGradingProfileIt.bIsEnabled)
+		return false;
+	}
+
+	// Collect all used color grading settings into this array.
+	TArray<const FDisplayClusterConfigurationViewport_ColorGradingRenderingSettings*> PostProcessList;
+
+	const FDisplayClusterConfigurationViewport_PerNodeColorGrading* ExistPerNodeColorGrading = InCameraSettings.PerNodeColorGrading.FindByPredicate([ClusterNodeId](const FDisplayClusterConfigurationViewport_PerNodeColorGrading& ColorGradingProfileIt)
 		{
-			for (const FString& ClusterNodeIt : ColorGradingProfileIt.ApplyPostProcessToObjects)
-			{
-				if (ClusterNodeId.Compare(ClusterNodeIt, ESearchCase::IgnoreCase) == 0)
+			// Only allowed profiles
+			return ColorGradingProfileIt.bIsEnabled && ColorGradingProfileIt.ApplyPostProcessToObjects.ContainsByPredicate([ClusterNodeId](const FString& ClusterNodeIt)
 				{
-					// Use cluster node PP
-					return PostprocessHelpers::ImplUpdateIncameraPerNodeColorGrading(DstViewport, InCameraSettings.AllNodesColorGrading, ColorGradingProfileIt);
-				}
-			}
+					return ClusterNodeId.Compare(ClusterNodeIt, ESearchCase::IgnoreCase) == 0;
+				});
+		});
+
+	if (ExistPerNodeColorGrading)
+	{
+		const bool bIncludeUseEntireClusterPostProcess = StageSettings->EntireClusterColorGrading.bEnableEntireClusterColorGrading && ExistPerNodeColorGrading->bEntireClusterColorGrading;
+		const bool bIncludeAllNodesColorGrading = InCameraSettings.AllNodesColorGrading.bEnableInnerFrustumAllNodesColorGrading && ExistPerNodeColorGrading->bAllNodesColorGrading;
+
+		// Cluster
+		if (bIncludeUseEntireClusterPostProcess)
+		{
+			PostProcessList.Add(&StageSettings->EntireClusterColorGrading.ColorGradingSettings);
+		}
+
+		// All Nodes
+		if (bIncludeAllNodesColorGrading)
+		{
+			PostProcessList.Add(&InCameraSettings.AllNodesColorGrading.ColorGradingSettings);
 		}
-	}
 
-	// run through dedicated all nodes pass only when per node list is empty
-	if (InCameraSettings.AllNodesColorGrading.bEnableInnerFrustumAllNodesColorGrading)
+		// Per-Node
+		PostProcessList.Add(&ExistPerNodeColorGrading->ColorGradingSettings);
+	}
+	else if (InCameraSettings.AllNodesColorGrading.bEnableInnerFrustumAllNodesColorGrading)
 	{
-		// all nodes only color grading
-		return PostprocessHelpers::ImplUpdateIncameraAllNodesColorGrading(DstViewport, InCameraSettings.AllNodesColorGrading);
+		const bool bEnableEntireClusterColorGrading = StageSettings->EntireClusterColorGrading.bEnableEntireClusterColorGrading && InCameraSettings.AllNodesColorGrading.bEnableEntireClusterColorGrading;
+
+		// Cluster
+		if (bEnableEntireClusterColorGrading)
+		{
+			PostProcessList.Add(&StageSettings->EntireClusterColorGrading.ColorGradingSettings);
+		}
+
+		// All Nodes
+		PostProcessList.Add(&InCameraSettings.AllNodesColorGrading.ColorGradingSettings);
 	}
 
-	return false;
+	return ImplUpdateFinalPerViewportPostProcessList(DstViewport, PostProcessList);
 }
 
 bool FDisplayClusterViewportConfigurationHelpers_Postprocess::UpdateLightcardPostProcessSettings(FDisplayClusterViewport& DstViewport, FDisplayClusterViewport& BaseViewport)
@@ -443,62 +338,270 @@ bool FDisplayClusterViewportConfigurationHelpers_Postprocess::ImplUpdateViewport
 	{
 		return false;
 	}
+
+	// Collect all used color grading settings into this array.
+	TArray<const FDisplayClusterConfigurationViewport_ColorGradingRenderingSettings*> PostProcessList;
+
+	const FDisplayClusterConfigurationViewport_PerViewportColorGrading* ExistPerViewportColorGrading = StageSettings->PerViewportColorGrading.FindByPredicate([InClusterViewportId](const FDisplayClusterConfigurationViewport_PerViewportColorGrading& ColorGradingProfileIt)
+		{
+			return ColorGradingProfileIt.bIsEnabled && ColorGradingProfileIt.ApplyPostProcessToObjects.ContainsByPredicate([InClusterViewportId](const FString& ViewportNameIt)
+				{
+					return InClusterViewportId.Compare(ViewportNameIt, ESearchCase::IgnoreCase) == 0;
+				});
+		});
+
+	// enable entire cluster only when global settings is on
+	const bool bUseEntireClusterPostProcess = StageSettings->EntireClusterColorGrading.bEnableEntireClusterColorGrading && (!ExistPerViewportColorGrading || ExistPerViewportColorGrading->bIsEntireClusterEnabled);
 	
-	for (const FDisplayClusterConfigurationViewport_PerViewportColorGrading& ColorGradingProfileIt : StageSettings->PerViewportColorGrading)
+	// Cluster
+	if (bUseEntireClusterPostProcess)
 	{
-		if (ColorGradingProfileIt.bIsEnabled)
+		PostProcessList.Add(&StageSettings->EntireClusterColorGrading.ColorGradingSettings);
+	}
+
+	// Per-Viewport
+	if (ExistPerViewportColorGrading)
+	{
+		PostProcessList.Add(&ExistPerViewportColorGrading->ColorGradingSettings);
+	}
+
+	return ImplUpdateFinalPerViewportPostProcessList(DstViewport, PostProcessList);
+}
+
+bool FDisplayClusterViewportConfigurationHelpers_Postprocess::ImplUpdateInnerFrustumColorGradingForOuterViewport(FDisplayClusterViewport& DstViewport, const FDisplayClusterConfigurationICVFX_CameraSettings& InCameraSettings)
+{
+	// per node color grading first (it includes all nodes blending too)
+	const FString& ClusterNodeId = DstViewport.GetClusterNodeId();
+	const FString& ClusterViewportId = DstViewport.GetId();
+
+	check(!ClusterNodeId.IsEmpty());
+
+	const FDisplayClusterConfigurationICVFX_StageSettings* StageSettings = DstViewport.Configuration->GetStageSettings();
+	if (!StageSettings)
+	{
+		return false;
+	}
+
+	// Collect all used color grading settings into this array.
+	TArray<const FDisplayClusterConfigurationViewport_ColorGradingRenderingSettings*> PostProcessList;
+
+	const FDisplayClusterConfigurationViewport_PerNodeColorGrading* ExistPerNodeColorGrading = InCameraSettings.PerNodeColorGrading.FindByPredicate([ClusterNodeId](const FDisplayClusterConfigurationViewport_PerNodeColorGrading& ColorGradingProfileIt)
 		{
-			for (const FString& ViewportNameIt : ColorGradingProfileIt.ApplyPostProcessToObjects)
-			{
-				if (InClusterViewportId.Compare(ViewportNameIt, ESearchCase::IgnoreCase) == 0)
+			// Only allowed profiles
+			return ColorGradingProfileIt.bIsEnabled && ColorGradingProfileIt.ApplyPostProcessToObjects.ContainsByPredicate([ClusterNodeId](const FString& ClusterNodeIt)
 				{
-					// Use per viewport blending
-					return PostprocessHelpers::ImplUpdatePerViewportColorGrading(DstViewport, ColorGradingProfileIt);
-				}
+					return ClusterNodeId.Compare(ClusterNodeIt, ESearchCase::IgnoreCase) == 0;
+				});
+		});
+
+	const FDisplayClusterConfigurationViewport_PerViewportColorGrading* ExistPerViewportColorGrading = StageSettings->PerViewportColorGrading.FindByPredicate([ClusterViewportId](const FDisplayClusterConfigurationViewport_PerViewportColorGrading& ColorGradingProfileIt)
+		{
+			return ColorGradingProfileIt.bIsEnabled && ColorGradingProfileIt.ApplyPostProcessToObjects.ContainsByPredicate([ClusterViewportId](const FString& ViewportNameIt)
+				{
+					return ClusterViewportId.Compare(ViewportNameIt, ESearchCase::IgnoreCase) == 0;
+				});
+		});
+
+	if (ExistPerNodeColorGrading)
+	{
+		const bool bIncludeUseEntireClusterPostProcess = StageSettings->EntireClusterColorGrading.bEnableEntireClusterColorGrading && ExistPerNodeColorGrading->bEntireClusterColorGrading;
+		const bool bIncludeAllNodesColorGrading = InCameraSettings.AllNodesColorGrading.bEnableInnerFrustumAllNodesColorGrading && ExistPerNodeColorGrading->bAllNodesColorGrading;
+
+		// Cluster
+		if (bIncludeUseEntireClusterPostProcess)
+		{
+			// enable entire cluster only when global settings is on
+			const bool bUseEntireClusterPostProcess = StageSettings->EntireClusterColorGrading.bEnableEntireClusterColorGrading && (!ExistPerViewportColorGrading || ExistPerViewportColorGrading->bIsEntireClusterEnabled);
+
+			// Cluster
+			if (bUseEntireClusterPostProcess)
+			{
+				PostProcessList.Add(&StageSettings->EntireClusterColorGrading.ColorGradingSettings);
+			}
+
+			// Per-Viewport
+			if (ExistPerViewportColorGrading)
+			{
+				PostProcessList.Add(&ExistPerViewportColorGrading->ColorGradingSettings);
 			}
 		}
+
+		// All Nodes
+		if (bIncludeAllNodesColorGrading)
+		{
+			PostProcessList.Add(&InCameraSettings.AllNodesColorGrading.ColorGradingSettings);
+		}
+
+		// Per-Node
+		PostProcessList.Add(&ExistPerNodeColorGrading->ColorGradingSettings);
 	}
+	else if (InCameraSettings.AllNodesColorGrading.bEnableInnerFrustumAllNodesColorGrading)
+	{
+		const bool bEnableEntireClusterColorGrading = StageSettings->EntireClusterColorGrading.bEnableEntireClusterColorGrading && InCameraSettings.AllNodesColorGrading.bEnableEntireClusterColorGrading;
 
-	// per viewport color grading is empty, entire cluster only
-	return PostprocessHelpers::ImplUpdateEntireClusterColorGrading(DstViewport);
+		// Cluster
+		if (bEnableEntireClusterColorGrading)
+		{
+			// enable entire cluster only when global settings is on
+			const bool bUseEntireClusterPostProcess = StageSettings->EntireClusterColorGrading.bEnableEntireClusterColorGrading && (!ExistPerViewportColorGrading || ExistPerViewportColorGrading->bIsEntireClusterEnabled);
+
+			// Cluster
+			if (bUseEntireClusterPostProcess)
+			{
+				PostProcessList.Add(&StageSettings->EntireClusterColorGrading.ColorGradingSettings);
+			}
+
+			// Per-Viewport
+			if (ExistPerViewportColorGrading)
+			{
+				PostProcessList.Add(&ExistPerViewportColorGrading->ColorGradingSettings);
+			}
+		}
+
+		// All Nodes
+		PostProcessList.Add(&InCameraSettings.AllNodesColorGrading.ColorGradingSettings);
+	}
+
+	return ImplUpdateFinalPerViewportPostProcessList(DstViewport, PostProcessList);
 }
 
-void FDisplayClusterViewportConfigurationHelpers_Postprocess::UpdateCameraPostProcessSettings(FDisplayClusterViewport& DstViewport, UDisplayClusterICVFXCameraComponent& InCameraComponent, const FDisplayClusterConfigurationICVFX_CameraSettings& InCameraSettings)
+void FDisplayClusterViewportConfigurationHelpers_Postprocess::ImplApplyICVFXCameraPostProcessesToViewport(FDisplayClusterViewport& DstViewport, UDisplayClusterICVFXCameraComponent& InSceneCameraComponent, const FDisplayClusterConfigurationICVFX_CameraSettings& InCfgCameraSettings, const EDisplayClusterViewportCameraPostProcessFlags InPostProcessingFlags)
 {
 	const FDisplayClusterConfigurationICVFX_StageSettings* StageSettings = DstViewport.Configuration->GetStageSettings();
-	ADisplayClusterRootActor* SceneRootActor = DstViewport.Configuration->GetRootActor(EDisplayClusterRootActorType::Scene);
-	if (!(StageSettings && SceneRootActor))
+	if (!(StageSettings))
 	{
 		return;
 	}
 
-	// check if frustum color grading is enabled	
-	if (StageSettings->EnableColorGrading && InCameraSettings.EnableInnerFrustumColorGrading)
+	const bool bIsLightcardViewport = EnumHasAnyFlags(DstViewport.GetRenderSettingsICVFX().RuntimeFlags, EDisplayClusterViewportRuntimeICVFXFlags::Lightcard);
+	if (bIsLightcardViewport)
 	{
-		FDisplayClusterConfigurationViewport_CustomPostprocessSettings CameraPPS;
-		CameraPPS.bIsOneFrame = true;
-		CameraPPS.BlendWeight = 1.f;
+		// LC viewports should not use settings from the ICVFX camera.
+		// Note: This use case needs to be clarified.
+		return;
+	}
+
+	// This function should only be used for InCamera and Outer viewports.
+	const bool bIsInCameraViewport = EnumHasAnyFlags(DstViewport.GetRenderSettingsICVFX().RuntimeFlags, EDisplayClusterViewportRuntimeICVFXFlags::InCamera);
+	const bool bIsOuterViewport = !EnumHasAnyFlags(DstViewport.GetRenderSettingsICVFX().RuntimeFlags, EDisplayClusterViewportRuntimeICVFXFlags::InternalResource);
+	if (!bIsInCameraViewport && !bIsOuterViewport)
+	{
+		return;
+	}
 
+	// All ICVFX camera post-process settings should be moved here.
+	
+	// Motion blur:
+	if (EnumHasAnyFlags(InPostProcessingFlags, EDisplayClusterViewportCameraPostProcessFlags::EnableICVFXMotionBlur))
+	{
+		DstViewport.UpdateConfiguration_CameraMotionBlur(GetICVFXCameraMotionBlurParameters(*StageSettings, InSceneCameraComponent, InCfgCameraSettings));
+	}
+
+	// Depth of field
+	if (EnumHasAnyFlags(InPostProcessingFlags, EDisplayClusterViewportCameraPostProcessFlags::EnableICVFXDepthOfFieldCompensation))
+	{
+		DstViewport.UpdateConfiguration_CameraDepthOfField(GetICVFXCameraDepthOfFieldParameters(*StageSettings, InSceneCameraComponent, InCfgCameraSettings));
+	}
+
+	// Always use postprocess from the actual camera
+	if (EnumHasAnyFlags(InPostProcessingFlags, EDisplayClusterViewportCameraPostProcessFlags::EnablePostProcess))
+	{
 		const bool bUseCameraPostprocess = true; // use internal rules of UDisplayClusterICVFXCameraComponent
 
 		// All logic was moved to the UDisplayClusterICVFXCameraComponent::GetCameraView() virtual function.
 		FMinimalViewInfo DesiredView;
-		IDisplayClusterViewport::GetCameraComponentView(&InCameraComponent, SceneRootActor->GetWorldDeltaSeconds(), bUseCameraPostprocess, DesiredView);
-		if (DesiredView.PostProcessBlendWeight > 0)
+
+		// PP is now always derived from the actual CineCamera component (ICVFXCameraComponent or from an external CineCameraActor).
+		UCineCameraComponent* ActualCineCameraComponent = InSceneCameraComponent.GetActualCineCameraComponent();
+		if (IDisplayClusterViewport::GetCameraComponentView(ActualCineCameraComponent, DstViewport.GetConfiguration().GetRootActorWorldDeltaSeconds(), bUseCameraPostprocess, DesiredView) && DesiredView.PostProcessBlendWeight > 0)
 		{
+			// Applies a filter to the post-processing settings.
+			FilterPostProcessSettings(DesiredView.PostProcessSettings, InPostProcessingFlags);
+
 			// Send camera postprocess to override
-			CameraPPS.bIsEnabled = true;
-			CameraPPS.PostProcessSettings = DesiredView.PostProcessSettings;
+			DstViewport.GetViewport_CustomPostProcessSettings().AddCustomPostProcess(IDisplayClusterViewport_CustomPostProcessSettings::ERenderPass::Override, DesiredView.PostProcessSettings, DesiredView.PostProcessBlendWeight, true);
 		}
+	}
 
-		PostprocessHelpers::ImplUpdateCustomPostprocess(DstViewport, CameraPPS.bIsEnabled, CameraPPS, IDisplayClusterViewport_CustomPostProcessSettings::ERenderPass::Override);
+	// check if frustum color grading is enabled
+	if (EnumHasAnyFlags(InPostProcessingFlags, EDisplayClusterViewportCameraPostProcessFlags::EnableICVFXColorGrading))
+	{
+		if (StageSettings->EnableColorGrading && InCfgCameraSettings.EnableInnerFrustumColorGrading)
+		{
+			if (bIsInCameraViewport)
+			{
+				// Use this function for all InCamera viewports
+				ImplUpdateInnerFrustumColorGrading(DstViewport, InCfgCameraSettings);
+			}
+			else if (bIsOuterViewport)
+			{
+				// Use this function for all Outer viewports
+				ImplUpdateInnerFrustumColorGradingForOuterViewport(DstViewport, InCfgCameraSettings);
+			}
+		}
+	}
+}
+
+void FDisplayClusterViewportConfigurationHelpers_Postprocess::FilterPostProcessSettings(FPostProcessSettings& InOutPostProcessSettings, const EDisplayClusterViewportCameraPostProcessFlags InPostProcessingFlags)
+{
+	if (!EnumHasAnyFlags(InPostProcessingFlags, EDisplayClusterViewportCameraPostProcessFlags::EnableDepthOfField))
+	{
+		// Do not override DoF PP settings from the CineCamera
+		InOutPostProcessSettings.bOverride_DepthOfFieldFstop = false;
+		InOutPostProcessSettings.bOverride_DepthOfFieldMinFstop = false;
+		InOutPostProcessSettings.bOverride_DepthOfFieldBladeCount = false;
+		InOutPostProcessSettings.bOverride_DepthOfFieldFocalDistance = false;
+		InOutPostProcessSettings.bOverride_DepthOfFieldSensorWidth = false;
+		InOutPostProcessSettings.bOverride_DepthOfFieldSqueezeFactor = false;
+	}
+}
 
-		if (!ImplUpdateInnerFrustumColorGrading(DstViewport, InCameraSettings))
+FDisplayClusterViewport_CameraMotionBlur FDisplayClusterViewportConfigurationHelpers_Postprocess::GetICVFXCameraMotionBlurParameters(const FDisplayClusterConfigurationICVFX_StageSettings& InStageSettings, UDisplayClusterICVFXCameraComponent& InSceneCameraComponent, const FDisplayClusterConfigurationICVFX_CameraSettings& InCfgCameraSettings)
+{
+	FDisplayClusterViewport_CameraMotionBlur OutParameters;
+	OutParameters.Mode = EDisplayClusterViewport_CameraMotionBlur::Undefined;
+
+	switch (InCfgCameraSettings.CameraMotionBlur.MotionBlurMode)
+	{
+	case EDisplayClusterConfigurationCameraMotionBlurMode::Override:
+		if (ADisplayClusterRootActor* SceneRootActor = static_cast<ADisplayClusterRootActor*>(InSceneCameraComponent.GetOwner()))
 		{
-			// This viewport doesn't use per-viewport PP
-			PostprocessHelpers::ImplRemoveCustomPostprocess(DstViewport, IDisplayClusterViewport_CustomPostProcessSettings::ERenderPass::FinalPerViewport);
+			UDisplayClusterCameraComponent* OuterCamera = SceneRootActor->GetDefaultCamera();
+			if (OuterCamera)
+			{
+				OutParameters.Mode = EDisplayClusterViewport_CameraMotionBlur::Override;
+
+				OutParameters.CameraLocation = OuterCamera->GetComponentLocation();
+				OutParameters.CameraRotation = OuterCamera->GetComponentRotation();
+
+				OutParameters.TranslationScale = InCfgCameraSettings.CameraMotionBlur.TranslationScale;
+			}
 		}
+		break;
+
+	case EDisplayClusterConfigurationCameraMotionBlurMode::On:
+		OutParameters.Mode = EDisplayClusterViewport_CameraMotionBlur::On;
+		break;
+
+	case EDisplayClusterConfigurationCameraMotionBlurMode::Off:
+	default:
+		OutParameters.Mode = EDisplayClusterViewport_CameraMotionBlur::Off;
+		break;
 	}
+
+	return OutParameters;
+}
+
+FDisplayClusterViewport_CameraDepthOfField FDisplayClusterViewportConfigurationHelpers_Postprocess::GetICVFXCameraDepthOfFieldParameters(const FDisplayClusterConfigurationICVFX_StageSettings& InStageSettings, UDisplayClusterICVFXCameraComponent& InSceneCameraComponent, const FDisplayClusterConfigurationICVFX_CameraSettings& InCfgCameraSettings)
+{
+	FDisplayClusterViewport_CameraDepthOfField OutParameters;
+
+	OutParameters.bEnableDepthOfFieldCompensation = InCfgCameraSettings.CameraDepthOfField.bEnableDepthOfFieldCompensation;
+	OutParameters.DistanceToWall = InCfgCameraSettings.CameraDepthOfField.DistanceToWall;
+	OutParameters.DistanceToWallOffset = InCfgCameraSettings.CameraDepthOfField.DistanceToWallOffset;
+	OutParameters.CompensationLUT = InCfgCameraSettings.CameraDepthOfField.GetCompensationLUT(InStageSettings);
+
+	return OutParameters;
 }
 
 void FDisplayClusterViewportConfigurationHelpers_Postprocess::UpdateCustomPostProcessSettings(FDisplayClusterViewport& DstViewport, const FDisplayClusterConfigurationViewport_CustomPostprocess& InCustomPostprocessConfiguration)
@@ -520,17 +623,17 @@ void FDisplayClusterViewportConfigurationHelpers_Postprocess::UpdatePerViewportP
 
 void FDisplayClusterViewportConfigurationHelpers_Postprocess::CopyBlendPostProcessSettings(FPostProcessSettings& OutputPP, const FDisplayClusterConfigurationViewport_ColorGradingRenderingSettings& InPPSettings)
 {
-	PostprocessHelpers::ImplBlendPostProcessSettings(OutputPP, InPPSettings, nullptr, nullptr);
+	PostprocessHelpers::ImplBlendPostProcessSettings(OutputPP, InPPSettings, nullptr, nullptr, nullptr);
 }
 
 void FDisplayClusterViewportConfigurationHelpers_Postprocess::PerNodeBlendPostProcessSettings(FPostProcessSettings& OutputPP, const FDisplayClusterConfigurationViewport_ColorGradingRenderingSettings& ClusterPPSettings, const FDisplayClusterConfigurationViewport_ColorGradingRenderingSettings& ViewportPPSettings, const FDisplayClusterConfigurationViewport_ColorGradingRenderingSettings& PerNodePPSettings)
 {
-	PostprocessHelpers::ImplBlendPostProcessSettings(OutputPP, ClusterPPSettings, &ViewportPPSettings, &PerNodePPSettings);
+	PostprocessHelpers::ImplBlendPostProcessSettings(OutputPP, ClusterPPSettings, &ViewportPPSettings, &PerNodePPSettings, nullptr);
 }
 
 void FDisplayClusterViewportConfigurationHelpers_Postprocess::BlendPostProcessSettings(FPostProcessSettings& OutputPP, const FDisplayClusterConfigurationViewport_ColorGradingRenderingSettings& ClusterPPSettings, const FDisplayClusterConfigurationViewport_ColorGradingRenderingSettings& ViewportPPSettings)
 {
-	PostprocessHelpers::ImplBlendPostProcessSettings(OutputPP, ClusterPPSettings, &ViewportPPSettings, nullptr);
+	PostprocessHelpers::ImplBlendPostProcessSettings(OutputPP, ClusterPPSettings, &ViewportPPSettings, nullptr, nullptr);
 }
 
 void FDisplayClusterViewportConfigurationHelpers_Postprocess::CopyPPSStructConditional(FDisplayClusterConfigurationViewport_ColorGradingRenderingSettings* OutViewportPPSettings, FPostProcessSettings* InPPS)
@@ -542,3 +645,54 @@ void FDisplayClusterViewportConfigurationHelpers_Postprocess::CopyPPSStruct(FDis
 {
 	PostprocessHelpers::ImplCopyPPSStruct(false, OutViewportPPSettings, InPPS);
 }
+
+bool FDisplayClusterViewportConfigurationHelpers_Postprocess::ImplUpdateFinalPerViewportPostProcessList(FDisplayClusterViewport& DstViewport, const TArray<const FDisplayClusterConfigurationViewport_ColorGradingRenderingSettings*>& InPostProcessList)
+{
+	FPostProcessSettings FinalPostProcessSettings;
+	float BlendWeight = 0;
+
+	const TArray<const FDisplayClusterConfigurationViewport_ColorGradingRenderingSettings*> PPList = InPostProcessList.FilterByPredicate([](const FDisplayClusterConfigurationViewport_ColorGradingRenderingSettings* ColorGradingPtr)
+		{
+			// Skip PostProcess with zero weight
+			return ColorGradingPtr && ColorGradingPtr->BlendWeight > 0;
+		});
+
+	// Now blend only up to 4 PP
+	check(PPList.Num() < 4);
+
+	switch (PPList.Num())
+	{
+	case 1:
+		PostprocessHelpers::ImplBlendPostProcessSettings(FinalPostProcessSettings, *PPList[0], nullptr, nullptr, nullptr);
+		BlendWeight = PPList[0]->BlendWeight;
+		break;
+
+	case 2:
+		PostprocessHelpers::ImplBlendPostProcessSettings(FinalPostProcessSettings, *PPList[0], PPList[1], nullptr, nullptr);
+		BlendWeight = PPList[0]->BlendWeight * PPList[1]->BlendWeight;
+		break;
+
+	case 3:
+		PostprocessHelpers::ImplBlendPostProcessSettings(FinalPostProcessSettings, *PPList[0], PPList[1], PPList[2], nullptr);
+		BlendWeight = PPList[0]->BlendWeight * PPList[1]->BlendWeight * PPList[2]->BlendWeight;
+		break;
+
+	case 4:
+		PostprocessHelpers::ImplBlendPostProcessSettings(FinalPostProcessSettings, *PPList[0], PPList[1], PPList[2], PPList[3]);
+		BlendWeight = PPList[0]->BlendWeight * PPList[1]->BlendWeight * PPList[2]->BlendWeight * PPList[3]->BlendWeight;
+		break;
+
+	default:
+	case 0:
+		return false;
+	}
+
+	if (BlendWeight > 0)
+	{
+		DstViewport.GetViewport_CustomPostProcessSettings().AddCustomPostProcess(IDisplayClusterViewport_CustomPostProcessSettings::ERenderPass::FinalPerViewport, FinalPostProcessSettings, BlendWeight, true);
+
+		return true;
+	}
+
+	return false;
+}
diff --git a/Engine/Plugins/Runtime/nDisplay/Source/DisplayCluster/Private/Render/Viewport/Configuration/DisplayClusterViewportConfigurationHelpers_Postprocess.h b/Engine/Plugins/Runtime/nDisplay/Source/DisplayCluster/Private/Render/Viewport/Configuration/DisplayClusterViewportConfigurationHelpers_Postprocess.h
--- a/Engine/Plugins/Runtime/nDisplay/Source/DisplayCluster/Private/Render/Viewport/Configuration/DisplayClusterViewportConfigurationHelpers_Postprocess.h
+++ b/Engine/Plugins/Runtime/nDisplay/Source/DisplayCluster/Private/Render/Viewport/Configuration/DisplayClusterViewportConfigurationHelpers_Postprocess.h
@@ -4,6 +4,12 @@
 
 #include "CoreMinimal.h"
 
+#include "Render/Viewport/Containers/DisplayClusterViewport_Enums.h"
+#include "Render/Viewport/Containers/DisplayClusterViewport_CameraMotionBlur.h"
+#include "Render/Viewport/Containers/DisplayClusterViewport_CameraDepthOfField.h"
+
+#include "DisplayClusterConfigurationTypes_Enums.h"
+
 class FDisplayClusterViewport;
 class UDisplayClusterICVFXCameraComponent;
 
@@ -11,6 +17,7 @@ struct FPostProcessSettings;
 struct FDisplayClusterConfigurationViewport_ColorGradingRenderingSettings;
 struct FDisplayClusterConfigurationViewport_CustomPostprocess;
 struct FDisplayClusterConfigurationICVFX_CameraSettings;
+struct FDisplayClusterConfigurationICVFX_StageSettings;
 
 /**
 * Postprocess configuration helper class.
@@ -20,7 +27,6 @@ class FDisplayClusterViewportConfigurationHelpers_Postprocess
 public:
 	static void UpdateCustomPostProcessSettings(FDisplayClusterViewport& DstViewport, const FDisplayClusterConfigurationViewport_CustomPostprocess& InCustomPostprocessConfiguration);
 	static void UpdatePerViewportPostProcessSettings(FDisplayClusterViewport& DstViewport);
-	static void UpdateCameraPostProcessSettings(FDisplayClusterViewport& DstViewport, UDisplayClusterICVFXCameraComponent& InCameraComponent, const FDisplayClusterConfigurationICVFX_CameraSettings& InCameraSettings);
 
 	static bool UpdateLightcardPostProcessSettings(FDisplayClusterViewport& DstViewport, FDisplayClusterViewport& BaseViewport);
 
@@ -35,7 +41,26 @@ public:
 	static void CopyBlendPostProcessSettings(FPostProcessSettings& OutputPP, const FDisplayClusterConfigurationViewport_ColorGradingRenderingSettings& InPPSettings);
 	static void PerNodeBlendPostProcessSettings(FPostProcessSettings& OutputPP, const FDisplayClusterConfigurationViewport_ColorGradingRenderingSettings& ClusterPPSettings, const FDisplayClusterConfigurationViewport_ColorGradingRenderingSettings& ViewportPPSettings, const FDisplayClusterConfigurationViewport_ColorGradingRenderingSettings& PerNodePPSettings);
 
+public:
+	/** Initialize PP settings for the viewport from the ICVFX camera component. */
+	static void ImplApplyICVFXCameraPostProcessesToViewport(FDisplayClusterViewport& DstViewport, UDisplayClusterICVFXCameraComponent& InSceneCameraComponent, const FDisplayClusterConfigurationICVFX_CameraSettings& InCfgCameraSettings, const EDisplayClusterViewportCameraPostProcessFlags InPostProcessingFlags);
+
+	/** Applies a filter to the post-processing settings.
+	* Note: If DoF is disabled in InPostProcessingFlags, this causes it to be removed from the PP settings as well.
+	*/
+	static void FilterPostProcessSettings(FPostProcessSettings& InOutPostProcessSettings, const EDisplayClusterViewportCameraPostProcessFlags InPostProcessingFlags);
+
 private:
 	static bool ImplUpdateInnerFrustumColorGrading(FDisplayClusterViewport& DstViewport, const FDisplayClusterConfigurationICVFX_CameraSettings& InCameraSettings);
+	static bool ImplUpdateInnerFrustumColorGradingForOuterViewport(FDisplayClusterViewport& DstViewport, const FDisplayClusterConfigurationICVFX_CameraSettings& InCameraSettings);
 	static bool ImplUpdateViewportColorGrading(FDisplayClusterViewport& DstViewport, const FString& InClusterViewportId);
+
+	/** Gets the MotionBlur parameters of camera*/
+	static FDisplayClusterViewport_CameraMotionBlur GetICVFXCameraMotionBlurParameters(const FDisplayClusterConfigurationICVFX_StageSettings& InStageSettings, UDisplayClusterICVFXCameraComponent& InSceneCameraComponent, const FDisplayClusterConfigurationICVFX_CameraSettings& InCfgCameraSettings);
+
+	/** Gets the depth of field parameters to store on the display cluster viewport */
+	static FDisplayClusterViewport_CameraDepthOfField GetICVFXCameraDepthOfFieldParameters(const FDisplayClusterConfigurationICVFX_StageSettings& InStageSettings, UDisplayClusterICVFXCameraComponent& InSceneCameraComponent, const FDisplayClusterConfigurationICVFX_CameraSettings& InCfgCameraSettings);
+
+	/** Blend input list of postprocess and set as a FinalPerViewport. */
+	static bool ImplUpdateFinalPerViewportPostProcessList(FDisplayClusterViewport& DstViewport, const TArray<const FDisplayClusterConfigurationViewport_ColorGradingRenderingSettings*>& InPostProcessList);
 };
diff --git a/Engine/Plugins/Runtime/nDisplay/Source/DisplayCluster/Private/Render/Viewport/Configuration/DisplayClusterViewportConfiguration_ICVFXCamera.cpp b/Engine/Plugins/Runtime/nDisplay/Source/DisplayCluster/Private/Render/Viewport/Configuration/DisplayClusterViewportConfiguration_ICVFXCamera.cpp
--- a/Engine/Plugins/Runtime/nDisplay/Source/DisplayCluster/Private/Render/Viewport/Configuration/DisplayClusterViewportConfiguration_ICVFXCamera.cpp
+++ b/Engine/Plugins/Runtime/nDisplay/Source/DisplayCluster/Private/Render/Viewport/Configuration/DisplayClusterViewportConfiguration_ICVFXCamera.cpp
@@ -105,17 +105,11 @@ bool FDisplayClusterViewportConfiguration_ICVFXCamera::Initialize()
 		return false;
 	}
 
-	ADisplayClusterRootActor* SceneRootActor = Configuration.GetRootActor(EDisplayClusterRootActorType::Scene);
-	if (!SceneRootActor)
-	{
-		return false;
-	}
-
 	// Applying the correct sequence of steps to use the projection policy math:
 	// SetupProjectionViewPoint()->CalculateView()->GetProjectionMatrix()
 	FMinimalViewInfo CameraViewInfo;
 	float CustomNearClippingPlane = -1; // a value less than zero means ignoring.
-	CameraProjectionPolicy->SetupProjectionViewPoint(nullptr, SceneRootActor->GetWorldDeltaSeconds(), CameraViewInfo, &CustomNearClippingPlane);
+	CameraProjectionPolicy->SetupProjectionViewPoint(nullptr, Configuration.GetRootActorWorldDeltaSeconds(), CameraViewInfo, &CustomNearClippingPlane);
 
 	CameraContext.ViewLocation = CameraViewInfo.Location;
 	CameraContext.ViewRotation = CameraViewInfo.Rotation;
diff --git a/Engine/Plugins/Runtime/nDisplay/Source/DisplayCluster/Private/Render/Viewport/Configuration/DisplayClusterViewportConfiguration_ProjectionPolicy.cpp b/Engine/Plugins/Runtime/nDisplay/Source/DisplayCluster/Private/Render/Viewport/Configuration/DisplayClusterViewportConfiguration_ProjectionPolicy.cpp
--- a/Engine/Plugins/Runtime/nDisplay/Source/DisplayCluster/Private/Render/Viewport/Configuration/DisplayClusterViewportConfiguration_ProjectionPolicy.cpp
+++ b/Engine/Plugins/Runtime/nDisplay/Source/DisplayCluster/Private/Render/Viewport/Configuration/DisplayClusterViewportConfiguration_ProjectionPolicy.cpp
@@ -19,7 +19,7 @@
 #include "Misc/DisplayClusterLog.h"
 
 #include "Components/DisplayClusterICVFXCameraComponent.h"
-
+#include "Components/DisplayClusterCameraComponent.h"
 #include "Containers/DisplayClusterProjectionCameraPolicySettings.h"
 
 namespace UE::DisplayCluster::Configuration::ProjectionPolicyHelpers
@@ -54,14 +54,25 @@ void FDisplayClusterViewportConfiguration_ProjectionPolicy::Update()
 			// ignore internal viewport
 			if (ViewportIt.IsValid() && !ViewportIt->IsInternalViewport())
 			{
-				// Support advanced logic for 'camera' projection policy
 				if (ViewportIt->GetProjectionPolicy().IsValid())
 				{
+					// Support advanced logic for 'camera' projection policy
 					if (ViewportIt->GetProjectionPolicy()->GetType().Compare(DisplayClusterProjectionStrings::projection::Camera) == 0)
 					{
 						UpdateCameraPolicy(*ViewportIt);
 					}
 
+					// Support postprocesses from the ViewPoint component
+					if (ViewportIt->GetProjectionPolicy()->ShouldUseViewPointComponentPostProcesses(ViewportIt.Get()))
+					{
+						// ViewPoint can use its own postprocess, so we apply it before anything else in case others override it.
+						if (UDisplayClusterCameraComponent* SceneViewPointCameraComponent = ViewportIt->GetViewPointCameraComponent(EDisplayClusterRootActorType::Scene))
+						{
+							// Apply the PP of the cameras referenced by the ViewPoint to the viewport.
+							SceneViewPointCameraComponent->ApplyViewPointComponentPostProcessesToViewport(ViewportIt.Get());
+						}
+					}
+
 					// Projection policies can override postprocess settings
 					// The camera PP override code has been moved from FDisplayClusterViewportConfiguration_ProjectionPolicy::UpdateCameraPolicy_Base() to the projection policy new function:
 					ViewportIt->GetProjectionPolicy()->UpdatePostProcessSettings(ViewportIt.Get());
diff --git a/Engine/Plugins/Runtime/nDisplay/Source/DisplayCluster/Private/Render/Viewport/DisplayClusterViewport_CustomPostProcessSettings.cpp b/Engine/Plugins/Runtime/nDisplay/Source/DisplayCluster/Private/Render/Viewport/DisplayClusterViewport_CustomPostProcessSettings.cpp
--- a/Engine/Plugins/Runtime/nDisplay/Source/DisplayCluster/Private/Render/Viewport/DisplayClusterViewport_CustomPostProcessSettings.cpp
+++ b/Engine/Plugins/Runtime/nDisplay/Source/DisplayCluster/Private/Render/Viewport/DisplayClusterViewport_CustomPostProcessSettings.cpp
@@ -1,8 +1,22 @@
 // Copyright Epic Games, Inc. All Rights Reserved.
 
 #include "Render/Viewport/DisplayClusterViewport_CustomPostProcessSettings.h"
+#include "Render/Viewport/DisplayClusterViewport.h"
 
+#include "Render/Viewport/Configuration/DisplayClusterViewportConfigurationHelpers_Postprocess.h"
+#include "DisplayClusterConfigurationTypes_Postprocess.h"
 
+int32 GDisplayClusterPostProcessConfigureForViewport = 1;
+static FAutoConsoleVariableRef CVarDisplayClusterPostProcessConfigureForViewport(
+	TEXT("nDisplay.render.postprocess.ConfigureForViewport"),
+	GDisplayClusterPostProcessConfigureForViewport,
+	TEXT("Enable changes to some postprocessing parameters depending on the viewport context. (DoF, etc.) (0 to disable).\n"),
+	ECVF_RenderThreadSafe
+);
+
+///////////////////////////////////////////////////////////////////////////////////////
+// FDisplayClusterViewport_CustomPostProcessSettings
+///////////////////////////////////////////////////////////////////////////////////////
 void FDisplayClusterViewport_CustomPostProcessSettings::AddCustomPostProcess(const ERenderPass InRenderPass, const FPostProcessSettings& InSettings, float BlendWeight, bool bSingleFrame)
 {
 	if (BlendWeight > 0.f) // Ignore PP with zero weights
@@ -19,16 +33,14 @@ void FDisplayClusterViewport_CustomPostProcessSettings::RemoveCustomPostProcess(
 	}
 }
 
-bool FDisplayClusterViewport_CustomPostProcessSettings::DoPostProcess(const ERenderPass InRenderPass, FPostProcessSettings* OutSettings, float* OutBlendWeight) const
+bool FDisplayClusterViewport_CustomPostProcessSettings::GetCustomPostProcess(const ERenderPass InRenderPass, FPostProcessSettings& OutSettings, float* OutBlendWeight) const
 {
 	const FPostprocessData* ExistSettings = PostprocessAsset.Find(InRenderPass);
 	if (ExistSettings && ExistSettings->bIsEnabled)
 	{
-		if (OutSettings != nullptr)
-		{
-			*OutSettings = ExistSettings->Settings;
-		}
+		OutSettings = ExistSettings->Settings;
 
+		// Returns the weight value, if appropriate.
 		if (OutBlendWeight != nullptr)
 		{
 			*OutBlendWeight = ExistSettings->BlendWeight;
@@ -51,3 +63,64 @@ void FDisplayClusterViewport_CustomPostProcessSettings::FinalizeFrame()
 		}
 	}
 }
+
+bool FDisplayClusterViewport_CustomPostProcessSettings::ApplyCustomPostProcess(IDisplayClusterViewport* InViewport, const uint32 InContextNum, const ERenderPass InRenderPass, FPostProcessSettings& InOutPPSettings, float* InOutBlendWeight) const
+{
+	bool bDidOverride = false;
+
+	switch (InRenderPass)
+	{
+	case IDisplayClusterViewport_CustomPostProcessSettings::ERenderPass::Start:
+	case IDisplayClusterViewport_CustomPostProcessSettings::ERenderPass::Override:
+		bDidOverride = GetCustomPostProcess(InRenderPass, InOutPPSettings, InOutBlendWeight);
+		break;
+
+	case IDisplayClusterViewport_CustomPostProcessSettings::ERenderPass::Final:
+	{
+		// Obtaining custom 'Final' PostProcess settings.
+		bDidOverride = GetCustomPostProcess(InRenderPass, InOutPPSettings, InOutBlendWeight);
+
+		float PerViewportPPWeight = 0;
+		FPostProcessSettings PerViewportPPSettings;
+
+		// The `Final` and `FinalPerViewport` are always applied together.
+		// If 'FinalPerViewport' is also used, apply nDisplay ColorGrading as well.
+		if (GetCustomPostProcess(IDisplayClusterViewport_CustomPostProcessSettings::ERenderPass::FinalPerViewport, PerViewportPPSettings, &PerViewportPPWeight))
+		{
+			bDidOverride = true;
+
+			// Extract nDisplay ColorGrading data from PostProcessSettings.
+			FDisplayClusterConfigurationViewport_ColorGradingRenderingSettings FinalColorGrading, PerViewportColorGrading;
+			FDisplayClusterViewportConfigurationHelpers_Postprocess::CopyPPSStruct(&FinalColorGrading, &InOutPPSettings);
+			FDisplayClusterViewportConfigurationHelpers_Postprocess::CopyPPSStructConditional(&PerViewportColorGrading, &PerViewportPPSettings);
+
+			// Blending both using our custom math instead of standard PPS blending
+			FDisplayClusterViewportConfigurationHelpers_Postprocess::BlendPostProcessSettings(InOutPPSettings, FinalColorGrading, PerViewportColorGrading);
+		}
+	}
+	break;
+
+	default:
+		break;
+	}
+
+	// Update post-processing settings for the viewport (DoF, Blur, etc.).
+	if (ConfigurePostProcessSettingsForViewport(InViewport, InContextNum, InRenderPass, InOutPPSettings))
+	{
+		bDidOverride = true;
+	}
+
+	return bDidOverride;
+}
+
+bool FDisplayClusterViewport_CustomPostProcessSettings::ConfigurePostProcessSettingsForViewport(IDisplayClusterViewport* InViewport, const uint32 InContextNum, const ERenderPass InRenderPass, FPostProcessSettings& InOutPostProcessSettings) const
+{
+	if (!InViewport || !GDisplayClusterPostProcessConfigureForViewport)
+	{
+		return false;
+	}
+
+	// Todo: Updates DoF and other PP settings for the current viewport.(JIRAs UE-219457,UE-219466)
+
+	return false;
+}
diff --git a/Engine/Plugins/Runtime/nDisplay/Source/DisplayCluster/Private/Render/Viewport/DisplayClusterViewport_CustomPostProcessSettings.h b/Engine/Plugins/Runtime/nDisplay/Source/DisplayCluster/Private/Render/Viewport/DisplayClusterViewport_CustomPostProcessSettings.h
--- a/Engine/Plugins/Runtime/nDisplay/Source/DisplayCluster/Private/Render/Viewport/DisplayClusterViewport_CustomPostProcessSettings.h
+++ b/Engine/Plugins/Runtime/nDisplay/Source/DisplayCluster/Private/Render/Viewport/DisplayClusterViewport_CustomPostProcessSettings.h
@@ -16,12 +16,31 @@ public:
 	virtual ~FDisplayClusterViewport_CustomPostProcessSettings() = default;
 
 public:
+	//~ Begin IDisplayClusterViewport_CustomPostProcessSettings
 	virtual void AddCustomPostProcess(const ERenderPass InRenderPass, const FPostProcessSettings& InSettings, float BlendWeight, bool bSingleFrame) override;
 	virtual void RemoveCustomPostProcess(const ERenderPass InRenderPass) override;
 
+	virtual bool ApplyCustomPostProcess(IDisplayClusterViewport* InViewport, const uint32 InContextNum, const ERenderPass InRenderPass, FPostProcessSettings& InOutPPSettings, float* InOutBlendWeight = nullptr) const override;
+	//~End IDisplayClusterViewport_CustomPostProcessSettings
+
 	void FinalizeFrame();
 
-	virtual bool DoPostProcess(const ERenderPass InRenderPass, FPostProcessSettings* OutSettings, float* OutBlendWeight = nullptr) const override;
+private:
+	/** Get the custom postprocess settings for the specified rendering pass.
+	* Returns false if no custom post process exists for the specified InRenderPass.
+	*/
+	bool GetCustomPostProcess(const ERenderPass InRenderPass, FPostProcessSettings& OutSettings, float* OutBlendWeight) const;
+
+	/** Applies changes to some postprocessing parameters depending on the viewport context. (DoF, etc.)
+	*
+	* @param InViewport       - DC viewport that is rendered with these settings.
+	* @param InContextNum     - Index of view that is being processed for this viewport
+	* @param InRenderPass     - PP rendering pass to be used. (The value of 'FinalPerViewport' is ignored.)
+	* @param InOutPPSettings  - (in,out) PostProcess settings that need to be changed.
+	*
+	* @return true, if the postprocess settings have been modified.
+	*/
+	bool ConfigurePostProcessSettingsForViewport(IDisplayClusterViewport* InViewport, const uint32 InContextNum, const ERenderPass InRenderPass, FPostProcessSettings& InOutPostProcessSettings) const;
 
 private:
 	struct FPostprocessData
diff --git a/Engine/Plugins/Runtime/nDisplay/Source/DisplayCluster/Private/Render/Viewport/DisplayClusterViewport_Math.cpp b/Engine/Plugins/Runtime/nDisplay/Source/DisplayCluster/Private/Render/Viewport/DisplayClusterViewport_Math.cpp
--- a/Engine/Plugins/Runtime/nDisplay/Source/DisplayCluster/Private/Render/Viewport/DisplayClusterViewport_Math.cpp
+++ b/Engine/Plugins/Runtime/nDisplay/Source/DisplayCluster/Private/Render/Viewport/DisplayClusterViewport_Math.cpp
@@ -86,8 +86,8 @@ bool FDisplayClusterViewport::GetViewPointCameraEye(const uint32 InContextNum, F
 	// Here we use the ViewPoint component as the eye position
 	if (UDisplayClusterCameraComponent* SceneCameraComponent = GetViewPointCameraComponent(EDisplayClusterRootActorType::Scene))
 	{
-		OutViewLocation = SceneCameraComponent->GetComponentLocation();
-		OutViewRotation = SceneCameraComponent->GetComponentRotation();
+		// Get eye position from the ViewPoint component.
+		SceneCameraComponent->GetEyePosition(GetConfiguration(), OutViewLocation, OutViewRotation);
 
 		// Calculate stereo ViewOffset:
 		OutViewOffset = ImplGetViewOffset(GetStereoEyeOffsetDistance(InContextNum), OutViewLocation, OutViewRotation);
diff --git a/Engine/Plugins/Runtime/nDisplay/Source/DisplayCluster/Private/Render/Viewport/DisplayClusterViewport_ViewPoint.cpp b/Engine/Plugins/Runtime/nDisplay/Source/DisplayCluster/Private/Render/Viewport/DisplayClusterViewport_ViewPoint.cpp
--- a/Engine/Plugins/Runtime/nDisplay/Source/DisplayCluster/Private/Render/Viewport/DisplayClusterViewport_ViewPoint.cpp
+++ b/Engine/Plugins/Runtime/nDisplay/Source/DisplayCluster/Private/Render/Viewport/DisplayClusterViewport_ViewPoint.cpp
@@ -176,11 +176,7 @@ bool FDisplayClusterViewport::SetupViewPoint(FMinimalViewInfo& InOutViewInfo)
 		// The projection policy can override these ViewPoint data.
 		if (ProjectionPolicy.IsValid())
 		{
-			if (ADisplayClusterRootActor* SceneRootActor = Configuration->GetRootActor(EDisplayClusterRootActorType::Scene))
-			{
-				const float DeltaTime = SceneRootActor->GetWorldDeltaSeconds();
-				ProjectionPolicy->SetupProjectionViewPoint(this, DeltaTime, InOutViewInfo, &CustomNearClippingPlane);
-			}
+			ProjectionPolicy->SetupProjectionViewPoint(this, Configuration->GetRootActorWorldDeltaSeconds(), InOutViewInfo, &CustomNearClippingPlane);
 		}
 
 		return true;
diff --git a/Engine/Plugins/Runtime/nDisplay/Source/DisplayCluster/Private/Render/Viewport/Preview/DisplayClusterViewportPreview_Render.cpp b/Engine/Plugins/Runtime/nDisplay/Source/DisplayCluster/Private/Render/Viewport/Preview/DisplayClusterViewportPreview_Render.cpp
--- a/Engine/Plugins/Runtime/nDisplay/Source/DisplayCluster/Private/Render/Viewport/Preview/DisplayClusterViewportPreview_Render.cpp
+++ b/Engine/Plugins/Runtime/nDisplay/Source/DisplayCluster/Private/Render/Viewport/Preview/DisplayClusterViewportPreview_Render.cpp
@@ -16,9 +16,6 @@
 
 #include "Render/Viewport/DisplayClusterViewportManager.h"
 #include "Render/Viewport/RenderFrame/DisplayClusterRenderFrameSettings.h"
-#include "Render/Viewport/Configuration/DisplayClusterViewportConfigurationHelpers_Postprocess.h"
-
-#include "DisplayClusterConfigurationTypes_Postprocess.h"
 
 #include "Misc/DisplayClusterLog.h"
 
@@ -106,38 +103,25 @@ FSceneView* FDisplayClusterViewportPreview::CalcSceneView(FSceneViewFamilyContex
 
 		InOutViewFamily.Views.Add(View);
 
+		// Configure postprocesses for the current viewport.
+		// The code below is based on the code from ULocalPlayer.
+		{
+			// ERenderPass::Start
 		View->StartFinalPostprocessSettings(ViewLocation);
+			InViewport->GetViewport_CustomPostProcessSettings().ApplyCustomPostProcess(InViewport, InContextNum, IDisplayClusterViewport_CustomPostProcessSettings::ERenderPass::Start, View->FinalPostProcessSettings);
 
-		FPostProcessSettings* FinalPostProcessingSettings = &View->FinalPostProcessSettings;
-		// Support start PPS for preview
-		InViewport->GetViewport_CustomPostProcessSettings().DoPostProcess(IDisplayClusterViewport_CustomPostProcessSettings::ERenderPass::Start, FinalPostProcessingSettings);
-
-		// Support override PPS for preview
+			// ERenderPass::Override
 		FPostProcessSettings OverridePostProcessingSettings;
 		float OverridePostProcessBlendWeight = 1.0f;
-		InViewport->GetViewport_CustomPostProcessSettings().DoPostProcess(IDisplayClusterViewport_CustomPostProcessSettings::ERenderPass::Override, &OverridePostProcessingSettings, &OverridePostProcessBlendWeight);
-
-		View->OverridePostProcessSettings(OverridePostProcessingSettings, OverridePostProcessBlendWeight);
-
-		// Support final PPS for preview
-		InViewport->GetViewport_CustomPostProcessSettings().DoPostProcess(IDisplayClusterViewport_CustomPostProcessSettings::ERenderPass::Final, FinalPostProcessingSettings);
-
-		FPostProcessSettings RequestedFinalPerViewportPPS;
-		// Get the final overall cluster + per-viewport PPS from nDisplay
-		if (InViewport->GetViewport_CustomPostProcessSettings().DoPostProcess(IDisplayClusterViewport_CustomPostProcessSettings::ERenderPass::FinalPerViewport, &RequestedFinalPerViewportPPS))
+			if (InViewport->GetViewport_CustomPostProcessSettings().ApplyCustomPostProcess(InViewport, InContextNum, IDisplayClusterViewport_CustomPostProcessSettings::ERenderPass::Override, OverridePostProcessingSettings, &OverridePostProcessBlendWeight))
 		{
-			FDisplayClusterConfigurationViewport_ColorGradingRenderingSettings InPPSnDisplay;
-			FDisplayClusterViewportConfigurationHelpers_Postprocess::CopyPPSStructConditional(&InPPSnDisplay, &RequestedFinalPerViewportPPS);
-
-			// Get the passed-in cumulative PPS from the game/viewport (includes all PPVs affecting this viewport)
-			FDisplayClusterConfigurationViewport_ColorGradingRenderingSettings InPPSCumulative;
-			FDisplayClusterViewportConfigurationHelpers_Postprocess::CopyPPSStruct(&InPPSCumulative, FinalPostProcessingSettings);
-
-			// Blend both together with our custom math instead of the default PPS blending
-			FDisplayClusterViewportConfigurationHelpers_Postprocess::BlendPostProcessSettings(*FinalPostProcessingSettings, InPPSCumulative, InPPSnDisplay);
+			View->OverridePostProcessSettings(OverridePostProcessingSettings, OverridePostProcessBlendWeight);
 		}
 
+			// ERenderPass::Final
+			InViewport->GetViewport_CustomPostProcessSettings().ApplyCustomPostProcess(InViewport, InContextNum, IDisplayClusterViewport_CustomPostProcessSettings::ERenderPass::Final, View->FinalPostProcessSettings);
 		View->EndFinalPostprocessSettings(ViewInitOptions);
+		}
 
 		// Setup view extension for this view
 		for (int32 ViewExt = 0; ViewExt < InOutViewFamily.ViewExtensions.Num(); ViewExt++)
diff --git a/Engine/Plugins/Runtime/nDisplay/Source/DisplayCluster/Public/Components/DisplayClusterCameraComponent.h b/Engine/Plugins/Runtime/nDisplay/Source/DisplayCluster/Public/Components/DisplayClusterCameraComponent.h
--- a/Engine/Plugins/Runtime/nDisplay/Source/DisplayCluster/Public/Components/DisplayClusterCameraComponent.h
+++ b/Engine/Plugins/Runtime/nDisplay/Source/DisplayCluster/Public/Components/DisplayClusterCameraComponent.h
@@ -7,6 +7,7 @@
 #include "Materials/MaterialInstanceDynamic.h"
 #include "Components/IDisplayClusterComponent.h"
 #include "Render/DisplayDevice/Containers/DisplayClusterDisplayDevice_Enums.h"
+#include "Render/Viewport/Containers/DisplayClusterViewport_Enums.h"
 
 #include "DisplayClusterCameraComponent.generated.h"
 
@@ -19,6 +20,8 @@ class IDisplayClusterViewportManager;
 class IDisplayClusterWarpPolicy;
 class IDisplayClusterViewportConfiguration;
 class IDisplayClusterViewportPreview;
+class IDisplayClusterViewport;
+class UCameraComponent;
 struct FMinimalViewInfo;
 
 UENUM()
@@ -29,6 +32,42 @@ enum class EDisplayClusterEyeStereoOffset : uint8
 	Right UMETA(DisplayName = "Right Eye"),
 };
 
+/**
+* Specifies the parameters to be used from the specified camera.
+*/
+USTRUCT(BlueprintType)
+struct DISPLAYCLUSTER_API FDisplayClusterCameraComponent_OuterViewportPostProcessSettings
+{
+	GENERATED_BODY()
+
+	/** Decodes parameters into flags. */
+	EDisplayClusterViewportCameraPostProcessFlags GetCameraPostProcessFlags() const;
+
+public:
+	/** Use the NearClippingPlane value from the specified cine camera. */
+	UPROPERTY(EditAnywhere, BlueprintReadWrite, Category = "Advanced", meta = (DisplayName = "Use Custom Near Clipping Plane"))
+	bool bEnableNearClippingPlane = false;
+
+	/** Use the PP settings from the specified camera. */
+	UPROPERTY(EditAnywhere, BlueprintReadWrite, Category = "Advanced", meta = (DisplayName = "Use Post Process"))
+	bool bEnablePostProcess = true;
+
+	/** Enable the DoF PP settings from the specified camera. */
+	UPROPERTY(EditAnywhere, BlueprintReadWrite, Category = "Advanced", meta = (DisplayName = "Use Depth Of Field"))
+	bool bEnableDepthOfField = false;
+
+	/** Use the DC Depth-Of-Field settings from the specified ICVFX camera. */
+	UPROPERTY(EditAnywhere, BlueprintReadWrite, Category = "Advanced", meta = (DisplayName = "Use ICVFX Depth Of Field Compensation"))
+	bool bEnableICVFXDepthOfFieldCompensation = false;
+
+	/** Use the DC ColorGrading from the specified ICVFX camera. */
+	UPROPERTY(EditAnywhere, BlueprintReadWrite, Category = "Advanced", meta = (DisplayName = "Use ICVFX Color Grading"))
+	bool bEnableICVFXColorGrading = true;
+
+	/** Use the DC Motion Blur settings from the specified ICVFX camera. */
+	UPROPERTY(EditAnywhere, BlueprintReadWrite, Category = "Advanced", meta = (DisplayName = "Use ICVFX Motion Blur"))
+	bool bEnableICVFXMotionBlur = false;
+};
 
 /**
  * 3D point in space used to render nDisplay viewports from
@@ -61,6 +100,9 @@ public:
 	 */
 	virtual void GetDesiredView(IDisplayClusterViewportConfiguration& InViewportConfiguration, FMinimalViewInfo& InOutViewInfo, float* OutCustomNearClippingPlane = nullptr);
 
+	/** Returns the position of the observer's eyes in the Stage. */
+	virtual void GetEyePosition(const IDisplayClusterViewportConfiguration& InViewportConfiguration, FVector& OutViewLocation, FRotator& OutViewRotation);
+
 	/**
 	 * All cluster viewports that reference this component will be created in the background on the current cluster node if the function returns true.
 	 */
@@ -102,6 +144,22 @@ public:
 	virtual void OnUpdateDisplayDeviceMeshAndMaterialInstance(IDisplayClusterViewportPreview& InViewportPreview, const EDisplayClusterDisplayDeviceMeshType InMeshType, const EDisplayClusterDisplayDeviceMaterialType InMaterialType, UMeshComponent* InMeshComponent, UMaterialInstanceDynamic* InMeshMaterialInstance) const
 	{ }
 
+	/** Apply the ViewPoint component's post-processes to the viewport.
+	* (Outer viewport camera)
+	*
+	* @param InViewport - viewport to be configured.
+	*/
+	virtual void ApplyViewPointComponentPostProcessesToViewport(IDisplayClusterViewport* InViewport);
+
+	/** Return a reference to the Camera component, which is used for Outer viewports.
+	* 
+	* return nullptr if the camera is not in use.
+	*/
+	virtual UCameraComponent* GetOuterViewportCameraComponent(const IDisplayClusterViewportConfiguration& InViewportConfiguration) const;
+
+protected:
+	/** Get Outer Viewport Camera view. */
+	virtual bool GetOuterViewportCameraDesiredViewInternal(const IDisplayClusterViewportConfiguration& InViewportConfiguration, FMinimalViewInfo& InOutViewInfo, float* OutCustomNearClippingPlane = nullptr) const;
 
 public:
 	/**
@@ -226,6 +284,24 @@ protected:
 	TObjectPtr<UTexture2D> SpriteTexture;
 #endif
 
+public:
+	/** Use the post process from the specified camera. This applies to all viewports that use this viewpoint. */
+	UPROPERTY(EditAnywhere, Category = "Outer Viewport Post Process", meta = (DisplayName = "Use Outer Viewport Camera"))
+	bool bEnableOuterViewportCamera = false;
+
+	/** The viewpoint location follows the camera location. */
+	UPROPERTY(EditAnywhere, Category = "Outer Viewport Post Process", meta = (DisplayName = "Follow Outer Viewport Camera"))
+	bool bFollowOuterViewportCamera = false;
+
+	/** The name of the camera component that is used as the PP source.
+	* (An empty string means that the active game camera is used). */
+	UPROPERTY(EditAnywhere, Category = "Outer Viewport Post Process", meta = (DisplayName = "Outer Viewport Camera Name"))
+	FString OuterViewportCameraName;
+
+	/** Additional settings that control how PP will be used. */
+	UPROPERTY(EditAnywhere, Category = "Outer Viewport Post Process", meta = (DisplayName = "Post Process Settings"))
+	FDisplayClusterCameraComponent_OuterViewportPostProcessSettings OuterViewportPostProcessSettings;
+
 private:
 	UPROPERTY(EditAnywhere, Category = "Stereo")
 	float InterpupillaryDistance;
diff --git a/Engine/Plugins/Runtime/nDisplay/Source/DisplayCluster/Public/Components/DisplayClusterICVFXCameraComponent.h b/Engine/Plugins/Runtime/nDisplay/Source/DisplayCluster/Public/Components/DisplayClusterICVFXCameraComponent.h
--- a/Engine/Plugins/Runtime/nDisplay/Source/DisplayCluster/Public/Components/DisplayClusterICVFXCameraComponent.h
+++ b/Engine/Plugins/Runtime/nDisplay/Source/DisplayCluster/Public/Components/DisplayClusterICVFXCameraComponent.h
@@ -9,14 +9,15 @@
 
 #include "DisplayClusterConfigurationTypes_ICVFX.h"
 #include "DisplayClusterEditorPropertyReference.h"
+#include "Render/Viewport/Containers/DisplayClusterViewport_Enums.h"
 #include "Render/Viewport/Containers/DisplayClusterViewport_CameraMotionBlur.h"
-#include "Render/Viewport/Containers/DisplayClusterViewport_CameraDepthOfField.h"
 
 #include "DisplayClusterICVFXCameraComponent.generated.h"
 
 struct FMinimalViewInfo;
 class SWidget;
 class UCameraComponent;
+class IDisplayClusterViewport;
 
 
 /**
@@ -46,11 +47,19 @@ public:
 
 	virtual void TickComponent(float DeltaTime, ELevelTick TickType, FActorComponentTickFunction* ThisTickFunction) override;
 
-public:
-	FDisplayClusterViewport_CameraMotionBlur GetMotionBlurParameters();
+	/** The ICVFXCamera component uses its own PostProcess and applies it to the viewport.
+	*
+	* @param InViewport - a viewport that need to be configured with a PP settings.
+	* @param InPostProcessFlags - Define the PostProcess set that will be used
+	*/
+	virtual void ApplyICVFXCameraPostProcessesToViewport(IDisplayClusterViewport* InDestViewport, const EDisplayClusterViewportCameraPostProcessFlags InPostProcessingFlags);
 
-	/** Gets the depth of field parameters to store on the display cluster viewport */
-	FDisplayClusterViewport_CameraDepthOfField GetDepthOfFieldParameters();
+public:
+	UE_DEPRECATED(5.5, "This function has been deprecated. Please use 'FDisplayClusterViewportConfigurationHelpers_Postprocess::GetICVFXCameraMotionBlurParameters()'.")
+	FDisplayClusterViewport_CameraMotionBlur GetMotionBlurParameters()
+	{
+		return FDisplayClusterViewport_CameraMotionBlur();
+	}
 
 	/**
 	 * Return the actual source camera, e.g. the camera component of the referenced cine camera.
@@ -68,10 +77,7 @@ public:
 	// Return unique camera name
 	FString GetCameraUniqueId() const;
 
-	const FDisplayClusterConfigurationICVFX_CameraSettings& GetCameraSettingsICVFX() const
-	{
-		return CameraSettings;
-	}
+	const FDisplayClusterConfigurationICVFX_CameraSettings& GetCameraSettingsICVFX() const;
 
 	/** Obtaining view information for the actual camera, such as the camera component to which the cine-camera is referencing.
 	 * The data from the CameraSettings variable is used in postprocess settings (EnableCameraPP, OverrideMotionBlur, etc.).
diff --git a/Engine/Plugins/Runtime/nDisplay/Source/DisplayCluster/Public/Render/Projection/IDisplayClusterProjectionPolicy.h b/Engine/Plugins/Runtime/nDisplay/Source/DisplayCluster/Public/Render/Projection/IDisplayClusterProjectionPolicy.h
--- a/Engine/Plugins/Runtime/nDisplay/Source/DisplayCluster/Public/Render/Projection/IDisplayClusterProjectionPolicy.h
+++ b/Engine/Plugins/Runtime/nDisplay/Source/DisplayCluster/Public/Render/Projection/IDisplayClusterProjectionPolicy.h
@@ -175,6 +175,12 @@ public:
 	virtual void SetupProjectionViewPoint(IDisplayClusterViewport* InViewport, const float InDeltaTime, FMinimalViewInfo& InOutViewInfo, float* OutCustomNearClippingPlane = nullptr)
 	{ }
 
+	/* Returns true if this projection policy should use post-processes from the ViewPoint component. **/
+	virtual bool ShouldUseViewPointComponentPostProcesses(IDisplayClusterViewport* InViewport) const
+	{
+		return true;
+	}
+
 	/** Projection policy can override PP
 	* 
 	* @param InViewport - a owner viewport
@@ -185,7 +191,7 @@ public:
 	/** Calculate view projection data
 	* 
 	* @param InViewport        - a owner viewport
-	* @param ViewIdx           - Index of view that is being processed for this viewport
+	* @param InContextNum      - Index of view that is being processed for this viewport
 	* @param InOutViewLocation - (in/out) View location with ViewOffset (i.e. left eye pre-computed location)
 	* @param InOutViewRotation - (in/out) View rotation
 	* @param ViewOffset        - Offset applied ot a camera location that gives us InOutViewLocation (i.e. right offset in world to compute right eye location)
@@ -200,13 +206,25 @@ public:
 	/** Gets projection matrix
 	* 
 	* @param InViewport   - a owner viewport
-	* @param ViewIdx      - Index of view that is being processed for this viewport
+	* @param InContextNum - Index of view that is being processed for this viewport
 	* @param OutPrjMatrix - (out) projection matrix
 	*
 	* @return - True if success
 	*/
 	virtual bool GetProjectionMatrix(IDisplayClusterViewport* InViewport, const uint32 InContextNum, FMatrix& OutPrjMatrix) = 0;
 
+	/** Returns true if the frustum has been rotated to fit the context size.
+	* 
+	* @param InViewport   - a owner viewport
+	* @param InContextNum - Index of view that is being processed for this viewport
+	* 
+	* @return true is success.
+	*/
+	virtual bool IsFrustumRotatedToFitContextSize(IDisplayClusterViewport* InViewport, const uint32 InContextNum)
+	{
+		return false;
+	}
+
 	/**
 	* Returns if a policy provides warp&blend feature
 	*
diff --git a/Engine/Plugins/Runtime/nDisplay/Source/DisplayCluster/Public/Render/Viewport/Containers/DisplayClusterViewport_Enums.h b/Engine/Plugins/Runtime/nDisplay/Source/DisplayCluster/Public/Render/Viewport/Containers/DisplayClusterViewport_Enums.h
--- a/Engine/Plugins/Runtime/nDisplay/Source/DisplayCluster/Public/Render/Viewport/Containers/DisplayClusterViewport_Enums.h
+++ b/Engine/Plugins/Runtime/nDisplay/Source/DisplayCluster/Public/Render/Viewport/Containers/DisplayClusterViewport_Enums.h
@@ -202,3 +202,40 @@ enum class EDisplayClusterViewportMediaState : uint8
 	InputLateOCIO = 1 << 3,
 };
 ENUM_CLASS_FLAGS(EDisplayClusterViewportMediaState);
+
+/**
+* A set of flags defining the post-process parameters to be received from the camera.
+*/
+enum class EDisplayClusterViewportCameraPostProcessFlags : uint8
+{
+	// Ignore all post-processing from the camera.
+	None = 0,
+
+	// Use the PP settings from the specified camera.
+	// (Supported by Camera, CineCamera and ICVFXCamera components).
+	EnablePostProcess = 1 << 0,
+
+	// Enable the CineCamera DoF PP settings from the specified camera.
+	// (Supported by CineCamera and ICVFXCamera components).
+	EnableDepthOfField = 1 << 1,
+
+	// Use the custom NearClippingPlane value from the specified cine camera.
+	// (Supported by CineCamera and ICVFXCamera components).
+	EnableNearClippingPlane = 1 << 2,
+
+	// Use the DC ColorGrading from the specified ICVFX camera.
+	// (Supported by ICVFXCamera component only).
+	EnableICVFXColorGrading = 1 << 3,
+
+	// Use the DC Motion Blur settings from the specified ICVFX camera.
+	// (Supported by ICVFXCamera component only).
+	EnableICVFXMotionBlur = 1 << 4,
+
+	// Use the DC Depth-Of-Field settings from the specified ICVFX camera.
+	// (Supported by ICVFXCamera component only).
+	EnableICVFXDepthOfFieldCompensation = 1 << 5,
+
+	// Apply all possible post-processing from the camera.
+	All = 0xFF
+};
+ENUM_CLASS_FLAGS(EDisplayClusterViewportCameraPostProcessFlags);
diff --git a/Engine/Plugins/Runtime/nDisplay/Source/DisplayCluster/Public/Render/Viewport/IDisplayClusterViewportConfiguration.h b/Engine/Plugins/Runtime/nDisplay/Source/DisplayCluster/Public/Render/Viewport/IDisplayClusterViewportConfiguration.h
--- a/Engine/Plugins/Runtime/nDisplay/Source/DisplayCluster/Public/Render/Viewport/IDisplayClusterViewportConfiguration.h
+++ b/Engine/Plugins/Runtime/nDisplay/Source/DisplayCluster/Public/Render/Viewport/IDisplayClusterViewportConfiguration.h
@@ -75,6 +75,9 @@ public:
 	 */
 	virtual UWorld* GetCurrentWorld() const = 0;
 
+	/** Return value from the DCRA function GetWorldDeltaSeconds(). */
+	virtual float GetRootActorWorldDeltaSeconds(const EDisplayClusterRootActorType InRootActorType = EDisplayClusterRootActorType::Scene) const = 0;
+
 	/**
 	 * Gets a reference to the DCRA by type.
 	 * If a DCRA with the specified type is not assigned, the default DCRA is used.
diff --git a/Engine/Plugins/Runtime/nDisplay/Source/DisplayCluster/Public/Render/Viewport/IDisplayClusterViewport_CustomPostProcessSettings.h b/Engine/Plugins/Runtime/nDisplay/Source/DisplayCluster/Public/Render/Viewport/IDisplayClusterViewport_CustomPostProcessSettings.h
--- a/Engine/Plugins/Runtime/nDisplay/Source/DisplayCluster/Public/Render/Viewport/IDisplayClusterViewport_CustomPostProcessSettings.h
+++ b/Engine/Plugins/Runtime/nDisplay/Source/DisplayCluster/Public/Render/Viewport/IDisplayClusterViewport_CustomPostProcessSettings.h
@@ -4,6 +4,9 @@
 
 #include "CoreMinimal.h"
 
+class IDisplayClusterViewport;
+struct FPostProcessSettings;
+
 /**
 * DC Viewport Postprocess interface
 */
@@ -13,22 +16,49 @@ public:
 	virtual ~IDisplayClusterViewport_CustomPostProcessSettings() = default;
 
 public:
+	/** The viewport's custom PostProcesses are blended in a specific order, which is specified by the values below.*/
 	enum class ERenderPass : uint8
 	{
+		/** These PP settings will be applied when the StartFinalPostprocessSettings() function is called */
 		Start = 0,
+
+		/** These PP settings will be applied when the OverrideFinalPostprocessSettings() function is called */
 		Override,
+
+		/** These PP settings will be applied when the EndFinalPostprocessSettings() function is called.
+		* The `Final` and `FinalPerViewport` are always applied together.*/
 		Final,
+
+		/** This rendering pass is for the container only, to separate ICVFX ColorGrading into a separate pass.
+		* Note: The value is ignored by ApplyCustomPostProcess(). */
 		FinalPerViewport,
 	};
 
 public:
-	// Override posproces, if defined
-	// * return true, if processed
-	virtual bool DoPostProcess(const ERenderPass InRenderPass, struct FPostProcessSettings* OutSettings, float* OutBlendWeight = nullptr) const = 0;
-
 	/** Add custom postprocess. */
 	virtual void AddCustomPostProcess(const ERenderPass InRenderPass, const FPostProcessSettings& InSettings, float BlendWeight, bool bSingleFrame) = 0;
 
-	/** remove custom postprocess. */
+	/** Deletes the custom postprocess for the specified rendering pass. */
 	virtual void RemoveCustomPostProcess(const ERenderPass InRenderPass) = 0;
+
+	/** Apply postprocess for the specified rendering pass to the InOutStartPostProcessingSettings.
+	* 
+	* @param InViewport       - DC viewport that is rendered with these settings.
+	* @param InContextNum     - Index of view that is being processed for this viewport
+	* @param InRenderPass     - PP rendering pass to be used. (The value of 'FinalPerViewport' is ignored.)
+	* @param InOutPPSettings  - (in,out) PostProcess settings that need to be changed.
+	* @param InOutBlendWeight - (opt, in, out) PostProcess weight parameter that need to be changed.
+	* 
+	* @return true, if the postprocess settings have been overridden.
+	*/
+	virtual bool ApplyCustomPostProcess(IDisplayClusterViewport* InViewport, const uint32 InContextNum, const ERenderPass InRenderPass, FPostProcessSettings& InOutPPSettings, float* InOutBlendWeight = nullptr) const = 0;
+
+public:
+	/** Override postprocess, if defined. */
+	UE_DEPRECATED(5.5, "This function has been deprecated. Please use 'ApplyCustomPostProcess()'.")
+	virtual bool DoPostProcess(const ERenderPass InRenderPass, struct FPostProcessSettings* OutSettings, float* OutBlendWeight = nullptr) const
+	{
+		return false;
+	}
+
 };
diff --git a/Engine/Plugins/Runtime/nDisplay/Source/DisplayCluster/Public/Render/Viewport/Misc/DisplayClusterViewportHelpers.h b/Engine/Plugins/Runtime/nDisplay/Source/DisplayCluster/Public/Render/Viewport/Misc/DisplayClusterViewportHelpers.h
new file mode 100644
--- /dev/null
+++ b/Engine/Plugins/Runtime/nDisplay/Source/DisplayCluster/Public/Render/Viewport/Misc/DisplayClusterViewportHelpers.h
@@ -0,0 +1,64 @@
+// Copyright Epic Games, Inc. All Rights Reserved.
+
+#pragma once
+
+#include "CoreMinimal.h"
+#include "DisplayClusterRootActor.h"
+#include "Render/Viewport/IDisplayClusterViewportConfiguration.h"
+
+/**
+* Helper functions for the DisplayClusterViewport
+*/
+namespace UE::DisplayClusterViewportHelpers
+{
+	/** Find a component by name from a RootActor of the specified type.
+	*
+	* @param InRootActorType - the search is performed within the DCRA of the given type.
+	* @param InComponentName - the name of the component
+	*/
+	template<class TComp>
+	TComp* GetRootActorComponentByName(const IDisplayClusterViewportConfiguration& InConfiguration, const EDisplayClusterRootActorType InRootActorType, const FString& InComponentName)
+	{
+		if (InComponentName.IsEmpty())
+		{
+			return nullptr;
+		}
+
+		if (const ADisplayClusterRootActor* RequestedRootActor = InConfiguration.GetRootActor(InRootActorType))
+		{
+			if (TComp* Component = RequestedRootActor->GetComponentByName<TComp>(InComponentName))
+			{
+				return Component;
+			}
+		}
+
+		return nullptr;
+	}
+
+	/** Return the same component (by class and name) from a DCRA of the specified type.
+	* If no other component exists, or if the current DCRA type matches the requested one,
+	* the input component is returned.
+	*
+	* @param InRootActorType - the search is performed within the DCRA of the given type.
+	* @param InComponent     - The component that is used for the search.
+	*
+	* @return A component with the same name and class belonging to a DCRA of this type.
+	*/
+	template<class TComp>
+	const TComp& GetMatchingComponentFromRootActor(const IDisplayClusterViewportConfiguration& InConfiguration, const EDisplayClusterRootActorType InRootActorType, const TComp& InComponent)
+	{
+		if (const ADisplayClusterRootActor* RequestedRootActor = InConfiguration.GetRootActor(InRootActorType))
+		{
+			// If the current DCRA type does not match the requested one.
+			if (RequestedRootActor != InComponent.GetOwner())
+			{
+				if (const TComp* OutComponent = RequestedRootActor->GetComponentByName<TComp>(InComponent.GetName()))
+				{
+					return *OutComponent;
+				}
+			}
+		}
+
+		return InComponent;
+	}
+};
diff --git a/Engine/Plugins/Runtime/nDisplay/Source/DisplayClusterConfiguration/Private/DisplayClusterConfigurationTypes_ICVFX.cpp b/Engine/Plugins/Runtime/nDisplay/Source/DisplayClusterConfiguration/Private/DisplayClusterConfigurationTypes_ICVFX.cpp
--- a/Engine/Plugins/Runtime/nDisplay/Source/DisplayClusterConfiguration/Private/DisplayClusterConfigurationTypes_ICVFX.cpp
+++ b/Engine/Plugins/Runtime/nDisplay/Source/DisplayClusterConfiguration/Private/DisplayClusterConfigurationTypes_ICVFX.cpp
@@ -53,7 +53,7 @@ FDisplayClusterConfigurationICVFX_CameraRenderSettings::FDisplayClusterConfigura
 
 void FDisplayClusterConfigurationICVFX_CameraRenderSettings::SetupViewInfo(const FDisplayClusterConfigurationICVFX_StageSettings& InStageSettings, FMinimalViewInfo& InOutViewInfo) const
 {
-	// CameraSettings can disable posprocess from this camera
+	// CameraSettings can disable postprocess from this camera
 	if (!bUseCameraComponentPostprocess)
 	{
 		InOutViewInfo.PostProcessSettings = FPostProcessSettings();
@@ -745,3 +745,21 @@ void FDisplayClusterConfigurationICVFX_CameraMotionBlur::SetupViewInfo(const FDi
 		InOutViewInfo.PostProcessSettings.bOverride_MotionBlurPerObjectSize = true;
 	}
 }
+
+///////////////////////////////////////////////////////////////////////////////////////
+// FDisplayClusterConfigurationICVFX_CameraDepthOfField
+///////////////////////////////////////////////////////////////////////////////////////
+UTexture2D* FDisplayClusterConfigurationICVFX_CameraDepthOfField::GetCompensationLUT(const FDisplayClusterConfigurationICVFX_StageSettings& InStageSettings) const
+{
+	if (DynamicCompensationLUT)
+	{
+		return ToRawPtr(DynamicCompensationLUT);
+	}
+	
+	if (CompensationLUT.IsValid())
+	{
+		return CompensationLUT.Get();
+	}
+
+	return nullptr;
+}
diff --git a/Engine/Plugins/Runtime/nDisplay/Source/DisplayClusterConfiguration/Public/DisplayClusterConfigurationTypes_ICVFX.h b/Engine/Plugins/Runtime/nDisplay/Source/DisplayClusterConfiguration/Public/DisplayClusterConfigurationTypes_ICVFX.h
--- a/Engine/Plugins/Runtime/nDisplay/Source/DisplayClusterConfiguration/Public/DisplayClusterConfigurationTypes_ICVFX.h
+++ b/Engine/Plugins/Runtime/nDisplay/Source/DisplayClusterConfiguration/Public/DisplayClusterConfigurationTypes_ICVFX.h
@@ -583,6 +583,13 @@ struct DISPLAYCLUSTERCONFIGURATION_API FDisplayClusterConfigurationICVFX_CameraD
 {
 	GENERATED_BODY()
 
+public:
+	/** Get the CompensationLUT texture corresponding to the current settings.
+	* 
+	* @param InStageSettings - it can be used for global LUT settings.
+	*/
+	UTexture2D* GetCompensationLUT(const FDisplayClusterConfigurationICVFX_StageSettings& InStageSettings) const;
+
 public:
 	/** Enables depth of field correction on the wall, which dynamically adjusts the size of the defocus circle of confusion to compensate for the real-world camera blur when shooting the wall */
 	UPROPERTY(BlueprintReadWrite, EditAnywhere, Category = NDisplay)
diff --git a/Engine/Plugins/Runtime/nDisplay/Source/DisplayClusterConfigurator/Private/DisplayClusterConfiguratorModule.cpp b/Engine/Plugins/Runtime/nDisplay/Source/DisplayClusterConfigurator/Private/DisplayClusterConfiguratorModule.cpp
--- a/Engine/Plugins/Runtime/nDisplay/Source/DisplayClusterConfigurator/Private/DisplayClusterConfiguratorModule.cpp
+++ b/Engine/Plugins/Runtime/nDisplay/Source/DisplayClusterConfigurator/Private/DisplayClusterConfiguratorModule.cpp
@@ -29,10 +29,12 @@
 #include "Views/Details/Media/DisplayClusterConfiguratorICVFXMediaCustomization.h"
 #include "Views/Details/Media/DisplayClusterConfiguratorMediaFullFrameCustomization.h"
 #include "Views/Details/Media/DisplayClusterConfiguratorMediaTileCustomization.h"
+#include "Views/Details/Components/DisplayClusterCameraComponentDetailsCustomization.h"
 #include "Views/Details/Policies/DisplayClusterConfiguratorPolicyDetailCustomization.h"
 
 #include "Blueprints/DisplayClusterBlueprint.h"
 #include "Components/DisplayClusterScreenComponent.h"
+#include "Components/DisplayClusterCameraComponent.h"
 #include "Components/DisplayClusterICVFXCameraComponent.h"
 #include "Misc/DisplayClusterObjectRef.h"
 #include "DisplayClusterRootActor.h"
@@ -206,7 +208,8 @@ void FDisplayClusterConfiguratorModule::RegisterCustomLayouts()
 	REGISTER_OBJECT_LAYOUT(UDisplayClusterConfigurationCluster, FDisplayClusterConfiguratorClusterDetailsCustomization);
 	REGISTER_OBJECT_LAYOUT(UDisplayClusterConfigurationClusterNode, FDisplayClusterConfiguratorBaseDetailCustomization);
 	REGISTER_OBJECT_LAYOUT(UDisplayClusterConfigurationViewport, FDisplayClusterConfiguratorViewportDetailsCustomization);
-	REGISTER_OBJECT_LAYOUT(UDisplayClusterScreenComponent, FDisplayClusterConfiguratorScreenDetailsCustomization);	
+	REGISTER_OBJECT_LAYOUT(UDisplayClusterScreenComponent, FDisplayClusterConfiguratorScreenDetailsCustomization);
+	REGISTER_OBJECT_LAYOUT(UDisplayClusterCameraComponent, FDisplayClusterCameraComponentDetailsCustomization);
 	REGISTER_OBJECT_LAYOUT(UDisplayClusterICVFXCameraComponent, FDisplayClusterICVFXCameraComponentDetailsCustomization);
 	
 	/**
diff --git a/Engine/Plugins/Runtime/nDisplay/Source/DisplayClusterConfigurator/Private/Views/Details/Components/DisplayClusterCameraComponentDetailsCustomization.cpp b/Engine/Plugins/Runtime/nDisplay/Source/DisplayClusterConfigurator/Private/Views/Details/Components/DisplayClusterCameraComponentDetailsCustomization.cpp
new file mode 100644
--- /dev/null
+++ b/Engine/Plugins/Runtime/nDisplay/Source/DisplayClusterConfigurator/Private/Views/Details/Components/DisplayClusterCameraComponentDetailsCustomization.cpp
@@ -0,0 +1,157 @@
+// Copyright Epic Games, Inc. All Rights Reserved.
+
+#include "DisplayClusterCameraComponentDetailsCustomization.h"
+
+#include "Components/DisplayClusterCameraComponent.h"
+#include "DisplayClusterRootActor.h"
+#include "Camera/CameraComponent.h"
+
+#include "DisplayClusterProjectionStrings.h"
+#include "DisplayClusterConfiguratorLog.h"
+
+#include "Render/DisplayDevice/Components/DisplayClusterDisplayDeviceBaseComponent.h"
+#include "Views/Details/Widgets/SDisplayClusterConfiguratorComponentPicker.h"
+#include "Views/Details/Widgets/SDisplayClusterConfigurationSearchableComboBox.h"
+
+#include "DetailLayoutBuilder.h"
+#include "DetailCategoryBuilder.h"
+#include "DetailWidgetRow.h"
+#include "PropertyHandle.h"
+
+#define LOCTEXT_NAMESPACE "DisplayClusterCameraComponentDetailsCustomization"
+
+void FDisplayClusterCameraComponentDetailsCustomization::CustomizeDetails(IDetailLayoutBuilder& InLayoutBuilder)
+{
+	Super::CustomizeDetails(InLayoutBuilder);
+
+	DetailLayout = &InLayoutBuilder;
+
+	if (!EditedObject.IsValid())
+	{
+		TArray<TWeakObjectPtr<UObject>> Objects;
+		InLayoutBuilder.GetObjectsBeingCustomized(Objects);
+
+		for (TWeakObjectPtr<UObject> Object : Objects)
+		{
+			if (Object->IsA<UDisplayClusterCameraComponent>())
+			{
+				EditedObject = Cast<UDisplayClusterCameraComponent>(Object.Get());
+			}
+		}
+	}
+
+	if (EditedObject.IsValid())
+	{
+		NoneOption = MakeShared<FString>("Active Engine Camera");
+
+		CameraHandle = InLayoutBuilder.GetProperty(GET_MEMBER_NAME_CHECKED(UDisplayClusterCameraComponent, OuterViewportCameraName));
+		check(CameraHandle->IsValidHandle());
+
+		RebuildCameraOptions();
+
+		if (IDetailPropertyRow* CameraPropertyRow = InLayoutBuilder.EditDefaultProperty(CameraHandle))
+		{
+			CameraPropertyRow->CustomWidget()
+				.NameContent()
+				[
+					CameraHandle->CreatePropertyNameWidget()
+				]
+				.ValueContent()
+				[
+					CreateCustomCameraWidget()
+				];
+		}
+	}
+}
+
+void FDisplayClusterCameraComponentDetailsCustomization::RebuildCameraOptions()
+{
+	CameraOptions.Reset();
+	UDisplayClusterCameraComponent* DestCameraComponent = EditedObject.Get();
+	check(DestCameraComponent != nullptr);
+
+	AActor* RootActor = GetRootActor();
+	check(RootActor);
+
+	TArray<UActorComponent*> ActorComponents;
+	RootActor->GetComponents(UCameraComponent::StaticClass(), ActorComponents);
+	for (UActorComponent* ActorComponent : ActorComponents)
+	{
+		const FString ComponentName = ActorComponent->GetName();
+		CameraOptions.Add(MakeShared<FString>(ComponentName));
+	}
+
+	// Component order not guaranteed, sort for consistency.
+	CameraOptions.Sort([](const TSharedPtr<FString>& A, const TSharedPtr<FString>& B)
+	{
+		// Default sort isn't compatible with TSharedPtr<FString>.
+		return *A < *B;
+	});
+
+	// Add None option
+	if (!DestCameraComponent->OuterViewportCameraName.IsEmpty())
+	{
+		CameraOptions.Add(NoneOption);
+	}
+}
+
+TSharedRef<SWidget> FDisplayClusterCameraComponentDetailsCustomization::CreateCustomCameraWidget()
+{
+	if (CameraComboBox.IsValid())
+	{
+		return CameraComboBox.ToSharedRef();
+	}
+
+	return SAssignNew(CameraComboBox, SDisplayClusterConfigurationSearchableComboBox)
+		.OptionsSource(&CameraOptions)
+		.OnGenerateWidget(this, &FDisplayClusterCameraComponentDetailsCustomization::MakeCameraOptionComboWidget)
+		.OnSelectionChanged(this, &FDisplayClusterCameraComponentDetailsCustomization::OnCameraSelected)
+		.ContentPadding(2)
+		.MaxListHeight(200.0f)
+		.Content()
+		[
+			SNew(STextBlock)
+			.Text(this, &FDisplayClusterCameraComponentDetailsCustomization::GetSelectedCameraText)
+			.Font(IDetailLayoutBuilder::GetDetailFont())
+		];
+}
+
+TSharedRef<SWidget> FDisplayClusterCameraComponentDetailsCustomization::MakeCameraOptionComboWidget(TSharedPtr<FString> InItem)
+{
+	return SNew(STextBlock)
+		.Text(FText::FromString(*InItem))
+		.Font(IDetailLayoutBuilder::GetDetailFont());
+}
+
+void FDisplayClusterCameraComponentDetailsCustomization::OnCameraSelected(TSharedPtr<FString> InCamera, ESelectInfo::Type SelectInfo)
+{
+	if (InCamera.IsValid())
+	{
+		// Handle empty case
+		if (InCamera->Equals(*NoneOption.Get()))
+		{
+			CameraHandle->SetValue(TEXT(""));
+		}
+		else
+		{
+			CameraHandle->SetValue(*InCamera.Get());
+		}
+
+		// Reset available options
+		RebuildCameraOptions();
+		CameraComboBox->ResetOptionsSource(&CameraOptions);
+		CameraComboBox->SetIsOpen(false);
+	}
+}
+
+FText FDisplayClusterCameraComponentDetailsCustomization::GetSelectedCameraText() const
+{
+	FString SelectedOption = EditedObject.Get()->OuterViewportCameraName;
+	if (SelectedOption.IsEmpty())
+	{
+		SelectedOption = *NoneOption.Get();
+	}
+	return FText::FromString(SelectedOption);
+}
+
+#undef LOCTEXT_NAMESPACE
diff --git a/Engine/Plugins/Runtime/nDisplay/Source/DisplayClusterConfigurator/Private/Views/Details/Components/DisplayClusterCameraComponentDetailsCustomization.h b/Engine/Plugins/Runtime/nDisplay/Source/DisplayClusterConfigurator/Private/Views/Details/Components/DisplayClusterCameraComponentDetailsCustomization.h
new file mode 100644
--- /dev/null
+++ b/Engine/Plugins/Runtime/nDisplay/Source/DisplayClusterConfigurator/Private/Views/Details/Components/DisplayClusterCameraComponentDetailsCustomization.h
@@ -0,0 +1,71 @@
+// Copyright Epic Games, Inc. All Rights Reserved.
+
+#pragma once
+
+#include "CoreMinimal.h"
+#include "Views/Details/DisplayClusterConfiguratorBaseDetailCustomization.h"
+#include "Types/SlateEnums.h"
+
+class UDisplayClusterCameraComponent;
+class SWidget;
+class SDisplayClusterConfigurationSearchableComboBox;
+
+/** Details panel customization for the UDisplayClusterCameraComponent object */
+class FDisplayClusterCameraComponentDetailsCustomization final : public FDisplayClusterConfiguratorBaseDetailCustomization
+{
+public:
+	using Super = FDisplayClusterConfiguratorBaseDetailCustomization;
+
+	static TSharedRef<IDetailCustomization> MakeInstance()
+	{
+		return MakeShared<FDisplayClusterCameraComponentDetailsCustomization>();
+	}
+
+public:
+	// Begin IDetailCustomization interface
+	virtual void CustomizeDetails(IDetailLayoutBuilder& DetailBuilder) override;
+	// End IDetailCustomization interface
+
+private:
+	/** Rebuilds the list of cameras to show in the dropdown menu of the Camera property widget */
+	void RebuildCameraOptions();
+
+	/** Creates a combo box widget to replace the default textbox of the Camera property of the UDisplayClusterCameraComponent */
+	TSharedRef<SWidget> CreateCustomCameraWidget();
+
+	/**
+	 * Creates a text block widget to use to display the specified item in the camera dropdown menu
+	 * @param InItem - The string item to make the text block for
+	 */
+	TSharedRef<SWidget> MakeCameraOptionComboWidget(TSharedPtr<FString> InItem);
+
+	/**
+	 * Raised when a camera is selected from the camera dropdown menu
+	 * @param InCamera - The camera item that was selected
+	 * @param SelectionInfo - Flag to indicate through what interface the selection was made
+	 */
+	void OnCameraSelected(TSharedPtr<FString> InCamera, ESelectInfo::Type SelectInfo);
+
+	/** Gets the text to display for the currently selected camera */
+	FText GetSelectedCameraText() const;
+
+private:
+	/** A weak reference to the UDisplayClusterCameraComponent object being edited by the details panel */
+	TWeakObjectPtr<UDisplayClusterCameraComponent> EditedObject;
+
+	/** Reference to the detail layout builder, used to force refresh the layout */
+	IDetailLayoutBuilder* DetailLayout = nullptr;
+
+private:
+	/** The list of camera items to display in the dropdown menu */
+	TArray<TSharedPtr<FString>> CameraOptions;
+
+	/** The property handle for the Camera property of the UDisplayClusterCameraComponent object */
+	TSharedPtr<IPropertyHandle> CameraHandle;
+
+	/** A cached pointer to the "None" option that is added to the list of options in the dropdown menu */
+	TSharedPtr<FString> NoneOption;
+
+	/** The combo box that is being displayed in the details panel for the Camera property */
+	TSharedPtr<SDisplayClusterConfigurationSearchableComboBox> CameraComboBox;
+};
diff --git a/Engine/Plugins/Runtime/nDisplay/Source/DisplayClusterProjection/Private/DisplayClusterProjectionModule.cpp b/Engine/Plugins/Runtime/nDisplay/Source/DisplayClusterProjection/Private/DisplayClusterProjectionModule.cpp
--- a/Engine/Plugins/Runtime/nDisplay/Source/DisplayClusterProjection/Private/DisplayClusterProjectionModule.cpp
+++ b/Engine/Plugins/Runtime/nDisplay/Source/DisplayClusterProjection/Private/DisplayClusterProjectionModule.cpp
@@ -153,4 +153,17 @@ bool FDisplayClusterProjectionModule::CameraPolicySetCamera(const TSharedPtr<IDi
 	return false;
 }
 
+UCameraComponent* FDisplayClusterProjectionModule::CameraPolicyGetCameraComponent(const TSharedPtr<IDisplayClusterProjectionPolicy, ESPMode::ThreadSafe>& InPolicy)
+{
+	if (InPolicy.IsValid())
+	{
+		if (FDisplayClusterProjectionCameraPolicy* CameraPolicyInstance = static_cast<FDisplayClusterProjectionCameraPolicy*>(InPolicy.Get()))
+		{
+			return CameraPolicyInstance->GetCameraComponent();
+		}
+	}
+
+	return nullptr;
+}
+
 IMPLEMENT_MODULE(FDisplayClusterProjectionModule, DisplayClusterProjection);
diff --git a/Engine/Plugins/Runtime/nDisplay/Source/DisplayClusterProjection/Private/DisplayClusterProjectionModule.h b/Engine/Plugins/Runtime/nDisplay/Source/DisplayClusterProjection/Private/DisplayClusterProjectionModule.h
--- a/Engine/Plugins/Runtime/nDisplay/Source/DisplayClusterProjection/Private/DisplayClusterProjectionModule.h
+++ b/Engine/Plugins/Runtime/nDisplay/Source/DisplayClusterProjection/Private/DisplayClusterProjectionModule.h
@@ -33,6 +33,8 @@ public:
 
 	virtual bool CameraPolicySetCamera(const TSharedPtr<IDisplayClusterProjectionPolicy, ESPMode::ThreadSafe>& InPolicy, UCameraComponent* const NewCamera, const FDisplayClusterProjectionCameraPolicySettings& CameraSettings) override;
 
+	virtual UCameraComponent* CameraPolicyGetCameraComponent(const TSharedPtr<IDisplayClusterProjectionPolicy, ESPMode::ThreadSafe>& InPolicy) override;
+
 private:
 	// Available factories
 	TMap<FString, TSharedPtr<IDisplayClusterProjectionPolicyFactory>> ProjectionPolicyFactories;
diff --git a/Engine/Plugins/Runtime/nDisplay/Source/DisplayClusterProjection/Private/Policy/Camera/DisplayClusterProjectionCameraPolicy.cpp b/Engine/Plugins/Runtime/nDisplay/Source/DisplayClusterProjection/Private/Policy/Camera/DisplayClusterProjectionCameraPolicy.cpp
--- a/Engine/Plugins/Runtime/nDisplay/Source/DisplayClusterProjection/Private/Policy/Camera/DisplayClusterProjectionCameraPolicy.cpp
+++ b/Engine/Plugins/Runtime/nDisplay/Source/DisplayClusterProjection/Private/Policy/Camera/DisplayClusterProjectionCameraPolicy.cpp
@@ -39,11 +39,7 @@ void FDisplayClusterProjectionCameraPolicy::UpdatePostProcessSettings(IDisplayCl
 	// The ICVFX camera postprocess is more detailed and is defined in the FDisplayClusterConfigurationICVFX_CameraSettings structure.
 	if (InViewport && !EnumHasAnyFlags(InViewport->GetRenderSettingsICVFX().RuntimeFlags, EDisplayClusterViewportRuntimeICVFXFlags::InCamera))
 	{
-		float DeltaTime = 0.0f;
-		if (ADisplayClusterRootActor* SceneRootActor = InViewport->GetConfiguration().GetRootActor(EDisplayClusterRootActorType::Scene))
-		{
-			DeltaTime = SceneRootActor->GetWorldDeltaSeconds();
-		}
+		const float DeltaTime = InViewport->GetConfiguration().GetRootActorWorldDeltaSeconds(EDisplayClusterRootActorType::Scene);
 
 		FMinimalViewInfo ViewInfo;
 		if (ImplSetupProjectionViewPoint(InViewport, DeltaTime, ViewInfo) && ViewInfo.PostProcessBlendWeight > 0.0f)
diff --git a/Engine/Plugins/Runtime/nDisplay/Source/DisplayClusterProjection/Private/Policy/Camera/DisplayClusterProjectionCameraPolicy.h b/Engine/Plugins/Runtime/nDisplay/Source/DisplayClusterProjection/Private/Policy/Camera/DisplayClusterProjectionCameraPolicy.h
--- a/Engine/Plugins/Runtime/nDisplay/Source/DisplayClusterProjection/Private/Policy/Camera/DisplayClusterProjectionCameraPolicy.h
+++ b/Engine/Plugins/Runtime/nDisplay/Source/DisplayClusterProjection/Private/Policy/Camera/DisplayClusterProjectionCameraPolicy.h
@@ -19,9 +19,7 @@ public:
 	FDisplayClusterProjectionCameraPolicy(const FString& ProjectionPolicyId, const FDisplayClusterConfigurationProjection* InConfigurationProjectionPolicy);
 
 public:
-	//////////////////////////////////////////////////////////////////////////////////////////////
-	// IDisplayClusterProjectionPolicy
-	//////////////////////////////////////////////////////////////////////////////////////////////
+	//~ Begin IDisplayClusterProjectionPolicy
 	virtual const FString& GetType() const override;
 
 	virtual bool HandleStartScene(IDisplayClusterViewport* InViewport) override;
@@ -38,14 +36,24 @@ public:
 		return true;
 	}
 
+	virtual bool ShouldUseViewPointComponentPostProcesses(IDisplayClusterViewport* InViewport) const override
+	{
+		// This projection policy uses its own camera with postprocessing.
+		// Disable the use of PP settings from the ViewPoint component.
+		return false;
+	}
+	//~ End IDisplayClusterProjectionPolicy
+
 public:
 	void SetCamera(UCameraComponent* const NewCamera, const FDisplayClusterProjectionCameraPolicySettings& InCameraSettings);
 
+	/** Returns the camera component that is currently in use. */
+	UCameraComponent* GetCameraComponent() const;
+
 private:
 	bool ImplSetupProjectionViewPoint(IDisplayClusterViewport* InViewport, const float InDeltaTime, FMinimalViewInfo& InOutViewInfo, float* OutCustomNearClippingPlane = nullptr) const;
 
 protected:
-	UCameraComponent* GetCameraComponent() const;
 
 private:
 	// Camera to use for rendering
diff --git a/Engine/Plugins/Runtime/nDisplay/Source/DisplayClusterProjection/Private/Policy/MPCDI/DisplayClusterProjectionMPCDIPolicy.cpp b/Engine/Plugins/Runtime/nDisplay/Source/DisplayClusterProjection/Private/Policy/MPCDI/DisplayClusterProjectionMPCDIPolicy.cpp
--- a/Engine/Plugins/Runtime/nDisplay/Source/DisplayClusterProjection/Private/Policy/MPCDI/DisplayClusterProjectionMPCDIPolicy.cpp
+++ b/Engine/Plugins/Runtime/nDisplay/Source/DisplayClusterProjection/Private/Policy/MPCDI/DisplayClusterProjectionMPCDIPolicy.cpp
@@ -308,6 +308,18 @@ bool FDisplayClusterProjectionMPCDIPolicy::GetProjectionMatrix(IDisplayClusterVi
 	return false;
 }
 
+bool FDisplayClusterProjectionMPCDIPolicy::IsFrustumRotatedToFitContextSize(IDisplayClusterViewport* InViewport, const uint32 InContextNum)
+{
+	check(IsInGameThread());
+
+	if (WarpBlendContexts.IsValidIndex(InContextNum))
+	{
+		return WarpBlendContexts[InContextNum].bFrustumRotated;
+	}
+
+	return false;
+}
+
 bool FDisplayClusterProjectionMPCDIPolicy::IsWarpBlendSupported()
 {
 	return true;
diff --git a/Engine/Plugins/Runtime/nDisplay/Source/DisplayClusterProjection/Private/Policy/MPCDI/DisplayClusterProjectionMPCDIPolicy.h b/Engine/Plugins/Runtime/nDisplay/Source/DisplayClusterProjection/Private/Policy/MPCDI/DisplayClusterProjectionMPCDIPolicy.h
--- a/Engine/Plugins/Runtime/nDisplay/Source/DisplayClusterProjection/Private/Policy/MPCDI/DisplayClusterProjectionMPCDIPolicy.h
+++ b/Engine/Plugins/Runtime/nDisplay/Source/DisplayClusterProjection/Private/Policy/MPCDI/DisplayClusterProjectionMPCDIPolicy.h
@@ -33,6 +33,7 @@ public:
 	virtual bool CalculateView(IDisplayClusterViewport* InViewport, const uint32 InContextNum, FVector& InOutViewLocation, FRotator& InOutViewRotation, const FVector& ViewOffset, const float WorldToMeters, const float NCP, const float FCP) override;
 	virtual bool GetProjectionMatrix(IDisplayClusterViewport* InViewport, const uint32 InContextNum, FMatrix& OutPrjMatrix) override;
 
+	virtual bool IsFrustumRotatedToFitContextSize(IDisplayClusterViewport* InViewport, const uint32 InContextNum) override;
 	virtual bool IsWarpBlendSupported() override;
 	virtual void ApplyWarpBlend_RenderThread(FRHICommandListImmediate& RHICmdList, const IDisplayClusterViewportProxy* InViewportProxy) override;
 
diff --git a/Engine/Plugins/Runtime/nDisplay/Source/DisplayClusterProjection/Public/IDisplayClusterProjection.h b/Engine/Plugins/Runtime/nDisplay/Source/DisplayClusterProjection/Public/IDisplayClusterProjection.h
--- a/Engine/Plugins/Runtime/nDisplay/Source/DisplayClusterProjection/Public/IDisplayClusterProjection.h
+++ b/Engine/Plugins/Runtime/nDisplay/Source/DisplayClusterProjection/Public/IDisplayClusterProjection.h
@@ -62,4 +62,9 @@ public:
 	*
 	*/
 	virtual bool CameraPolicySetCamera(const TSharedPtr<class IDisplayClusterProjectionPolicy, ESPMode::ThreadSafe>& InPolicy, class UCameraComponent* const NewCamera, const struct FDisplayClusterProjectionCameraPolicySettings& CameraSettings) = 0;
+
+	/**
+	* Get camera policy camera
+	*/
+	virtual class UCameraComponent* CameraPolicyGetCameraComponent(const TSharedPtr<IDisplayClusterProjectionPolicy, ESPMode::ThreadSafe>& InPolicy) = 0;
 };
diff --git a/Engine/Plugins/Runtime/nDisplay/Source/DisplayClusterWarp/Private/Components/DisplayClusterInFrustumFitCameraComponent.cpp b/Engine/Plugins/Runtime/nDisplay/Source/DisplayClusterWarp/Private/Components/DisplayClusterInFrustumFitCameraComponent.cpp
--- a/Engine/Plugins/Runtime/nDisplay/Source/DisplayClusterWarp/Private/Components/DisplayClusterInFrustumFitCameraComponent.cpp
+++ b/Engine/Plugins/Runtime/nDisplay/Source/DisplayClusterWarp/Private/Components/DisplayClusterInFrustumFitCameraComponent.cpp
@@ -156,13 +156,12 @@ void UDisplayClusterInFrustumFitCameraComponent::GetDesiredView(IDisplayClusterV
 	const UDisplayClusterInFrustumFitCameraComponent& ConfigurationCameraComponent = GetConfigurationInFrustumFitCameraComponent(InViewportConfiguration);
 	if (ConfigurationCameraComponent.IsEnabled())
 	{
-		if (ADisplayClusterRootActor* SceneRootActor = InViewportConfiguration.GetRootActor(EDisplayClusterRootActorType::Scene))
+		UCameraComponent* CameraComponent = ConfigurationCameraComponent.GetExternalCameraComponent();
+		const float DeltaSeconds = InViewportConfiguration.GetRootActorWorldDeltaSeconds(EDisplayClusterRootActorType::Scene);
+		if (IDisplayClusterViewport::GetCameraComponentView(CameraComponent, DeltaSeconds, ConfigurationCameraComponent.bUseCameraPostprocess, InOutViewInfo, OutCustomNearClippingPlane))
 		{
-			if (IDisplayClusterViewport::GetCameraComponentView(ConfigurationCameraComponent.GetExternalCameraComponent(), SceneRootActor->GetWorldDeltaSeconds(), ConfigurationCameraComponent.bUseCameraPostprocess, InOutViewInfo, OutCustomNearClippingPlane))
-			{
-				// 1. Use external camera for rendering
-				return;
-			}
+			// 1. Use external camera for rendering
+			return;
 		}
 
 		if (IDisplayClusterViewport::GetPlayerCameraView(InViewportConfiguration.GetCurrentWorld(), ConfigurationCameraComponent.bUseCameraPostprocess, InOutViewInfo))
diff --git a/Engine/Plugins/Runtime/nDisplay/Source/DisplayClusterWarp/Private/WarpBlend/Math/DisplayClusterWarpBlendMath_Frustum.cpp b/Engine/Plugins/Runtime/nDisplay/Source/DisplayClusterWarp/Private/WarpBlend/Math/DisplayClusterWarpBlendMath_Frustum.cpp
--- a/Engine/Plugins/Runtime/nDisplay/Source/DisplayClusterWarp/Private/WarpBlend/Math/DisplayClusterWarpBlendMath_Frustum.cpp
+++ b/Engine/Plugins/Runtime/nDisplay/Source/DisplayClusterWarp/Private/WarpBlend/Math/DisplayClusterWarpBlendMath_Frustum.cpp
@@ -188,7 +188,8 @@ bool FDisplayClusterWarpBlendMath_Frustum::EndCalcFrustum()
 	FMatrix RegionMatrix = GeometryContext.GetRegionMatrix();
 
 	// Apply frustum rotation to fit context size:
-	if (ShouldRotateFrustumToFitContextSize())
+	WarpData.WarpContext.bFrustumRotated = ShouldRotateFrustumToFitContextSize();
+	if (WarpData.WarpContext.bFrustumRotated)
 	{
 		static FQuat PostRotation = FRotator(0, 0, -90).Quaternion();
 
diff --git a/Engine/Plugins/Runtime/nDisplay/Source/DisplayClusterWarp/Public/Containers/DisplayClusterWarpContext.h b/Engine/Plugins/Runtime/nDisplay/Source/DisplayClusterWarp/Public/Containers/DisplayClusterWarpContext.h
--- a/Engine/Plugins/Runtime/nDisplay/Source/DisplayClusterWarp/Public/Containers/DisplayClusterWarpContext.h
+++ b/Engine/Plugins/Runtime/nDisplay/Source/DisplayClusterWarp/Public/Containers/DisplayClusterWarpContext.h
@@ -13,6 +13,10 @@ struct FDisplayClusterWarpContext
 {
 	bool bIsValid = false;
 
+	// Frustum has been rotated to fit the size of the context.
+	// This value contains the result of the ShouldRotateFrustumToFitContextSize() function.
+	bool bFrustumRotated = false;
+
 	// Camera
 	FRotator Rotation = FRotator::ZeroRotator;
 	FVector  Location = FVector::ZeroVector;
