nDisplay: Fixed logic that changes the resolution of the ICVFX camera viewport.
- The resolution for a sensor size with a height>width is now adjusted by width.
- Otherwise, the old method was used, which was height-adjustable.

diff --git a/Engine/Plugins/Runtime/nDisplay/Source/DisplayClusterConfiguration/Private/DisplayClusterConfigurationTypes_ICVFX.cpp b/Engine/Plugins/Runtime/nDisplay/Source/DisplayClusterConfiguration/Private/DisplayClusterConfigurationTypes_ICVFX.cpp
--- a/Engine/Plugins/Runtime/nDisplay/Source/DisplayClusterConfiguration/Private/DisplayClusterConfigurationTypes_ICVFX.cpp
+++ b/Engine/Plugins/Runtime/nDisplay/Source/DisplayClusterConfiguration/Private/DisplayClusterConfigurationTypes_ICVFX.cpp
@@ -326,14 +326,29 @@ FIntPoint FDisplayClusterConfigurationICVFX_CameraSettings::GetCameraFrameSize(c
 		const double CroppedSensorHeightRatio = CropedSensorHeight / InCineCameraComponent.Filmback.SensorHeight;
 
 		// Adapt camera resolution to the filmback sensor aspect ratio
-		// We keep the width, but adjust the height to match the aspect ratio of the Fimlmback sensor.
-		const double CameraFrameHeight = (InCineCameraComponent.Filmback.SensorHeight > 0.f && InCineCameraComponent.Filmback.SensorWidth > 0.f)
-			? CameraFrameSize.X / (InCineCameraComponent.Filmback.SensorWidth / InCineCameraComponent.Filmback.SensorHeight)
-			: CameraFrameSize.Y;
+		const double SensorAspectRatio = (InCineCameraComponent.Filmback.SensorHeight > 0.f && InCineCameraComponent.Filmback.SensorWidth > 0.f)
+			? InCineCameraComponent.Filmback.SensorWidth / InCineCameraComponent.Filmback.SensorHeight
+			: -1;
+
+		if (SensorAspectRatio < 0)
+		{
+			// The CineCamera sensor size has invalid values.
+			return CameraFrameSize;
+		}
+
+		// Use the max size as a base.
+		const double MaxSize = CameraFrameSize.GetMax();
+
+		const double CameraFrameWidth  = (SensorAspectRatio >= 1.0)
+			? MaxSize
+			: MaxSize * SensorAspectRatio;
+		const double CameraFrameHeight = (SensorAspectRatio >= 1.0)
+			? MaxSize / SensorAspectRatio
+			: MaxSize;
 
 		// Get cropped camera size
 		const FIntPoint CroppedCameraFrameSize(
-			FMath::RoundToInt(CameraFrameSize.X * CroppedSensorWidthRatio),
+			FMath::RoundToInt(CameraFrameWidth  * CroppedSensorWidthRatio),
 			FMath::RoundToInt(CameraFrameHeight * CroppedSensorHeightRatio)
 		);
 
