nDisplay: Fixed media settings propagation.

The method is now to copy the settings from the archetype in PostLoad of the icvfx camera and the root actor. The reason is that the previous method of restoring the values right after serialization was causing the reversal of all the settings, which in practice meant that icvfx camera settings could not be altered on the instaance and only on the archetype.

#rnx

diff --git a/Engine/Plugins/Runtime/nDisplay/Source/DisplayCluster/Private/Game/EngineClasses/Scene/DisplayClusterICVFXCameraComponent.cpp b/Engine/Plugins/Runtime/nDisplay/Source/DisplayCluster/Private/Game/EngineClasses/Scene/DisplayClusterICVFXCameraComponent.cpp
--- a/Engine/Plugins/Runtime/nDisplay/Source/DisplayCluster/Private/Game/EngineClasses/Scene/DisplayClusterICVFXCameraComponent.cpp
+++ b/Engine/Plugins/Runtime/nDisplay/Source/DisplayCluster/Private/Game/EngineClasses/Scene/DisplayClusterICVFXCameraComponent.cpp
@@ -39,7 +39,7 @@ void UDisplayClusterICVFXCameraComponent::PostLoad()
 {
 	Super::PostLoad();
 
-	PRAGMA_DISABLE_DEPRECATION_WARNINGS
+	PRAGMA_DISABLE_DEPRECATION_WARNINGS;
 	const int32 CustomVersion = GetLinkerCustomVersion(FDisplayClusterICVFXCameraCustomVersion::GUID);
 	if (CustomVersion < FDisplayClusterICVFXCameraCustomVersion::UpdateChromakeyConfig)
 	{
@@ -60,7 +60,16 @@ void UDisplayClusterICVFXCameraComponent::PostLoad()
 			CameraSettings.Chromakey.ChromakeySettingsSource = EDisplayClusterConfigurationICVFX_ChromakeySettingsSource::ICVFXCamera;
 		}
 	}
-	PRAGMA_ENABLE_DEPRECATION_WARNINGS
+	PRAGMA_ENABLE_DEPRECATION_WARNINGS;
+
+	// Propagate Media settings from the Archetype. Works around instanced property limitations.
+	if (!IsTemplate())
+	{
+		if (const UDisplayClusterICVFXCameraComponent* Archetype = Cast<UDisplayClusterICVFXCameraComponent>(GetArchetype()))
+		{
+			CameraSettings.RenderSettings.Media = Archetype->CameraSettings.RenderSettings.Media;
+		}
+	}
 }
 
 void UDisplayClusterICVFXCameraComponent::PostApplyToComponent()
diff --git a/Engine/Plugins/Runtime/nDisplay/Source/DisplayCluster/Private/Game/EngineClasses/Scene/DisplayClusterRootActor.cpp b/Engine/Plugins/Runtime/nDisplay/Source/DisplayCluster/Private/Game/EngineClasses/Scene/DisplayClusterRootActor.cpp
--- a/Engine/Plugins/Runtime/nDisplay/Source/DisplayCluster/Private/Game/EngineClasses/Scene/DisplayClusterRootActor.cpp
+++ b/Engine/Plugins/Runtime/nDisplay/Source/DisplayCluster/Private/Game/EngineClasses/Scene/DisplayClusterRootActor.cpp
@@ -86,6 +86,95 @@ namespace UE::DisplayCluster::RootActor
 			}
 		}
 	}
+
+	/**
+	 * Syncs default media settings to instance.
+	 *
+	 * @param InstanceDCRA Instance of the nDisplay root actor.
+	 */
+	static void PropagateMediaSettingsFromDefault(ADisplayClusterRootActor* InstanceDCRA)
+	{
+		if (!InstanceDCRA || InstanceDCRA->IsTemplate())
+		{
+			return;
+		}
+
+		// Note: ICVFX Camera Component Media property propagation is done in the ICVFXCamera itself.
+
+		// Propagate Backbuffer and Viewports MediaSettings property
+		//
+		const UDisplayClusterConfigurationData* const ArchetypeConfig = InstanceDCRA->GetDefaultConfigDataFromAsset();
+		const UDisplayClusterConfigurationData* const InstanceConfig = InstanceDCRA->GetConfigData();
+
+		if (!IsValid(ArchetypeConfig)
+			|| !IsValid(InstanceConfig)
+			|| !IsValid(ArchetypeConfig->Cluster)
+			|| !IsValid(InstanceConfig->Cluster))
+		{
+			return;
+		}
+
+		// Iterate over instance nodes, find the archetype, and copy the Media setting.
+		for (auto ItNode = InstanceConfig->Cluster->Nodes.CreateIterator(); ItNode; ++ItNode)
+		{
+			const FString& NodeName = ItNode.Key();
+			TObjectPtr<UDisplayClusterConfigurationClusterNode> InstanceNode = ItNode.Value();
+
+			if (!InstanceNode)
+			{
+				continue;
+			}
+
+			TObjectPtr<UDisplayClusterConfigurationClusterNode>* ArchetypeNodePtr = ArchetypeConfig->Cluster->Nodes.Find(NodeName);
+
+			if (!ArchetypeNodePtr)
+			{
+				continue;
+			}
+
+			TObjectPtr<UDisplayClusterConfigurationClusterNode> ArchetypeNode = *ArchetypeNodePtr;
+
+			if (!ArchetypeNode)
+			{
+				continue;
+			}
+
+			// Copy node media settings.
+			InstanceNode->MediaSettings = ArchetypeNode->MediaSettings;
+
+			// Viewports Media property
+			//
+
+			// Iterate over instance viewports, find the archetype, and deep-copy the Media setting.
+			for (auto ItViewport = InstanceNode->Viewports.CreateIterator(); ItViewport; ++ItViewport)
+			{
+				const FString& ViewportName = ItViewport.Key();
+				TObjectPtr<UDisplayClusterConfigurationViewport> InstanceViewport = ItViewport.Value();
+
+				if (!InstanceViewport)
+				{
+					continue;
+				}
+
+				TObjectPtr<UDisplayClusterConfigurationViewport>* ArchetypeViewportPtr = ArchetypeNode->Viewports.Find(ViewportName);
+
+				if (!ArchetypeViewportPtr)
+				{
+					continue;
+				}
+
+				TObjectPtr<UDisplayClusterConfigurationViewport> ArchetypeViewport = *ArchetypeViewportPtr;
+
+				if (!ArchetypeViewport)
+				{
+					continue;
+				}
+
+				// Copy viewport media settings.
+				InstanceViewport->RenderSettings.Media = ArchetypeViewport->RenderSettings.Media;
+			}
+		}
+	}
 };
 
 //////////////////////////////////////////////////////////////////////////////////////////////
@@ -744,6 +833,12 @@ void ADisplayClusterRootActor::InitializeRootActor()
 	{
 		BuildHierarchy();
 	}
+
+	// Propagate Media settings that should always override the instance settings.
+	// This is currently required due to limitation of instanced property propagation,
+	// where the instanced media settings may be incorrectly considered dirty and therefore
+	// not propagated from the template to the instance.
+	UE::DisplayCluster::RootActor::PropagateMediaSettingsFromDefault(this);
 }
 
 void ADisplayClusterRootActor::UpdateProceduralMeshComponentData(const UProceduralMeshComponent* InProceduralMeshComponent)
diff --git a/Engine/Plugins/Runtime/nDisplay/Source/DisplayClusterConfiguration/Private/DisplayClusterConfigurationTypes_ICVFX.cpp b/Engine/Plugins/Runtime/nDisplay/Source/DisplayClusterConfiguration/Private/DisplayClusterConfigurationTypes_ICVFX.cpp
--- a/Engine/Plugins/Runtime/nDisplay/Source/DisplayClusterConfiguration/Private/DisplayClusterConfigurationTypes_ICVFX.cpp
+++ b/Engine/Plugins/Runtime/nDisplay/Source/DisplayClusterConfiguration/Private/DisplayClusterConfigurationTypes_ICVFX.cpp
@@ -61,32 +61,6 @@ void FDisplayClusterConfigurationICVFX_CameraRenderSettings::SetupViewInfo(const
 	}
 }
 
-
-bool FDisplayClusterConfigurationICVFX_CameraRenderSettings::Serialize(FArchive& Ar)
-{	
-	// When loading, overwrite Media settings with defaults unless this is the archetype
-
-	if (!Ar.IsLoading() && !Ar.IsSaving())
-	{
-		return true;
-	}
-
-	UScriptStruct& Struct = *StaticStruct();
-
-	if (Ar.IsLoading() && !FDisplayClusterConfigurationUtils::IsSerializingTemplate(Ar))
-	{
-		const FDisplayClusterConfigurationMediaICVFX MediaOriginal = Media;
-		Struct.SerializeTaggedProperties(Ar, (uint8*)this, &Struct, nullptr);
-		Media = MediaOriginal;
-	}
-	else
-	{
-		Struct.SerializeTaggedProperties(Ar, (uint8*)this, &Struct, nullptr);
-	}
-
-	return true;
-}
-
 ///////////////////////////////////////////////////////////////////////////////////////
 // FDisplayClusterConfigurationICVFX_LightcardCustomOCIO
 ///////////////////////////////////////////////////////////////////////////////////////
diff --git a/Engine/Plugins/Runtime/nDisplay/Source/DisplayClusterConfiguration/Public/DisplayClusterConfigurationTypes_ICVFX.h b/Engine/Plugins/Runtime/nDisplay/Source/DisplayClusterConfiguration/Public/DisplayClusterConfigurationTypes_ICVFX.h
--- a/Engine/Plugins/Runtime/nDisplay/Source/DisplayClusterConfiguration/Public/DisplayClusterConfigurationTypes_ICVFX.h
+++ b/Engine/Plugins/Runtime/nDisplay/Source/DisplayClusterConfiguration/Public/DisplayClusterConfigurationTypes_ICVFX.h
@@ -520,17 +520,8 @@ public:
 public:
 	/** Propagates general render related settings to the view info. */
 	void SetupViewInfo(const FDisplayClusterConfigurationICVFX_StageSettings& InStageSettings, FMinimalViewInfo& InOutViewInfo) const;
-
-	/** Custom serialization to skip Media on instances */
-	bool Serialize(FArchive& Ar);
 };
 
-/** Makes the given struct use the custom serializer */
-template<> struct TStructOpsTypeTraits<FDisplayClusterConfigurationICVFX_CameraRenderSettings> 
-	: public TStructOpsTypeTraitsBase2<FDisplayClusterConfigurationICVFX_CameraRenderSettings>
-{
-	enum { WithSerializer = true };
-};
 
 USTRUCT(BlueprintType)
 struct DISPLAYCLUSTERCONFIGURATION_API FDisplayClusterConfigurationICVFX_CameraMotionBlurOverridePPS
