nDisplay: Hide lightcard in Editor, but keep it in the nDisplay preview.

The trick to make this work was firstly to ensure the LCs are rendered unlit, then checking scene color in the material to force to a low opacity when there is any light.

ICVFX panel opacity is still controlled via "IsProxy()" since that render has no lights as well.

In those cases the opacity was made zero, but if a slightly higher value is desired, change the material again (the opacity section is clearly marked with a comment box), and also change the value of ProxyOpacity in DisplayClusterLightCardActor.cpp.

#virtualized

diff --git a/Engine/Plugins/Runtime/nDisplay/Source/DisplayCluster/Private/Game/EngineClasses/Scene/DisplayClusterLightCardActor.cpp b/Engine/Plugins/Runtime/nDisplay/Source/DisplayCluster/Private/Game/EngineClasses/Scene/DisplayClusterLightCardActor.cpp
--- a/Engine/Plugins/Runtime/nDisplay/Source/DisplayCluster/Private/Game/EngineClasses/Scene/DisplayClusterLightCardActor.cpp
+++ b/Engine/Plugins/Runtime/nDisplay/Source/DisplayCluster/Private/Game/EngineClasses/Scene/DisplayClusterLightCardActor.cpp
@@ -57,7 +57,7 @@ ADisplayClusterLightCardActor::ADisplayClusterLightCardActor(const FObjectInitia
 	, Pitch(0.f)
 	, Yaw(0.f)
 	, Scale(FVector2D(1.f))
-	, RadialOffset(-0.5)
+	, RadialOffset(-1)
 	, bAlwaysFlushToWall(true)
 	, Mask(EDisplayClusterLightCardMask::Circle)
 	, Texture(nullptr)
@@ -316,7 +316,7 @@ void ADisplayClusterLightCardActor::UpdateLightCardMaterialInstance()
 	if (UMaterialInstanceDynamic* LightCardMaterialInstance = Cast<UMaterialInstanceDynamic>(LightCardComponent->GetMaterial(0)))
 	{
 		// Showing proxy with low opacity to make it less distracting when it doesn't line up well with its projection in the Light Card Editor.
-		constexpr float ProxyOpacity = 0.25;
+		constexpr float ProxyOpacity = 0;
 
 		LightCardMaterialInstance->SetVectorParameterValue(TEXT("CardColor"), Color);
 		LightCardMaterialInstance->SetScalarParameterValue(TEXT("Temperature"), Temperature);
diff --git a/Engine/Plugins/Runtime/nDisplay/Source/DisplayCluster/Private/Render/Viewport/DisplayClusterViewportManager.cpp b/Engine/Plugins/Runtime/nDisplay/Source/DisplayCluster/Private/Render/Viewport/DisplayClusterViewportManager.cpp
--- a/Engine/Plugins/Runtime/nDisplay/Source/DisplayCluster/Private/Render/Viewport/DisplayClusterViewportManager.cpp
+++ b/Engine/Plugins/Runtime/nDisplay/Source/DisplayCluster/Private/Render/Viewport/DisplayClusterViewportManager.cpp
@@ -192,6 +192,12 @@ void IDisplayClusterViewportManager::SetupEngineShowFlags(const EDisplayClusterV
 		InOutEngineShowFlags.SetDeferredLighting(0);
 		InOutEngineShowFlags.SetVirtualTexturePrimitives(0);
 		InOutEngineShowFlags.SetRectLights(0);
+
+		// Finish making it essentially Unlit
+		InOutEngineShowFlags.SetLightFunctions(0);
+		InOutEngineShowFlags.SetDynamicShadows(0);
+		InOutEngineShowFlags.SetLighting(0);
+
 		break;
 
 	default:
