[nDisplay] - Generic barrier sync callback now provides an additional information that allows to detect cluster node ID based on a thread marker/ID

- There is a new TMap<FString, FString> property that maps the barrier synchronization threads to the owning cluster nodes
- Updated Rivermax sync policy to report cluster node IDs instead of the sync thread IDs

#rnx

diff --git a/Engine/Plugins/Runtime/nDisplay/Source/DisplayCluster/Private/Network/Barrier/DisplayClusterBarrier.cpp b/Engine/Plugins/Runtime/nDisplay/Source/DisplayCluster/Private/Network/Barrier/DisplayClusterBarrier.cpp
--- a/Engine/Plugins/Runtime/nDisplay/Source/DisplayCluster/Private/Network/Barrier/DisplayClusterBarrier.cpp
+++ b/Engine/Plugins/Runtime/nDisplay/Source/DisplayCluster/Private/Network/Barrier/DisplayClusterBarrier.cpp
@@ -284,8 +284,9 @@ void FDisplayClusterBarrier::HandleBarrierPreSyncStart()
 
 void FDisplayClusterBarrier::HandleBarrierPreSyncEnd()
 {
-	// Execute sync delegate with data requested
-	BarrierPreSyncEndDelegate.ExecuteIfBound(Name, ClientsRequestData, ClientsResponseData);
+	// Prepare callback data and call the handler
+	FDisplayClusterBarrierPreSyncEndDelegateData PreSyncEndCallbackData{ Name, ClientsRequestData, ClientsResponseData };
+	BarrierPreSyncEndDelegate.ExecuteIfBound(PreSyncEndCallbackData);
 
 	// We can clean request data now before next iteration
 	ClientsRequestData.Empty(CallersAllowed.Num());
diff --git a/Engine/Plugins/Runtime/nDisplay/Source/DisplayCluster/Private/Network/Barrier/DisplayClusterBarrier.h b/Engine/Plugins/Runtime/nDisplay/Source/DisplayCluster/Private/Network/Barrier/DisplayClusterBarrier.h
--- a/Engine/Plugins/Runtime/nDisplay/Source/DisplayCluster/Private/Network/Barrier/DisplayClusterBarrier.h
+++ b/Engine/Plugins/Runtime/nDisplay/Source/DisplayCluster/Private/Network/Barrier/DisplayClusterBarrier.h
@@ -42,6 +42,7 @@ public:
 	// Remove specified caller from the sync pipeline
 	virtual void UnregisterSyncCaller(const FString& CallerId) override;
 
+	// Returns barrier PreSyncEnd delegate
 	virtual FDisplayClusterBarrierPreSyncEndDelegate& GetPreSyncEndDelegate() override
 	{
 		return BarrierPreSyncEndDelegate;
@@ -103,9 +104,9 @@ private:
 	FDisplayClusterBarrierTimeoutEvent BarrierTimeoutEvent;
 
 	// Request data from the callers
-	FClientsCommData ClientsRequestData;
+	TMap<FString, TArray<uint8>> ClientsRequestData;
 	// Response data for the callers
-	FClientsCommData ClientsResponseData;
+	TMap<FString, TArray<uint8>> ClientsResponseData;
 
 	// Diagnostics data
 	double BarrierWaitTimeStart   = 0;
diff --git a/Engine/Plugins/Runtime/nDisplay/Source/DisplayCluster/Private/Network/Barrier/IDisplayClusterBarrier.h b/Engine/Plugins/Runtime/nDisplay/Source/DisplayCluster/Private/Network/Barrier/IDisplayClusterBarrier.h
--- a/Engine/Plugins/Runtime/nDisplay/Source/DisplayCluster/Private/Network/Barrier/IDisplayClusterBarrier.h
+++ b/Engine/Plugins/Runtime/nDisplay/Source/DisplayCluster/Private/Network/Barrier/IDisplayClusterBarrier.h
@@ -19,6 +19,22 @@ enum class EDisplayClusterBarrierWaitResult : uint8
 };
 
 
+/**
+ * Barrier PreSyncEnd callback data
+ */
+struct FDisplayClusterBarrierPreSyncEndDelegateData
+{
+	/** Barrier ID */
+	const FString& BarrierId;
+
+	/** Binary data provided on sync request (thread marker - to - data mapping) */
+	const TMap<FString, TArray<uint8>>& RequestData;
+
+	/** Binary data to respond (thread marker - to - data mapping) */
+	TMap<FString, TArray<uint8>>& ResponseData;
+};
+
+
 /**
  * Thread barrier interface
  */
@@ -50,8 +66,7 @@ public:
 	virtual void UnregisterSyncCaller(const FString& CallerId) = 0;
 
 	/** Barrier PreSyncEnd delegate. Called when all calling threads arrived right before opening the gate. */
-	typedef TMap<FString, TArray<uint8>> FClientsCommData;
-	DECLARE_DELEGATE_ThreeParams(FDisplayClusterBarrierPreSyncEndDelegate, const FString&, const FClientsCommData&, FClientsCommData&);
+	DECLARE_DELEGATE_OneParam(FDisplayClusterBarrierPreSyncEndDelegate, FDisplayClusterBarrierPreSyncEndDelegateData&);
 	virtual FDisplayClusterBarrierPreSyncEndDelegate& GetPreSyncEndDelegate() = 0;
 
 	// Barrier timout notification (provides BarrierName and CallersTimedOut in parameters)
diff --git a/Engine/Plugins/Runtime/nDisplay/Source/DisplayCluster/Private/Network/Service/GenericBarrier/DisplayClusterGenericBarrierClient.cpp b/Engine/Plugins/Runtime/nDisplay/Source/DisplayCluster/Private/Network/Service/GenericBarrier/DisplayClusterGenericBarrierClient.cpp
--- a/Engine/Plugins/Runtime/nDisplay/Source/DisplayCluster/Private/Network/Service/GenericBarrier/DisplayClusterGenericBarrierClient.cpp
+++ b/Engine/Plugins/Runtime/nDisplay/Source/DisplayCluster/Private/Network/Service/GenericBarrier/DisplayClusterGenericBarrierClient.cpp
@@ -1,24 +1,26 @@
 // Copyright Epic Games, Inc. All Rights Reserved.
 
 #include "Network/Service/GenericBarrier/DisplayClusterGenericBarrierClient.h"
-#include "Network/Service/GenericBarrier/DisplayClusterGenericBarrierService.h"
-#include "Network/Service/GenericBarrier/DisplayClusterGenericBarrierStrings.h"
-#include "Network/Packet/DisplayClusterPacketInternal.h"
-#include "Network/Listener/DisplayClusterHelloMessageStrings.h"
-
-#include "Config/IPDisplayClusterConfigManager.h"
 
 #include "Cluster/IPDisplayClusterClusterManager.h"
 #include "Cluster/Controller/IDisplayClusterClusterNodeController.h"
 
-#include "DisplayClusterEnums.h"
+#include "Config/IPDisplayClusterConfigManager.h"
+
 #include "Misc/DisplayClusterGlobals.h"
 #include "Misc/DisplayClusterHelpers.h"
 #include "Misc/DisplayClusterLog.h"
 #include "Misc/DisplayClusterStrings.h"
 #include "Misc/DisplayClusterTypesConverter.h"
 
+#include "Network/Barrier/IDisplayClusterBarrier.h"
+#include "Network/Service/GenericBarrier/DisplayClusterGenericBarrierService.h"
+#include "Network/Service/GenericBarrier/DisplayClusterGenericBarrierStrings.h"
+#include "Network/Packet/DisplayClusterPacketInternal.h"
+#include "Network/Listener/DisplayClusterHelloMessageStrings.h"
+
 #include "DisplayClusterConfigurationTypes.h"
+#include "DisplayClusterEnums.h"
 
 
 FDisplayClusterGenericBarrierClient::FDisplayClusterGenericBarrierClient()
@@ -362,6 +364,22 @@ EDisplayClusterCommResult FDisplayClusterGenericBarrierClient::SyncOnBarrierWith
 	return Response->GetCommResult();
 }
 
+FDisplayClusterGenericBarrierService* FDisplayClusterGenericBarrierClient::GetGenericBarrierService() const
+{
+	if (IPDisplayClusterClusterManager* const ClusterMgr = GDisplayCluster->GetPrivateClusterMgr())
+	{
+		if (IDisplayClusterClusterNodeController* const ClusterNodeCtrl = ClusterMgr->GetClusterNodeController())
+		{
+			if (FDisplayClusterService* const Service = ClusterNodeCtrl->GetGenericBarriersServer())
+			{
+				return static_cast<FDisplayClusterGenericBarrierService*>(Service);
+			}
+		}
+	}
+
+	return nullptr;
+}
+
 bool FDisplayClusterGenericBarrierClient::ConfigureBarrierSyncDelegate(const FString& BarrierId, bool bSetup)
 {
 	// @note
@@ -376,9 +394,9 @@ bool FDisplayClusterGenericBarrierClient::ConfigureBarrierSyncDelegate(const FSt
 		{
 			if (ClusterMgr->IsPrimary())
 			{
-				if (IDisplayClusterClusterNodeController* ClusterNodeCtrl = ClusterMgr->GetClusterNodeController())
+				if (IDisplayClusterClusterNodeController* const ClusterNodeCtrl = ClusterMgr->GetClusterNodeController())
 				{
-					if (FDisplayClusterService* Service = ClusterNodeCtrl->GetGenericBarriersServer())
+					if (FDisplayClusterService* const Service = ClusterNodeCtrl->GetGenericBarriersServer())
 					{
 						FDisplayClusterGenericBarrierService* const GBService = static_cast<FDisplayClusterGenericBarrierService*>(Service);
 						if (TSharedPtr<IDisplayClusterBarrier, ESPMode::ThreadSafe> Barrier = GBService->GetBarrier(BarrierId))
@@ -389,7 +407,7 @@ bool FDisplayClusterGenericBarrierClient::ConfigureBarrierSyncDelegate(const FSt
 								if (!BarrierSyncDelegates.Contains(BarrierId))
 								{
 									BarrierSyncDelegates.Emplace(BarrierId);
-									Barrier->GetPreSyncEndDelegate().BindRaw(this, &FDisplayClusterGenericBarrierClient::OnPreBarrierSyncEnd);
+									Barrier->GetPreSyncEndDelegate().BindRaw(this, &FDisplayClusterGenericBarrierClient::OnBarrierSync);
 								}
 							}
 							// Release
@@ -412,14 +430,37 @@ bool FDisplayClusterGenericBarrierClient::ConfigureBarrierSyncDelegate(const FSt
 	return false;
 }
 
-void FDisplayClusterGenericBarrierClient::OnPreBarrierSyncEnd(const FString& BarrierId, const TMap<FString, TArray<uint8>>& RequestData, TMap<FString, TArray<uint8>>& ResponseData)
+void FDisplayClusterGenericBarrierClient::OnBarrierSync(FDisplayClusterBarrierPreSyncEndDelegateData& SyncData)
 {
-	if (FOnGenericBarrierSynchronizationDelegate* const Delegate = BarrierSyncDelegates.Find(BarrierId))
+	// Make sure the delegate is set
+	FOnGenericBarrierSynchronizationDelegate* const Delegate = BarrierSyncDelegates.Find(SyncData.BarrierId);
+	if (!Delegate || !Delegate->IsBound())
 	{
-		if (Delegate->IsBound())
-		{
-			FGenericBarrierSynchronizationDelegateData CallbackData{ BarrierId, RequestData, ResponseData };
-			Delegate->Execute(CallbackData);
-		}
+		return;
 	}
+
+	// Access GB server
+	FDisplayClusterGenericBarrierService* const GBService = GetGenericBarrierService();
+	checkSlow(GBService);
+	if (!GBService)
+	{
+		return;
+	}
+
+	// Get barrier info
+	const TSharedPtr<FDisplayClusterGenericBarrierService::FBarrierInfo> BarrierInfo = GBService->GetBarrierInfo(SyncData.BarrierId);
+	checkSlow(BarrierInfo.IsValid());
+	if (!BarrierInfo.IsValid())
+	{
+		return;
+	}
+
+	// Last thing before triggering the callback. For optimization purposes, we disable further barrier info updates.
+	// At this point we got a sync callback, therefore all the thread markers and cluster nodes for this barrier
+	// are known. No need to update anymore as we won't get any new information.
+	GBService->SetBarrierInfoUpdateLocked(SyncData.BarrierId, true);
+
+	// Now forward data to the handler
+	FGenericBarrierSynchronizationDelegateData CallbackData{ SyncData.BarrierId, BarrierInfo->ThreadToNodeMapping, SyncData.RequestData, SyncData.ResponseData };
+	Delegate->Execute(CallbackData);
 }
diff --git a/Engine/Plugins/Runtime/nDisplay/Source/DisplayCluster/Private/Network/Service/GenericBarrier/DisplayClusterGenericBarrierClient.h b/Engine/Plugins/Runtime/nDisplay/Source/DisplayCluster/Private/Network/Service/GenericBarrier/DisplayClusterGenericBarrierClient.h
--- a/Engine/Plugins/Runtime/nDisplay/Source/DisplayCluster/Private/Network/Service/GenericBarrier/DisplayClusterGenericBarrierClient.h
+++ b/Engine/Plugins/Runtime/nDisplay/Source/DisplayCluster/Private/Network/Service/GenericBarrier/DisplayClusterGenericBarrierClient.h
@@ -8,6 +8,9 @@
 #include "Network/Protocol/IDisplayClusterProtocolGenericBarrier.h"
 #include "Cluster/IDisplayClusterGenericBarriersClient.h"
 
+class FDisplayClusterGenericBarrierService;
+struct FDisplayClusterBarrierPreSyncEndDelegateData;
+
 
 /**
  * Generic barriers TCP client
@@ -57,11 +60,15 @@ public:
 	virtual EDisplayClusterCommResult SyncOnBarrierWithData(const FString& BarrierId, const FString& UniqueThreadMarker, const TArray<uint8>& RequestData, TArray<uint8>& OutResponseData, EBarrierControlResult& Result) override;
 
 private:
+
+	/** Provides easy access to the GB service */
+	FDisplayClusterGenericBarrierService* GetGenericBarrierService() const;
+
 	/** Setup/release sync delegate for a specific barrier */
 	bool ConfigureBarrierSyncDelegate(const FString& BarrierId, bool bSetup);
 
 	/** Callback on barrier sync phase end */
-	void OnPreBarrierSyncEnd(const FString& BarrierId, const TMap<FString, TArray<uint8>>& RequestData, TMap<FString, TArray<uint8>>& ResponseData);
+	void OnBarrierSync(FDisplayClusterBarrierPreSyncEndDelegateData& SyncData);
 
 private:
 	// Holds synchronization delegates for the barriers operated by this client
diff --git a/Engine/Plugins/Runtime/nDisplay/Source/DisplayCluster/Private/Network/Service/GenericBarrier/DisplayClusterGenericBarrierService.cpp b/Engine/Plugins/Runtime/nDisplay/Source/DisplayCluster/Private/Network/Service/GenericBarrier/DisplayClusterGenericBarrierService.cpp
--- a/Engine/Plugins/Runtime/nDisplay/Source/DisplayCluster/Private/Network/Service/GenericBarrier/DisplayClusterGenericBarrierService.cpp
+++ b/Engine/Plugins/Runtime/nDisplay/Source/DisplayCluster/Private/Network/Service/GenericBarrier/DisplayClusterGenericBarrierService.cpp
@@ -63,6 +63,29 @@ TSharedPtr<IDisplayClusterBarrier, ESPMode::ThreadSafe> FDisplayClusterGenericBa
 	return BarrierFound ? *BarrierFound : nullptr;
 }
 
+TSharedPtr<FDisplayClusterGenericBarrierService::FBarrierInfo> FDisplayClusterGenericBarrierService::GetBarrierInfo(const FString& BarrierId) const
+{
+	FScopeLock Lock(&BarriersInfoCS);
+
+	const FBarrierInfoWrapper* BarrierInfoWrapper = BarriersInfo.Find(BarrierId);
+	checkSlow(BarrierInfoWrapper);
+
+	return BarrierInfoWrapper ? BarrierInfoWrapper->BarrierInfo : TSharedPtr<FDisplayClusterGenericBarrierService::FBarrierInfo>();
+}
+
+void FDisplayClusterGenericBarrierService::SetBarrierInfoUpdateLocked(const FString& BarrierId, bool bLocked)
+{
+	FScopeLock Lock(&BarriersInfoCS);
+
+	FBarrierInfoWrapper* BarrierInfoWrapper = BarriersInfo.Find(BarrierId);
+	checkSlow(BarrierInfoWrapper);
+
+	if (BarrierInfoWrapper)
+	{
+		BarrierInfoWrapper->bBarrierInfoLockedOut = bLocked;
+	}
+}
+
 TSharedPtr<IDisplayClusterSession> FDisplayClusterGenericBarrierService::CreateSession(FDisplayClusterSessionInfo& SessionInfo)
 {
 	SessionInfo.SessionName = FString::Printf(TEXT("%s_%lu_%s_%s"),
@@ -359,6 +382,12 @@ EDisplayClusterCommResult FDisplayClusterGenericBarrierService::ReleaseBarrier(c
 				FPlatformProcess::ReturnSynchEventToPool(*BarrierCreatedEvent);
 				BarrierCreationEvents.Remove(BarrierId);
 			}
+
+			// Remove clients info of this barrier
+			{
+				FScopeLock LockInfo(&BarriersInfoCS);
+				BarriersInfo.Remove(BarrierId);
+			}
 		}
 		else
 		{
@@ -406,6 +435,9 @@ EDisplayClusterCommResult FDisplayClusterGenericBarrierService::SyncOnBarrier(co
 		return EDisplayClusterCommResult::NotAllowed;
 	}
 
+	// Save info about this caller
+	UpdateBarrierInformation(BarrierId, GetSessionInfoCache().NodeId.Get(FString()), UniqueThreadMarker);
+
 	// Sync on the barrier
 	(*Barrier)->Wait(UniqueThreadMarker);
 
@@ -447,9 +479,30 @@ EDisplayClusterCommResult FDisplayClusterGenericBarrierService::SyncOnBarrierWit
 		return EDisplayClusterCommResult::NotAllowed;
 	}
 
+	// Save info about this caller
+	UpdateBarrierInformation(BarrierId, GetSessionInfoCache().NodeId.Get(FString()), UniqueThreadMarker);
+
 	// Sync on the barrier
 	(*Barrier)->WaitWithData(UniqueThreadMarker, RequestData, OutResponseData);
 
 	Result = EBarrierControlResult::SynchronizedSuccessfully;
 	return EDisplayClusterCommResult::Ok;
 }
+
+void FDisplayClusterGenericBarrierService::UpdateBarrierInformation(const FString& BarrierId, const FString& NodeId, const FString& ThreadMarker)
+{
+	FScopeLock Lock(&BarriersInfoCS);
+
+	FBarrierInfoWrapper& BarrierInfoWrapper = BarriersInfo.FindOrAdd(BarrierId);
+
+	// Update if allowed
+	if (!BarrierInfoWrapper.bBarrierInfoLockedOut)
+	{
+		// Update NodeToThread map
+		TSet<FString>& NodeInfo = BarrierInfoWrapper.BarrierInfo->NodeToThreadsMapping.FindOrAdd(NodeId);
+		NodeInfo.Add(ThreadMarker);
+
+		// Update ThreadToNode map
+		BarrierInfoWrapper.BarrierInfo->ThreadToNodeMapping.Emplace(ThreadMarker, NodeId);
+	}
+}
diff --git a/Engine/Plugins/Runtime/nDisplay/Source/DisplayCluster/Private/Network/Service/GenericBarrier/DisplayClusterGenericBarrierService.h b/Engine/Plugins/Runtime/nDisplay/Source/DisplayCluster/Private/Network/Service/GenericBarrier/DisplayClusterGenericBarrierService.h
--- a/Engine/Plugins/Runtime/nDisplay/Source/DisplayCluster/Private/Network/Service/GenericBarrier/DisplayClusterGenericBarrierService.h
+++ b/Engine/Plugins/Runtime/nDisplay/Source/DisplayCluster/Private/Network/Service/GenericBarrier/DisplayClusterGenericBarrierService.h
@@ -20,6 +20,20 @@ class FDisplayClusterGenericBarrierService
 	, public    IDisplayClusterSessionPacketHandler<FDisplayClusterPacketInternal, true>
 	, protected IDisplayClusterProtocolGenericBarrier
 {
+public:
+
+	/**
+	 * Additional barrier information that might be useful outside of the server
+	 */
+	struct FBarrierInfo
+	{
+		// Holds a set of thread markers bound to the owning cluster node
+		TMap<FString, TSet<FString>> NodeToThreadsMapping;
+
+		// Holds thread marker - to - cluster node mapping
+		TMap<FString, FString> ThreadToNodeMapping;
+	};
+
 public:
 	FDisplayClusterGenericBarrierService();
 	virtual ~FDisplayClusterGenericBarrierService();
@@ -32,8 +46,18 @@ public:
 	virtual FString GetProtocolName() const override;
 
 public:
+
+	// Returns barrier by ID
 	TSharedPtr<IDisplayClusterBarrier, ESPMode::ThreadSafe> GetBarrier(const FString& BarrierId);
 
+	// Returns barrier information
+	TSharedPtr<FDisplayClusterGenericBarrierService::FBarrierInfo> GetBarrierInfo(const FString& BarrierId) const;
+
+	// Enables/disables barrier info update. The update procedure is called every time a thread joins the barrier
+	// to synchronize. More threads we have more CPU time is consumed for this. For optimization purposes, this
+	// function allows to stop collecting barrier information and continue updating at any time.
+	void SetBarrierInfoUpdateLocked(const FString& BarrierId, bool bLocked);
+
 protected:
 	// Creates session instance for this service
 	virtual TSharedPtr<IDisplayClusterSession> CreateSession(FDisplayClusterSessionInfo& SessionInfo) override;
@@ -56,6 +80,12 @@ private:
 	virtual EDisplayClusterCommResult SyncOnBarrierWithData(const FString& BarrierId, const FString& UniqueThreadMarker, const TArray<uint8>& RequestData, TArray<uint8>& OutResponseData, EBarrierControlResult& Result) override;
 
 private:
+
+	// Caches information about barrier users
+	void UpdateBarrierInformation(const FString& BarrierId, const FString& NodeId, const FString& ThreadMarker);
+
+private:
+
 	// Barriers
 	TMap<FString, TSharedPtr<IDisplayClusterBarrier, ESPMode::ThreadSafe>> Barriers;
 
@@ -64,4 +94,25 @@ private:
 
 	// Critical section for internal data access synchronization
 	mutable FCriticalSection BarriersCS;
+
+private:
+
+	/**
+	 * A helper structure that wraps the actual barrier info holder
+	 * with an additional data for optimization purposes.
+	 */
+	struct FBarrierInfoWrapper
+	{
+		// The actual barrier information holder
+		TSharedRef<FBarrierInfo> BarrierInfo = MakeShared<FBarrierInfo>();
+
+		// Whether barrier information is locked from being updated
+		bool bBarrierInfoLockedOut = false;
+	};
+
+	// Holds extra information per-barrier
+	TMap<FString, FBarrierInfoWrapper> BarriersInfo;
+
+	// Critical section to synchronize access to the barriers information container
+	mutable FCriticalSection BarriersInfoCS;
 };
diff --git a/Engine/Plugins/Runtime/nDisplay/Source/DisplayCluster/Public/Cluster/IDisplayClusterGenericBarriersClient.h b/Engine/Plugins/Runtime/nDisplay/Source/DisplayCluster/Public/Cluster/IDisplayClusterGenericBarriersClient.h
--- a/Engine/Plugins/Runtime/nDisplay/Source/DisplayCluster/Public/Cluster/IDisplayClusterGenericBarriersClient.h
+++ b/Engine/Plugins/Runtime/nDisplay/Source/DisplayCluster/Public/Cluster/IDisplayClusterGenericBarriersClient.h
@@ -6,17 +6,20 @@
 
 
 /**
- * Barreir synchonrization callback data
+ * Barrier synchonrization callback data
  */
 struct FGenericBarrierSynchronizationDelegateData
 {
 	/** Barrier ID */
-	const FString BarrierId;
+	const FString& BarrierId;
 
-	/** Binary data provided on sync request (node Id - to - data mapping) */
+	/** Associates thread markers with the owning cluster nodes (thread marker - to - cluster node) */
+	const TMap<FString, FString>& ThreadToNodeMap;
+
+	/** Binary data provided on sync request (thread marker - to - data mapping) */
 	const TMap<FString, TArray<uint8>>& RequestData;
 
-	/** Binary data to respond (node Id - to - data mapping) */
+	/** Binary data to respond (thread marker - to - data mapping) */
 	TMap<FString, TArray<uint8>>& ResponseData;
 };
 
diff --git a/Engine/Plugins/VirtualProduction/Rivermax/RivermaxSync/Source/RivermaxSync/Private/MediaOutputSynchronizationPolicyRivermax.cpp b/Engine/Plugins/VirtualProduction/Rivermax/RivermaxSync/Source/RivermaxSync/Private/MediaOutputSynchronizationPolicyRivermax.cpp
--- a/Engine/Plugins/VirtualProduction/Rivermax/RivermaxSync/Source/RivermaxSync/Private/MediaOutputSynchronizationPolicyRivermax.cpp
+++ b/Engine/Plugins/VirtualProduction/Rivermax/RivermaxSync/Source/RivermaxSync/Private/MediaOutputSynchronizationPolicyRivermax.cpp
@@ -418,14 +418,26 @@ void FMediaOutputSynchronizationPolicyRivermaxHandler::HandleBarrierSync(FGeneri
 	const int32 PtpUnsyncFramesPerReport = UE::RivermaxSync::CVarRivermaxPtpUnsyncFramesPerReport.GetValueOnAnyThread();
 	bool bShouldReportPtpMismatches = (PtpUnsyncFramesPerReport >= 0) && !(PtpBaseNodeData->LastRenderedFrameNumber[0] % PtpUnsyncFramesPerReport);
 
-	// Iterate over the nodes and detect PTP de-syncs.
-	for (const TPair<FString, TArray<uint8>>& NodePresentedFrame : BarrierSyncData.RequestData)
+	// Iterate over the presentation requests and detect PTP de-syncs.
+	for (const TPair<FString, TArray<uint8>>& FramePresentedInfo : BarrierSyncData.RequestData)
 	{
+		const FString& ThreadPresentedFrame = FramePresentedInfo.Key;
+
+		// Find the corresponding cluster node based on the synchronization thread name from request
+		const FString* NodePresentedFrame = BarrierSyncData.ThreadToNodeMap.Find(ThreadPresentedFrame);
+		check(NodePresentedFrame);
+
+		// Don't leave nullptr just in case (this is not expected)
+		if (!NodePresentedFrame)
+		{
+			NodePresentedFrame = &ThreadPresentedFrame;
+		}
+
 		// Barrier data with unexpected sizes are a logical error.
-		check(NodePresentedFrame.Value.Num() == sizeof(FMediaSyncBarrierData));
+		check(FramePresentedInfo.Value.Num() == sizeof(FMediaSyncBarrierData));
 
 		// Get the node data in struct format.
-		const FMediaSyncBarrierData* const NodeData = reinterpret_cast<const FMediaSyncBarrierData* const>(NodePresentedFrame.Value.GetData());
+		const FMediaSyncBarrierData* const NodeData = reinterpret_cast<const FMediaSyncBarrierData* const>(FramePresentedInfo.Value.GetData());
 
 		// Skip base node comparing with itself
 		if (NodeData == PtpBaseNodeData)
@@ -448,7 +460,7 @@ void FMediaOutputSynchronizationPolicyRivermaxHandler::HandleBarrierSync(FGeneri
 				*PtpBaseNodeId,
 				*PtpBaseNodeData->LastRenderedFrameNumbersAsString(),
 				*PtpBaseNodeData->PresentedFrameBoundaryNumbersAsString(),
-				*NodePresentedFrame.Key,
+				**NodePresentedFrame,
 				*NodeData->LastRenderedFrameNumbersAsString(),
 				*NodeData->PresentedFrameBoundaryNumbersAsString());
 
@@ -461,7 +473,7 @@ void FMediaOutputSynchronizationPolicyRivermaxHandler::HandleBarrierSync(FGeneri
 				*PtpBaseNodeId,
 				*PtpBaseNodeData->LastRenderedFrameNumbersAsString(),
 				*PtpBaseNodeData->PresentedFrameBoundaryNumbersAsString(),
-				*NodePresentedFrame.Key,
+				**NodePresentedFrame,
 				*NodeData->LastRenderedFrameNumbersAsString(),
 				*NodeData->PresentedFrameBoundaryNumbersAsString());
 		}
@@ -469,7 +481,7 @@ void FMediaOutputSynchronizationPolicyRivermaxHandler::HandleBarrierSync(FGeneri
 		// Collect Vsync deltas for reporting purposes.
 		if (bShouldReportPtpMismatches && VsyncDelta)
 		{
-			PtpMismatchedNodes.Add(NodePresentedFrame.Key, VsyncDelta);
+			PtpMismatchedNodes.Add(*NodePresentedFrame, VsyncDelta);
 		}
 	}
 
