[ObjectMixer] [nDisplay] Fix cases where Object Mixer would deselect DisplayCluster root actors when changing their properties

diff --git a/Engine/Plugins/Editor/ObjectMixer/ObjectMixer/Source/ObjectMixer/Private/Views/List/Modes/ObjectMixerOutlinerMode.cpp b/Engine/Plugins/Editor/ObjectMixer/ObjectMixer/Source/ObjectMixer/Private/Views/List/Modes/ObjectMixerOutlinerMode.cpp
--- a/Engine/Plugins/Editor/ObjectMixer/ObjectMixer/Source/ObjectMixer/Private/Views/List/Modes/ObjectMixerOutlinerMode.cpp
+++ b/Engine/Plugins/Editor/ObjectMixer/ObjectMixer/Source/ObjectMixer/Private/Views/List/Modes/ObjectMixerOutlinerMode.cpp
@@ -1880,8 +1880,8 @@ void FObjectMixerOutlinerMode::OnItemAdded(FSceneOutlinerTreeItemPtr Item)
 		{
 			++FilteredActorCount;
 
-			// Synchronize selection with editor (or always select if sync is disabled)
-			if (!ShouldSyncSelectionFromEditor() || GetSelectedActorsInEditor().Contains(ActorItem->Actor.Get()))
+			// Synchronize selection with editor if enabled
+			if (ShouldSyncSelectionFromEditor() && GetSelectedActorsInEditor().Contains(ActorItem->Actor.Get()))
 			{
 				SceneOutliner->SetItemSelection(Item, true);
 			}
diff --git a/Engine/Plugins/Editor/ObjectMixer/ObjectMixer/Source/ObjectMixer/Public/Views/List/Modes/ObjectMixerOutlinerMode.h b/Engine/Plugins/Editor/ObjectMixer/ObjectMixer/Source/ObjectMixer/Public/Views/List/Modes/ObjectMixerOutlinerMode.h
--- a/Engine/Plugins/Editor/ObjectMixer/ObjectMixer/Source/ObjectMixer/Public/Views/List/Modes/ObjectMixerOutlinerMode.h
+++ b/Engine/Plugins/Editor/ObjectMixer/ObjectMixer/Source/ObjectMixer/Public/Views/List/Modes/ObjectMixerOutlinerMode.h
@@ -331,8 +331,6 @@ protected:
 	void SynchronizeComponentAndActorSelection();
 	void SynchronizeSelectedActorDescs();
 
-	void SynchronizeCustomActorSelection();
-
 	void OnActorEditorContextSubsystemChanged();
 
 	/** Filter factories */
diff --git a/Engine/Plugins/Runtime/nDisplay/Source/DisplayClusterLightCardEditor/Private/Outliner/SDisplayClusterLightCardOutliner.cpp b/Engine/Plugins/Runtime/nDisplay/Source/DisplayClusterLightCardEditor/Private/Outliner/SDisplayClusterLightCardOutliner.cpp
--- a/Engine/Plugins/Runtime/nDisplay/Source/DisplayClusterLightCardEditor/Private/Outliner/SDisplayClusterLightCardOutliner.cpp
+++ b/Engine/Plugins/Runtime/nDisplay/Source/DisplayClusterLightCardEditor/Private/Outliner/SDisplayClusterLightCardOutliner.cpp
@@ -59,7 +59,15 @@ void SDisplayClusterLightCardOutliner::SetRootActor(ADisplayClusterRootActor* Ne
 	{
 		PreviouslySelectedActors = LightCardEditorPtr.Pin()->GetSelectedActorsAs<AActor>();
 	}
-	
+
+	for (TWeakObjectPtr<AActor> Actor : NonOutlinerSelectedActors)
+	{
+		if (Actor.IsValid())
+		{
+			PreviouslySelectedActors.AddUnique(Actor.Get());
+		}
+	}
+
 	RootActor = NewRootActor;
 
 	FillActorList();
@@ -87,19 +95,29 @@ void SDisplayClusterLightCardOutliner::GetSelectedActors(TArray<AActor*>& OutSel
 			OutSelectedActors.Add(MatchingTreeItem->Actor.Get());
 		}
 	}
+
+	for (TWeakObjectPtr<AActor> Actor : NonOutlinerSelectedActors)
+	{
+		if (Actor.IsValid())
+		{
+			OutSelectedActors.Add(Actor.Get());
+		}
+	}
 }
 
 void SDisplayClusterLightCardOutliner::SelectActors(const TArray<AActor*>& ActorsToSelect)
 {
-	IDisplayClusterOperator::Get().GetOperatorViewModel()->OnOutlinerSelectionChanged().Broadcast(ActorsToSelect);
+	NonOutlinerSelectedActors.Empty();
 
 	if (bIsOutlinerChangingSelection)
 	{
-		// Selection change initiated from outliner
+		// Selection change initiated from outliner; assume the selection is already applied and non-outliner actors are deselected
+		IDisplayClusterOperator::Get().GetOperatorViewModel()->OnOutlinerSelectionChanged().Broadcast(ActorsToSelect);
 		return;
 	}
 	
 	TArray<FSceneOutlinerTreeItemPtr> SelectedTreeItems;
+	TArray<AActor*> SelectedActors;
 	CachedOutlinerItems.Reset();
 	
 	for (const TSharedPtr<FStageActorTreeItem>& TreeItem : StageActorTreeItems)
@@ -109,12 +127,23 @@ void SDisplayClusterLightCardOutliner::SelectActors(const TArray<AActor*>& Actor
 			if (FSceneOutlinerTreeItemPtr OutlinerItem = SceneOutliner->GetTreeItem(TreeItem->Actor.Get()))
 			{
 				SelectedTreeItems.Add(OutlinerItem);
+				SelectedActors.Add(TreeItem->Actor.Get());
 			}
 		}
 	}
 
+	for (AActor* Actor : ActorsToSelect)
+	{
+		if (!SelectedActors.Contains(Actor))
+		{
+			NonOutlinerSelectedActors.Add(Actor);
+		}
+	}
+
 	CachedOutlinerItems = SelectedTreeItems;
 	SceneOutliner->SetItemSelection(SelectedTreeItems, true);
+
+	IDisplayClusterOperator::Get().GetOperatorViewModel()->OnOutlinerSelectionChanged().Broadcast(ActorsToSelect);
 }
 
 void SDisplayClusterLightCardOutliner::RestoreCachedSelection()
@@ -140,6 +169,14 @@ void SDisplayClusterLightCardOutliner::RestoreCachedSelection()
 		}
 	}
 
+	for (TWeakObjectPtr<AActor> Actor : NonOutlinerSelectedActors)
+	{
+		if (Actor.IsValid())
+		{
+			ActorsToSelect.Add(Actor.Get());
+		}
+	}
+
 	if (LightCardEditorPtr.IsValid())
 	{
 		LightCardEditorPtr.Pin()->SelectActors(MoveTemp(ActorsToSelect));
@@ -373,6 +410,9 @@ void SDisplayClusterLightCardOutliner::OnOutlinerSelectionChanged(FSceneOutliner
 	ESelectInfo::Type Type)
 {
 	MostRecentSelectedItem = TreeItem;
+	
+	// User made a selection directly in the outliner, so override the non-outliner selection
+	NonOutlinerSelectedActors.Empty();
 
 	if (LightCardEditorPtr.IsValid())
 	{
diff --git a/Engine/Plugins/Runtime/nDisplay/Source/DisplayClusterLightCardEditor/Private/Outliner/SDisplayClusterLightCardOutliner.h b/Engine/Plugins/Runtime/nDisplay/Source/DisplayClusterLightCardEditor/Private/Outliner/SDisplayClusterLightCardOutliner.h
--- a/Engine/Plugins/Runtime/nDisplay/Source/DisplayClusterLightCardEditor/Private/Outliner/SDisplayClusterLightCardOutliner.h
+++ b/Engine/Plugins/Runtime/nDisplay/Source/DisplayClusterLightCardEditor/Private/Outliner/SDisplayClusterLightCardOutliner.h
@@ -106,6 +106,12 @@ private:
 	/** A flattened list of all the light card actors being displayed in the tree view */
 	TArray<TSharedPtr<FStageActorTreeItem>> StageActorTreeItems;
 
+	/**
+	 * Additional selected actors that aren't visible in the outliner hierarchy.
+	 * We store these to avoid trampling the light card editor's selection, which may include actors that aren't displayed in the outliner.
+	 */
+	TArray<TWeakObjectPtr<AActor>> NonOutlinerSelectedActors;
+
 	/** The active root actor whose light cards are being displayed */
 	TWeakObjectPtr<ADisplayClusterRootActor> RootActor;
 
