nDisplay: Media settings is now ShowInnerPropertiesOnly to avoid repeating the category name.

Also renamed "Enable" to "Enable Media" for clarity and consistency.

#rnx

diff --git a/Engine/Plugins/Runtime/nDisplay/Source/DisplayClusterConfiguration/Public/DisplayClusterConfigurationTypes.h b/Engine/Plugins/Runtime/nDisplay/Source/DisplayClusterConfiguration/Public/DisplayClusterConfigurationTypes.h
--- a/Engine/Plugins/Runtime/nDisplay/Source/DisplayClusterConfiguration/Public/DisplayClusterConfigurationTypes.h
+++ b/Engine/Plugins/Runtime/nDisplay/Source/DisplayClusterConfiguration/Public/DisplayClusterConfigurationTypes.h
@@ -359,7 +359,7 @@ public:
 	TMap<FString, FDisplayClusterConfigurationPostprocess> Postprocess;
 
 	// Media settings
-	UPROPERTY(EditDefaultsOnly, BlueprintReadWrite, Category = "Media", meta = (DisplayName = "Media"))
+	UPROPERTY(EditDefaultsOnly, BlueprintReadWrite, Category = "Media", meta = (DisplayName = "Media", ShowOnlyInnerProperties))
 	FDisplayClusterConfigurationMediaNodeBackbuffer MediaSettings;
 
 #if WITH_EDITORONLY_DATA
diff --git a/Engine/Plugins/Runtime/nDisplay/Source/DisplayClusterConfiguration/Public/DisplayClusterConfigurationTypes_ICVFX.h b/Engine/Plugins/Runtime/nDisplay/Source/DisplayClusterConfiguration/Public/DisplayClusterConfigurationTypes_ICVFX.h
--- a/Engine/Plugins/Runtime/nDisplay/Source/DisplayClusterConfiguration/Public/DisplayClusterConfigurationTypes_ICVFX.h
+++ b/Engine/Plugins/Runtime/nDisplay/Source/DisplayClusterConfiguration/Public/DisplayClusterConfigurationTypes_ICVFX.h
@@ -503,7 +503,7 @@ public:
 	FDisplayClusterConfigurationPostRender_Override Replace;
 
 	// Media settings
-	UPROPERTY(EditDefaultsOnly, BlueprintReadWrite, Category = NDisplay, meta = (DisplayName = "Media"))
+	UPROPERTY(EditDefaultsOnly, BlueprintReadWrite, Category = NDisplay, meta = (DisplayName = "Media", ShowOnlyInnerProperties))
 	FDisplayClusterConfigurationMediaICVFX Media;
 
 	UPROPERTY()
diff --git a/Engine/Plugins/Runtime/nDisplay/Source/DisplayClusterConfiguration/Public/DisplayClusterConfigurationTypes_Media.h b/Engine/Plugins/Runtime/nDisplay/Source/DisplayClusterConfiguration/Public/DisplayClusterConfigurationTypes_Media.h
--- a/Engine/Plugins/Runtime/nDisplay/Source/DisplayClusterConfiguration/Public/DisplayClusterConfigurationTypes_Media.h
+++ b/Engine/Plugins/Runtime/nDisplay/Source/DisplayClusterConfiguration/Public/DisplayClusterConfigurationTypes_Media.h
@@ -70,7 +70,7 @@ struct DISPLAYCLUSTERCONFIGURATION_API FDisplayClusterConfigurationMediaNodeBack
 
 public:
 	/** Enable/disable media */
-	UPROPERTY(EditAnywhere, BlueprintReadWrite, Category = "Media")
+	UPROPERTY(EditAnywhere, BlueprintReadWrite, Category = "Media", meta = (DisplayName = "Enable Media"))
 	bool bEnable = false;
 
 	/** Media outputs to use */
@@ -96,7 +96,7 @@ struct DISPLAYCLUSTERCONFIGURATION_API FDisplayClusterConfigurationMediaViewport
 
 public:
 	/** Enable/disable media */
-	UPROPERTY(EditAnywhere, BlueprintReadWrite, Category = "Media")
+	UPROPERTY(EditAnywhere, BlueprintReadWrite, Category = "Media", meta=(DisplayName="Enable Media"))
 	bool bEnable = false;
 
 	/** Media source to use */
@@ -261,7 +261,7 @@ struct DISPLAYCLUSTERCONFIGURATION_API FDisplayClusterConfigurationMediaICVFX
 
 public:
 	/** Enable/disable media */
-	UPROPERTY(EditAnywhere, BlueprintReadWrite, Category = "Media")
+	UPROPERTY(EditAnywhere, BlueprintReadWrite, Category = "Media", meta = (DisplayName = "Enable Media"))
 	bool bEnable = false;
 
 	/** Media frame split type */
