Fixed assigning an AActor* to a TSoftObjectPtr<ADisplayClusterRootActor>.

diff --git a/Engine/Plugins/Runtime/nDisplay/Source/DisplayCluster/Private/Game/EngineClasses/Scene/DisplayClusterChromakeyCardStageActorComponent.cpp b/Engine/Plugins/Runtime/nDisplay/Source/DisplayCluster/Private/Game/EngineClasses/Scene/DisplayClusterChromakeyCardStageActorComponent.cpp
--- a/Engine/Plugins/Runtime/nDisplay/Source/DisplayCluster/Private/Game/EngineClasses/Scene/DisplayClusterChromakeyCardStageActorComponent.cpp
+++ b/Engine/Plugins/Runtime/nDisplay/Source/DisplayCluster/Private/Game/EngineClasses/Scene/DisplayClusterChromakeyCardStageActorComponent.cpp
@@ -34,7 +34,7 @@ void UDisplayClusterChromakeyCardStageActorComponent::PostEditChangeProperty(FPr
 		const int32 ArrayIndex = PropertyChangedEvent.GetArrayIndex(PropertyName.ToString());
 		if (ArrayIndex >= 0 && ArrayIndex < ICVFXCameras.Num())
 		{
-			RootActor = ICVFXCameras[ArrayIndex].OtherActor.Get();
+			RootActor = Cast<ADisplayClusterRootActor>(ICVFXCameras[ArrayIndex].OtherActor.Get());
 			RemoveInvalidComponents();
 		}
 	}
