nDisplay: Default to showing section selector in details of NDC blueprint editor.

#rnx

diff --git a/Engine/Plugins/Runtime/nDisplay/Source/DisplayClusterConfigurator/Private/DisplayClusterConfiguratorBlueprintEditor.cpp b/Engine/Plugins/Runtime/nDisplay/Source/DisplayClusterConfigurator/Private/DisplayClusterConfiguratorBlueprintEditor.cpp
--- a/Engine/Plugins/Runtime/nDisplay/Source/DisplayClusterConfigurator/Private/DisplayClusterConfiguratorBlueprintEditor.cpp
+++ b/Engine/Plugins/Runtime/nDisplay/Source/DisplayClusterConfigurator/Private/DisplayClusterConfiguratorBlueprintEditor.cpp
@@ -2,73 +2,72 @@
 
 #include "DisplayClusterConfiguratorBlueprintEditor.h"
 
-#include "Framework/Application/SlateApplication.h"
-#include "IDisplayClusterConfigurator.h"
+#include "Blueprints/DisplayClusterBlueprint.h"
+#include "Blueprints/DisplayClusterBlueprintGeneratedClass.h"
+#include "ClusterConfiguration/DisplayClusterConfiguratorClusterUtils.h"
+#include "Components/DisplayClusterCameraComponent.h"
+#include "Components/DisplayClusterOriginComponent.h"
+#include "Components/DisplayClusterScreenComponent.h"
+#include "Components/DisplayClusterXformComponent.h"
 #include "DisplayClusterConfigurationTypes.h"
 #include "DisplayClusterConfiguratorCommands.h"
 #include "DisplayClusterConfiguratorEditorMode.h"
 #include "DisplayClusterConfiguratorEditorSubsystem.h"
+#include "DisplayClusterConfiguratorLog.h"
 #include "DisplayClusterConfiguratorModule.h"
+#include "DisplayClusterConfiguratorPropertyUtils.h"
 #include "DisplayClusterConfiguratorStyle.h"
 #include "DisplayClusterConfiguratorUtils.h"
-#include "DisplayClusterConfiguratorLog.h"
-
-#include "DisplayClusterRootActor.h"
-#include "DisplayClusterProjectionStrings.h"
-#include "Blueprints/DisplayClusterBlueprint.h"
-#include "Components/DisplayClusterCameraComponent.h"
-#include "Components/DisplayClusterOriginComponent.h"
-#include "Components/DisplayClusterScreenComponent.h"
-#include "Components/DisplayClusterXformComponent.h"
-
-#include "ClusterConfiguration/DisplayClusterConfiguratorClusterUtils.h"
-#include "DisplayClusterConfiguratorPropertyUtils.h"
 #include "DisplayClusterConfiguratorVersionUtils.h"
+#include "DisplayClusterProjectionStrings.h"
+#include "DisplayClusterRootActor.h"
+#include "IDisplayClusterConfigurator.h"
+#include "Settings/DisplayClusterConfiguratorSettings.h"
+#include "Views/DisplayClusterConfiguratorToolbar.h"
 #include "Views/OutputMapping/DisplayClusterConfiguratorViewOutputMapping.h"
+#include "Views/SCSEditor/SDisplayClusterConfiguratorComponentCombo.h"
 #include "Views/TreeViews/Cluster/DisplayClusterConfiguratorViewCluster.h"
 #include "Views/Viewport/DisplayClusterConfiguratorSCSEditorViewport.h"
 #include "Views/Viewport/DisplayClusterConfiguratorSCSEditorViewportClient.h"
-#include "Views/SCSEditor/SDisplayClusterConfiguratorComponentCombo.h"
-#include "Views/DisplayClusterConfiguratorToolbar.h"
-#include "Settings/DisplayClusterConfiguratorSettings.h"
 
-#include "Engine/SimpleConstructionScript.h"
-
-#include "Components/ActorComponent.h"
 #include "Camera/CameraComponent.h"
-#include "GameFramework/Actor.h"
-
+#include "ComponentAssetBroker.h"
+#include "Components/ActorComponent.h"
+#include "DesktopPlatformModule.h"
 #include "EditorDirectories.h"
 #include "EditorSupportDelegates.h"
 #include "EditorViewportTabContent.h"
-#include "ISCSEditorUICustomization.h"
-#include "SBlueprintEditorToolbar.h"
-#include "SubobjectEditorExtensionContext.h"
-#include "SKismetInspector.h"
-#include "SSCSEditor.h"
-#include "ToolMenu.h"
-#include "ToolMenus.h"
-#include "ComponentAssetBroker.h"
-#include "DesktopPlatformModule.h"
-#include "SSubobjectEditor.h"
-#include "SubobjectDataSubsystem.h"
-#include "Tools/BaseAssetToolkit.h"
 #include "Engine/Selection.h"
+#include "Engine/SimpleConstructionScript.h"
 #include "Engine/TextureRenderTarget2D.h"
+#include "Framework/Application/SlateApplication.h"
 #include "Framework/Commands/GenericCommands.h"
 #include "Framework/Notifications/NotificationManager.h"
-#include "Widgets/SOverlay.h"
-#include "Widgets/Docking/SDockTab.h"
-#include "Widgets/Notifications/SNotificationList.h"
+#include "GameFramework/Actor.h"
 #include "HAL/ConsoleManager.h"
 #include "HAL/PlatformApplicationMisc.h"
+#include "ISCSEditorUICustomization.h"
 #include "Kismet2/BlueprintEditorUtils.h"
 #include "Kismet2/DebuggerCommands.h"
-#include "Subsystems/PanelExtensionSubsystem.h"
-#include "ScopedTransaction.h"
-#include "Blueprints/DisplayClusterBlueprintGeneratedClass.h"
 #include "Misc/DisplayClusterHelpers.h"
 #include "MPCDI/DisplayClusterConfiguratorMPCDIImporter.h"
+#include "SBlueprintEditorToolbar.h"
+#include "ScopedTransaction.h"
+#include "SKismetInspector.h"
+#include "SKismetInspector.h"
+#include "SSCSEditor.h"
+#include "SSubobjectEditor.h"
+#include "SubobjectDataSubsystem.h"
+#include "SubobjectEditorExtensionContext.h"
+#include "Subsystems/PanelExtensionSubsystem.h"
+#include "ToolMenu.h"
+#include "ToolMenus.h"
+#include "Tools/BaseAssetToolkit.h"
+#include "Widgets/Docking/SDockTab.h"
+#include "Widgets/Notifications/SNotificationList.h"
+#include "Widgets/SOverlay.h"
 
+
+
 #define LOCTEXT_NAMESPACE "DisplayClusterConfiguratorBlueprintEditor"
 
@@ -1159,6 +1158,24 @@
 	FBlueprintEditor::Compile();
 }
 
+void FDisplayClusterConfiguratorBlueprintEditor::CreateDefaultTabContents(const TArray<UBlueprint*>& InBlueprints)
+{
+	Super::CreateDefaultTabContents(InBlueprints);
+
+	Inspector =
+		SNew(SKismetInspector)
+		.HideNameArea(true)
+		.ViewIdentifier(FName("BlueprintInspector"))
+		.Kismet2(SharedThis(this))
+		.OnFinishedChangingProperties(FOnFinishedChangingProperties::FDelegate::CreateSP(this, &FDisplayClusterConfiguratorBlueprintEditor::OnFinishedChangingProperties))
+		.ShowSectionSelector(true);
+}
+
+void FDisplayClusterConfiguratorBlueprintEditor::OnFinishedChangingProperties(const FPropertyChangedEvent& PropertyChangedEvent)
+{
+	Super::OnFinishedChangingProperties(PropertyChangedEvent);
+}
+
 void FDisplayClusterConfiguratorBlueprintEditor::OnSelectionUpdated(const TArray<TSharedPtr<class FSubobjectEditorTreeNode>>& SelectedNodes)
 {
 	if (CurrentSelectionSource == ESelectionSource::Refresh)
diff --git a/Engine/Plugins/Runtime/nDisplay/Source/DisplayClusterConfigurator/Private/DisplayClusterConfiguratorBlueprintEditor.h b/Engine/Plugins/Runtime/nDisplay/Source/DisplayClusterConfigurator/Private/DisplayClusterConfiguratorBlueprintEditor.h
--- a/Engine/Plugins/Runtime/nDisplay/Source/DisplayClusterConfigurator/Private/DisplayClusterConfiguratorBlueprintEditor.h
+++ b/Engine/Plugins/Runtime/nDisplay/Source/DisplayClusterConfigurator/Private/DisplayClusterConfiguratorBlueprintEditor.h
@@ -28,6 +28,8 @@ class FSubobjectEditorTreeNode;
 class FDisplayClusterConfiguratorBlueprintEditor
 	: public IDisplayClusterConfiguratorBlueprintEditor
 {
+	using Super = IDisplayClusterConfiguratorBlueprintEditor;
+
 public:
 	/** Registers any panel extension factories needed for the blueprint editor */
 	static void RegisterPanelExtensionFactory();
@@ -145,6 +147,8 @@ protected:
 	virtual bool OnRequestClose(EAssetEditorCloseReason InCloseReason) override;
 	virtual void OnClose() override;
 	virtual void Compile() override;
+	virtual void CreateDefaultTabContents(const TArray<UBlueprint*>& InBlueprints) override;
+	virtual void OnFinishedChangingProperties(const FPropertyChangedEvent& PropertyChangedEvent) override;
 	//~ End FBlueprintEditor Interface
 
 	// SSCS Implementation
