[nDisplay] Don't update actor gizmos until drag is finished (due to performance implications)

diff --git a/Engine/Plugins/Runtime/nDisplay/Source/DisplayClusterLightCardEditor/Private/Viewport/DisplayClusterLightCardEditorViewportClient.cpp b/Engine/Plugins/Runtime/nDisplay/Source/DisplayClusterLightCardEditor/Private/Viewport/DisplayClusterLightCardEditorViewportClient.cpp
--- a/Engine/Plugins/Runtime/nDisplay/Source/DisplayClusterLightCardEditor/Private/Viewport/DisplayClusterLightCardEditorViewportClient.cpp
+++ b/Engine/Plugins/Runtime/nDisplay/Source/DisplayClusterLightCardEditor/Private/Viewport/DisplayClusterLightCardEditorViewportClient.cpp
@@ -804,6 +804,18 @@ void FDisplayClusterLightCardEditorViewportClient::TrackingStopped()
 	{
 		InputMode = EInputMode::Idle;
 
+		for (const FDisplayClusterWeakStageActorPtr& ActorProxy : SelectedActors)
+		{
+			const FActorProxy* FoundProxy = ActorProxies.FindByKey(ActorProxy.AsActor());
+			if (FoundProxy && FoundProxy->Proxy == ActorProxy)
+			{
+				if (IDisplayClusterStageActor* StageActor = Cast<IDisplayClusterStageActor>(FoundProxy->LevelInstance.Get()))
+				{
+					StageActor->UpdateEditorGizmos();
+				}
+			}
+		}
+
 		DragWidgetOffset = FVector::ZeroVector;
 		EndTransaction();
 
@@ -2498,9 +2510,13 @@ void FDisplayClusterLightCardEditorViewportClient::PropagateActorTransform(const
 		// Allows MU to receive the update in real-time.
 		LevelInstance->PostEditMove(false);
 
-		if (IDisplayClusterStageActor* StageActor = Cast<IDisplayClusterStageActor>(LevelInstance))
+		// Gizmo update can be expensive, so prevent doing it every frame while dragging
+		if (InputMode != EInputMode::DraggingActor)
 		{
-			StageActor->UpdateEditorGizmos();
+			if (IDisplayClusterStageActor* StageActor = Cast<IDisplayClusterStageActor>(LevelInstance))
+			{
+				StageActor->UpdateEditorGizmos();
+			}
 		}
 	}
 }
