nDisplay: Fixed bugs for MPCDI.
- Fixed DisplayClusterTypesConverter::ToString() functions for FVector and FVector2D types: system functions use "%3.3f", which compresses the value. New DC functions are used to get full-size float.
- Added 2d profile selection and new properties to the viewport projection policy.
- The importer now passes all profiles as a parameter.
- Fixed mpcdi importer bug where the system relative file path did not match the UE relative path. As a result of this bug, when we selected the mpcdi file, the system base relative path which was then not recognized as a UE relative path in subsequent code. Similar pieces of logic have also been fixed.
- added a new parameter DesiredFOV instead of the 'BufferToWorldDistance' The screen distance from ViewOrigin is now calculated based on this value.
- The RootComponent is now used as the default ViewOrigin. The DefaultViewPoint cannot be used as the origin.

diff --git a/Engine/Plugins/Runtime/nDisplay/Source/DisplayCluster/Public/Misc/DisplayClusterTypesConverter.h b/Engine/Plugins/Runtime/nDisplay/Source/DisplayCluster/Public/Misc/DisplayClusterTypesConverter.h
--- a/Engine/Plugins/Runtime/nDisplay/Source/DisplayCluster/Public/Misc/DisplayClusterTypesConverter.h
+++ b/Engine/Plugins/Runtime/nDisplay/Source/DisplayCluster/Public/Misc/DisplayClusterTypesConverter.h
@@ -60,13 +60,23 @@ template <> inline FString DisplayClusterTypesConverter::ToString<>(const int32&
 template <> inline FString DisplayClusterTypesConverter::ToString<>(const uint32& From)     { return ToString(static_cast<int32>(From)); }
 template <> inline FString DisplayClusterTypesConverter::ToString<>(const float& From)      { return FString::SanitizeFloat(From); }
 template <> inline FString DisplayClusterTypesConverter::ToString<>(const double& From)     { return FString::Printf(TEXT("%lf"), From); }
-template <> inline FString DisplayClusterTypesConverter::ToString<>(const FVector& From)    { return From.ToString(); }
-template <> inline FString DisplayClusterTypesConverter::ToString<>(const FVector2D& From)  { return From.ToString(); }
 template <> inline FString DisplayClusterTypesConverter::ToString<>(const FRotator& From)   { return From.ToString(); }
 template <> inline FString DisplayClusterTypesConverter::ToString<>(const FMatrix& From)    { return From.ToString(); }
 template <> inline FString DisplayClusterTypesConverter::ToString<>(const FQuat& From)      { return From.ToString(); }
 template <> inline FString DisplayClusterTypesConverter::ToString<>(const FIntPoint& From)  { return From.ToString(); }
 
+// The TVector::ToString() function uses "%3.3f" which compresses the value. A proprietary function is used to get a full-size float.
+template <> inline FString DisplayClusterTypesConverter::ToString<>(const FVector& From)
+{
+	return FString::Printf(TEXT("X=%f Y=%f Z=%f"), From.X, From.Y, From.Z);
+}
+
+// The TVector2D::ToString() function uses "%3.3f" which compresses the value. A proprietary function is used to get a full-size float.
+template <> inline FString DisplayClusterTypesConverter::ToString<>(const FVector2D& From)
+{
+	return FString::Printf(TEXT("X=%f Y=%f"), From.X, From.Y);
+}
+
 // We can't just use FTimecode ToString as that loses information.
 template <> inline FString DisplayClusterTypesConverter::ToString<>(const FTimecode& From)
 {
diff --git a/Engine/Plugins/Runtime/nDisplay/Source/DisplayClusterConfigurator/Private/DisplayClusterConfiguratorBlueprintEditor.cpp b/Engine/Plugins/Runtime/nDisplay/Source/DisplayClusterConfigurator/Private/DisplayClusterConfiguratorBlueprintEditor.cpp
--- a/Engine/Plugins/Runtime/nDisplay/Source/DisplayClusterConfigurator/Private/DisplayClusterConfiguratorBlueprintEditor.cpp
+++ b/Engine/Plugins/Runtime/nDisplay/Source/DisplayClusterConfigurator/Private/DisplayClusterConfiguratorBlueprintEditor.cpp
@@ -642,7 +642,11 @@ bool FDisplayClusterConfiguratorBlueprintEditor::LoadWithOpenFileDialog()
 		{
 			FEditorDirectories::Get().SetLastDirectory(ELastDirectory::GENERIC_IMPORT, OpenFilenames[0]);
 
-			return LoadFromFile(OpenFilenames[0]);
+			// The system relative path is not associated with the project directory.
+			// We must convert it to the relative path of the UE project
+			const FString RelativeFilePath = DisplayClusterHelpers::filesystem::GetRelativePathForConfigResource(FPaths::ConvertRelativePathToFull(OpenFilenames[0]));
+
+			return LoadFromFile(RelativeFilePath);
 		}
 	}
 
@@ -1106,10 +1110,14 @@ void FDisplayClusterConfiguratorBlueprintEditor::ImportMPCDI_Clicked()
 			const FString& FileName = OpenFileNames[0];
 			FEditorDirectories::Get().SetLastDirectory(ELastDirectory::GENERIC_IMPORT, FileName);
 
+			// The system relative path is not associated with the project directory.
+			// We must convert it to the relative path of the UE project
+			const FString RelativeFilePath = DisplayClusterHelpers::filesystem::GetRelativePathForConfigResource(FPaths::ConvertRelativePathToFull(FileName));
+
 			UDisplayClusterBlueprint* Blueprint = LoadedBlueprint.Get();
 
 			FDisplayClusterConfiguratorMPCDIImporterParams ImportParams {};
-			if (FDisplayClusterConfiguratorMPCDIImporter::ImportMPCDIIntoBlueprint(FileName, Blueprint, ImportParams))
+			if (FDisplayClusterConfiguratorMPCDIImporter::ImportMPCDIIntoBlueprint(RelativeFilePath, Blueprint, ImportParams))
 			{
 				GetEditorData()->MarkPackageDirty();
 				ClusterChanged();
diff --git a/Engine/Plugins/Runtime/nDisplay/Source/DisplayClusterConfigurator/Private/MPCDI/DisplayClusterConfiguratorMPCDIImporter.cpp b/Engine/Plugins/Runtime/nDisplay/Source/DisplayClusterConfigurator/Private/MPCDI/DisplayClusterConfiguratorMPCDIImporter.cpp
--- a/Engine/Plugins/Runtime/nDisplay/Source/DisplayClusterConfigurator/Private/MPCDI/DisplayClusterConfiguratorMPCDIImporter.cpp
+++ b/Engine/Plugins/Runtime/nDisplay/Source/DisplayClusterConfigurator/Private/MPCDI/DisplayClusterConfiguratorMPCDIImporter.cpp
@@ -1,6 +1,7 @@
 // Copyright Epic Games, Inc. All Rights Reserved.
 
 #include "DisplayClusterConfiguratorMPCDIImporter.h"
+#include "DisplayClusterConfiguratorLog.h"
 
 #include "DisplayClusterProjectionStrings.h"
 #include "IDisplayClusterWarp.h"
@@ -15,16 +16,27 @@
 #include "Engine/SCS_Node.h"
 #include "Engine/SimpleConstructionScript.h"
 #include "Misc/DisplayClusterHelpers.h"
+#include "Misc/DisplayClusterProjectionHelpers.h"
 
 bool FDisplayClusterConfiguratorMPCDIImporter::ImportMPCDIIntoBlueprint(const FString& InFilePath, UDisplayClusterBlueprint* InBlueprint, const FDisplayClusterConfiguratorMPCDIImporterParams& InParams)
 {
 	const FString MPCDIFileFullPath = DisplayClusterHelpers::filesystem::GetFullPathForConfigResource(InFilePath);
+	if (MPCDIFileFullPath.IsEmpty())
+	{
+		UE_LOG(DisplayClusterConfiguratorLog, Error, TEXT("Could not find the MPCDI file '%s'."), *InFilePath);
+		return false;
+	}
 
 	TMap<FString, TMap<FString, FDisplayClusterWarpMPCDIAttributes>> MPCDIFile;
-	IDisplayClusterWarp::Get().ReadMPCDFileStructure(InFilePath, MPCDIFile);
-
+	if (!IDisplayClusterWarp::Get().ReadMPCDFileStructure(InFilePath, MPCDIFile))
+	{
+		UE_LOG(DisplayClusterConfiguratorLog, Error, TEXT("Could not read the data from the MPCDI file '%s'."), *InFilePath);
+		return false;
+	}
 
-	FName ViewOriginComponentName = !InParams.ViewOriginComponentName.IsNone() ? InParams.ViewOriginComponentName : TEXT("DefaultViewPoint");
+	// The ViewPoint component cannot be used as a ViewOrigin. They have different purposes.
+	// The default ViewOrigin is the RootComponent.
+	const FName ViewOriginComponentName = !InParams.ViewOriginComponentName.IsNone() ? InParams.ViewOriginComponentName : TEXT("RootComponent");
 
 	USCS_Node* ViewOriginNode = nullptr;
 	UDisplayClusterCameraComponent* ViewOriginComponent = Cast<UDisplayClusterCameraComponent>(FindBlueprintSceneComponent(InBlueprint, ViewOriginComponentName, &ViewOriginNode));
@@ -44,7 +56,7 @@ bool FDisplayClusterConfiguratorMPCDIImporter::ImportMPCDIIntoBlueprint(const FS
 					if (ScreenNode)
 					{
 						UDisplayClusterScreenComponent* ScreenComponent = CastChecked<UDisplayClusterScreenComponent>(ScreenNode->GetActualComponentTemplate(InBlueprint->GetGeneratedClass()));
-						ConfigureScreenFromRegion(ScreenComponent, ViewOriginComponent, Region.Value, InParams);
+						ConfigureScreenComponentFrom2DProfileRegion(ScreenComponent, ViewOriginComponent, Region.Value, InParams);
 
 						// If a new screen had to be created, it needs to be parented to a valid parent component
 						if (!bFoundExistingScreen)
@@ -92,14 +104,13 @@ bool FDisplayClusterConfiguratorMPCDIImporter::ImportMPCDIIntoBlueprint(const FS
 					ProjectionPolicyViewModel.SetParameterValue(DisplayClusterProjectionStrings::cfg::mpcdi::Region, Region.Key);
 					ProjectionPolicyViewModel.SetParameterValue(DisplayClusterProjectionStrings::cfg::mpcdi::EnablePreview, TEXT("true"));
 
-					if (Region.Value.ProfileType == EDisplayClusterWarpProfileType::warp_2D)
-					{
-						ProjectionPolicyViewModel.SetParameterValue(DisplayClusterProjectionStrings::cfg::mpcdi::MPCDIType, TEXT("2d"));
+					// Always pass the mpcdi profile type through the projection policy parameter.
+					const FString ProfileName = UE::DisplayClusterProjectionHelpers::MPCDI::ProfileTypeToString(Region.Value.ProfileType);
+					ProjectionPolicyViewModel.SetParameterValue(DisplayClusterProjectionStrings::cfg::mpcdi::MPCDIType, ProfileName);
 
-						if (InParams.bCreateStageGeometryComponents)
-						{
-							ProjectionPolicyViewModel.SetParameterValue(DisplayClusterProjectionStrings::cfg::mpcdi::Component, GetScreenNameForRegion(Region.Key));
-						}
+					if (Region.Value.ProfileType == EDisplayClusterWarpProfileType::warp_2D && InParams.bCreateStageGeometryComponents)
+					{
+						ProjectionPolicyViewModel.SetParameterValue(DisplayClusterProjectionStrings::cfg::mpcdi::Component, GetScreenNameForRegion(Region.Key));
 					}
 
 					if (!InParams.ParentComponentName.IsNone())
@@ -147,17 +158,37 @@ USCS_Node* FDisplayClusterConfiguratorMPCDIImporter::FindOrCreateScreenNodeForRe
 	return InBlueprint->SimpleConstructionScript->CreateNode(UDisplayClusterScreenComponent::StaticClass(), *ScreenName);
 }
 
-void FDisplayClusterConfiguratorMPCDIImporter::ConfigureScreenFromRegion(
+void FDisplayClusterConfiguratorMPCDIImporter::ConfigureScreenComponentFrom2DProfileRegion(
 	UDisplayClusterScreenComponent* InScreenComponent,
 	UDisplayClusterCameraComponent* InViewOriginComponent,
 	const FDisplayClusterWarpMPCDIAttributes& InAttributes,
 	const FDisplayClusterConfiguratorMPCDIImporterParams& InParams)
 {
-	const FVector2D BufferActualSize(InAttributes.Buffer.Resolution.X * InParams.BufferToWorldScale, InAttributes.Buffer.Resolution.Y * InParams.BufferToWorldScale);
-	const FVector2D RegionActualPosition((InAttributes.Region.Pos.X - 0.5) * BufferActualSize.X, (InAttributes.Region.Pos.Y - 0.5) * BufferActualSize.Y);
-	const FVector2D RegionActualSize(InAttributes.Region.Size.X * BufferActualSize.X, InAttributes.Region.Size.Y * BufferActualSize.Y);
+	// Computes the buffer size in world units from the resolution in pixels.
+	const FVector2D BufferSize(InAttributes.Buffer.Resolution.X * InParams.Profile2DParams.BufferPixelsToWorldUnits, InAttributes.Buffer.Resolution.Y * InParams.Profile2DParams.BufferPixelsToWorldUnits);
 
-	FVector ScreenPosition = FVector(InParams.BufferToWorldDistance, RegionActualPosition.X + RegionActualSize.X * 0.5, RegionActualPosition.Y + RegionActualSize.Y * 0.5);
+	// Calculates the position and size of the region.
+	const FVector2D  RegionPos(InAttributes.Region.Pos.X  * BufferSize.X, InAttributes.Region.Pos.Y  * BufferSize.Y);
+	const FVector2D RegionSize(InAttributes.Region.Size.X * BufferSize.X, InAttributes.Region.Size.Y * BufferSize.Y);
+
+	float FocalLength = 0.f;
+	// Moves the buffer position along the X axis to achieve a DesiredFOV.
+	if (InParams.Profile2DParams.DesiredFOV > 0.f && InParams.Profile2DParams.DesiredFOV < 180.f)
+	{
+		// Convert FOV to focal length,
+		// 
+		// fov = 2 * atan(d/(2*f))
+		// where,
+		//   d = sensor dimension
+		//   f = focal length
+		// 
+		// f = 0.5 * d * (1/tan(fov/2))
+		const float TanHalfFOV = FMath::Tan(FMath::DegreesToRadians(InParams.Profile2DParams.DesiredFOV * 0.5f));
+		FocalLength = (BufferSize.X * 0.5f) / TanHalfFOV;
+	}
+
+	const FVector2D RegionCenterPos = RegionPos + (RegionSize * 0.5) - (BufferSize * 0.5);
+	FVector ScreenPosition = FVector(FocalLength, RegionCenterPos.X, RegionCenterPos.Y);
 
 	if (InViewOriginComponent)
 	{
@@ -183,7 +214,7 @@ void FDisplayClusterConfiguratorMPCDIImporter::ConfigureScreenFromRegion(
 	}
 
 	InScreenComponent->SetRelativeLocation(ScreenPosition);
-	InScreenComponent->SetScreenSize(RegionActualSize);
+	InScreenComponent->SetScreenSize(RegionSize);
 }
 
 UDisplayClusterConfigurationViewport* FDisplayClusterConfiguratorMPCDIImporter::FindOrCreateViewportForRegion(UDisplayClusterBlueprint* InBlueprint, const FString& RegionId, bool& bOutFoundExistingViewport)
@@ -250,6 +281,11 @@ USceneComponent* FDisplayClusterConfiguratorMPCDIImporter::FindBlueprintSceneCom
 	USceneComponent* FoundComponent = nullptr;
 	if (CDO)
 	{
+		if (ComponentName.IsNone())
+		{
+			return FoundComponent = CDO->GetRootComponent();
+		}
+
 		for (UActorComponent* Component : CDO->GetComponents())
 		{
 			if (USceneComponent* SceneComponent = Cast<USceneComponent>(Component))
@@ -266,6 +302,11 @@ USceneComponent* FDisplayClusterConfiguratorMPCDIImporter::FindBlueprintSceneCom
 	// If a native component was not found, check the SCS to see if one exists there
 	if (!FoundComponent)
 	{
+		if (ComponentName.IsNone())
+		{
+			return nullptr;
+		}
+
 		if (USCS_Node* FoundNode = InBlueprint->SimpleConstructionScript->FindSCSNode(ComponentName))
 		{
 			if (FoundNode->ComponentTemplate && FoundNode->ComponentTemplate->IsA<USceneComponent>())
diff --git a/Engine/Plugins/Runtime/nDisplay/Source/DisplayClusterConfigurator/Private/MPCDI/DisplayClusterConfiguratorMPCDIImporter.h b/Engine/Plugins/Runtime/nDisplay/Source/DisplayClusterConfigurator/Private/MPCDI/DisplayClusterConfiguratorMPCDIImporter.h
--- a/Engine/Plugins/Runtime/nDisplay/Source/DisplayClusterConfigurator/Private/MPCDI/DisplayClusterConfiguratorMPCDIImporter.h
+++ b/Engine/Plugins/Runtime/nDisplay/Source/DisplayClusterConfigurator/Private/MPCDI/DisplayClusterConfiguratorMPCDIImporter.h
@@ -13,14 +13,21 @@ class UDisplayClusterConfigurationViewport;
 class UDisplayClusterBlueprint;
 class UDisplayClusterScreenComponent;
 
-/** USed to pass parameters into the MPCDI importer */
-struct FDisplayClusterConfiguratorMPCDIImporterParams
+/** Additional parameters for the mpcdi 2d profile. */
+struct FDisplayClusterConfiguratorMPCDIProfile2DImporterParams
 {
 	/** The amount to scale an MPCDI buffer by when converting from pixels to world coordinates */
-	float BufferToWorldScale = 0.1;
+	float BufferPixelsToWorldUnits = 0.1;
+
+	/** Desired FOV. The focal length will be calculated for that FOV and the geometry will be shifted by that value. */
+	float DesiredFOV = 90.f;
+};
 
-	/** The distance from the origin to position an MPCDI buffer in world coordinates */
-	float BufferToWorldDistance = 100;
+/** Used to pass parameters into the MPCDI importer */
+struct FDisplayClusterConfiguratorMPCDIImporterParams
+{
+	/** Additional parameters for the mpcdi 2d profile. */
+	FDisplayClusterConfiguratorMPCDIProfile2DImporterParams Profile2DParams;
 
 	/** The name of the component to make the parent of the MPCDI geometry */
 	FName ParentComponentName = NAME_None;
@@ -68,7 +75,7 @@ private:
 	 * @param InAttributes - The attributes of the region to configure
 	 * @param InParams - The import parameters 
 	 **/
-	static void ConfigureScreenFromRegion(UDisplayClusterScreenComponent* InScreenComponent, UDisplayClusterCameraComponent* InViewOriginComponent, const FDisplayClusterWarpMPCDIAttributes& InAttributes, const FDisplayClusterConfiguratorMPCDIImporterParams& InParams);
+	static void ConfigureScreenComponentFrom2DProfileRegion(UDisplayClusterScreenComponent* InScreenComponent, UDisplayClusterCameraComponent* InViewOriginComponent, const FDisplayClusterWarpMPCDIAttributes& InAttributes, const FDisplayClusterConfiguratorMPCDIImporterParams& InParams);
 
 	/**
 	 * Attempts to find an existing viewport configuration for the specified region, and if one is not found, creates a new viewport configuration
diff --git a/Engine/Plugins/Runtime/nDisplay/Source/DisplayClusterConfigurator/Private/Views/Details/Policies/DisplayClusterConfiguratorPolicyDetailCustomization.cpp b/Engine/Plugins/Runtime/nDisplay/Source/DisplayClusterConfigurator/Private/Views/Details/Policies/DisplayClusterConfiguratorPolicyDetailCustomization.cpp
--- a/Engine/Plugins/Runtime/nDisplay/Source/DisplayClusterConfigurator/Private/Views/Details/Policies/DisplayClusterConfiguratorPolicyDetailCustomization.cpp
+++ b/Engine/Plugins/Runtime/nDisplay/Source/DisplayClusterConfigurator/Private/Views/Details/Policies/DisplayClusterConfiguratorPolicyDetailCustomization.cpp
@@ -834,6 +834,24 @@ void FDisplayClusterConfiguratorProjectionCustomization::CreateMPCDIPolicy(UDisp
 	}
 	else if (Setting == TypePFM)
 	{
+		const FString MPCDIProfile = DisplayClusterProjectionStrings::cfg::mpcdi::Profiles::mpcdi_a3d;
+		const TSharedPtr<FPolicyParameterInfoCombo> MPCDIProfilesCombo = MakeShared<FPolicyParameterInfoCombo>(
+			"Profile",
+			DisplayClusterProjectionStrings::cfg::mpcdi::MPCDIType,
+			Blueprint,
+			ConfigurationViewports,
+			TArray<FString>{
+				// Note: The MPCDI/PFM configuration currently supports only 2 mpcdi profiles.
+				DisplayClusterProjectionStrings::cfg::mpcdi::Profiles::mpcdi_2d,
+				DisplayClusterProjectionStrings::cfg::mpcdi::Profiles::mpcdi_a3d
+			},
+			& MPCDIProfile,
+			false);
+
+		CustomPolicyParameters.Add(MPCDIProfilesCombo);
+		MPCDIProfilesCombo->SetOnSelectedDelegate(FPolicyParameterInfoCombo::FOnItemSelected::CreateLambda(RefreshPolicy));
+		const FString CurrentMPCDIProfile = MPCDIProfilesCombo->GetOrAddCustomParameterValueText().ToString();
+
 		CustomPolicyParameters.Add(MakeShared<FPolicyParameterInfoFile>(
 			"File",
 			DisplayClusterProjectionStrings::cfg::mpcdi::FilePFM,
@@ -841,6 +859,31 @@ void FDisplayClusterConfiguratorProjectionCustomization::CreateMPCDIPolicy(UDisp
 			ConfigurationViewports,
 			TArray<FString>{"pfm"}));
 
+		if (CurrentMPCDIProfile == DisplayClusterProjectionStrings::cfg::mpcdi::Profiles::mpcdi_2d)
+		{
+			// Adds additional parameters for the mpcdi 2d profile.
+			// <buffer Xresolution, Yresolution> <region x, y, xsize, ysize>
+
+			CustomPolicyParameters.Add(MakeShared<FPolicyParameterInfoResolution>(
+				"Buffer Resolution",
+				DisplayClusterProjectionStrings::cfg::mpcdi::Attributes::Buffer::Resolution,
+				Blueprint,
+				ConfigurationViewports));
+
+			CustomPolicyParameters.Add(MakeShared<FPolicyParameterInfoNormalizedVector2D>(
+				"Region Pos",
+				DisplayClusterProjectionStrings::cfg::mpcdi::Attributes::Region::Pos,
+				Blueprint,
+				ConfigurationViewports));
+
+			const FString RegionDefaultSize("X=1, Y=1");
+			CustomPolicyParameters.Add(MakeShared<FPolicyParameterInfoNormalizedVector2D>(
+				"Region Size",
+				DisplayClusterProjectionStrings::cfg::mpcdi::Attributes::Region::Size,
+				Blueprint,
+				ConfigurationViewports,&RegionDefaultSize));
+		}
+
 		CustomPolicyParameters.Add(MakeShared<FPolicyParameterInfoFile>(
 			"Alpha Mask",
 			DisplayClusterProjectionStrings::cfg::mpcdi::FileAlpha,
diff --git a/Engine/Plugins/Runtime/nDisplay/Source/DisplayClusterConfigurator/Private/Views/Details/Policies/DisplayClusterConfiguratorPolicyParameterCustomization.cpp b/Engine/Plugins/Runtime/nDisplay/Source/DisplayClusterConfigurator/Private/Views/Details/Policies/DisplayClusterConfiguratorPolicyParameterCustomization.cpp
--- a/Engine/Plugins/Runtime/nDisplay/Source/DisplayClusterConfigurator/Private/Views/Details/Policies/DisplayClusterConfiguratorPolicyParameterCustomization.cpp
+++ b/Engine/Plugins/Runtime/nDisplay/Source/DisplayClusterConfigurator/Private/Views/Details/Policies/DisplayClusterConfiguratorPolicyParameterCustomization.cpp
@@ -34,6 +34,23 @@
 
 FLinearColor OrangeLabelBackgroundColor(0.8f, 0.3f, 0.0f);
 
+
+int32 GDisplayClusterConfiguratorMinResolution = 0;
+static FAutoConsoleVariableRef CVarDisplayClusterConfiguratorMinResolution(
+	TEXT("nDisplay.UI.MinResolution"),
+	GDisplayClusterConfiguratorMinResolution,
+	TEXT("The minimum value of the resolution parameter in the nDisplay UI. (Default = 0)"),
+	ECVF_Default
+);
+
+int32 GDisplayClusterConfiguratorMaxResolution = 16384;
+static FAutoConsoleVariableRef CVarDisplayClusterConfiguratorMaxResolution(
+	TEXT("nDisplay.UI.MaxResolution"),
+	GDisplayClusterConfiguratorMaxResolution,
+	TEXT("The maximum value of the resolution parameter in the nDisplay UI. (Default = 16384)"),
+	ECVF_Default
+);
+
 //////////////////////////////////////////////////////////////////////////////////////////////
 // Policy Parameter Configuration
 //////////////////////////////////////////////////////////////////////////////////////////////
@@ -493,11 +510,33 @@ FReply FPolicyParameterInfoFile::OnChangePathClicked()
 	return FReply::Handled();
 }
 
+/**
+* FPolicyParameterInfoFloatReference
+*/
 TSharedRef<SWidget> FPolicyParameterInfoFloatReference::MakeFloatInputWidget(TSharedRef<TOptional<float>>& ProxyValue, const FText& Label,
-                                                                     bool bRotationInDegrees,
-                                                                     const FLinearColor& LabelColor,
-                                                                     const FLinearColor& LabelBackgroundColor)
+											const bool bRotationInDegrees,
+											const FLinearColor& LabelColor, const FLinearColor& LabelBackgroundColor,
+											const float* InMinValue, const float* InMaxValue)
 {
+	TOptional<float> MinValue, MaxValue;
+
+	if (bRotationInDegrees)
+	{
+		MinValue = 0;
+		MaxValue = 360;
+	}
+
+	if (InMinValue)
+	{
+		MinValue = *InMinValue;
+	}
+	
+	if (InMaxValue)
+	{
+		MaxValue = *InMaxValue;
+	}
+
+
 	return
 		SNew(SNumericEntryBox<float>)
 			.Value(this, &FPolicyParameterInfoFloatReference::OnGetValue, ProxyValue)
@@ -507,20 +546,51 @@ TSharedRef<SWidget> FPolicyParameterInfoFloatReference::MakeFloatInputWidget(TSh
 			.LabelVAlign(VAlign_Fill)
 			.LabelPadding(0)
 			.AllowSpin(bRotationInDegrees)
-			.MaxSliderValue(bRotationInDegrees ? 360.0f : TOptional<float>())
-			.MinSliderValue(bRotationInDegrees ? 0.0f : TOptional<float>())
+			.MaxSliderValue(MaxValue)
+			.MinSliderValue(MinValue)
 			.Label()
 		[
 			SNumericEntryBox<float>::BuildLabel(Label, LabelColor, LabelBackgroundColor)
 		];
 }
 
-void FPolicyParameterInfoFloatReference::OnValueCommitted(float NewValue, ETextCommit::Type CommitType, TSharedRef<TOptional<float>> Value)
+/**
+* FPolicyParameterInfoIntReference
+*/
+TSharedRef<SWidget> FPolicyParameterInfoIntReference::MakeIntInputWidget(TSharedRef<TOptional<int32>>& ProxyValue, const FText& Label,
+										const FLinearColor& LabelColor, const FLinearColor& LabelBackgroundColor,
+										const int32* InMinValue, const int32* InMaxValue)
 {
-	*Value = NewValue;
-	FormatTextAndUpdateParameter();
+	TOptional<int32> MinValue, MaxValue;
+
+	if (InMinValue)
+	{
+		MinValue = *InMinValue;
+	}
+
+	if (InMaxValue)
+	{
+		MaxValue = *InMaxValue;
+	}
+	return
+		SNew(SNumericEntryBox<int32>)
+			.Value(this, &FPolicyParameterInfoIntReference::OnGetValue, ProxyValue)
+			.Font(IDetailLayoutBuilder::GetDetailFont())
+			.UndeterminedString(NSLOCTEXT("PropertyEditor", "MultipleValues", "Multiple Values"))
+			.OnValueCommitted(this, &FPolicyParameterInfoIntReference::OnValueCommitted, ProxyValue)
+			.LabelVAlign(VAlign_Fill)
+			.LabelPadding(0)
+			.MaxSliderValue(MaxValue)
+			.MinSliderValue(MinValue)
+			.Label()
+		[
+			SNumericEntryBox<float>::BuildLabel(Label, LabelColor, LabelBackgroundColor)
+		];
 }
 
+/**
+* FPolicyParameterInfoMatrix
+*/
 const FString FPolicyParameterInfoMatrix::BaseMatrixString = DisplayClusterTypesConverter::template ToString<FMatrix>(FMatrix(FPlane(1.0f, 0.0f, 0.0f, 0.0f), FPlane(0.0f, 1.0f, 0.0f, 0.0f), FPlane(0.0f, 0.0f, 0.0f, 1.0f), FPlane(0.0f, 0.0f, 1.0f, 0.0f)));
 
 FPolicyParameterInfoMatrix::FPolicyParameterInfoMatrix(const FString& InDisplayName, const FString& InKey,
@@ -1082,4 +1152,167 @@ void FPolicyParameterInfoFrustumAngle::FormatTextAndUpdateParameter()
 	UpdateCustomParameterValueText(AngleString);
 }
 
+/**
+* FPolicyParameterInfoResolution
+*/
+const FString FPolicyParameterInfoResolution::BaseResolutionString = FString("X=2048, Y=2048");
+
+FPolicyParameterInfoResolution::FPolicyParameterInfoResolution(const FString& InDisplayName, const FString& InKey,
+	UDisplayClusterBlueprint* InBlueprint, const TArray<TWeakObjectPtr<UDisplayClusterConfigurationViewport>>& InConfigurationViewports)
+	: FPolicyParameterInfoIntReference(InDisplayName, InKey, InBlueprint, InConfigurationViewports, &BaseResolutionString)
+	, CachedX(MakeShared<TOptional<int32>>())
+	, CachedY(MakeShared<TOptional<int32>>())
+{
+	if (!DoParametersMatchForAllViewports())
+	{
+		for (const TWeakObjectPtr<UDisplayClusterConfigurationViewport>& Viewport : ConfigurationViewports)
+		{
+			UDisplayClusterConfigurationViewport* ConfigurationViewport = Viewport.Get();
+			check(ConfigurationViewport != nullptr);
+
+			if (const FString* ParameterValue = ConfigurationViewport->ProjectionPolicy.Parameters.Find(GetParameterKey()))
+			{
+				const FIntPoint Resolution = DisplayClusterTypesConverter::template FromString<FIntPoint>(*ParameterValue);
+
+				ResetOrSetCachedValue(CachedX, FMath::Max(0, Resolution.X));
+				ResetOrSetCachedValue(CachedY, FMath::Max(0, Resolution.Y));
+			}
+		}
+	}
+	else
+	{
+		const FString TextValue = GetOrAddCustomParameterValueText().ToString();
+		const FIntPoint Resolution = DisplayClusterTypesConverter::template FromString<FIntPoint>(*TextValue);
+
+		*CachedX = FMath::Max(0, Resolution.X);
+		*CachedY = FMath::Max(0, Resolution.Y);
+	}
+}
+
+void FPolicyParameterInfoResolution::CreateCustomRowWidget(IDetailChildrenBuilder& InDetailWidgetRow)
+{
+	InDetailWidgetRow.AddCustomRow(GetParameterDisplayName())
+		.Visibility(MakeAttributeRaw(this, &FPolicyParameterInfoResolution::IsParameterVisible))
+		.NameContent()
+		[
+			SNew(STextBlock)
+				.Font(IDetailLayoutBuilder::GetDetailFont())
+				.Text(GetParameterDisplayName())
+				.ToolTipText(this, &FPolicyParameterInfo::GetParameterTooltip)
+		]
+		.ValueContent()
+		.MinDesiredWidth(375.0f)
+		.MaxDesiredWidth(375.0f)
+		[
+			SNew(SHorizontalBox)
+				+ SHorizontalBox::Slot()
+				.Padding(FMargin(0.0f, 2.0f, 3.0f, 2.0f))
+				[
+					MakeIntInputWidget(CachedX, LOCTEXT("Width", "W"),
+						FLinearColor::White, SNumericEntryBox<float>::RedLabelBackgroundColor,
+						&GDisplayClusterConfiguratorMinResolution, &GDisplayClusterConfiguratorMaxResolution)
+				]
+				+ SHorizontalBox::Slot()
+				.Padding(FMargin(0.0f, 2.0f, 3.0f, 2.0f))
+				[
+					MakeIntInputWidget(CachedY, LOCTEXT("Height", "H"),
+						FLinearColor::White, SNumericEntryBox<float>::GreenLabelBackgroundColor,
+						&GDisplayClusterConfiguratorMinResolution, &GDisplayClusterConfiguratorMaxResolution)
+				]
+		];
+}
+
+void FPolicyParameterInfoResolution::FormatTextAndUpdateParameter()
+{
+	const bool bAllValuesSet = CachedX->IsSet() && CachedY->IsSet();
+	if (!bAllValuesSet)
+	{
+		return;
+	}
+
+	const FIntPoint Resolution(*CachedX.Get(), *CachedY.Get());
+	UpdateCustomParameterValueText(Resolution.ToString());
+}
+
+/**
+* FPolicyParameterInfoNormalizedVector2D
+*/
+const FString FPolicyParameterInfoNormalizedVector2D::BaseVectorValueStr = FString("X=0, Y=0");
+
+FPolicyParameterInfoNormalizedVector2D::FPolicyParameterInfoNormalizedVector2D(const FString& InDisplayName, const FString& InKey,
+	UDisplayClusterBlueprint* InBlueprint, const TArray<TWeakObjectPtr<UDisplayClusterConfigurationViewport>>& InConfigurationViewports, const FString* InDefaultValue)
+	: FPolicyParameterInfoFloatReference(InDisplayName, InKey, InBlueprint, InConfigurationViewports, InDefaultValue ? InDefaultValue : &BaseVectorValueStr)
+	, CachedX(MakeShared<TOptional<float>>())
+	, CachedY(MakeShared<TOptional<float>>())
+{
+	if (!DoParametersMatchForAllViewports())
+	{
+		for (const TWeakObjectPtr<UDisplayClusterConfigurationViewport>& Viewport : ConfigurationViewports)
+		{
+			UDisplayClusterConfigurationViewport* ConfigurationViewport = Viewport.Get();
+			check(ConfigurationViewport != nullptr);
+
+			if (const FString* ParameterValue = ConfigurationViewport->ProjectionPolicy.Parameters.Find(GetParameterKey()))
+			{
+				const FVector2D Pts = DisplayClusterTypesConverter::template FromString<FVector2D>(*ParameterValue);
+
+				ResetOrSetCachedValue(CachedX, FMath::Max(0.f, FMath::Min(1.f, Pts.X)));
+				ResetOrSetCachedValue(CachedY, FMath::Max(0.f, FMath::Min(1.f, Pts.Y)));
+			}
+		}
+	}
+	else
+	{
+		const FString TextValue = GetOrAddCustomParameterValueText().ToString();
+		const FVector2D Pts = DisplayClusterTypesConverter::template FromString<FVector2D>(*TextValue);
+
+		*CachedX = FMath::Max(0.f, FMath::Min(1.f, Pts.X));
+		*CachedY = FMath::Max(0.f, FMath::Min(1.f, Pts.Y));
+	}
+}
+
+void FPolicyParameterInfoNormalizedVector2D::CreateCustomRowWidget(IDetailChildrenBuilder& InDetailWidgetRow)
+{
+	const float MinVectorComponentValue = 0.f;
+	const float MaxVectorComponentValue = 1.0f;
+
+	InDetailWidgetRow.AddCustomRow(GetParameterDisplayName())
+		.Visibility(MakeAttributeRaw(this, &FPolicyParameterInfoNormalizedVector2D::IsParameterVisible))
+		.NameContent()
+		[
+			SNew(STextBlock)
+				.Font(IDetailLayoutBuilder::GetDetailFont())
+				.Text(GetParameterDisplayName())
+				.ToolTipText(this, &FPolicyParameterInfo::GetParameterTooltip)
+		]
+		.ValueContent()
+		.MinDesiredWidth(375.0f)
+		.MaxDesiredWidth(375.0f)
+		[
+			SNew(SHorizontalBox)
+				+ SHorizontalBox::Slot()
+				.Padding(FMargin(0.0f, 2.0f, 3.0f, 2.0f))
+				[
+					MakeFloatInputWidget(CachedX, LOCTEXT("NormV2D_X", "X"), false, FLinearColor::White, SNumericEntryBox<float>::RedLabelBackgroundColor, &MinVectorComponentValue, &MaxVectorComponentValue)
+				]
+				+ SHorizontalBox::Slot()
+				.Padding(FMargin(0.0f, 2.0f, 3.0f, 2.0f))
+				[
+					MakeFloatInputWidget(CachedY, LOCTEXT("NormV2D_Y", "Y"), false, FLinearColor::White, SNumericEntryBox<float>::GreenLabelBackgroundColor, &MinVectorComponentValue, &MaxVectorComponentValue)
+				]
+		];
+}
+
+void FPolicyParameterInfoNormalizedVector2D::FormatTextAndUpdateParameter()
+{
+	const bool bAllValuesSet = CachedX->IsSet() && CachedY->IsSet();
+	if (!bAllValuesSet)
+	{
+		return;
+	}
+
+	const FVector2D Pts(*CachedX.Get(), *CachedY.Get());
+	UpdateCustomParameterValueText(DisplayClusterTypesConverter::ToString(Pts));
+}
+
 #undef LOCTEXT_NAMESPACE
diff --git a/Engine/Plugins/Runtime/nDisplay/Source/DisplayClusterConfigurator/Private/Views/Details/Policies/DisplayClusterConfiguratorPolicyParameterCustomization.h b/Engine/Plugins/Runtime/nDisplay/Source/DisplayClusterConfigurator/Private/Views/Details/Policies/DisplayClusterConfiguratorPolicyParameterCustomization.h
--- a/Engine/Plugins/Runtime/nDisplay/Source/DisplayClusterConfigurator/Private/Views/Details/Policies/DisplayClusterConfiguratorPolicyParameterCustomization.h
+++ b/Engine/Plugins/Runtime/nDisplay/Source/DisplayClusterConfigurator/Private/Views/Details/Policies/DisplayClusterConfiguratorPolicyParameterCustomization.h
@@ -344,18 +344,32 @@ public:
 	}
 
 protected:
+	/** Converts values to text. */
 	virtual void FormatTextAndUpdateParameter() = 0;
 	
-	TSharedRef<SWidget> MakeFloatInputWidget(TSharedRef<TOptional<float>>& ProxyValue, const FText& Label, bool bRotationInDegrees,
-	                                         const FLinearColor& LabelColor, const FLinearColor& LabelBackgroundColor);
+	/** Creates SNumericEntryBox<float> widget instance. */
+	TSharedRef<SWidget> MakeFloatInputWidget(TSharedRef<TOptional<float>>& ProxyValue, const FText& Label,
+											const bool bRotationInDegrees,
+											const FLinearColor& LabelColor, const FLinearColor& LabelBackgroundColor,
+											const float* InMinValue = nullptr, const float* InMaxValue = nullptr);
 
+	/** Returns an optional float value. */
 	TOptional<float> OnGetValue(TSharedRef<TOptional<float>> Value) const
 	{
 		return Value.Get();
 	}
 
-	void OnValueCommitted(float NewValue, ETextCommit::Type CommitType, TSharedRef<TOptional<float>> Value);
+	/** Callback that is called when the value changes. */
+	void OnValueCommitted(float NewValue, ETextCommit::Type CommitType, TSharedRef<TOptional<float>> Value)
+	{
+		*Value = NewValue;
+		FormatTextAndUpdateParameter();
+	}
 
+	/** Reset or set cached value.
+	* This function can be called multiple times when editing a value for a group selection.
+	* In this case, when the values in the group differ, a reset is used.
+	*/
 	static void ResetOrSetCachedValue(TSharedRef<TOptional<float>>& InCachedValue, float InCompareValue)
 	{
 		if (InCachedValue->IsSet() && InCachedValue.Get() != InCompareValue)
@@ -369,6 +383,61 @@ protected:
 	};
 };
 
+/**
+ * Policy info for modifying positive int32 reference, typically in a resolution
+ */
+class FPolicyParameterInfoIntReference : public FPolicyParameterInfo
+{
+public:
+	FPolicyParameterInfoIntReference(
+		const FString& InDisplayName,
+		const FString& InKey,
+		UDisplayClusterBlueprint* InBlueprint,
+		const TArray<TWeakObjectPtr<UDisplayClusterConfigurationViewport>>& InConfigurationViewports,
+		const FString* InInitialValue = nullptr) :
+		FPolicyParameterInfo(InDisplayName, InKey, InBlueprint, InConfigurationViewports, InInitialValue)
+	{
+	}
+
+protected:
+	/** Converts values to text. */
+	virtual void FormatTextAndUpdateParameter() = 0;
+
+	/** Creates SNumericEntryBox<int32> widget instance. */
+	TSharedRef<SWidget> MakeIntInputWidget(TSharedRef<TOptional<int32>>& ProxyValue, const FText& Label,
+										const FLinearColor& LabelColor, const FLinearColor& LabelBackgroundColor,
+										const int32* InMinValue = nullptr, const int32* InMaxValue = nullptr);
+
+	/** Returns an optional int32 value. */
+	TOptional<int32> OnGetValue(TSharedRef<TOptional<int32>> Value) const
+	{
+		return Value.Get();
+	}
+
+	/** Callback that is called when the value changes. */
+	void OnValueCommitted(int32 NewValue, ETextCommit::Type CommitType, TSharedRef<TOptional<int32>> Value)
+	{
+		*Value = NewValue;
+		FormatTextAndUpdateParameter();
+	}
+
+	/** Reset or set cached value.
+	* This function can be called multiple times when editing a value for a group selection.
+	* In this case, when the values in the group differ, a reset is used.
+	*/
+	static void ResetOrSetCachedValue(TSharedRef<TOptional<int32>>& InCachedValue, int32 InCompareValue)
+	{
+		if (InCachedValue->IsSet() && InCachedValue.Get() != InCompareValue)
+		{
+			InCachedValue->Reset();
+		}
+		else
+		{
+			*InCachedValue = InCompareValue;
+		}
+	};
+};
+
 /**
  * Policy info for matrix.
  */
@@ -506,4 +575,58 @@ private:
 	mutable TSharedRef<TOptional<float>> CachedAngleB;
 
 	static const FString BaseFrustumPlanesString;
-};
\ No newline at end of file
+};
+
+/**
+ * Policy info for resolution
+ */
+class FPolicyParameterInfoResolution final : public FPolicyParameterInfoIntReference
+{
+public:
+	FPolicyParameterInfoResolution(
+		const FString& InDisplayName,
+		const FString& InKey,
+		UDisplayClusterBlueprint* InBlueprint,
+		const TArray<TWeakObjectPtr<UDisplayClusterConfigurationViewport>>& InConfigurationViewports);
+
+	// FPolicyParameterInfo
+	virtual void CreateCustomRowWidget(IDetailChildrenBuilder& InDetailWidgetRow) override;
+	// ~FPolicyParameterInfo
+
+private:
+	virtual void FormatTextAndUpdateParameter() override;
+
+private:
+	mutable TSharedRef<TOptional<int32>> CachedX;
+	mutable TSharedRef<TOptional<int32>> CachedY;
+
+	static const FString BaseResolutionString;
+};
+
+/**
+ * Policy information for a normalized vector 2d whose component values must be between 0 and 1.
+ */
+class FPolicyParameterInfoNormalizedVector2D final : public FPolicyParameterInfoFloatReference
+{
+public:
+	FPolicyParameterInfoNormalizedVector2D(
+		const FString& InDisplayName,
+		const FString& InKey,
+		UDisplayClusterBlueprint* InBlueprint,
+		const TArray<TWeakObjectPtr<UDisplayClusterConfigurationViewport>>& InConfigurationViewports,
+		const FString* InDefaultValue = nullptr
+	);
+
+	// FPolicyParameterInfo
+	virtual void CreateCustomRowWidget(IDetailChildrenBuilder& InDetailWidgetRow) override;
+	// ~FPolicyParameterInfo
+
+private:
+	virtual void FormatTextAndUpdateParameter() override;
+
+private:
+	mutable TSharedRef<TOptional<float>> CachedX;
+	mutable TSharedRef<TOptional<float>> CachedY;
+
+	static const FString BaseVectorValueStr;
+};
diff --git a/Engine/Plugins/Runtime/nDisplay/Source/DisplayClusterProjection/Private/Policy/MPCDI/DisplayClusterProjectionMPCDIPolicy_Config.cpp b/Engine/Plugins/Runtime/nDisplay/Source/DisplayClusterProjection/Private/Policy/MPCDI/DisplayClusterProjectionMPCDIPolicy_Config.cpp
--- a/Engine/Plugins/Runtime/nDisplay/Source/DisplayClusterProjection/Private/Policy/MPCDI/DisplayClusterProjectionMPCDIPolicy_Config.cpp
+++ b/Engine/Plugins/Runtime/nDisplay/Source/DisplayClusterProjection/Private/Policy/MPCDI/DisplayClusterProjectionMPCDIPolicy_Config.cpp
@@ -11,6 +11,7 @@
 #include "DisplayClusterProjectionStrings.h"
 
 #include "Misc/DisplayClusterHelpers.h"
+#include "Misc/DisplayClusterProjectionHelpers.h"
 #include "Misc/FileHelper.h"
 
 #include "DisplayClusterRootActor.h"
@@ -39,20 +40,26 @@ bool FDisplayClusterProjectionMPCDIPolicy_ConfigParser::ReadConfig()
 		{
 			UE_LOG(LogDisplayClusterProjectionMPCDI, Error, TEXT("Undefined mpcdi type key '%s'='%s'"), DisplayClusterProjectionStrings::cfg::mpcdi::MPCDITypeKey, *MPCDITypeKey);
 		}
+
 		return false;
 	}
 
-	if (MPCDITypeKey.Compare(DisplayClusterProjectionStrings::cfg::mpcdi::TypeMPCDI) == 0)
+	if (!ImplGetBaseConfig())
 	{
-		return ImplGetMPCDIConfig() && ImplGetBaseConfig();
+		return false;
 	}
 
-	if (MPCDITypeKey.Compare(DisplayClusterProjectionStrings::cfg::mpcdi::TypePFM) == 0)
+	if (MPCDITypeKey.Compare(DisplayClusterProjectionStrings::cfg::mpcdi::TypeMPCDI) == 0)
 	{
-		return ImplGetPFMConfig() && ImplGetMPCDIAttributes() && ImplGetBaseConfig();
+		return ImplGetMPCDIConfig();
+	}
+	else if (MPCDITypeKey.Compare(DisplayClusterProjectionStrings::cfg::mpcdi::TypePFM) == 0)
+	{
+		return ImplGetPFMConfig();
 	}
 
 	UE_LOG(LogDisplayClusterProjectionMPCDI, Error, TEXT("Unknown mpcdi type key '%s'='%s'"), DisplayClusterProjectionStrings::cfg::mpcdi::MPCDITypeKey, *MPCDITypeKey);
+
 	return false;
 }
 
@@ -103,30 +110,6 @@ bool FDisplayClusterProjectionMPCDIPolicy_ConfigParser::ImplGetMPCDIConfig()
 		return false;
 	}
 
-	// Screen component
-	FString ScreenComponentName;
-	if (DisplayClusterHelpers::map::template ExtractValue(ConfigParameters, DisplayClusterProjectionStrings::cfg::mpcdi::Component, ScreenComponentName))
-	{
-		if (!ScreenComponentName.IsEmpty())
-		{
-			if (ADisplayClusterRootActor* SceneRootActor = Viewport->GetConfiguration().GetRootActor(EDisplayClusterRootActorType::Scene))
-			{
-				if (UDisplayClusterScreenComponent* ScreenComp = SceneRootActor->GetComponentByName<UDisplayClusterScreenComponent>(ScreenComponentName))
-				{
-					ScreenComponent = ScreenComp;
-				}
-			}
-
-			if (ADisplayClusterRootActor* PreviewRootActor = Viewport->GetConfiguration().GetRootActor(EDisplayClusterRootActorType::Preview))
-			{
-				if (UDisplayClusterScreenComponent* PreviewScreenComp = PreviewRootActor->GetComponentByName<UDisplayClusterScreenComponent>(ScreenComponentName))
-				{
-					PreviewScreenComponent = PreviewScreenComp;
-				}
-			}
-		}
-	}
-
 	return true;
 }
 
@@ -145,33 +128,6 @@ bool FDisplayClusterProjectionMPCDIPolicy_ConfigParser::ImplGetPFMConfig()
 		return false;
 	}
 
-	// MPCDIType (optional)
-	FString MPCDITypeStr;
-	if (!DisplayClusterHelpers::map::template ExtractValue(ConfigParameters, DisplayClusterProjectionStrings::cfg::mpcdi::MPCDIType, MPCDITypeStr))
-	{
-		MPCDIAttributes.ProfileType = EDisplayClusterWarpProfileType::warp_A3D;
-	}
-	else
-	{
-		MPCDIAttributes.ProfileType = EDisplayClusterWarpProfileType::Invalid;
-
-		static const TArray<FString> Profiles({ "2d","3d","a3d","sl" });
-		for (int32 ProfileIndex = 0; ProfileIndex < Profiles.Num(); ++ProfileIndex)
-		{
-			if (!MPCDITypeStr.Compare(Profiles[ProfileIndex], ESearchCase::IgnoreCase))
-			{
-				MPCDIAttributes.ProfileType = (EDisplayClusterWarpProfileType)ProfileIndex;
-				break;
-			}
-		}
-
-		if (MPCDIAttributes.ProfileType == EDisplayClusterWarpProfileType::Invalid)
-		{
-			UE_LOG(LogDisplayClusterProjectionMPCDI, Error, TEXT("Argument '%s' has unknown value '%s'"), DisplayClusterProjectionStrings::cfg::mpcdi::MPCDIType, *MPCDITypeStr);
-			return false;
-		}
-	}
-
 	// Default is UE scale, cm
 	PFMFileScale = 1;
 	if (DisplayClusterHelpers::map::template ExtractValueFromString(ConfigParameters, DisplayClusterProjectionStrings::cfg::mpcdi::WorldScale, PFMFileScale))
@@ -231,13 +187,76 @@ bool FDisplayClusterProjectionMPCDIPolicy_ConfigParser::ImplGetBaseConfig()
 		UE_LOG(LogDisplayClusterProjectionMPCDI, Verbose, TEXT("Found EnablePreview value - %s"), bEnablePreview ? TEXT("true") : TEXT("false"));
 	}
 
-	return true;
-}
+	// Screen component
+	FString ScreenComponentName;
+	if (DisplayClusterHelpers::map::template ExtractValue(ConfigParameters, DisplayClusterProjectionStrings::cfg::mpcdi::Component, ScreenComponentName))
+	{
+		if (!ScreenComponentName.IsEmpty())
+		{
+			if (ADisplayClusterRootActor* SceneRootActor = Viewport->GetConfiguration().GetRootActor(EDisplayClusterRootActorType::Scene))
+			{
+				if (UDisplayClusterScreenComponent* ScreenComp = SceneRootActor->GetComponentByName<UDisplayClusterScreenComponent>(ScreenComponentName))
+				{
+					ScreenComponent = ScreenComp;
+				}
+			}
 
-bool FDisplayClusterProjectionMPCDIPolicy_ConfigParser::ImplGetMPCDIAttributes()
-{
-	// Todo: Read parameters with names from the DisplayClusterProjectionStrings::cfg::mpcdi::Attributes namespace
-	// into a local variable named MPCDIAttributes;
+			if (ADisplayClusterRootActor* PreviewRootActor = Viewport->GetConfiguration().GetRootActor(EDisplayClusterRootActorType::Preview))
+			{
+				if (UDisplayClusterScreenComponent* PreviewScreenComp = PreviewRootActor->GetComponentByName<UDisplayClusterScreenComponent>(ScreenComponentName))
+				{
+					PreviewScreenComponent = PreviewScreenComp;
+				}
+			}
+		}
+	}
+
+	// MPCDIType (optional)
+	FString MPCDITypeStr;
+	if (!DisplayClusterHelpers::map::template ExtractValue(ConfigParameters, DisplayClusterProjectionStrings::cfg::mpcdi::MPCDIType, MPCDITypeStr))
+	{
+		MPCDIAttributes.ProfileType = EDisplayClusterWarpProfileType::warp_A3D;
+	}
+	else
+	{
+		MPCDIAttributes.ProfileType = UE::DisplayClusterProjectionHelpers::MPCDI::ProfileTypeFromString(MPCDITypeStr);
+
+		if (MPCDIAttributes.ProfileType == EDisplayClusterWarpProfileType::Invalid)
+		{
+			UE_LOG(LogDisplayClusterProjectionMPCDI, Error, TEXT("Argument '%s' has unknown value '%s'"), DisplayClusterProjectionStrings::cfg::mpcdi::MPCDIType, *MPCDITypeStr);
+			return false;
+		}
+	}
+
+	switch (MPCDIAttributes.ProfileType)
+	{
+	case EDisplayClusterWarpProfileType::warp_2D:
+	// Reads additional properties of the mpcdi 2d profile.
+	{
+		FIntPoint BufferRes;
+		if (DisplayClusterHelpers::map::template ExtractValueFromString(ConfigParameters, DisplayClusterProjectionStrings::cfg::mpcdi::Attributes::Buffer::Resolution, BufferRes))
+		{
+			UE_LOG(LogDisplayClusterProjectionMPCDI, Verbose, TEXT("Found buffer resolution %s"), *DisplayClusterTypesConverter::ToString(BufferRes));
+			MPCDIAttributes.Buffer.Resolution = BufferRes;
+		}
+
+		FVector2D RegionPos(0,0), RegionSize(1,1);
+		if (DisplayClusterHelpers::map::template ExtractValueFromString(ConfigParameters, DisplayClusterProjectionStrings::cfg::mpcdi::Attributes::Region::Pos, RegionPos)
+		&& DisplayClusterHelpers::map::template ExtractValueFromString(ConfigParameters, DisplayClusterProjectionStrings::cfg::mpcdi::Attributes::Region::Size, RegionSize))
+		{
+			UE_LOG(LogDisplayClusterProjectionMPCDI, Verbose, TEXT("Found region pos (%s) and size (%s)"), *DisplayClusterTypesConverter::ToString(RegionPos), *DisplayClusterTypesConverter::ToString(RegionSize));
+
+			MPCDIAttributes.Region.Pos = RegionPos;
+			MPCDIAttributes.Region.Size= RegionSize;
+		}
+	}
+	break;
+
+	default:
+		// Todo: Read parameters with names from the DisplayClusterProjectionStrings::cfg::mpcdi::Attributes namespace
+		// into a local variable named MPCDIAttributes;
+		break;
+	}
 
 	return true;
 }
diff --git a/Engine/Plugins/Runtime/nDisplay/Source/DisplayClusterProjection/Private/Policy/MPCDI/DisplayClusterProjectionMPCDIPolicy_Config.h b/Engine/Plugins/Runtime/nDisplay/Source/DisplayClusterProjection/Private/Policy/MPCDI/DisplayClusterProjectionMPCDIPolicy_Config.h
--- a/Engine/Plugins/Runtime/nDisplay/Source/DisplayClusterProjection/Private/Policy/MPCDI/DisplayClusterProjectionMPCDIPolicy_Config.h
+++ b/Engine/Plugins/Runtime/nDisplay/Source/DisplayClusterProjection/Private/Policy/MPCDI/DisplayClusterProjectionMPCDIPolicy_Config.h
@@ -83,9 +83,6 @@ private:
 	/** Implementation of parameter reading for the case when we initialize from an PFM file. */
 	bool ImplGetPFMConfig();
 
-	/** Implementation of MPCDI attribute reading for the case when we initialize from a PFM file. */
-	bool ImplGetMPCDIAttributes();
-
 private:
 	// The projection policy parameters used to read this configuration
 	const TMap<FString, FString>& ConfigParameters;
diff --git a/Engine/Plugins/Runtime/nDisplay/Source/DisplayClusterProjection/Public/DisplayClusterProjectionStrings.h b/Engine/Plugins/Runtime/nDisplay/Source/DisplayClusterProjection/Public/DisplayClusterProjectionStrings.h
--- a/Engine/Plugins/Runtime/nDisplay/Source/DisplayClusterProjection/Public/DisplayClusterProjectionStrings.h
+++ b/Engine/Plugins/Runtime/nDisplay/Source/DisplayClusterProjection/Public/DisplayClusterProjectionStrings.h
@@ -118,41 +118,38 @@ namespace DisplayClusterProjectionStrings
 
 			static constexpr const TCHAR* Component = TEXT("screen_component");
 
+			namespace Profiles
+			{
+				static constexpr const TCHAR* mpcdi_2d  = TEXT("2d");
+				static constexpr const TCHAR* mpcdi_3d  = TEXT("3d");
+				static constexpr const TCHAR* mpcdi_a3d = TEXT("a3d");
+				static constexpr const TCHAR* mpcdi_sl  = TEXT("sl");
+			}
+
 			namespace Attributes
 			{
 				namespace Buffer
 				{
-					static constexpr const TCHAR* Resolution = TEXT("BufferResolution");
+					static constexpr const TCHAR* Resolution = TEXT("BufferRes");
 				}
 
 				namespace Region
 				{
-					static constexpr const TCHAR* Resolution = TEXT("RegionResolution");
-					static constexpr const TCHAR* Pos = TEXT("RegionPos");
-					static constexpr const TCHAR* Size = TEXT("RegionSize");
+					static constexpr const TCHAR* Resolution = TEXT("RegionRes");
+					static constexpr const TCHAR* Pos        = TEXT("RegionPos");
+					static constexpr const TCHAR* Size       = TEXT("RegionSize");
 				}
 
 				namespace Frustum
 				{
-					static constexpr const TCHAR* Pitch = TEXT("FrustumPitch");
-					static constexpr const TCHAR* Yaw   = TEXT("FrustumYaw");
-					static constexpr const TCHAR* Roll  = TEXT("FrustumRoll");
-					
-					namespace Angle
-					{
-						static constexpr const TCHAR* Left   = TEXT("FrustumAngleLeft");
-						static constexpr const TCHAR* Right  = TEXT("FrustumAngleRight");
-						static constexpr const TCHAR* Top    = TEXT("FrustumAngleTop");
-						static constexpr const TCHAR* Bottom = TEXT("FrustumAngleBottom");
-					}
+					static constexpr const TCHAR* Rotation = TEXT("FrustumRot");
+					static constexpr const TCHAR* Angles   = TEXT("FrustumAngles");
 				}
 
 				namespace CoordinateFrame
 				{
-					static constexpr const TCHAR* Pos   = TEXT("CoordinateFramePos");
-					static constexpr const TCHAR* Yaw   = TEXT("CoordinateFrameYaw");
-					static constexpr const TCHAR* Pitch = TEXT("CoordinateFramePitch");
-					static constexpr const TCHAR* Roll  = TEXT("CoordinateFrameRoll");
+					static constexpr const TCHAR* Position = TEXT("CoordFramePos");
+					static constexpr const TCHAR* Rotation = TEXT("CoordFrameRot");
 				}
 			}
 		}
diff --git a/Engine/Plugins/Runtime/nDisplay/Source/DisplayClusterProjection/Public/Misc/DisplayClusterProjectionHelpers.h b/Engine/Plugins/Runtime/nDisplay/Source/DisplayClusterProjection/Public/Misc/DisplayClusterProjectionHelpers.h
new file mode 100644
--- /dev/null
+++ b/Engine/Plugins/Runtime/nDisplay/Source/DisplayClusterProjection/Public/Misc/DisplayClusterProjectionHelpers.h
@@ -0,0 +1,59 @@
+// Copyright Epic Games, Inc. All Rights Reserved.
+
+#pragma once
+
+#include "Misc/DisplayClusterHelpers.h"
+#include "DisplayClusterProjectionStrings.h"
+#include "Containers/DisplayClusterWarpEnums.h"
+
+namespace UE::DisplayClusterProjectionHelpers
+{
+	namespace MPCDI
+	{
+		/** Convert the mpcdi profile type from a string to an enum value. */
+		static EDisplayClusterWarpProfileType ProfileTypeFromString(const FString& InProfileTypeName)
+		{
+			if (!InProfileTypeName.Compare(DisplayClusterProjectionStrings::cfg::mpcdi::Profiles::mpcdi_2d, ESearchCase::IgnoreCase))
+			{
+				return EDisplayClusterWarpProfileType::warp_2D;
+			}
+			else if (!InProfileTypeName.Compare(DisplayClusterProjectionStrings::cfg::mpcdi::Profiles::mpcdi_3d, ESearchCase::IgnoreCase))
+			{
+				return EDisplayClusterWarpProfileType::warp_3D;
+			}
+			else if (!InProfileTypeName.Compare(DisplayClusterProjectionStrings::cfg::mpcdi::Profiles::mpcdi_a3d, ESearchCase::IgnoreCase))
+			{
+				return EDisplayClusterWarpProfileType::warp_A3D;
+			}
+			else if (!InProfileTypeName.Compare(DisplayClusterProjectionStrings::cfg::mpcdi::Profiles::mpcdi_sl, ESearchCase::IgnoreCase))
+			{
+				return EDisplayClusterWarpProfileType::warp_SL;
+			}
+
+			return EDisplayClusterWarpProfileType::Invalid;
+		}
+
+		/** Convert the mpcdi profile type from an enum value to a string. */
+		static FString ProfileTypeToString(const EDisplayClusterWarpProfileType InProfileType)
+		{
+			switch (InProfileType)
+			{
+			case EDisplayClusterWarpProfileType::warp_2D:
+				return DisplayClusterProjectionStrings::cfg::mpcdi::Profiles::mpcdi_2d;
+			case EDisplayClusterWarpProfileType::warp_3D:
+				return DisplayClusterProjectionStrings::cfg::mpcdi::Profiles::mpcdi_3d;
+			case EDisplayClusterWarpProfileType::warp_A3D:
+				return DisplayClusterProjectionStrings::cfg::mpcdi::Profiles::mpcdi_a3d;
+			case EDisplayClusterWarpProfileType::warp_SL:
+				return DisplayClusterProjectionStrings::cfg::mpcdi::Profiles::mpcdi_sl;
+
+			default:
+				break;
+			}
+
+			// Returns an empty string if an invalid profile type is used.
+			return FString();
+		}
+	}
+};
+
diff --git a/Engine/Plugins/Runtime/nDisplay/Source/DisplayClusterWarp/Private/WarpBlend/DisplayClusterWarpBlend.cpp b/Engine/Plugins/Runtime/nDisplay/Source/DisplayClusterWarp/Private/WarpBlend/DisplayClusterWarpBlend.cpp
--- a/Engine/Plugins/Runtime/nDisplay/Source/DisplayClusterWarp/Private/WarpBlend/DisplayClusterWarpBlend.cpp
+++ b/Engine/Plugins/Runtime/nDisplay/Source/DisplayClusterWarp/Private/WarpBlend/DisplayClusterWarpBlend.cpp
@@ -208,6 +208,18 @@ UMeshComponent* FDisplayClusterWarpBlend::GetOrCreatePreviewMeshComponentImpl(ID
 	case EDisplayClusterWarpGeometryType::WarpMap:
 		if (ProjectionPolicy.IsValid())
 		{
+			// use the existing DCRA component
+			if (USceneComponent* PreviewMeshComponent = GeometryContext.GeometryProxy.PreviewMeshComponentRef.GetOrFindSceneComponent())
+			{
+				if (PreviewMeshComponent->IsA<UMeshComponent>())
+				{
+					UMeshComponent* ExistMeshComponent = static_cast<UMeshComponent*>(PreviewMeshComponent);
+
+					bExistingComponent = true;
+					return ExistMeshComponent;
+				}
+			}
+
 			// create a new mesh component
 			bExistingComponent = false;
 
