nDisplay: Fixed visibility for the tile viewport rendering
- The tile viewport copies the settings from the original viewport at the very end.

diff --git a/Engine/Plugins/Runtime/nDisplay/Source/DisplayCluster/Private/Game/EngineClasses/Scene/DisplayClusterRootActor.cpp b/Engine/Plugins/Runtime/nDisplay/Source/DisplayCluster/Private/Game/EngineClasses/Scene/DisplayClusterRootActor.cpp
--- a/Engine/Plugins/Runtime/nDisplay/Source/DisplayCluster/Private/Game/EngineClasses/Scene/DisplayClusterRootActor.cpp
+++ b/Engine/Plugins/Runtime/nDisplay/Source/DisplayCluster/Private/Game/EngineClasses/Scene/DisplayClusterRootActor.cpp
@@ -762,7 +762,7 @@ bool ADisplayClusterRootActor::GetHiddenInGamePrimitives(TSet<FPrimitiveComponen
 		PrimitiveComponentsArray.AddDefaulted(NumIterThreads);
 
 		// Start the iteration parallel threads
-		ParallelFor(NumIterThreads, [NumIterThreads, CurrentWorld, &Actors, &PrimitiveComponentsArray](int32 Index)
+		ParallelFor(NumIterThreads, [NumIterThreads, CurrentWorld, &Actors, CurrentRootActor = this, &PrimitiveComponentsArray](int32 Index)
 			{
 				// Using inline allocator for efficiency
 				constexpr int32 MaxExpectedComponentsPerActor = 64;
@@ -775,7 +775,9 @@ bool ADisplayClusterRootActor::GetHiddenInGamePrimitives(TSet<FPrimitiveComponen
 				{
 					const AActor* Actor = Actors[ActorIdx];
 
-					if (IsValid(Actor))
+					// RootActor visibility is already handled by this GetHiddenInGamePrimitives() function. (See the code above.)
+					// The code below is generic for all actor types and hides all primitives without following the DCRA rules used above.
+					if (IsValid(Actor) && Actor != CurrentRootActor)
 					{
 						Actor->GetComponents(PrimitiveComponents);
 						for (UPrimitiveComponent* PrimComp : PrimitiveComponents)
diff --git a/Engine/Plugins/Runtime/nDisplay/Source/DisplayCluster/Private/Render/Viewport/Configuration/DisplayClusterViewportConfigurationHelpers_Tile.cpp b/Engine/Plugins/Runtime/nDisplay/Source/DisplayCluster/Private/Render/Viewport/Configuration/DisplayClusterViewportConfigurationHelpers_Tile.cpp
--- a/Engine/Plugins/Runtime/nDisplay/Source/DisplayCluster/Private/Render/Viewport/Configuration/DisplayClusterViewportConfigurationHelpers_Tile.cpp
+++ b/Engine/Plugins/Runtime/nDisplay/Source/DisplayCluster/Private/Render/Viewport/Configuration/DisplayClusterViewportConfigurationHelpers_Tile.cpp
@@ -220,18 +220,29 @@ FDisplayClusterViewport* FDisplayClusterViewportConfigurationHelpers_Tile::GetOr
 
 		// Setup as tile.
 		InOutRenderSettings.TileSettings = FDisplayClusterViewport_TileSettings(InSourceViewport.GetId(), InTilePos, InTileSize, TileFlags);
+	}
+
+	return TileViewport;
+}
 
+void FDisplayClusterViewportConfigurationHelpers_Tile::ConfigureTileViewport(FDisplayClusterViewport& InSourceViewport, const FIntPoint& InTilePos)
+{
+	const FDisplayClusterViewport_RenderSettings& SourceRenderSettings = InSourceViewport.GetRenderSettings();
+	if (SourceRenderSettings.TileSettings.GetType() == EDisplayClusterViewportTileType::Source)
+	{
+		if (FDisplayClusterViewport* TileViewport = FindTileViewport(InSourceViewport, InTilePos))
+		{
 		// Copy internal render settings from the source:
 		TileViewport->GetCustomPostProcessSettings() = InSourceViewport.GetCustomPostProcessSettings();
 		TileViewport->GetVisibilitySettingsImpl() = InSourceViewport.GetVisibilitySettingsImpl();
 		TileViewport->GetCameraMotionBlurImpl() = InSourceViewport.GetCameraMotionBlurImpl();
 		TileViewport->GetCameraDepthOfFieldImpl() = InSourceViewport.GetCameraDepthOfFieldImpl();
+			TileViewport->GetRenderSettingsICVFXImpl() = InSourceViewport.GetRenderSettingsICVFX();
 
 		// Copy OCIO.
 		TileViewport->SetOpenColorIO(InSourceViewport.GetOpenColorIO());
 	}
-
-	return TileViewport;
+	}
 }
 
 FDisplayClusterViewport_OverscanSettings FDisplayClusterViewportConfigurationHelpers_Tile::GetTileOverscanSettings(const FDisplayClusterConfigurationTile_Overscan& InTileOverscan)
diff --git a/Engine/Plugins/Runtime/nDisplay/Source/DisplayCluster/Private/Render/Viewport/Configuration/DisplayClusterViewportConfigurationHelpers_Tile.h b/Engine/Plugins/Runtime/nDisplay/Source/DisplayCluster/Private/Render/Viewport/Configuration/DisplayClusterViewportConfigurationHelpers_Tile.h
--- a/Engine/Plugins/Runtime/nDisplay/Source/DisplayCluster/Private/Render/Viewport/Configuration/DisplayClusterViewportConfigurationHelpers_Tile.h
+++ b/Engine/Plugins/Runtime/nDisplay/Source/DisplayCluster/Private/Render/Viewport/Configuration/DisplayClusterViewportConfigurationHelpers_Tile.h
@@ -23,6 +23,9 @@ public:
 	/** Get or create a new tile viewport from source. */
 	static FDisplayClusterViewport* GetOrCreateTileViewport(FDisplayClusterViewport& InSourceViewport, const FIntPoint& TilePos);
 
+	/** Configure tile viewport from source. */
+	static void ConfigureTileViewport(FDisplayClusterViewport& InSourceViewport, const FIntPoint& TilePos);
+
 	/** Get unique viewport name for tile. */
 	static FString GetUniqueViewportNameForTile(const FString& InViewportId, const FIntPoint& TilePos);
 
diff --git a/Engine/Plugins/Runtime/nDisplay/Source/DisplayCluster/Private/Render/Viewport/Configuration/DisplayClusterViewportConfiguration_Tile.cpp b/Engine/Plugins/Runtime/nDisplay/Source/DisplayCluster/Private/Render/Viewport/Configuration/DisplayClusterViewportConfiguration_Tile.cpp
--- a/Engine/Plugins/Runtime/nDisplay/Source/DisplayCluster/Private/Render/Viewport/Configuration/DisplayClusterViewportConfiguration_Tile.cpp
+++ b/Engine/Plugins/Runtime/nDisplay/Source/DisplayCluster/Private/Render/Viewport/Configuration/DisplayClusterViewportConfiguration_Tile.cpp
@@ -20,36 +20,55 @@
 ///////////////////////////////////////////////////////////////////
 // FDisplayClusterViewportConfiguration_Tile
 ///////////////////////////////////////////////////////////////////
-void FDisplayClusterViewportConfiguration_Tile::Update()
+template <typename Predicate>
+void FDisplayClusterViewportConfiguration_Tile::ForEachTileViewport(Predicate Pred) const
+{
+	if (FDisplayClusterViewportManager* ViewportManager = Configuration.GetViewportManagerImpl())
 {
-	ImplBeginReallocateViewports();
-
-	const FDisplayClusterRenderFrameSettings& RenderFrameSettings = Configuration.GetRenderFrameSettings();
-	const FDisplayClusterConfigurationICVFX_StageSettings* StageSettings = Configuration.GetStageSettings();
-	FDisplayClusterViewportManager* ViewportManager = Configuration.GetViewportManagerImpl();
-
-	// Split the source viewports into multiple tiles.
 	for (const TSharedPtr<FDisplayClusterViewport, ESPMode::ThreadSafe>& ViewportIt : ViewportManager->ImplGetCurrentRenderFrameViewports())
 	{
 		if (ViewportIt.IsValid() && ViewportIt->GetRenderSettings().TileSettings.GetType() == EDisplayClusterViewportTileType::Source)
 		{
-			// Split the source viewport into multiple tiles:
 			const FIntPoint& TileSize = ViewportIt->GetRenderSettings().TileSettings.GetSize();
+
+				// Iterate over all tile viewports
 			for (int32 PosX = 0; PosX < TileSize.X; PosX++)
 			{
 				for (int32 PosY = 0; PosY < TileSize.Y; PosY++)
 				{
-					FDisplayClusterViewportConfigurationHelpers_Tile::GetOrCreateTileViewport(*ViewportIt, FIntPoint(PosX, PosY));
+						::Invoke(Pred, *ViewportIt, FIntPoint(PosX, PosY));
 				}
 			}
 		}
 	}
+	}
+	else
+	{
+		UE_LOG(LogDisplayClusterViewport, Warning, TEXT("Tile:ForEachTileViewport - ViewportManager not found. "));
+	}
+}
+
+void FDisplayClusterViewportConfiguration_Tile::Update()
+{
+	ImplBeginReallocateViewports();
+
+	ForEachTileViewport([](FDisplayClusterViewport& InSourceViewport, const FIntPoint& InTilePos)
+	{
+		// Split the source viewports into multiple tiles.
+		FDisplayClusterViewportConfigurationHelpers_Tile::GetOrCreateTileViewport(InSourceViewport, InTilePos);
+	});
 
 	ImplFinishReallocateViewports();
 }
 
 void FDisplayClusterViewportConfiguration_Tile::PostUpdate()
-{ }
+{
+	ForEachTileViewport([](FDisplayClusterViewport& InSourceViewport, const FIntPoint& InTilePos)
+	{
+		// Copy settings for the source viewports to the all tiles.
+		FDisplayClusterViewportConfigurationHelpers_Tile::ConfigureTileViewport(InSourceViewport, InTilePos);
+	});
+}
 
 void FDisplayClusterViewportConfiguration_Tile::ImplBeginReallocateViewports() const
 {
@@ -63,6 +82,10 @@ void FDisplayClusterViewportConfiguration_Tile::ImplBeginReallocateViewports() c
 			}
 		}
 	}
+	else
+	{
+		UE_LOG(LogDisplayClusterViewport, Warning, TEXT("Tile:BeginReallocate - ViewportManager not found. "));
+	}
 }
 
 void FDisplayClusterViewportConfiguration_Tile::ImplFinishReallocateViewports() const
@@ -91,4 +114,8 @@ void FDisplayClusterViewportConfiguration_Tile::ImplFinishReallocateViewports()
 
 		UnusedViewports.Empty();
 	}
+	else
+	{
+		UE_LOG(LogDisplayClusterViewport, Warning, TEXT("Tile:FinishReallocate - ViewportManager not found. "));
+	}
 }
diff --git a/Engine/Plugins/Runtime/nDisplay/Source/DisplayCluster/Private/Render/Viewport/Configuration/DisplayClusterViewportConfiguration_Tile.h b/Engine/Plugins/Runtime/nDisplay/Source/DisplayCluster/Private/Render/Viewport/Configuration/DisplayClusterViewportConfiguration_Tile.h
--- a/Engine/Plugins/Runtime/nDisplay/Source/DisplayCluster/Private/Render/Viewport/Configuration/DisplayClusterViewportConfiguration_Tile.h
+++ b/Engine/Plugins/Runtime/nDisplay/Source/DisplayCluster/Private/Render/Viewport/Configuration/DisplayClusterViewportConfiguration_Tile.h
@@ -13,7 +13,7 @@ class FDisplayClusterViewportConfiguration;
 class FDisplayClusterViewportConfiguration_Tile
 {
 public:
-	FDisplayClusterViewportConfiguration_Tile(FDisplayClusterViewportConfiguration& InConfiguration)
+	FDisplayClusterViewportConfiguration_Tile(const FDisplayClusterViewportConfiguration& InConfiguration)
 		: Configuration(InConfiguration)
 	{ }
 
@@ -33,7 +33,15 @@ private:
 	/** Delete unused Tile viewports. */
 	void ImplFinishReallocateViewports() const;
 
+	/**
+	 * Iterate over all tile viewports
+	 *
+	 * @param Pred The predicate to match.
+	 */
+	template <typename Predicate>
+	void ForEachTileViewport(Predicate Pred) const;
+
 private:
 	// Viewport configuration API
-	FDisplayClusterViewportConfiguration& Configuration;
+	const FDisplayClusterViewportConfiguration& Configuration;
 };
