nDisplay: The logic for creating 3D geometry for the mpcdi 2D profile has been reworked.
- DCRA with an imported mpcdi file and with a manually assigned PFM file now creates identical geometry and identical images
- The mpcdi 2d profile geometry settings are now publicly available at FDisplayClusterWarpMPCDIAttributes::FProfile2DSettings.
  * DesiredFOV - The focal length will be calculated for that FOV and the geometry will be shifted by that value.
  * BufferPixelsToWorldUnits - The amount to scale an MPCDI buffer by when converting from pixels to world coordinates.

diff --git a/Engine/Plugins/Runtime/nDisplay/Source/DisplayClusterConfigurator/Private/MPCDI/DisplayClusterConfiguratorMPCDIImporter.cpp b/Engine/Plugins/Runtime/nDisplay/Source/DisplayClusterConfigurator/Private/MPCDI/DisplayClusterConfiguratorMPCDIImporter.cpp
--- a/Engine/Plugins/Runtime/nDisplay/Source/DisplayClusterConfigurator/Private/MPCDI/DisplayClusterConfiguratorMPCDIImporter.cpp
+++ b/Engine/Plugins/Runtime/nDisplay/Source/DisplayClusterConfigurator/Private/MPCDI/DisplayClusterConfiguratorMPCDIImporter.cpp
@@ -164,32 +164,13 @@ void FDisplayClusterConfiguratorMPCDIImporter::ConfigureScreenComponentFrom2DPro
 	const FDisplayClusterWarpMPCDIAttributes& InAttributes,
 	const FDisplayClusterConfiguratorMPCDIImporterParams& InParams)
 {
-	// Computes the buffer size in world units from the resolution in pixels.
-	const FVector2D BufferSize(InAttributes.Buffer.Resolution.X * InParams.Profile2DParams.BufferPixelsToWorldUnits, InAttributes.Buffer.Resolution.Y * InParams.Profile2DParams.BufferPixelsToWorldUnits);
-
-	// Calculates the position and size of the region.
-	const FVector2D  RegionPos(InAttributes.Region.Pos.X  * BufferSize.X, InAttributes.Region.Pos.Y  * BufferSize.Y);
-	const FVector2D RegionSize(InAttributes.Region.Size.X * BufferSize.X, InAttributes.Region.Size.Y * BufferSize.Y);
-
-	float FocalLength = 0.f;
-	// Moves the buffer position along the X axis to achieve a DesiredFOV.
-	if (InParams.Profile2DParams.DesiredFOV > 0.f && InParams.Profile2DParams.DesiredFOV < 180.f)
+	FVector ScreenPosition;
+	FVector2D ScreenSize;
+	if (!InAttributes.CalcProfile2DScreen(ScreenPosition, ScreenSize))
 	{
-		// Convert FOV to focal length,
-		// 
-		// fov = 2 * atan(d/(2*f))
-		// where,
-		//   d = sensor dimension
-		//   f = focal length
-		// 
-		// f = 0.5 * d * (1/tan(fov/2))
-		const float TanHalfFOV = FMath::Tan(FMath::DegreesToRadians(InParams.Profile2DParams.DesiredFOV * 0.5f));
-		FocalLength = (BufferSize.X * 0.5f) / TanHalfFOV;
+		return;
 	}
 
-	const FVector2D RegionCenterPos = RegionPos + (RegionSize * 0.5) - (BufferSize * 0.5);
-	FVector ScreenPosition = FVector(FocalLength, RegionCenterPos.X, RegionCenterPos.Y);
-
 	if (InViewOriginComponent)
 	{
 		// We want to position the screen component to be in front of the view origin it is assigned to, so we must compute
@@ -214,7 +195,7 @@ void FDisplayClusterConfiguratorMPCDIImporter::ConfigureScreenComponentFrom2DPro
 	}
 
 	InScreenComponent->SetRelativeLocation(ScreenPosition);
-	InScreenComponent->SetScreenSize(RegionSize);
+	InScreenComponent->SetScreenSize(ScreenSize);
 }
 
 UDisplayClusterConfigurationViewport* FDisplayClusterConfiguratorMPCDIImporter::FindOrCreateViewportForRegion(UDisplayClusterBlueprint* InBlueprint, const FString& RegionId, bool& bOutFoundExistingViewport)
diff --git a/Engine/Plugins/Runtime/nDisplay/Source/DisplayClusterConfigurator/Private/MPCDI/DisplayClusterConfiguratorMPCDIImporter.h b/Engine/Plugins/Runtime/nDisplay/Source/DisplayClusterConfigurator/Private/MPCDI/DisplayClusterConfiguratorMPCDIImporter.h
--- a/Engine/Plugins/Runtime/nDisplay/Source/DisplayClusterConfigurator/Private/MPCDI/DisplayClusterConfiguratorMPCDIImporter.h
+++ b/Engine/Plugins/Runtime/nDisplay/Source/DisplayClusterConfigurator/Private/MPCDI/DisplayClusterConfiguratorMPCDIImporter.h
@@ -13,22 +13,9 @@ class UDisplayClusterConfigurationViewport;
 class UDisplayClusterBlueprint;
 class UDisplayClusterScreenComponent;
 
-/** Additional parameters for the mpcdi 2d profile. */
-struct FDisplayClusterConfiguratorMPCDIProfile2DImporterParams
-{
-	/** The amount to scale an MPCDI buffer by when converting from pixels to world coordinates */
-	float BufferPixelsToWorldUnits = 0.1;
-
-	/** Desired FOV. The focal length will be calculated for that FOV and the geometry will be shifted by that value. */
-	float DesiredFOV = 90.f;
-};
-
 /** Used to pass parameters into the MPCDI importer */
 struct FDisplayClusterConfiguratorMPCDIImporterParams
 {
-	/** Additional parameters for the mpcdi 2d profile. */
-	FDisplayClusterConfiguratorMPCDIProfile2DImporterParams Profile2DParams;
-
 	/** The name of the component to make the parent of the MPCDI geometry */
 	FName ParentComponentName = NAME_None;
 
diff --git a/Engine/Plugins/Runtime/nDisplay/Source/DisplayClusterWarp/Private/WarpBlend/Exporter/DisplayClusterWarpBlendExporter_WarpMap.cpp b/Engine/Plugins/Runtime/nDisplay/Source/DisplayClusterWarp/Private/WarpBlend/Exporter/DisplayClusterWarpBlendExporter_WarpMap.cpp
--- a/Engine/Plugins/Runtime/nDisplay/Source/DisplayClusterWarp/Private/WarpBlend/Exporter/DisplayClusterWarpBlendExporter_WarpMap.cpp
+++ b/Engine/Plugins/Runtime/nDisplay/Source/DisplayClusterWarp/Private/WarpBlend/Exporter/DisplayClusterWarpBlendExporter_WarpMap.cpp
@@ -120,31 +120,28 @@ using namespace UE::DisplayClusterWarp::PFM;
 //---------------------------------------------------------------
 // FDisplayClusterWarpBlendExporter_WarpMap
 //---------------------------------------------------------------
-float FDisplayClusterWarpBlendExporter_WarpMap::Get2DProfilePixelsToUnitScale()
+void FDisplayClusterWarpBlendExporter_WarpMap::Get2DProfileGeometry(const FDisplayClusterWarpMPCDIAttributes& InAttributes, TArray<FVector>& OutGeometryPoints, TArray<FVector>* OutNormal, TArray<FVector2D>* OutUV)
 {
-	return 0.1f; // 10 pixels = 1 centimeter
-}
-
-void FDisplayClusterWarpBlendExporter_WarpMap::Get2DProfileGeometry(const FDisplayClusterWarpMPCDIAttributes& InMPCDIAttributes, TArray<FVector>& OutGeometryPoints, TArray<FVector>* OutNormal, TArray<FVector2D>* OutUV)
-{
-	const float PixelsToUnitScale = FDisplayClusterWarpBlendExporter_WarpMap::Get2DProfilePixelsToUnitScale();
-
-	// export each region as a plain
-	const float Width = InMPCDIAttributes.Buffer.Resolution.X * PixelsToUnitScale;
-	const float Height = InMPCDIAttributes.Buffer.Resolution.Y * PixelsToUnitScale;
+	FVector ScreenPosition;
+	FVector2D ScreenSize;
+	if (!InAttributes.CalcProfile2DScreen(ScreenPosition, ScreenSize))
+	{
+		return;
+	}
 
 	// Rectangle stands on the floor, centered at zero
-	const float X0 = Width * (InMPCDIAttributes.Region.Pos.X) - (Width * 0.5f);
-	const float X1 = Width * (InMPCDIAttributes.Region.Pos.X + InMPCDIAttributes.Region.Size.X) - (Width * 0.5f);
-	const float Y0 = Height * (InMPCDIAttributes.Region.Pos.Y);
-	const float Y1 = Height * (InMPCDIAttributes.Region.Pos.Y + InMPCDIAttributes.Region.Size.Y);
+	const float X = ScreenPosition.X;
+	const float Y0 = ScreenPosition.Y - (ScreenSize.X * 0.5f);
+	const float Y1 = ScreenPosition.Y + (ScreenSize.X * 0.5f);
+	const float Z0 = ScreenPosition.Z - (ScreenSize.Y * 0.5f);
+	const float Z1 = ScreenPosition.Z + (ScreenSize.Y * 0.5f);
 
 	// Create a vertices
 	// Y - right, Z - Up
-	OutGeometryPoints.Add(FVector(0, X0, Y1));
-	OutGeometryPoints.Add(FVector(0, X1, Y1));
-	OutGeometryPoints.Add(FVector(0, X0, Y0));
-	OutGeometryPoints.Add(FVector(0, X1, Y0));
+	OutGeometryPoints.Add(FVector(X, Y0, Z1));
+	OutGeometryPoints.Add(FVector(X, Y1, Z1));
+	OutGeometryPoints.Add(FVector(X, Y0, Z0));
+	OutGeometryPoints.Add(FVector(X, Y1, Z0));
 
 	if (OutNormal)
 	{
diff --git a/Engine/Plugins/Runtime/nDisplay/Source/DisplayClusterWarp/Private/WarpBlend/Exporter/DisplayClusterWarpBlendExporter_WarpMap.h b/Engine/Plugins/Runtime/nDisplay/Source/DisplayClusterWarp/Private/WarpBlend/Exporter/DisplayClusterWarpBlendExporter_WarpMap.h
--- a/Engine/Plugins/Runtime/nDisplay/Source/DisplayClusterWarp/Private/WarpBlend/Exporter/DisplayClusterWarpBlendExporter_WarpMap.h
+++ b/Engine/Plugins/Runtime/nDisplay/Source/DisplayClusterWarp/Private/WarpBlend/Exporter/DisplayClusterWarpBlendExporter_WarpMap.h
@@ -14,9 +14,6 @@ public:
 	 */
 	static bool ExportWarpMap(const class FDisplayClusterWarpBlend_GeometryContext& InContext, struct FDisplayClusterWarpGeometryOBJ& Dst, uint32 InMaxDimension = 0);
 
-	/** Return pixels to Unit scale for fake geometry of MPCDI profile 2D. */
-	static float Get2DProfilePixelsToUnitScale();
-
 	/** [2D] Return fake geometry points for region in mpcdi 2D profile. */
 	static void Get2DProfileGeometry(const struct FDisplayClusterWarpMPCDIAttributes& InMPCDIAttributes, TArray<FVector>& OutGeometryPoints, TArray<FVector>* OutNormal = nullptr, TArray<FVector2D>* OutUV = nullptr);
 };
diff --git a/Engine/Plugins/Runtime/nDisplay/Source/DisplayClusterWarp/Public/Containers/DisplayClusterWarpContainers.h b/Engine/Plugins/Runtime/nDisplay/Source/DisplayClusterWarp/Public/Containers/DisplayClusterWarpContainers.h
--- a/Engine/Plugins/Runtime/nDisplay/Source/DisplayClusterWarp/Public/Containers/DisplayClusterWarpContainers.h
+++ b/Engine/Plugins/Runtime/nDisplay/Source/DisplayClusterWarp/Public/Containers/DisplayClusterWarpContainers.h
@@ -148,6 +148,20 @@ struct FDisplayClusterWarpMPCDIAttributes
 	// MPCDI profile type
 	EDisplayClusterWarpProfileType ProfileType = EDisplayClusterWarpProfileType::Invalid;
 
+	/**
+	* Special settings for creating 3D geometry of the mpcdi 2D profile.
+	*/
+	struct FProfile2DSettings
+	{
+		// The focal length will be calculated for that FOV and the geometry will be shifted by that value.
+		float DesiredFOV = 90.f;
+
+		/** The amount to scale an MPCDI buffer by when converting from pixels to world coordinates. */
+		float BufferPixelsToWorldUnits = 0.1f;
+
+	} Profile2DSettings;
+
+
 	// Special flags
 	EDisplayClusterWarpMPCDIAttributesFlags Flags = EDisplayClusterWarpMPCDIAttributesFlags::None;
 
@@ -209,4 +223,46 @@ struct FDisplayClusterWarpMPCDIAttributes
 		FVector Roll;
 
 	} CoordinateFrame;
+
+	/** 
+	* Calculates the position and screen size for a 2d profile.
+	* Returns true if the success.
+	*/
+	inline bool CalcProfile2DScreen(FVector& OutScreenPosition, FVector2D& OutScreenSize) const
+	{
+		if (ProfileType != EDisplayClusterWarpProfileType::warp_2D)
+		{
+			return false;
+		}
+
+		// Computes the buffer size in world units from the resolution in pixels.
+		const FVector2D BufferSize(Buffer.Resolution.X* Profile2DSettings.BufferPixelsToWorldUnits, Buffer.Resolution.Y* Profile2DSettings.BufferPixelsToWorldUnits);
+
+		OutScreenSize = FVector2D(Region.Size.X * BufferSize.X, Region.Size.Y * BufferSize.Y);
+
+		// Calculates the position and size of the region.
+		const FVector2D  RegionPos(Region.Pos.X* BufferSize.X, Region.Pos.Y* BufferSize.Y);
+
+		float FocalLength = 0.f;
+		// Moves the buffer position along the X axis to achieve a DesiredFOV.
+		if (Profile2DSettings.DesiredFOV > 0.f && Profile2DSettings.DesiredFOV < 180.f)
+		{
+			// Convert FOV to focal length,
+			// 
+			// fov = 2 * atan(d/(2*f))
+			// where,
+			//   d = sensor dimension
+			//   f = focal length
+			// 
+			// f = 0.5 * d * (1/tan(fov/2))
+			const float TanHalfFOV = FMath::Tan(FMath::DegreesToRadians(Profile2DSettings.DesiredFOV * 0.5f));
+			FocalLength = (BufferSize.X * 0.5f) / TanHalfFOV;
+		}
+
+		const FVector2D RegionCenterPos = RegionPos + (OutScreenSize * 0.5) - (BufferSize * 0.5);
+
+		OutScreenPosition = FVector(FocalLength, RegionCenterPos.X, RegionCenterPos.Y);
+
+		return true;
+	}
 };
