nDisplay: Remove lingering viewport media custom serialize since this has been fixed at root actor level already.

#rnx

diff --git a/Engine/Plugins/Runtime/nDisplay/Source/DisplayClusterConfiguration/Private/DisplayClusterConfigurationTypes_Viewport.cpp b/Engine/Plugins/Runtime/nDisplay/Source/DisplayClusterConfiguration/Private/DisplayClusterConfigurationTypes_Viewport.cpp
--- a/Engine/Plugins/Runtime/nDisplay/Source/DisplayClusterConfiguration/Private/DisplayClusterConfigurationTypes_Viewport.cpp
+++ b/Engine/Plugins/Runtime/nDisplay/Source/DisplayClusterConfiguration/Private/DisplayClusterConfigurationTypes_Viewport.cpp
@@ -93,30 +93,3 @@ EDisplayClusterViewportICVFXFlags FDisplayClusterConfigurationViewport_ICVFX::Ge
 
 	return OutFlags;
 }
-
-
-bool FDisplayClusterConfigurationViewport_RenderSettings::Serialize(FArchive& Ar)
-{
-	// When loading, overwrite Media settings with defaults unless this is the archetype
-
-	if (!Ar.IsLoading() && !Ar.IsSaving())
-	{
-		return true;
-	}
-
-	UScriptStruct& Struct = *StaticStruct();
-
-	if (Ar.IsLoading() && !FDisplayClusterConfigurationUtils::IsSerializingTemplate(Ar))
-	{
-		const FDisplayClusterConfigurationMediaViewport MediaOriginal = Media;
-		Struct.SerializeTaggedProperties(Ar, (uint8*)this, &Struct, nullptr);
-		Media = MediaOriginal;
-	}
-	else
-	{
-		Struct.SerializeTaggedProperties(Ar, (uint8*)this, &Struct, nullptr);
-	}
-
-	return true;
-}
-
diff --git a/Engine/Plugins/Runtime/nDisplay/Source/DisplayClusterConfiguration/Public/DisplayClusterConfigurationTypes_Viewport.h b/Engine/Plugins/Runtime/nDisplay/Source/DisplayClusterConfiguration/Public/DisplayClusterConfigurationTypes_Viewport.h
--- a/Engine/Plugins/Runtime/nDisplay/Source/DisplayClusterConfiguration/Public/DisplayClusterConfigurationTypes_Viewport.h
+++ b/Engine/Plugins/Runtime/nDisplay/Source/DisplayClusterConfiguration/Public/DisplayClusterConfigurationTypes_Viewport.h
@@ -113,17 +113,6 @@ public:
 	UPROPERTY()
 	int RenderFamilyGroup = -1;
 
-public:
-	/** Custom serialization to skip Media on instances */
-	bool Serialize(FArchive& Ar);
-};
-
-
-/** Makes the given struct use the custom serializer */
-template<> struct TStructOpsTypeTraits<FDisplayClusterConfigurationViewport_RenderSettings>
-: public TStructOpsTypeTraitsBase2<FDisplayClusterConfigurationViewport_RenderSettings>
-{
-	enum { WithSerializer = true };
 };
 
 
