DisplayCluster: Fix DC.Stats.Frame race condition.

diff --git a/Engine/Plugins/Runtime/nDisplay/Source/DisplayCluster/Private/Render/Viewport/DisplayClusterViewportFrameStatsViewExtension.cpp b/Engine/Plugins/Runtime/nDisplay/Source/DisplayCluster/Private/Render/Viewport/DisplayClusterViewportFrameStatsViewExtension.cpp
--- a/Engine/Plugins/Runtime/nDisplay/Source/DisplayCluster/Private/Render/Viewport/DisplayClusterViewportFrameStatsViewExtension.cpp
+++ b/Engine/Plugins/Runtime/nDisplay/Source/DisplayCluster/Private/Render/Viewport/DisplayClusterViewportFrameStatsViewExtension.cpp
@@ -93,10 +93,18 @@ void FDisplayClusterViewportFrameStatsViewExtension::BeginRenderViewFamily(FScen
 {
 	if (IsActive())
 	{
-		const FTimecode CurrentTimecode = FApp::GetTimecode();
+		const FTimecode GameThreadTimecode = FApp::GetTimecode();
 
-		EncodedTimecode = (static_cast<uint8>(CurrentTimecode.Hours) << 24u) | (static_cast<uint8>(CurrentTimecode.Minutes) << 16u) | (static_cast<uint8>(CurrentTimecode.Seconds) << 8u) | static_cast<uint8>(CurrentTimecode.Frames);
-		FrameCount = GFrameCounter;
+		// Local lambda copy prevents race condition from further game thread updates before rendering.
+		ENQUEUE_RENDER_COMMAND(DCFrameTimeCode)(
+		[CurrentTimecode = GameThreadTimecode, WeakThis = AsWeak()](FRHICommandListImmediate& RHICmdList)
+		{
+			const TSharedPtr<FDisplayClusterViewportFrameStatsViewExtension> DCSVE = StaticCastSharedPtr<FDisplayClusterViewportFrameStatsViewExtension>(WeakThis.Pin());
+			if (DCSVE.IsValid())
+			{
+				DCSVE->EncodedTimecode_RenderThread = (static_cast<uint8>(CurrentTimecode.Hours) << 24u) | (static_cast<uint8>(CurrentTimecode.Minutes) << 16u) | (static_cast<uint8>(CurrentTimecode.Seconds) << 8u) | static_cast<uint8>(CurrentTimecode.Frames);
+			}
+		});
 	}
 }
 
@@ -124,8 +132,8 @@ FScreenPassTexture FDisplayClusterViewportFrameStatsViewExtension::PostProcessPa
 		FDisplayClusterFrameStatsShaderParameters* Parameters = GraphBuilder.AllocParameters<FDisplayClusterFrameStatsShaderParameters>();
 		Parameters->MiniFontTexture = UE::DisplayClusterViewExtension::GetMiniFontTexture();
 		Parameters->RenderTargets[0] = FrameStatsRenderTargetBinding;
-		Parameters->FrameCount = FrameCount;
-		Parameters->Timecode = EncodedTimecode;
+		Parameters->FrameCount = View.Family->FrameCounter; // Copy of GFrameCounter
+		Parameters->Timecode = EncodedTimecode_RenderThread;
 
 		AddDrawScreenPass(GraphBuilder, RDG_EVENT_NAME("DisplayClusterViewportFrameStatsPass"), View, Viewport, Viewport, VertexShader, FrameStatsPixelShader, Parameters);
 	}
diff --git a/Engine/Plugins/Runtime/nDisplay/Source/DisplayCluster/Private/Render/Viewport/DisplayClusterViewportFrameStatsViewExtension.h b/Engine/Plugins/Runtime/nDisplay/Source/DisplayCluster/Private/Render/Viewport/DisplayClusterViewportFrameStatsViewExtension.h
--- a/Engine/Plugins/Runtime/nDisplay/Source/DisplayCluster/Private/Render/Viewport/DisplayClusterViewportFrameStatsViewExtension.h
+++ b/Engine/Plugins/Runtime/nDisplay/Source/DisplayCluster/Private/Render/Viewport/DisplayClusterViewportFrameStatsViewExtension.h
@@ -38,6 +38,7 @@ private:
 	bool IsActive() const;
 
 private:
-	std::atomic<uint32> FrameCount;
-	std::atomic<uint32> EncodedTimecode;
+
+	// Encoded timecode
+	uint32 EncodedTimecode_RenderThread;
 };
