[DisplayClusterColorGrading] Add "Include Entire Cluster Color Grading" option to Color Grading panel details

diff --git a/Engine/Plugins/Runtime/nDisplay/Source/DisplayClusterColorGrading/Private/DataModelGenerators/DisplayClusterColorGradingGenerator_RootActor.cpp b/Engine/Plugins/Runtime/nDisplay/Source/DisplayClusterColorGrading/Private/DataModelGenerators/DisplayClusterColorGradingGenerator_RootActor.cpp
--- a/Engine/Plugins/Runtime/nDisplay/Source/DisplayClusterColorGrading/Private/DataModelGenerators/DisplayClusterColorGradingGenerator_RootActor.cpp
+++ b/Engine/Plugins/Runtime/nDisplay/Source/DisplayClusterColorGrading/Private/DataModelGenerators/DisplayClusterColorGradingGenerator_RootActor.cpp
@@ -804,11 +804,12 @@ public:
 		{
 			const TMap<FName, int32> SortOrder =
 			{
-				{ TEXT("DetailView_PerNode"), 0},
-				{ TEXT("DetailView_Exposure"), 1},
-				{ TEXT("DetailView_ColorGrading"), 2},
-				{ TEXT("DetailView_WhiteBalance"), 3},
-				{ TEXT("DetailView_Misc"), 4}
+				{ TEXT("DetailView_AllNodes"), 0},
+				{ TEXT("DetailView_PerNode"), 1},
+				{ TEXT("DetailView_Exposure"), 2},
+				{ TEXT("DetailView_ColorGrading"), 3},
+				{ TEXT("DetailView_WhiteBalance"), 4},
+				{ TEXT("DetailView_Misc"), 5}
 			};
 
 			for (const TPair<FName, int32>& SortPair : SortOrder)
@@ -863,6 +864,9 @@ private:
 		}
 		else
 		{
+			IDetailCategoryBuilder& AllNodeSettingsCategoryBuilder = DetailBuilder.EditCategory(TEXT("DetailView_AllNodes"), LOCTEXT("DetailView_AllNodesDisplayName", "All Nodes Settings"));
+			AllNodeSettingsCategoryBuilder.AddProperty(DetailBuilder.GetProperty(GET_MEMBER_NAME_CHECKED(UDisplayClusterICVFXCameraComponent, CameraSettings.AllNodesColorGrading.bEnableEntireClusterColorGrading)));
+
 			AddColorGradingSettings(DetailBuilder.GetProperty(GET_MEMBER_NAME_CHECKED(UDisplayClusterICVFXCameraComponent, CameraSettings.AllNodesColorGrading.ColorGradingSettings)));
 		}
 	}
@@ -1019,6 +1023,8 @@ void FDisplayClusterColorGradingGenerator_ICVFXCamera::GenerateDataModel(IProper
 		FColorGradingEditorDataModel::FColorGradingGroup AllNodesGroup = CreateColorGradingGroup(AllNodesColorGradingHandle);
 		AllNodesGroup.EditConditionPropertyHandle = AllNodesColorGradingHandle->GetChildHandle(GET_MEMBER_NAME_CHECKED(FDisplayClusterConfigurationViewport_AllNodesColorGrading, bEnableInnerFrustumAllNodesColorGrading));
 
+		AllNodesGroup.DetailsViewCategories.Add(TEXT("DetailView_AllNodes"));
+
 		AllNodesGroup.GroupHeaderWidget = SNew(SHorizontalBox)
 			+ SHorizontalBox::Slot()
 			.AutoWidth()
