Fix 'availalbe' typos

diff --git a/Engine/Plugins/Media/MediaIOFramework/Source/MediaIOCore/Private/Player/MediaIOCorePlayerBase.cpp b/Engine/Plugins/Media/MediaIOFramework/Source/MediaIOCore/Private/Player/MediaIOCorePlayerBase.cpp
--- a/Engine/Plugins/Media/MediaIOFramework/Source/MediaIOCore/Private/Player/MediaIOCorePlayerBase.cpp
+++ b/Engine/Plugins/Media/MediaIOFramework/Source/MediaIOCore/Private/Player/MediaIOCorePlayerBase.cpp
@@ -1117,7 +1117,7 @@ void FMediaIOCorePlayerBase::AddVideoSample(const TSharedRef<FMediaIOCoreTexture
 		const bool bCanUseGPUTextureTransfer = CanUseGPUTextureTransfer();
 		const bool bIsAwaitingForGPUTransfer = Sample->IsAwaitingForGPUTransfer();
 
-		// So, when GPUDirect is availalbe we transfer data on the rendering thread. The callback AddVideoSampleAfterGPUTransfer_RenderThread
+		// So, when GPUDirect is available we transfer data on the rendering thread. The callback AddVideoSampleAfterGPUTransfer_RenderThread
 		// will put the sample to the pool when finished.
 		if (bCanUseGPUTextureTransfer && bIsAwaitingForGPUTransfer)
 		{
diff --git a/Engine/Plugins/Media/MediaIOFramework/Source/MediaIOCore/Public/MediaIOCoreDefinitions.h b/Engine/Plugins/Media/MediaIOFramework/Source/MediaIOCore/Public/MediaIOCoreDefinitions.h
--- a/Engine/Plugins/Media/MediaIOFramework/Source/MediaIOCore/Public/MediaIOCoreDefinitions.h
+++ b/Engine/Plugins/Media/MediaIOFramework/Source/MediaIOCore/Public/MediaIOCoreDefinitions.h
@@ -77,7 +77,7 @@ enum class EMediaIOInterlaceFieldOrder : uint8
 };
 
 /**
- * Various sample evaluation types availalbe in MediaIO.
+ * Various sample evaluation types available in MediaIO.
  */
 UENUM()
 enum class EMediaIOSampleEvaluationType : uint8
diff --git a/Engine/Plugins/Runtime/nDisplay/Source/DisplayCluster/Public/Blueprints/DisplayClusterBlueprintLib.h b/Engine/Plugins/Runtime/nDisplay/Source/DisplayCluster/Public/Blueprints/DisplayClusterBlueprintLib.h
--- a/Engine/Plugins/Runtime/nDisplay/Source/DisplayCluster/Public/Blueprints/DisplayClusterBlueprintLib.h
+++ b/Engine/Plugins/Runtime/nDisplay/Source/DisplayCluster/Public/Blueprints/DisplayClusterBlueprintLib.h
@@ -29,7 +29,7 @@ public:
 
 	/** [DEPRECATED] Return Display Cluster API interface. */
 	UE_DEPRECATED(5.4, "GetAPI function has been deprecated. All functions are now available in UDisplayClusterBlueprintLib.")
-	UFUNCTION(BlueprintPure, meta = (DeprecatedFunction, DeprecationMessage = "GetAPI has been deprecated. All functions are now availalbe in the main blueprint functions list under 'nDisplay' category."))
+	UFUNCTION(BlueprintPure, meta = (DeprecatedFunction, DeprecationMessage = "GetAPI has been deprecated. All functions are now available in the main blueprint functions list under 'nDisplay' category."))
 	static void GetAPI(TScriptInterface<IDisplayClusterBlueprintAPI>& OutAPI);
 
 public:
diff --git a/Engine/Plugins/Runtime/nDisplay/Source/DisplayClusterMedia/Private/Capture/DisplayClusterMediaCaptureTile.cpp b/Engine/Plugins/Runtime/nDisplay/Source/DisplayClusterMedia/Private/Capture/DisplayClusterMediaCaptureTile.cpp
--- a/Engine/Plugins/Runtime/nDisplay/Source/DisplayClusterMedia/Private/Capture/DisplayClusterMediaCaptureTile.cpp
+++ b/Engine/Plugins/Runtime/nDisplay/Source/DisplayClusterMedia/Private/Capture/DisplayClusterMediaCaptureTile.cpp
@@ -18,7 +18,7 @@ FDisplayClusterMediaCaptureTile::FDisplayClusterMediaCaptureTile(
 bool FDisplayClusterMediaCaptureTile::GetCaptureSizeFromConfig(FIntPoint& OutSize) const
 {
 	// The upper level logic always tries to get actual capture size from the game proxy. If the game proxy
-	// is not availalbe (camera/viewport is disabled or deactivated on start and therefore not being rendered),
+	// is not available (camera/viewport is disabled or deactivated on start and therefore not being rendered),
 	// it tries to acquire capture size from config. Being here means exactly this case.
 	//
 	// The problem is the tiles can't know their size until they start to render. There is a bunch of intermediate
diff --git a/Engine/Plugins/Runtime/nDisplay/Source/DisplayClusterProjection/Public/Blueprints/DisplayClusterProjectionBlueprintLib.h b/Engine/Plugins/Runtime/nDisplay/Source/DisplayClusterProjection/Public/Blueprints/DisplayClusterProjectionBlueprintLib.h
--- a/Engine/Plugins/Runtime/nDisplay/Source/DisplayClusterProjection/Public/Blueprints/DisplayClusterProjectionBlueprintLib.h
+++ b/Engine/Plugins/Runtime/nDisplay/Source/DisplayClusterProjection/Public/Blueprints/DisplayClusterProjectionBlueprintLib.h
@@ -21,7 +21,7 @@ public:
 
 	/** [DEPRECATED] Returns Display Cluster Projection API interface. */
 	UE_DEPRECATED(5.4, "GetAPI function has been deprecated. All functions are now available in UDisplayClusterProjectionBlueprintLib.")
-	UFUNCTION(BlueprintPure, meta = (DeprecatedFunction, DeprecationMessage = "GetAPI has been deprecated. All functions are now availalbe in the main blueprint functions list under 'nDisplay' category."))
+	UFUNCTION(BlueprintPure, meta = (DeprecatedFunction, DeprecationMessage = "GetAPI has been deprecated. All functions are now available in the main blueprint functions list under 'nDisplay' category."))
 	static void GetAPI(TScriptInterface<IDisplayClusterProjectionBlueprintAPI>& OutAPI);
 
 public:
