[ColorGrading] [ColorCorrectRegions] [nDisplay] Remove unused ColorGradingListItem and related code

diff --git a/Engine/Plugins/Editor/ColorGrading/Source/ColorGradingEditor/Private/ColorGradingEditorModule.cpp b/Engine/Plugins/Editor/ColorGrading/Source/ColorGradingEditor/Private/ColorGradingEditorModule.cpp
--- a/Engine/Plugins/Editor/ColorGrading/Source/ColorGradingEditor/Private/ColorGradingEditorModule.cpp
+++ b/Engine/Plugins/Editor/ColorGrading/Source/ColorGradingEditor/Private/ColorGradingEditorModule.cpp
@@ -23,8 +23,6 @@ void FColorGradingEditorModule::StartupModule()
 	FColorGradingEditorDataModel::RegisterColorGradingDataModelGenerator<APostProcessVolume>(
 		FGetDetailsDataModelGenerator::CreateStatic(&FColorGradingDataModelGenerator_PostProcessVolume::MakeInstance));
 
-	FColorGradingListItem::RegisterColorGradingListItemGenerator<FColorGradingListItemGenerator_PostProcessVolume>();
-
 	FColorGradingMixerObjectFilterRegistry::RegisterActorClassToPlace(APostProcessVolume::StaticClass());
 
 	FColorGradingMixerObjectFilterRegistry::RegisterObjectClassToFilter(APostProcessVolume::StaticClass());
diff --git a/Engine/Plugins/Editor/ColorGrading/Source/ColorGradingEditor/Private/ColorGradingListItem.cpp b/Engine/Plugins/Editor/ColorGrading/Source/ColorGradingEditor/Private/ColorGradingListItem.cpp
deleted file mode 100644
diff --git a/Engine/Plugins/Editor/ColorGrading/Source/ColorGradingEditor/Private/DataModelGenerators/ColorGradingDataModelGenerator_PostProcessVolume.cpp b/Engine/Plugins/Editor/ColorGrading/Source/ColorGradingEditor/Private/DataModelGenerators/ColorGradingDataModelGenerator_PostProcessVolume.cpp
--- a/Engine/Plugins/Editor/ColorGrading/Source/ColorGradingEditor/Private/DataModelGenerators/ColorGradingDataModelGenerator_PostProcessVolume.cpp
+++ b/Engine/Plugins/Editor/ColorGrading/Source/ColorGradingEditor/Private/DataModelGenerators/ColorGradingDataModelGenerator_PostProcessVolume.cpp
@@ -218,21 +218,4 @@ void FColorGradingDataModelGenerator_PostProcessVolume::AddPropertyToColorGradin
 	}
 }
 
-TArray<TSubclassOf<AActor>> FColorGradingListItemGenerator_PostProcessVolume::GetActorClassesForListItems() const
-{
-	return { APostProcessVolume::StaticClass() };
-}
-
-void FColorGradingListItemGenerator_PostProcessVolume::GenerateColorGradingListItems(AActor* InActor, TArray<FColorGradingListItemRef>& OutList) const
-{
-	if (APostProcessVolume* PostProcessVolume = Cast<APostProcessVolume>(InActor))
-	{
-		FColorGradingListItemRef ListItemRef = MakeShared<FColorGradingListItem>(PostProcessVolume);
-		ListItemRef->IsItemEnabled = CREATE_IS_ENABLED_LAMBDA(PostProcessVolume, PostProcessVolume->bEnabled);
-		ListItemRef->OnItemEnabledChanged = CREATE_ON_ENABLED_CHANGED_LAMBDA(PostProcessVolume, PostProcessVolume->bEnabled);
-
-		OutList.Add(ListItemRef);
-	}
-}
-
 #undef LOCTEXT_NAMESPACE
\ No newline at end of file
diff --git a/Engine/Plugins/Editor/ColorGrading/Source/ColorGradingEditor/Private/DataModelGenerators/ColorGradingDataModelGenerator_PostProcessVolume.h b/Engine/Plugins/Editor/ColorGrading/Source/ColorGradingEditor/Private/DataModelGenerators/ColorGradingDataModelGenerator_PostProcessVolume.h
--- a/Engine/Plugins/Editor/ColorGrading/Source/ColorGradingEditor/Private/DataModelGenerators/ColorGradingDataModelGenerator_PostProcessVolume.h
+++ b/Engine/Plugins/Editor/ColorGrading/Source/ColorGradingEditor/Private/DataModelGenerators/ColorGradingDataModelGenerator_PostProcessVolume.h
@@ -21,13 +21,4 @@ public:
 
 private:
 	void AddPropertyToColorGradingElement(const TSharedPtr<IPropertyHandle>& PropertyHandle, FColorGradingEditorDataModel::FColorGradingElement& ColorGradingElement);
-};
-
-/** Color grading list item generator for an APostProcessVolume */
-class FColorGradingListItemGenerator_PostProcessVolume : public IColorGradingListItemGenerator
-{
-	//~ IColorGradingListItemGenerator interface
-	virtual TArray<TSubclassOf<AActor>> GetActorClassesForListItems() const override;
-	virtual void GenerateColorGradingListItems(AActor* InActor, TArray<FColorGradingListItemRef>& OutList) const override;
-	//~ End IColorGradingListItemGenerator interface
 };
\ No newline at end of file
diff --git a/Engine/Plugins/Editor/ColorGrading/Source/ColorGradingEditor/Public/ColorGradingEditorDataModel.h b/Engine/Plugins/Editor/ColorGrading/Source/ColorGradingEditor/Public/ColorGradingEditorDataModel.h
--- a/Engine/Plugins/Editor/ColorGrading/Source/ColorGradingEditor/Public/ColorGradingEditorDataModel.h
+++ b/Engine/Plugins/Editor/ColorGrading/Source/ColorGradingEditor/Public/ColorGradingEditorDataModel.h
@@ -5,8 +5,6 @@
 #include "CoreMinimal.h"
 #include "PropertyEditorDelegates.h"
 
-#include "ColorGradingListItem.h"
-
 class IPropertyHandle;
 class IPropertyRowGenerator;
 class SWidget;
diff --git a/Engine/Plugins/Editor/ColorGrading/Source/ColorGradingEditor/Public/ColorGradingListItem.h b/Engine/Plugins/Editor/ColorGrading/Source/ColorGradingEditor/Public/ColorGradingListItem.h
deleted file mode 100644
diff --git a/Engine/Plugins/Editor/ColorGrading/Source/ColorGradingEditor/Public/SColorGradingPanel.h b/Engine/Plugins/Editor/ColorGrading/Source/ColorGradingEditor/Public/SColorGradingPanel.h
--- a/Engine/Plugins/Editor/ColorGrading/Source/ColorGradingEditor/Public/SColorGradingPanel.h
+++ b/Engine/Plugins/Editor/ColorGrading/Source/ColorGradingEditor/Public/SColorGradingPanel.h
@@ -9,7 +9,6 @@
 #include "SelectionInterface/IObjectMixerSelectionInterface.h"
 
 #include "ColorGradingPanelState.h"
-#include "ColorGradingListItem.h"
 
 class FColorGradingEditorDataModel;
 class FObjectMixerEditorList;
diff --git a/Engine/Plugins/Experimental/ColorCorrectRegions/Source/ColorCorrectRegionsEditor/Private/ColorCorrectRegionsEditorModule.cpp b/Engine/Plugins/Experimental/ColorCorrectRegions/Source/ColorCorrectRegionsEditor/Private/ColorCorrectRegionsEditorModule.cpp
--- a/Engine/Plugins/Experimental/ColorCorrectRegions/Source/ColorCorrectRegionsEditor/Private/ColorCorrectRegionsEditorModule.cpp
+++ b/Engine/Plugins/Experimental/ColorCorrectRegions/Source/ColorCorrectRegionsEditor/Private/ColorCorrectRegionsEditorModule.cpp
@@ -25,8 +25,6 @@ void FColorCorrectRegionsEditorModule::StartupModule()
 	FColorGradingEditorDataModel::RegisterColorGradingDataModelGenerator<AColorCorrectRegion>(
 		FGetDetailsDataModelGenerator::CreateStatic(&FColorGradingDataModelGenerator_ColorCorrectRegion::MakeInstance));
 
-	FColorGradingListItem::RegisterColorGradingListItemGenerator<FColorGradingListItemGenerator_ColorCorrectRegion>();
-
 	FColorGradingMixerObjectFilterRegistry::RegisterActorClassToPlace(AColorCorrectionRegion::StaticClass());
 	FColorGradingMixerObjectFilterRegistry::RegisterActorClassToPlace(AColorCorrectionWindow::StaticClass());
 
diff --git a/Engine/Plugins/Experimental/ColorCorrectRegions/Source/ColorCorrectRegionsEditor/Private/ColorGradingDataModelGenerator_ColorCorrectRegion.cpp b/Engine/Plugins/Experimental/ColorCorrectRegions/Source/ColorCorrectRegionsEditor/Private/ColorGradingDataModelGenerator_ColorCorrectRegion.cpp
--- a/Engine/Plugins/Experimental/ColorCorrectRegions/Source/ColorCorrectRegionsEditor/Private/ColorGradingDataModelGenerator_ColorCorrectRegion.cpp
+++ b/Engine/Plugins/Experimental/ColorCorrectRegions/Source/ColorCorrectRegionsEditor/Private/ColorGradingDataModelGenerator_ColorCorrectRegion.cpp
@@ -203,21 +203,4 @@ bool FColorGradingDataModelGenerator_ColorCorrectRegion::FilterDetailsViewProper
 	}
 }
 
-TArray<TSubclassOf<AActor>> FColorGradingListItemGenerator_ColorCorrectRegion::GetActorClassesForListItems() const
-{ 
-	return { AColorCorrectionRegion::StaticClass() };
-}
-
-void FColorGradingListItemGenerator_ColorCorrectRegion::GenerateColorGradingListItems(AActor* InActor, TArray<FColorGradingListItemRef>& OutList) const
-{
-	if (AColorCorrectionRegion* ColorCorrectRegion = Cast<AColorCorrectionRegion>(InActor))
-	{
-		FColorGradingListItemRef ListItemRef = MakeShared<FColorGradingListItem>(ColorCorrectRegion);
-		ListItemRef->IsItemEnabled = CREATE_IS_ENABLED_LAMBDA(ColorCorrectRegion, ColorCorrectRegion->Enabled);
-		ListItemRef->OnItemEnabledChanged = CREATE_ON_ENABLED_CHANGED_LAMBDA(ColorCorrectRegion, ColorCorrectRegion->Enabled);
-
-		OutList.Add(ListItemRef);
-	}
-}
-
 #undef LOCTEXT_NAMESPACE
\ No newline at end of file
diff --git a/Engine/Plugins/Experimental/ColorCorrectRegions/Source/ColorCorrectRegionsEditor/Private/ColorGradingDataModelGenerator_ColorCorrectRegion.h b/Engine/Plugins/Experimental/ColorCorrectRegions/Source/ColorCorrectRegionsEditor/Private/ColorGradingDataModelGenerator_ColorCorrectRegion.h
--- a/Engine/Plugins/Experimental/ColorCorrectRegions/Source/ColorCorrectRegionsEditor/Private/ColorGradingDataModelGenerator_ColorCorrectRegion.h
+++ b/Engine/Plugins/Experimental/ColorCorrectRegions/Source/ColorCorrectRegionsEditor/Private/ColorGradingDataModelGenerator_ColorCorrectRegion.h
@@ -24,13 +24,4 @@ private:
 	FColorGradingEditorDataModel::FColorGradingElement CreateColorGradingElement(const TSharedRef<IDetailTreeNode>& GroupNode, FText ElementLabel);
 
 	bool FilterDetailsViewProperties(const TSharedRef<IDetailTreeNode>& InDetailTreeNode);
-};
-
-/** Color grading list item generator for an AColorCorrectionRegion */
-class FColorGradingListItemGenerator_ColorCorrectRegion : public IColorGradingListItemGenerator
-{
-	//~ IColorGradingListItemGenerator interface
-	virtual TArray<TSubclassOf<AActor>> GetActorClassesForListItems() const override;
-	virtual void GenerateColorGradingListItems(AActor* InActor, TArray<FColorGradingListItemRef>& OutList) const override;
-	//~ End IColorGradingListItemGenerator interface
 };
\ No newline at end of file
diff --git a/Engine/Plugins/Runtime/nDisplay/Source/DisplayClusterColorGrading/Private/DataModelGenerators/DisplayClusterColorGradingGenerator_RootActor.cpp b/Engine/Plugins/Runtime/nDisplay/Source/DisplayClusterColorGrading/Private/DataModelGenerators/DisplayClusterColorGradingGenerator_RootActor.cpp
--- a/Engine/Plugins/Runtime/nDisplay/Source/DisplayClusterColorGrading/Private/DataModelGenerators/DisplayClusterColorGradingGenerator_RootActor.cpp
+++ b/Engine/Plugins/Runtime/nDisplay/Source/DisplayClusterColorGrading/Private/DataModelGenerators/DisplayClusterColorGradingGenerator_RootActor.cpp
@@ -1353,41 +1353,4 @@ TSharedRef<SWidget> FDisplayClusterColorGradingGenerator_ICVFXCamera::GetNodeCom
 	return MenuBuilder.MakeWidget();
 }
 
-TArray<TSubclassOf<AActor>> FDisplayClusterColorGradingListItemGenerator_RootActor::GetActorClassesForListItems() const
-{
-	return { ADisplayClusterRootActor::StaticClass() };
-}
-
-void FDisplayClusterColorGradingListItemGenerator_RootActor::GenerateColorGradingListItems(AActor* InActor, TArray<FColorGradingListItemRef>& OutList) const
-{
-	if (ADisplayClusterRootActor* RootActor = Cast<ADisplayClusterRootActor>(InActor))
-	{
-		FColorGradingListItemRef ListItemRef = MakeShared<FColorGradingListItem>(RootActor);
-		ListItemRef->IsItemEnabled = CREATE_IS_ENABLED_LAMBDA(RootActor, RootActor->GetConfigData()->StageSettings.EnableColorGrading);
-		ListItemRef->OnItemEnabledChanged = CREATE_ON_ENABLED_CHANGED_LAMBDA(RootActor, RootActor->GetConfigData()->StageSettings.EnableColorGrading);
-
-		OutList.Add(ListItemRef);
-	}
-}
-
-TArray<TSubclassOf<AActor>> FDisplayClusterColorGradingListItemGenerator_ICVFXCamera::GetActorClassesForListItems() const
-{
-	return { ADisplayClusterRootActor::StaticClass() };
-}
-
-void FDisplayClusterColorGradingListItemGenerator_ICVFXCamera::GenerateColorGradingListItems(AActor* InActor, TArray<FColorGradingListItemRef>& OutList) const
-{
-	if (ADisplayClusterRootActor* RootActor = Cast<ADisplayClusterRootActor>(InActor))
-	{
-		RootActor->ForEachComponent<UDisplayClusterICVFXCameraComponent>(false, [this, RootActor, &OutList](UDisplayClusterICVFXCameraComponent* ICVFXCameraComponent)
-		{
-			FColorGradingListItemRef ICVFXCameraListItemRef = MakeShared<FColorGradingListItem>(RootActor, ICVFXCameraComponent);
-			ICVFXCameraListItemRef->IsItemEnabled = CREATE_IS_ENABLED_LAMBDA(ICVFXCameraComponent, ICVFXCameraComponent->CameraSettings.EnableInnerFrustumColorGrading);
-			ICVFXCameraListItemRef->OnItemEnabledChanged = CREATE_ON_ENABLED_CHANGED_LAMBDA(ICVFXCameraComponent, ICVFXCameraComponent->CameraSettings.EnableInnerFrustumColorGrading);
-
-			OutList.Add(ICVFXCameraListItemRef);
-		});
-	}
-}
-
 #undef LOCTEXT_NAMESPACE
diff --git a/Engine/Plugins/Runtime/nDisplay/Source/DisplayClusterColorGrading/Private/DataModelGenerators/DisplayClusterColorGradingGenerator_RootActor.h b/Engine/Plugins/Runtime/nDisplay/Source/DisplayClusterColorGrading/Private/DataModelGenerators/DisplayClusterColorGradingGenerator_RootActor.h
--- a/Engine/Plugins/Runtime/nDisplay/Source/DisplayClusterColorGrading/Private/DataModelGenerators/DisplayClusterColorGradingGenerator_RootActor.h
+++ b/Engine/Plugins/Runtime/nDisplay/Source/DisplayClusterColorGrading/Private/DataModelGenerators/DisplayClusterColorGradingGenerator_RootActor.h
@@ -102,22 +102,4 @@ private:
 private:
 	/** A list of camera components that are being represented by the data model */
 	TArray<TWeakObjectPtr<UDisplayClusterICVFXCameraComponent>> CameraComponents;
-};
-
-/** Color grading list item generator for an nDisplay root actor */
-class FDisplayClusterColorGradingListItemGenerator_RootActor : public IColorGradingListItemGenerator
-{
-	//~ IColorGradingListItemGenerator interface
-	virtual TArray<TSubclassOf<AActor>> GetActorClassesForListItems() const override;
-	virtual void GenerateColorGradingListItems(AActor* InActor, TArray<FColorGradingListItemRef>& OutList) const override;
-	//~ End IColorGradingListItemGenerator interface
-};
-
-/** Color grading list item generator for an nDisplay ICVFX camera component */
-class FDisplayClusterColorGradingListItemGenerator_ICVFXCamera : public IColorGradingListItemGenerator
-{
-	//~ IColorGradingListItemGenerator interface
-	virtual TArray<TSubclassOf<AActor>> GetActorClassesForListItems() const override;
-	virtual void GenerateColorGradingListItems(AActor* InActor, TArray<FColorGradingListItemRef>& OutList) const override;
-	//~ End IColorGradingListItemGenerator interface
 };
\ No newline at end of file
diff --git a/Engine/Plugins/Runtime/nDisplay/Source/DisplayClusterColorGrading/Private/DisplayClusterColorGradingModule.cpp b/Engine/Plugins/Runtime/nDisplay/Source/DisplayClusterColorGrading/Private/DisplayClusterColorGradingModule.cpp
--- a/Engine/Plugins/Runtime/nDisplay/Source/DisplayClusterColorGrading/Private/DisplayClusterColorGradingModule.cpp
+++ b/Engine/Plugins/Runtime/nDisplay/Source/DisplayClusterColorGrading/Private/DisplayClusterColorGradingModule.cpp
@@ -26,9 +26,6 @@ void FDisplayClusterColorGradingModule::StartupModule()
 	FColorGradingEditorDataModel::RegisterColorGradingDataModelGenerator<UDisplayClusterICVFXCameraComponent>(
 		FGetDetailsDataModelGenerator::CreateStatic(&FDisplayClusterColorGradingGenerator_ICVFXCamera::MakeInstance));
 
-	FColorGradingListItem::RegisterColorGradingListItemGenerator<FDisplayClusterColorGradingListItemGenerator_RootActor>();
-	FColorGradingListItem::RegisterColorGradingListItemGenerator<FDisplayClusterColorGradingListItemGenerator_ICVFXCamera>();
-
 	FColorGradingMixerObjectFilterRegistry::RegisterActorClassToPlace(ADisplayClusterRootActor::StaticClass());
 
 	FColorGradingMixerObjectFilterRegistry::RegisterObjectClassToFilter(ADisplayClusterRootActor::StaticClass());
