nDisplay: Fixed overscan issues for the InCamera and viewport.
- Fixed a bug where Frustum Overscan values were not shifting correctly.
- Added new CVar "nDisplay.render.custom_frustum.min_percent".
- Fixed viewport overscan:
  * negative values are not allowed, clamped.
  * Percentage used to change frustum, since RTT size can be changed by other multipliers (e.g. preview scale, etc.)

diff --git a/Engine/Plugins/Runtime/nDisplay/Source/DisplayCluster/Private/Render/Viewport/Containers/DisplayClusterViewport_CustomFrustumRuntimeSettings.cpp b/Engine/Plugins/Runtime/nDisplay/Source/DisplayCluster/Private/Render/Viewport/Containers/DisplayClusterViewport_CustomFrustumRuntimeSettings.cpp
--- a/Engine/Plugins/Runtime/nDisplay/Source/DisplayCluster/Private/Render/Viewport/Containers/DisplayClusterViewport_CustomFrustumRuntimeSettings.cpp
+++ b/Engine/Plugins/Runtime/nDisplay/Source/DisplayCluster/Private/Render/Viewport/Containers/DisplayClusterViewport_CustomFrustumRuntimeSettings.cpp
@@ -23,15 +23,57 @@ static FAutoConsoleVariableRef CVarDisplayClusterRenderCustomFrustumMaxValue(
 	ECVF_Default
 );
 
+int32 GDisplayClusterRenderCustomFrustumMinSizeValue = 1;
+static FAutoConsoleVariableRef CVarDisplayClusterRenderCustomFrustumMinSizeValue(
+	TEXT("nDisplay.render.custom_frustum.min_percent_size"),
+	GDisplayClusterRenderCustomFrustumMinSizeValue,
+	TEXT("Min size in percent for custom frustum (default 1).\n"),
+	ECVF_Default
+);
+
 namespace UE::DisplayCluster::Viewport::CustomFrustumHelpers
 {
-	/** Clamp percent for custom frustum settings. */
+	/** Clamp percent for a custom frustum settings. */
 	static inline double ClampPercent(double InValue)
 	{
 		const double MaxCustomFrustumValue = double(GDisplayClusterRenderCustomFrustumMaxValue) / 100;
 
 		return FMath::Clamp(InValue, -MaxCustomFrustumValue, MaxCustomFrustumValue);
 	}
+
+	/** Adjusts the overscan values so that the final size is larger than the minimum size.
+	* 
+	* @param InOutFromValue - (in, out) FromValue, where 1 means 100%
+	* @param InOutToValue   - (in, out) ToValue, where 1 means 100%
+	*/
+	static inline void AdjustOverscanValuesToEnforceMinimumSize(double& InOutFromValue, double& InOutToValue)
+	{
+		const double MinSize = FMath::Max(GDisplayClusterRenderCustomFrustumMinSizeValue * 0.01, 0);
+
+		// Expected size in percentage = 100 % + From + To 
+		// 1 means 100 % because all percentages were multiplied by 0.01.
+		const double ExpectedSize = InOutFromValue + InOutToValue + 1;
+
+		// If negative values are used, the size may decrease to zero or less.
+		// The MinSize value is used as a limit.
+		if (ExpectedSize < MinSize)
+		{
+			// Gets the negative values only
+			const double NegativeFromValue = -FMath::Min(InOutFromValue, 0);
+			const double NegativeToValue   = -FMath::Min(InOutToValue, 0);
+
+			// When using negative values, check that the result size is larger than the minimum.
+			const double NegativeRange = NegativeFromValue + NegativeToValue;
+			if (NegativeRange > 0)
+			{
+				// Adjust From and To values to the minimum size.
+				const double AdjustMult = (MinSize - ExpectedSize) / NegativeRange;
+
+				InOutFromValue += NegativeFromValue * AdjustMult;
+				InOutToValue   += NegativeToValue * AdjustMult;
+			}
+		}
+	}
 };
 
 ///////////////////////////////////////////////////////////////////////////////////////////
@@ -47,16 +89,12 @@ bool FDisplayClusterViewport_CustomFrustumRuntimeSettings::UpdateProjectionAngle
 {
 	if (InRuntimeSettings.bIsEnabled)
 	{
-		const double Horizontal = InOutRight - InOutLeft;
-		const double Vertical = InOutTop - InOutBottom;
+		const FVector2D SizeFOV(InOutRight - InOutLeft, InOutTop - InOutBottom);
 
-		// Use the inner region of the texture as the base of the frustum.
-		const FIntPoint InnerSize = InRenderTargetSize - InRuntimeSettings.CustomFrustumPixels.Size();
-
-		InOutLeft -= Horizontal * InRuntimeSettings.CustomFrustumPixels.Left / InnerSize.X;
-		InOutRight += Horizontal * InRuntimeSettings.CustomFrustumPixels.Right / InnerSize.X;
-		InOutBottom -= Vertical * InRuntimeSettings.CustomFrustumPixels.Bottom / InnerSize.Y;
-		InOutTop += Vertical * InRuntimeSettings.CustomFrustumPixels.Top / InnerSize.Y;
+		InOutLeft   -= SizeFOV.X * InRuntimeSettings.CustomFrustumPercent.Left;
+		InOutRight  += SizeFOV.X * InRuntimeSettings.CustomFrustumPercent.Right;
+		InOutBottom -= SizeFOV.Y * InRuntimeSettings.CustomFrustumPercent.Bottom;
+		InOutTop    += SizeFOV.Y * InRuntimeSettings.CustomFrustumPercent.Top;
 
 		return true;
 	}
@@ -109,6 +147,10 @@ void FDisplayClusterViewport_CustomFrustumRuntimeSettings::UpdateCustomFrustumSe
 		return;
 	}
 
+	// Clamp minimum size of frustum
+	CustomFrustumHelpers::AdjustOverscanValuesToEnforceMinimumSize(InOutRuntimeSettings.CustomFrustumPercent.Left,   InOutRuntimeSettings.CustomFrustumPercent.Right);
+	CustomFrustumHelpers::AdjustOverscanValuesToEnforceMinimumSize(InOutRuntimeSettings.CustomFrustumPercent.Bottom, InOutRuntimeSettings.CustomFrustumPercent.Top);
+
 	// Calc pixels from percent
 	InOutRuntimeSettings.CustomFrustumPixels.Left   = FMath::RoundToInt(Size.X * InOutRuntimeSettings.CustomFrustumPercent.Left);
 	InOutRuntimeSettings.CustomFrustumPixels.Right  = FMath::RoundToInt(Size.X * InOutRuntimeSettings.CustomFrustumPercent.Right);
diff --git a/Engine/Plugins/Runtime/nDisplay/Source/DisplayCluster/Private/Render/Viewport/Containers/DisplayClusterViewport_OverscanRuntimeSettings.cpp b/Engine/Plugins/Runtime/nDisplay/Source/DisplayCluster/Private/Render/Viewport/Containers/DisplayClusterViewport_OverscanRuntimeSettings.cpp
--- a/Engine/Plugins/Runtime/nDisplay/Source/DisplayCluster/Private/Render/Viewport/Containers/DisplayClusterViewport_OverscanRuntimeSettings.cpp
+++ b/Engine/Plugins/Runtime/nDisplay/Source/DisplayCluster/Private/Render/Viewport/Containers/DisplayClusterViewport_OverscanRuntimeSettings.cpp
@@ -29,7 +29,14 @@ namespace UE::DisplayCluster::Viewport::OverscanHelpers
 	{
 		const double MaxCustomFrustumValue = double(GDisplayClusterRenderOverscanMaxValue) / 100;
 
-		return FMath::Clamp(InValue, -MaxCustomFrustumValue, MaxCustomFrustumValue);
+		// We can't use negative overscan values.
+		// The idea behind the overscan is to add extra space on the sides of the RTT.
+		// Note: this only applies to regular viewports (Outers, etc.).
+		// 
+		// The inner frustum viewport has its own implementation for the overscan feature called “CustomFrustum”.
+		// (see FDisplayClusterViewport_CustomFrustumRuntimeSettings)
+
+		return FMath::Clamp(InValue, 0, MaxCustomFrustumValue);
 	}
 };
 
@@ -46,16 +53,12 @@ bool FDisplayClusterViewport_OverscanRuntimeSettings::UpdateProjectionAngles(
 {
 	if (InOverscanRuntimeSettings.bIsEnabled)
 	{
-		double Horizontal = InOutRight - InOutLeft;
-		double Vertical = InOutTop - InOutBottom;
+		const FVector2D SizeFOV(InOutRight - InOutLeft, InOutTop - InOutBottom);
 
-		// Use the inner region of the texture as the base of the frustum.
-		const FIntPoint InnerSize = InRenderTargetSize - InOverscanRuntimeSettings.OverscanPixels.Size();
-
-		InOutLeft   -= Horizontal * InOverscanRuntimeSettings.OverscanPixels.Left / InnerSize.X;
-		InOutRight  += Horizontal * InOverscanRuntimeSettings.OverscanPixels.Right / InnerSize.X;
-		InOutBottom -= Vertical * InOverscanRuntimeSettings.OverscanPixels.Bottom / InnerSize.Y;
-		InOutTop    += Vertical * InOverscanRuntimeSettings.OverscanPixels.Top / InnerSize.Y;
+		InOutLeft   -= SizeFOV.X * InOverscanRuntimeSettings.OverscanPercent.Left;
+		InOutRight  += SizeFOV.X * InOverscanRuntimeSettings.OverscanPercent.Right;
+		InOutBottom -= SizeFOV.Y * InOverscanRuntimeSettings.OverscanPercent.Bottom;
+		InOutTop    += SizeFOV.Y * InOverscanRuntimeSettings.OverscanPercent.Top;
 
 		return true;
 	}
