nDisplay: Fixed tile rendering issue
- fixed visibility settings for tile viewport
- ICVFX settings are not copied for the tile viewport
- Code cleanup

diff --git a/Engine/Plugins/Runtime/nDisplay/Source/DisplayCluster/Private/Render/Viewport/Configuration/DisplayClusterViewportConfiguration.cpp b/Engine/Plugins/Runtime/nDisplay/Source/DisplayCluster/Private/Render/Viewport/Configuration/DisplayClusterViewportConfiguration.cpp
--- a/Engine/Plugins/Runtime/nDisplay/Source/DisplayCluster/Private/Render/Viewport/Configuration/DisplayClusterViewportConfiguration.cpp
+++ b/Engine/Plugins/Runtime/nDisplay/Source/DisplayCluster/Private/Render/Viewport/Configuration/DisplayClusterViewportConfiguration.cpp
@@ -384,15 +384,15 @@ bool FDisplayClusterViewportConfiguration::ImplUpdateConfiguration(EDisplayClust
 	}
 
 	ConfigurationICVFX.Update();
-	ConfigurationTile.Update();
-
 	ConfigurationProjectionPolicy.Update();
-
 	ConfigurationICVFX.PostUpdate();
-	ConfigurationTile.PostUpdate();
 
 	ImplUpdateConfigurationVisibility();
 
+	// Tiled viewports should be created and updated at the very end, when all base viewports are already set up.
+	// Because tile viewports copy configuration data from the base viewport that is used to create the tile.
+	ConfigurationTile.Update();
+
 	if (!InViewportNames)
 	{
 		// Update postprocess for current cluster node
diff --git a/Engine/Plugins/Runtime/nDisplay/Source/DisplayCluster/Private/Render/Viewport/Configuration/DisplayClusterViewportConfigurationHelpers_Tile.cpp b/Engine/Plugins/Runtime/nDisplay/Source/DisplayCluster/Private/Render/Viewport/Configuration/DisplayClusterViewportConfigurationHelpers_Tile.cpp
--- a/Engine/Plugins/Runtime/nDisplay/Source/DisplayCluster/Private/Render/Viewport/Configuration/DisplayClusterViewportConfigurationHelpers_Tile.cpp
+++ b/Engine/Plugins/Runtime/nDisplay/Source/DisplayCluster/Private/Render/Viewport/Configuration/DisplayClusterViewportConfigurationHelpers_Tile.cpp
@@ -24,47 +24,18 @@
 
 void FDisplayClusterViewportConfigurationHelpers_Tile::UpdateICVFXCameraViewportTileSettings(FDisplayClusterViewport& InSourceViewport, const FDisplayClusterConfigurationMediaICVFX& InCameraMediaSettings)
 {
-	// Nothing to do if media tiling is not configured
-	const bool bMediaEnabled = InCameraMediaSettings.bEnable;
-	const bool bMediaTiled = (InCameraMediaSettings.SplitType == EDisplayClusterConfigurationMediaSplitType::UniformTiles);
-	if (!bMediaEnabled || !bMediaTiled)
+	if (!InSourceViewport.CanSplitIntoTiles() || !InCameraMediaSettings.ShouldMediaICVFXSplitIntoTiles())
 	{
+		// Disable tile splittings for this viewport.
 		InSourceViewport.GetRenderSettingsImpl().TileSettings = { };
 		return;
 	}
 
-	// Validate the layout
-	const bool bIsLayoutValid = FDisplayClusterConfigurationTile_Settings::IsValid(InCameraMediaSettings.TiledSplitLayout);
-	if (!bIsLayoutValid)
-	{
-		InSourceViewport.GetRenderSettingsImpl().TileSettings = { };
-		return;
-	}
-
-	// Check if viewport is able to split
-	const bool bCanSplitIntoTiles = InSourceViewport.CanSplitIntoTiles();
-	if (!bCanSplitIntoTiles)
-	{
-		InSourceViewport.GetRenderSettingsImpl().TileSettings = { };
-		return;
-	}
-
-	// Find if this cluster node is allowed to render unbound tiles
-	const FString ThisNodeId = InSourceViewport.GetClusterNodeId();
-	const bool bAllowRenderUnbound = InCameraMediaSettings.ClusterNodesToRenderUnboundTiles.ItemNames.ContainsByPredicate([&ThisNodeId](const FString& Item)
-		{
-			return ThisNodeId.Equals(Item, ESearchCase::IgnoreCase);
-		});
-
-	// Generate flags
-	EDisplayClusterViewportTileFlags TileFlags = EDisplayClusterViewportTileFlags::None;
-	if (bAllowRenderUnbound)
-	{
-		TileFlags |= EDisplayClusterViewportTileFlags::AllowUnboundRender;
-	}
+	// Get additional tile flags
+	const EDisplayClusterViewportTileFlags TileFlags = InCameraMediaSettings.GetMediaICVFXTileFlags(InSourceViewport.GetClusterNodeId());
 
 	// Generate overscan settings
-	const FDisplayClusterViewport_OverscanSettings& OverscanSettings = FDisplayClusterViewportConfigurationHelpers_Tile::GetTileOverscanSettings(InCameraMediaSettings.TileOverscan);
+	const FDisplayClusterViewport_OverscanSettings OverscanSettings = FDisplayClusterViewportConfigurationHelpers_Tile::GetTileOverscanSettings(InCameraMediaSettings.TileOverscan);
 
 	// Set this viewport as a source for tile rendering.
 	FDisplayClusterViewport_TileSettings& OutTileSettings = InSourceViewport.GetRenderSettingsImpl().TileSettings;
@@ -140,6 +111,8 @@ FDisplayClusterViewport* FDisplayClusterViewportConfigurationHelpers_Tile::FindT
 
 FDisplayClusterViewport* FDisplayClusterViewportConfigurationHelpers_Tile::GetOrCreateTileViewport(FDisplayClusterViewport& InSourceViewport, const FIntPoint& InTilePos)
 {
+	// Note: At this point, the viewports should already be configured.
+
 	const FDisplayClusterViewport_RenderSettings& SourceRenderSettings = InSourceViewport.GetRenderSettings();
 	if (SourceRenderSettings.TileSettings.GetType() != EDisplayClusterViewportTileType::Source)
 	{
@@ -167,31 +140,33 @@ FDisplayClusterViewport* FDisplayClusterViewportConfigurationHelpers_Tile::GetOr
 	// Note: The source viewport must already be configured.
 	if (TileViewport)
 	{
-		// Gain direct access to internal resources of the NewViewport:
-		FDisplayClusterViewport_RenderSettings& InOutRenderSettings = TileViewport->GetRenderSettingsImpl();
-
 		// Reset runtime flags from prev frame.
 		// Also this function update media states.
+		// Note: Save the new media states after calling ResetRuntimeParameters().
 		TileViewport->ResetRuntimeParameters();
+		const EDisplayClusterViewportMediaState SavedTileViewportMediaStates = TileViewport->GetRenderSettings().GetMediaStates();
 
-		// Override rendering settings from the source viewport, except for media states.
-		{
-			// Save the new media states after calling ResetRuntimeParameters().
-			const EDisplayClusterViewportMediaState TileMediaStates = TileViewport->GetRenderSettings().GetMediaStates();
+		// Gain direct access to internal resources of the NewViewport:
+		FDisplayClusterViewport_RenderSettings& InOutRenderSettings = TileViewport->GetRenderSettingsImpl();
 
 			// Copy all the settings from the source viewport, but some of them still need to be overridden.
 			InOutRenderSettings = SourceRenderSettings;
 
-			// Restore media states
-			InOutRenderSettings.AssignMediaStates(TileMediaStates);
-		}
+		// Restore media states for the tile viewport.
+		InOutRenderSettings.AssignMediaStates(SavedTileViewportMediaStates);
 
 		// Don't show Tile composing viewports on frame target
 		InOutRenderSettings.bVisible = false;
 
-		// Disable custom frustum settings and override overscan settings.
-		// Use custom overscan settings for the tile rendering.
+		// Override custom frustum settings and override overscan settings.
+		// The size of the original viewport after overscan, custom frustum, etc. is used as
+		// the base size for tiling. (see the FDisplayClusterViewport::UpdateFrameContexts() function)
+
+		// Disables the custom frustum for the tile viewport (it has already been applied to the original viewport).
 		InOutRenderSettings.CustomFrustumSettings = FDisplayClusterViewport_CustomFrustumSettings();
+
+		// Disables the overscan for the tile viewport (it has already been applied to the original viewport).
+		// Override by the overscan settings obtained by splitting the tile.
 		InOutRenderSettings.OverscanSettings = SourceRenderSettings.TileSettings.GetOverscanSettings();
 
 		// Optimize overscan values for edge tiles
@@ -220,29 +195,34 @@ FDisplayClusterViewport* FDisplayClusterViewportConfigurationHelpers_Tile::GetOr
 
 		// Setup as tile.
 		InOutRenderSettings.TileSettings = FDisplayClusterViewport_TileSettings(InSourceViewport.GetId(), InTilePos, InTileSize, TileFlags);
-	}
 
-	return TileViewport;
-}
-
-void FDisplayClusterViewportConfigurationHelpers_Tile::ConfigureTileViewport(FDisplayClusterViewport& InSourceViewport, const FIntPoint& InTilePos)
+		// Copy the other viewport settings:
 {
-	const FDisplayClusterViewport_RenderSettings& SourceRenderSettings = InSourceViewport.GetRenderSettings();
-	if (SourceRenderSettings.TileSettings.GetType() == EDisplayClusterViewportTileType::Source)
-	{
-		if (FDisplayClusterViewport* TileViewport = FindTileViewport(InSourceViewport, InTilePos))
-		{
-		// Copy internal render settings from the source:
+			// Use the OCIO of the original viewport.
+			// In some cases, OCIO may be applied during the post-processing phase of rendering.
+			TileViewport->SetOpenColorIO(InSourceViewport.GetOpenColorIO());
+
+			// The tile viewport uses the same post-processing settings as the original viewport.
 		TileViewport->GetCustomPostProcessSettings() = InSourceViewport.GetCustomPostProcessSettings();
+
+			// The tile viewport uses the same visibility settings as the original viewport.
 		TileViewport->GetVisibilitySettingsImpl() = InSourceViewport.GetVisibilitySettingsImpl();
+
+			// Use the same settings for camera motion blur and DoF.
 		TileViewport->GetCameraMotionBlurImpl() = InSourceViewport.GetCameraMotionBlurImpl();
 		TileViewport->GetCameraDepthOfFieldImpl() = InSourceViewport.GetCameraDepthOfFieldImpl();
-			TileViewport->GetRenderSettingsICVFXImpl() = InSourceViewport.GetRenderSettingsICVFX();
 
-		// Copy OCIO.
-		TileViewport->SetOpenColorIO(InSourceViewport.GetOpenColorIO());
+			// Ignore `OverscanRuntimeSettings` and `CustomFrustumRuntimeSettings`:
+			// These values are calculated later, in FDisplayClusterViewport::UpdateFrameContexts() function
+			// from RenderSettings.CustomFrustumSettings and RenderSettings.OverscanSettings
+
+			// Ignore `RenderSettingsICVFX`: tile viewport exists outside of the ICVFX architecture.
+
+			// Ignore `ViewportRemap` and`PostRenderSettings`: it only applies to the final source viewport.
 	}
 	}
+
+	return TileViewport;
 }
 
 FDisplayClusterViewport_OverscanSettings FDisplayClusterViewportConfigurationHelpers_Tile::GetTileOverscanSettings(const FDisplayClusterConfigurationTile_Overscan& InTileOverscan)
diff --git a/Engine/Plugins/Runtime/nDisplay/Source/DisplayCluster/Private/Render/Viewport/Configuration/DisplayClusterViewportConfigurationHelpers_Tile.h b/Engine/Plugins/Runtime/nDisplay/Source/DisplayCluster/Private/Render/Viewport/Configuration/DisplayClusterViewportConfigurationHelpers_Tile.h
--- a/Engine/Plugins/Runtime/nDisplay/Source/DisplayCluster/Private/Render/Viewport/Configuration/DisplayClusterViewportConfigurationHelpers_Tile.h
+++ b/Engine/Plugins/Runtime/nDisplay/Source/DisplayCluster/Private/Render/Viewport/Configuration/DisplayClusterViewportConfigurationHelpers_Tile.h
@@ -23,9 +23,6 @@ public:
 	/** Get or create a new tile viewport from source. */
 	static FDisplayClusterViewport* GetOrCreateTileViewport(FDisplayClusterViewport& InSourceViewport, const FIntPoint& TilePos);
 
-	/** Configure tile viewport from source. */
-	static void ConfigureTileViewport(FDisplayClusterViewport& InSourceViewport, const FIntPoint& TilePos);
-
 	/** Get unique viewport name for tile. */
 	static FString GetUniqueViewportNameForTile(const FString& InViewportId, const FIntPoint& TilePos);
 
diff --git a/Engine/Plugins/Runtime/nDisplay/Source/DisplayCluster/Private/Render/Viewport/Configuration/DisplayClusterViewportConfiguration_Tile.cpp b/Engine/Plugins/Runtime/nDisplay/Source/DisplayCluster/Private/Render/Viewport/Configuration/DisplayClusterViewportConfiguration_Tile.cpp
--- a/Engine/Plugins/Runtime/nDisplay/Source/DisplayCluster/Private/Render/Viewport/Configuration/DisplayClusterViewportConfiguration_Tile.cpp
+++ b/Engine/Plugins/Runtime/nDisplay/Source/DisplayCluster/Private/Render/Viewport/Configuration/DisplayClusterViewportConfiguration_Tile.cpp
@@ -61,15 +61,6 @@ void FDisplayClusterViewportConfiguration_Tile::Update()
 	ImplFinishReallocateViewports();
 }
 
-void FDisplayClusterViewportConfiguration_Tile::PostUpdate()
-{
-	ForEachTileViewport([](FDisplayClusterViewport& InSourceViewport, const FIntPoint& InTilePos)
-	{
-		// Copy settings for the source viewports to the all tiles.
-		FDisplayClusterViewportConfigurationHelpers_Tile::ConfigureTileViewport(InSourceViewport, InTilePos);
-	});
-}
-
 void FDisplayClusterViewportConfiguration_Tile::ImplBeginReallocateViewports() const
 {
 	if (FDisplayClusterViewportManager* ViewportManager = Configuration.GetViewportManagerImpl())
diff --git a/Engine/Plugins/Runtime/nDisplay/Source/DisplayCluster/Private/Render/Viewport/Configuration/DisplayClusterViewportConfiguration_Tile.h b/Engine/Plugins/Runtime/nDisplay/Source/DisplayCluster/Private/Render/Viewport/Configuration/DisplayClusterViewportConfiguration_Tile.h
--- a/Engine/Plugins/Runtime/nDisplay/Source/DisplayCluster/Private/Render/Viewport/Configuration/DisplayClusterViewportConfiguration_Tile.h
+++ b/Engine/Plugins/Runtime/nDisplay/Source/DisplayCluster/Private/Render/Viewport/Configuration/DisplayClusterViewportConfiguration_Tile.h
@@ -23,9 +23,6 @@ public:
 	/** Update Tile viewports for a new frame. */
 	void Update();
 
-	/** Post-Update Tile viewports for a new frame. */
-	void PostUpdate();
-
 private:
 	/** Mark all Tile viewports as unused before updating. */
 	void ImplBeginReallocateViewports() const;
diff --git a/Engine/Plugins/Runtime/nDisplay/Source/DisplayCluster/Private/Render/Viewport/DisplayClusterViewport_Rules.cpp b/Engine/Plugins/Runtime/nDisplay/Source/DisplayCluster/Private/Render/Viewport/DisplayClusterViewport_Rules.cpp
--- a/Engine/Plugins/Runtime/nDisplay/Source/DisplayCluster/Private/Render/Viewport/DisplayClusterViewport_Rules.cpp
+++ b/Engine/Plugins/Runtime/nDisplay/Source/DisplayCluster/Private/Render/Viewport/DisplayClusterViewport_Rules.cpp
@@ -204,6 +204,12 @@ bool FDisplayClusterViewport::IsRenderEnabledByMedia() const
 
 bool FDisplayClusterViewport::CanSplitIntoTiles() const
 {
+	if (!RenderSettings.bEnable || RenderSettings.bSkipRendering || RenderSettings.bFreezeRendering)
+	{
+		// When this viewport is not rendering, ignore tile splitting.
+		return false;
+	}
+
 	// Ignore internal tile viewports.
 	switch(RenderSettings.TileSettings.GetType())
 	{
diff --git a/Engine/Plugins/Runtime/nDisplay/Source/DisplayClusterConfiguration/Private/DisplayClusterConfigurationTypes_Media.cpp b/Engine/Plugins/Runtime/nDisplay/Source/DisplayClusterConfiguration/Private/DisplayClusterConfigurationTypes_Media.cpp
--- a/Engine/Plugins/Runtime/nDisplay/Source/DisplayClusterConfiguration/Private/DisplayClusterConfigurationTypes_Media.cpp
+++ b/Engine/Plugins/Runtime/nDisplay/Source/DisplayClusterConfiguration/Private/DisplayClusterConfigurationTypes_Media.cpp
@@ -41,6 +41,48 @@ bool FDisplayClusterConfigurationMediaViewport::IsMediaOutputAssigned() const
 ///////////////////////////////////////////////////
 // FDisplayClusterConfigurationMediaICVFX
 
+bool FDisplayClusterConfigurationMediaICVFX::ShouldMediaICVFXSplitIntoTiles() const
+{
+	// Nothing to do if media is not enabled
+	if (!bEnable)
+	{
+		return false;
+	}
+
+	// Nothing to do if media tiling is not used
+	if (SplitType != EDisplayClusterConfigurationMediaSplitType::UniformTiles)
+	{
+		return false;
+	}
+
+	// Validate the layout
+	if (!FDisplayClusterConfigurationTile_Settings::IsValid(TiledSplitLayout))
+	{
+		return false;
+	}
+
+	// Allow the use of tile splitting.
+	return true;
+}
+
+EDisplayClusterViewportTileFlags FDisplayClusterConfigurationMediaICVFX::GetMediaICVFXTileFlags(const FString& NodeId) const
+{
+	// Find if this cluster node is allowed to render unbound tiles
+	const bool bAllowRenderUnbound = ClusterNodesToRenderUnboundTiles.ItemNames.ContainsByPredicate([&NodeId](const FString& Item)
+		{
+			return NodeId.Equals(Item, ESearchCase::IgnoreCase);
+		});
+
+	// Generate flags
+	EDisplayClusterViewportTileFlags OutTileFlags = EDisplayClusterViewportTileFlags::None;
+	if (bAllowRenderUnbound)
+	{
+		OutTileFlags |= EDisplayClusterViewportTileFlags::AllowUnboundRender;
+	}
+
+	return OutTileFlags;
+}
+
 bool FDisplayClusterConfigurationMediaICVFX::HasAnyMediaInputAssigned(const FString& NodeId, EDisplayClusterConfigurationMediaSplitType InSplitType) const
 {
 	// Nothing to do if a different split type requested
diff --git a/Engine/Plugins/Runtime/nDisplay/Source/DisplayClusterConfiguration/Public/DisplayClusterConfigurationTypes_Media.h b/Engine/Plugins/Runtime/nDisplay/Source/DisplayClusterConfiguration/Public/DisplayClusterConfigurationTypes_Media.h
--- a/Engine/Plugins/Runtime/nDisplay/Source/DisplayClusterConfiguration/Public/DisplayClusterConfigurationTypes_Media.h
+++ b/Engine/Plugins/Runtime/nDisplay/Source/DisplayClusterConfiguration/Public/DisplayClusterConfigurationTypes_Media.h
@@ -8,6 +8,8 @@
 #include "DisplayClusterConfigurationTypes_MediaSync.h"
 #include "DisplayClusterConfigurationTypes_Tile.h"
 
+#include "Render/Viewport/Containers/DisplayClusterViewport_TileSettings.h"
+
 #include "MediaPlayer.h"
 #include "MediaSource.h"
 #include "MediaTexture.h"
@@ -305,6 +307,11 @@ public:
 	bool bLateOCIOPass = false;
 
 public:
+	/** Returns true if these media settings can be used for tiling. */
+	bool ShouldMediaICVFXSplitIntoTiles() const;
+
+	/** Returns additional tile settings for the cluster node. */
+	EDisplayClusterViewportTileFlags GetMediaICVFXTileFlags(const FString& NodeId) const;
 
 	/** Returns true if any media source of specified split-type is bound */
 	bool HasAnyMediaInputAssigned(const FString& NodeId, EDisplayClusterConfigurationMediaSplitType SplitType) const;
