nDisplay: Added a per-lightcard option for in-frustum visibility.
- This CL adds a feature that allows both types of lightmaps, ???over??? and ???under???, to be rendered simultaneously.
- reworked ICVFX shader
  * the single camera use case now renders in a single pass.
  * Implemented UVLC per-viewport light card use cases (always under/over).
  * reduced permutations
- FDisplayClusterViewport_VisibilitySettings, FDisplayClusterViewportLightCardManager classes have been reworked.
- An bug in the ADisplayClusterRootActor::GetTypedPrimitives() function has been fixed. When calling CollectPrimitiveComponentsImpl() function, bForceHide was used instead of bCollectChildrenVisualizationComponent.
- Numerous bugs in primitive component collection logic have been fixed. All visibility related logic has been moved to the FDisplayClusterViewportConfigurationHelpers_Visibility class.
- Optimization: Viewports (lightcards) that use an empty ShowList now do not allocate any resources and are not rendered.
- Fixed bug: The LightCardGamma parameter was not set for an outer viewport, resulting in it not being passed to the icvfx shader.

diff --git a/Engine/Plugins/Runtime/nDisplay/Shaders/Private/MPCDIOverlayShaders.usf b/Engine/Plugins/Runtime/nDisplay/Shaders/Private/MPCDIOverlayShaders.usf
--- a/Engine/Plugins/Runtime/nDisplay/Shaders/Private/MPCDIOverlayShaders.usf
+++ b/Engine/Plugins/Runtime/nDisplay/Shaders/Private/MPCDIOverlayShaders.usf
@@ -32,34 +32,137 @@ OutputPS Passthrough_PS(OutputVS IN)
 	return OUT;
 }
 
+float3 ComposeOverlayColor(float4 InColor, float4 InOverlayColor)
+{
+	return InOverlayColor.xyz + InColor.xyz * InOverlayColor.w;
+}
+
+// Additive render pass with 2 overlays
+// a = in\out, b,c = ordered overlays (a + b + c)
+// c.xyz + (b.xyz + a.xyz * b.w)*c.w = a.xyz * (b.w *c.w) + b.xyz*c.w + c.xyz
+float4 ComposeRenderPassTwoOverlays(float4 InOverlay1, float4 InOverlay2)
+{
+	return float4(ComposeOverlayColor(InOverlay1, InOverlay2), InOverlay1.w * InOverlay2.w);
+}
+
 ////////////////////////////////////////////////////////////////////////////////
 // LightCard
 ////////////////////////////////////////////////////////////////////////////////
-float4 GetLightCardTextureColor(float2 InUV)
+float4 NormalizeLightCardColor(float4 LightCardColorLin)
 {
-	float4 OutColor = LightCardTexture.Sample(LightCardSampler, InUV);
-	return float4(pow(OutColor.xyz, LightCardGamma), OutColor.w);
+	return float4(pow(LightCardColorLin.xyz, LightCardGamma), LightCardColorLin.w);
 }
 
-float4 GetUVLightCardTextureColor(float2 InUV)
+float4 NormalizeLightCardColorInvAlpha(float4 LightCardColorLin)
 {
-	float4 OutColor = UVLightCardTexture.Sample(UVLightCardSampler, InUV);
+	return float4(pow(LightCardColorLin.xyz, LightCardGamma), 1 - LightCardColorLin.w);
+}
 
-	// UVLightCard alpha not inverted
-	return float4(pow(OutColor.xyz, LightCardGamma), 1 - OutColor.w);
+float4 GetLightCardUnderTextureColor(float2 InUV)
+{
+	return NormalizeLightCardColor(LightCardUnderTexture.Sample(LightCardUnderSampler, InUV));
 }
 
-float3 ComposeOverlayColor(float4 InColor, float4 InOverlayColor)
+float4 GetLightCardOverTextureColor(float2 InUV)
 {
-	return InOverlayColor.xyz + InColor.xyz * InOverlayColor.w;
+	return NormalizeLightCardColor(LightCardOverTexture.Sample(LightCardOverSampler, InUV));
 }
 
-// Additive render pass with 2 overlays
-// a = in\out, b,c = ordered overlays (a + b + c)
-// c.xyz + (b.xyz + a.xyz * b.w)*c.w = a.xyz * (b.w *c.w) + b.xyz*c.w + c.xyz
-float4 ComposeRenderPassTwoOverlays(float4 InOverlay1, float4 InOverlay2)
+float4 GetUVLightCardUnderTextureColor(float2 InUV)
 {
-	return float4(ComposeOverlayColor(InOverlay1, InOverlay2), InOverlay1.w * InOverlay2.w);
+	float4 UVLightCardUnderColor = NormalizeLightCardColorInvAlpha(UVLightCardUnderTexture.Sample(UVLightCardUnderSampler, InUV));
+	
+	if (LightCardMode & UVLIGHTCARD_MERGE)
+	{
+		// UVLIGHTCARD_MERGE - special rendering mode, for Outer viewports that are forced to render LightCards only in "over" or "under" mode.
+		// both UV lightcard textures are rendered under or over the inner frustum.
+
+		float4 UVLightCardOverColor= NormalizeLightCardColorInvAlpha(UVLightCardOverTexture.Sample(UVLightCardOverSampler, InUV));
+
+		return ComposeRenderPassTwoOverlays(UVLightCardUnderColor, UVLightCardOverColor);
+	}
+
+	return UVLightCardUnderColor;
+}
+
+float4 GetUVLightCardOverTextureColor(float2 InUV)
+{
+	float4 UVLightCardOverColor = NormalizeLightCardColorInvAlpha(UVLightCardOverTexture.Sample(UVLightCardOverSampler, InUV));
+
+	if (LightCardMode & UVLIGHTCARD_MERGE)
+	{
+		// UVLIGHTCARD_MERGE - special rendering mode, for Outer viewports that are forced to render LightCards only in "over" or "under" mode.
+		// both UV lightcard textures are rendered under or over the inner frustum.
+
+		float4 UVLightCardUnderColor = NormalizeLightCardColorInvAlpha(UVLightCardUnderTexture.Sample(UVLightCardUnderSampler, InUV));
+
+		return ComposeRenderPassTwoOverlays(UVLightCardUnderColor, UVLightCardOverColor);
+	}
+
+	return UVLightCardOverColor;
+}
+
+// This function expects that (LightCardMode & (LIGHTCARD_UNDER | UVLIGHTCARD_UNDER)) != 0
+float4 GetLightCardUnderColor(OutputVS IN, float2 OverlayUV)
+{
+// Only the mesh has a UV_Chromakey parameter.
+#if MESH_WARP
+	if( (LightCardMode & (LIGHTCARD_UNDER | UVLIGHTCARD_UNDER)) == (LIGHTCARD_UNDER | UVLIGHTCARD_UNDER))
+	{
+		return ComposeRenderPassTwoOverlays(GetLightCardUnderTextureColor(OverlayUV), GetUVLightCardUnderTextureColor(IN.UV_Chromakey.xy));
+	}
+
+	if(LightCardMode & UVLIGHTCARD_UNDER)
+	{
+		return GetUVLightCardUnderTextureColor(IN.UV_Chromakey.xy);
+	}
+#endif
+
+	return GetLightCardUnderTextureColor(OverlayUV);
+}
+
+
+// This function expects that (LightCardMode & (LIGHTCARD_OVER | UVLIGHTCARD_OVER)) != 0
+float4 GetLightCardOverColor(OutputVS IN, float2 OverlayUV)
+{
+// Only the mesh has a UV_Chromakey parameter.
+#if MESH_WARP
+	if( (LightCardMode & (LIGHTCARD_OVER | UVLIGHTCARD_OVER)) == (LIGHTCARD_OVER | UVLIGHTCARD_OVER))
+	{
+		return ComposeRenderPassTwoOverlays(GetLightCardOverTextureColor(OverlayUV), GetUVLightCardOverTextureColor(IN.UV_Chromakey.xy));
+	}
+
+	if(LightCardMode & UVLIGHTCARD_OVER)
+	{
+		return GetUVLightCardOverTextureColor(IN.UV_Chromakey.xy);
+	}
+#endif
+
+	return GetLightCardOverTextureColor(OverlayUV);
+}
+
+float3 ComposeLightCardUnderColor(float4 InColor, OutputVS IN, float2 OverlayUV)
+{
+	if( LightCardMode & (LIGHTCARD_UNDER | UVLIGHTCARD_UNDER))
+	{
+		float4 LightCardUnderColor = GetLightCardUnderColor(IN, OverlayUV);
+
+		return ComposeOverlayColor(InColor, LightCardUnderColor);
+	}
+	
+	return InColor;
+}
+
+float3 ComposeLightCardOverColor(float4 InColor, OutputVS IN, float2 OverlayUV)
+{
+	if( LightCardMode & (LIGHTCARD_OVER | UVLIGHTCARD_OVER))
+	{
+		float4 LightCardOverColor = GetLightCardOverColor(IN, OverlayUV);
+
+		return ComposeOverlayColor(InColor, LightCardOverColor);
+	}
+	
+	return InColor;
 }
 
 ////////////////////////////////////////////////////////////////////////////////
@@ -298,20 +401,6 @@ OutputPS IcvfxWarpPS(OutputVS IN)
 	float4 OverlayUVW = mul(WarpedUV, OverlayProjectionMatrix);
 	float2 OverlayUV = OverlayUVW.xy / OverlayUVW.w;
 
-	#if LIGHTCARD_OVER || UVLIGHTCARD_OVER || LIGHTCARD_UNDER || UVLIGHTCARD_UNDER
-		#if (LIGHTCARD_OVER && UVLIGHTCARD_OVER) || (LIGHTCARD_UNDER && UVLIGHTCARD_UNDER)
-			float4 LightCardColor = ComposeRenderPassTwoOverlays(GetLightCardTextureColor(OverlayUV), GetUVLightCardTextureColor(IN.UV_Chromakey.xy));
-		#else
-			// Render overlays in separate render-pass
-			#if LIGHTCARD_UNDER || LIGHTCARD_OVER
-				float4 LightCardColor = GetLightCardTextureColor(OverlayUV);
-			#endif
-			#if UVLIGHTCARD_UNDER || UVLIGHTCARD_OVER
-				float4 LightCardColor = GetUVLightCardTextureColor(IN.UV_Chromakey.xy);
-			#endif
-		#endif
-	#endif
-
 	// Compose all:
 #if VIEWPORT_INPUT
 
@@ -323,29 +412,34 @@ OutputPS IcvfxWarpPS(OutputVS IN)
 		float4 OutColor = float4(InputTexture.Sample(InputSampler, ViewportUV).rgb, 1.0f);
 	#endif
 
-	#if LIGHTCARD_UNDER || UVLIGHTCARD_UNDER
-		OutColor.xyz = ComposeOverlayColor(OutColor, LightCardColor);
-	#endif
+	
+	// Compose a “Under” lightcard, if used.
+	OutColor.xyz = ComposeLightCardUnderColor(OutColor, IN, OverlayUV);
 	
 	#if INNER_CAMERA
 		float4 CameraColor = GetInnerCameraColor(WarpedUV, IN);
 		OutColor.xyz = lerp(OutColor.xyz, CameraColor.xyz, CameraColor.w);
 	#endif
 
-	#if LIGHTCARD_OVER || UVLIGHTCARD_OVER
-		OutColor.xyz = ComposeOverlayColor(OutColor, LightCardColor);
-	#endif
+
+	// Compose a “Over” lightcard, if used.
+	OutColor.xyz = ComposeLightCardOverColor(OutColor,IN, OverlayUV);
 
 #else 
+	float4 OutColor;
+
 	// MultiCam or Final renderpass (only cam or overlay per pass))
 	#if INNER_CAMERA
-		float4 OutColor = GetInnerCameraColor(WarpedUV, IN);
+		OutColor = GetInnerCameraColor(WarpedUV, IN);
 	#else
-		#if LIGHTCARD_OVER || UVLIGHTCARD_OVER
-			float4 OutColor = LightCardColor;
-		#else
-			float4 OutColor = float4(0,0,0,0);
-		#endif
+		if( LightCardMode & (LIGHTCARD_OVER | UVLIGHTCARD_OVER))
+		{
+			OutColor = GetLightCardOverColor(IN, OverlayUV);
+		}
+		else
+		{
+			OutColor = float4(0,0,0,0);
+		}
 	#endif
 
 #endif
diff --git a/Engine/Plugins/Runtime/nDisplay/Shaders/Private/MPCDIUtils.ush b/Engine/Plugins/Runtime/nDisplay/Shaders/Private/MPCDIUtils.ush
--- a/Engine/Plugins/Runtime/nDisplay/Shaders/Private/MPCDIUtils.ush
+++ b/Engine/Plugins/Runtime/nDisplay/Shaders/Private/MPCDIUtils.ush
@@ -31,8 +31,10 @@ Texture2D InnerCameraTexture;
 Texture2D ChromakeyCameraTexture;
 Texture2D ChromakeyMarkerTexture;
 
-Texture2D LightCardTexture;
-Texture2D UVLightCardTexture;
+Texture2D LightCardUnderTexture;
+Texture2D LightCardOverTexture;
+Texture2D UVLightCardUnderTexture;
+Texture2D UVLightCardOverTexture;
 
 SamplerState InputSampler;
 SamplerState WarpMapSampler;
@@ -43,8 +45,10 @@ SamplerState InnerCameraSampler;
 SamplerState ChromakeyCameraSampler;
 SamplerState ChromakeyMarkerSampler;
 
-SamplerState LightCardSampler;
-SamplerState UVLightCardSampler;
+SamplerState LightCardUnderSampler;
+SamplerState LightCardOverSampler;
+SamplerState UVLightCardUnderSampler;
+SamplerState UVLightCardOverSampler;
 
 float4x4 MeshToStageProjectionMatrix;
 float4x4 ViewportTextureProjectionMatrix;
@@ -66,6 +70,29 @@ float4 OverlappedInnerCameraSoftEdges[MAX_INNER_CAMERAS_AMOUNT];
 
 float LightCardGamma;
 
+// These are the flags for the 'LightCardMode' parameter:
+
+// LIGHTCARD_UNDER - LightCard is rendered under InCamera. The image from the "LightCardUnderTexture" texture is used.
+#define LIGHTCARD_UNDER   (1 << 0)
+
+// LIGHTCARD_OVER - LightCard is rendered over InCamera. The image from the "LightCardOverTexture" texture is used.
+#define LIGHTCARD_OVER    (1 << 1)
+
+// UVLIGHTCARD_UNDER - UV LightCard is rendered under InCamera. The image from the "UVLightCardUnderTexture" texture is used.
+#define UVLIGHTCARD_UNDER (1 << 2)
+
+// UVLIGHTCARD_OVER - UV LightCard is rendered over InCamera. The image from the "UVLightCardOverTexture" texture is used.
+#define UVLIGHTCARD_OVER  (1 << 3)
+
+// UVLIGHTCARD_MERGE - special rendering mode, for Outer viewports that are forced to render LightCards only in "over" or "under" mode.
+// both UV lightcard textures are rendered under or over the inner frustum.
+// Images from the "UVLightCardOverTexture" and "UVLightCardOverTexture" textures are used.
+// Can only be combined in two ways: (UVLIGHTCARD_MERGE|UVLIGHTCARD_UNDER) or (UVLIGHTCARD_MERGE|UVLIGHTCARD_OVER)
+#define UVLIGHTCARD_MERGE (1 << 4)
+
+// Several flags in this parameter: LIGHTCARD_UNDER, LIGHTCARD_OVER, UVLIGHTCARD_UNDER, UVLIGHTCARD_OVER, UVLIGHTCARD_MERGE
+int LightCardMode;
+
 float4 InnerCameraSoftEdge;
 
 float4 InnerCameraBorderColor;
diff --git a/Engine/Plugins/Runtime/nDisplay/Source/DisplayCluster/Private/Game/EngineClasses/Scene/DisplayClusterLightCardActor.cpp b/Engine/Plugins/Runtime/nDisplay/Source/DisplayCluster/Private/Game/EngineClasses/Scene/DisplayClusterLightCardActor.cpp
--- a/Engine/Plugins/Runtime/nDisplay/Source/DisplayCluster/Private/Game/EngineClasses/Scene/DisplayClusterLightCardActor.cpp
+++ b/Engine/Plugins/Runtime/nDisplay/Source/DisplayCluster/Private/Game/EngineClasses/Scene/DisplayClusterLightCardActor.cpp
@@ -59,6 +59,7 @@ ADisplayClusterLightCardActor::ADisplayClusterLightCardActor(const FObjectInitia
 	, Scale(FVector2D(1.f))
 	, RadialOffset(-1)
 	, bAlwaysFlushToWall(true)
+	, PerLightcardRenderMode(EDisplayClusterConfigurationICVFX_PerLightcardRenderMode::Default)
 	, Mask(EDisplayClusterLightCardMask::Circle)
 	, Texture(nullptr)
 	, Color(FLinearColor(1.f, 1.f, 1.f, 1.f))
diff --git a/Engine/Plugins/Runtime/nDisplay/Source/DisplayCluster/Private/Game/EngineClasses/Scene/DisplayClusterRootActor.cpp b/Engine/Plugins/Runtime/nDisplay/Source/DisplayCluster/Private/Game/EngineClasses/Scene/DisplayClusterRootActor.cpp
--- a/Engine/Plugins/Runtime/nDisplay/Source/DisplayCluster/Private/Game/EngineClasses/Scene/DisplayClusterRootActor.cpp
+++ b/Engine/Plugins/Runtime/nDisplay/Source/DisplayCluster/Private/Game/EngineClasses/Scene/DisplayClusterRootActor.cpp
@@ -56,10 +56,17 @@
 
 namespace UE::DisplayCluster::RootActor
 {
+	/** Collect components by type
+	* 
+	* @param OutPrimitives - (out) Output variable for collected components.
+	* @param pComp         - (in) The component from which the collection will be made.
+	* @param bForceHide    - (opt) if true, collects components with any value of the bHiddenInGame property.
+	* @param bCollectChildrenVisualizationComponent - If true collects child components.
+	*/
 	template <typename TComp>
-	void CollectPrimitiveComponentsImpl(TSet<FPrimitiveComponentId>& OutPrimitives, TComp* pComp, bool bForceHide = false, const bool bCollectChildrenVisualizationComponent = true)
+	void CollectPrimitiveComponentsImpl(TSet<FPrimitiveComponentId>& OutPrimitives, const TComp* pComp, bool bForceHide = false, const bool bCollectChildrenVisualizationComponent = true)
 	{
-		if (UPrimitiveComponent* PrimComp = Cast<UPrimitiveComponent>(pComp))
+		if (const UPrimitiveComponent* PrimComp = Cast<const UPrimitiveComponent>(pComp))
 		{
 			if (PrimComp->bHiddenInGame
 			|| bForceHide
@@ -68,18 +75,19 @@ namespace UE::DisplayCluster::RootActor
 #endif
 				)
 			{
-				OutPrimitives.Add(PrimComp->GetPrimitiveSceneId());
+				const FPrimitiveComponentId CompId = PrimComp->GetPrimitiveSceneId();
+				OutPrimitives.Add(CompId);
 			}
 		}
 
 		if (bCollectChildrenVisualizationComponent)
 		{
-			if (USceneComponent* SceneComp = Cast<USceneComponent>(pComp))
+			if (const USceneComponent* SceneComp = Cast<const USceneComponent>(pComp))
 			{
 				TArray<USceneComponent*> ChildrenComponents;
 				SceneComp->GetChildrenComponents(false, ChildrenComponents);
 
-				for (USceneComponent* CompIt : ChildrenComponents)
+				for (const USceneComponent* CompIt : ChildrenComponents)
 				{
 					CollectPrimitiveComponentsImpl(OutPrimitives, CompIt, bForceHide, bCollectChildrenVisualizationComponent);
 				}
@@ -612,7 +620,7 @@ int ADisplayClusterRootActor::GetInnerFrustumPriority(const FString& InnerFrustu
 }
 
 template <typename TComp>
-void ADisplayClusterRootActor::GetTypedPrimitives(TSet<FPrimitiveComponentId>& OutPrimitives, const TArray<FString>* InCompNames, bool bCollectChildrenVisualizationComponent) const
+void ADisplayClusterRootActor::GetTypedPrimitives(TSet<FPrimitiveComponentId>& OutPrimitives, const TArray<FString>* InCompNames, bool bForceHide, bool bCollectChildrenVisualizationComponent) const
 {
 	using namespace UE::DisplayCluster::RootActor;
 
@@ -628,22 +636,22 @@ void ADisplayClusterRootActor::GetTypedPrimitives(TSet<FPrimitiveComponentId>& O
 				if (InCompNames->Find(CompIt->GetName()) != INDEX_NONE)
 				{
 					// add only comp from names list
-					CollectPrimitiveComponentsImpl(OutPrimitives, CompIt, bCollectChildrenVisualizationComponent);
+					CollectPrimitiveComponentsImpl(OutPrimitives, CompIt, bForceHide, bCollectChildrenVisualizationComponent);
 				}
 			}
 			else
 			{
-				CollectPrimitiveComponentsImpl(OutPrimitives, CompIt, bCollectChildrenVisualizationComponent);
+				CollectPrimitiveComponentsImpl(OutPrimitives, CompIt, bForceHide, bCollectChildrenVisualizationComponent);
 			}
 		}
 	}
 }
 
-bool ADisplayClusterRootActor::FindPrimitivesByName(const TArray<FString>& InNames, TSet<FPrimitiveComponentId>& OutPrimitives)
+bool ADisplayClusterRootActor::FindPrimitivesByName(const TArray<FString>& InNames, TSet<FPrimitiveComponentId>& OutPrimitives, bool bForceHide)
 {
-	GetTypedPrimitives<UActorComponent>(OutPrimitives, &InNames, false);
+	GetTypedPrimitives<UActorComponent>(OutPrimitives, &InNames, bForceHide, false);
 
-	return true;
+	return !OutPrimitives.IsEmpty();
 }
 
 // Gather components not rendered in game
diff --git a/Engine/Plugins/Runtime/nDisplay/Source/DisplayCluster/Private/Render/Viewport/Configuration/DisplayClusterViewportConfigurationHelpers.cpp b/Engine/Plugins/Runtime/nDisplay/Source/DisplayCluster/Private/Render/Viewport/Configuration/DisplayClusterViewportConfigurationHelpers.cpp
--- a/Engine/Plugins/Runtime/nDisplay/Source/DisplayCluster/Private/Render/Viewport/Configuration/DisplayClusterViewportConfigurationHelpers.cpp
+++ b/Engine/Plugins/Runtime/nDisplay/Source/DisplayCluster/Private/Render/Viewport/Configuration/DisplayClusterViewportConfigurationHelpers.cpp
@@ -88,10 +88,6 @@ void FDisplayClusterViewportConfigurationHelpers::UpdateBaseViewportSetting(FDis
 {
 	// Gain direct access to internal settings of the viewport:
 	FDisplayClusterViewport_RenderSettings&           InOutRenderSettings = DstViewport.GetRenderSettingsImpl();
-	FDisplayClusterViewport_RenderSettingsICVFX& InOutRenderSettingsICVFX = DstViewport.GetRenderSettingsICVFXImpl();
-
-	// Reset runtime flags from prev frame:
-	DstViewport.ResetRuntimeParameters();
 
 	// UDisplayClusterConfigurationViewport
 	{
@@ -137,11 +133,4 @@ void FDisplayClusterViewportConfigurationHelpers::UpdateBaseViewportSetting(FDis
 		InOutRenderSettings.StereoGPUIndex = InRenderSettings.StereoGPUIndex;
 		InOutRenderSettings.RenderTargetRatio = InRenderSettings.RenderTargetRatio;
 	}
-
-	// FDisplayClusterConfigurationViewport_ICVFX property:
-	if(const FDisplayClusterConfigurationICVFX_StageSettings* StageSettings = DstViewport.Configuration->GetStageSettings())
-	{
-		InOutRenderSettingsICVFX.Flags = InConfigurationViewport.ICVFX.GetViewportICVFXFlags(*StageSettings);
-		InOutRenderSettingsICVFX.ICVFX.LightCardMode = InConfigurationViewport.ICVFX.GetLightCardRenderMode(*StageSettings);
-	}
 }
diff --git a/Engine/Plugins/Runtime/nDisplay/Source/DisplayCluster/Private/Render/Viewport/Configuration/DisplayClusterViewportConfigurationHelpers_ICVFX.cpp b/Engine/Plugins/Runtime/nDisplay/Source/DisplayCluster/Private/Render/Viewport/Configuration/DisplayClusterViewportConfigurationHelpers_ICVFX.cpp
--- a/Engine/Plugins/Runtime/nDisplay/Source/DisplayCluster/Private/Render/Viewport/Configuration/DisplayClusterViewportConfigurationHelpers_ICVFX.cpp
+++ b/Engine/Plugins/Runtime/nDisplay/Source/DisplayCluster/Private/Render/Viewport/Configuration/DisplayClusterViewportConfigurationHelpers_ICVFX.cpp
@@ -49,13 +49,13 @@
 namespace UE::DisplayCluster::Viewport::ConfigurationHelpers_ICVFX
 {
 	// Return unique ICVFX name
-	static FString ImplGetNameICVFX(const FString& InClusterNodeId, const FString& InViewportId, const FString& InResourceId)
+	static FString ImplGetNameICVFX(const FString& InClusterNodeId, const FString& InViewportId, const FString& InResourceId, const FString& InResourceSuffixId)
 	{
 		check(!InClusterNodeId.IsEmpty());
 		check(!InViewportId.IsEmpty());
 		check(!InResourceId.IsEmpty());
 
-		return FString::Printf(TEXT("%s_%s_%s_%s"), *InClusterNodeId, DisplayClusterViewportStrings::icvfx::prefix, *InViewportId, *InResourceId);
+		return FString::Printf(TEXT("%s_%s_%s_%s%s"), *InClusterNodeId, DisplayClusterViewportStrings::icvfx::prefix, *InViewportId, *InResourceId, *InResourceSuffixId);
 	}
 };
 using namespace UE::DisplayCluster::Viewport::ConfigurationHelpers_ICVFX;
@@ -63,7 +63,7 @@ using namespace UE::DisplayCluster::Viewport::ConfigurationHelpers_ICVFX;
 ////////////////////////////////////////////////////////////////////////////////
 // FDisplayClusterViewportConfigurationHelpers_ICVFX
 ////////////////////////////////////////////////////////////////////////////////
-bool FDisplayClusterViewportConfigurationHelpers_ICVFX::CreateProjectionPolicyICVFX(FDisplayClusterViewportConfiguration& InConfiguration, const FString& InViewportId, const FString& InResourceId, bool bIsCameraProjection, TSharedPtr<IDisplayClusterProjectionPolicy, ESPMode::ThreadSafe>& OutProjPolicy)
+bool FDisplayClusterViewportConfigurationHelpers_ICVFX::CreateProjectionPolicyICVFX(FDisplayClusterViewportConfiguration& InConfiguration, const FString& InViewportId, const FString& InResourceId, const FString& InResourceSuffixId, bool bIsCameraProjection, TSharedPtr<IDisplayClusterProjectionPolicy, ESPMode::ThreadSafe>& OutProjPolicy)
 {
 	const FString& ClusterNodeId = InConfiguration.GetClusterNodeId();
 	if (ClusterNodeId.IsEmpty())
@@ -75,11 +75,11 @@ bool FDisplayClusterViewportConfigurationHelpers_ICVFX::CreateProjectionPolicyIC
 	CameraProjectionPolicyConfig.Type = bIsCameraProjection ? DisplayClusterProjectionStrings::projection::Camera : DisplayClusterProjectionStrings::projection::Link;
 
 	// Create projection policy for viewport
-	OutProjPolicy = FDisplayClusterViewportManager::CreateProjectionPolicy(ImplGetNameICVFX(ClusterNodeId, InViewportId, InResourceId), &CameraProjectionPolicyConfig);
+	OutProjPolicy = FDisplayClusterViewportManager::CreateProjectionPolicy(ImplGetNameICVFX(ClusterNodeId, InViewportId, InResourceId, InResourceSuffixId), &CameraProjectionPolicyConfig);
 
 	if (!OutProjPolicy.IsValid())
 	{
-		UE_LOG(LogDisplayClusterViewport, Error, TEXT("ICVFX Viewport '%s': projection policy for resource '%s' not created for node '%s'."), *InViewportId, *InResourceId, *ClusterNodeId);
+		UE_LOG(LogDisplayClusterViewport, Error, TEXT("ICVFX Viewport '%s': projection policy for resource '%s%s' not created for node '%s'."), *InViewportId, *InResourceId, *InResourceSuffixId,*ClusterNodeId);
 		return false;
 	}
 
@@ -111,15 +111,15 @@ bool FDisplayClusterViewportConfigurationHelpers_ICVFX::UpdateCameraProjectionSe
 
 bool FDisplayClusterViewportConfigurationHelpers_ICVFX::CreateProjectionPolicyCameraICVFX(FDisplayClusterViewportConfiguration& InConfiguration, UDisplayClusterICVFXCameraComponent& InCameraComponent, const FDisplayClusterConfigurationICVFX_CameraSettings& InCameraSettings, TSharedPtr<IDisplayClusterProjectionPolicy, ESPMode::ThreadSafe>& OutProjPolicy)
 {
-	return CreateProjectionPolicyICVFX(InConfiguration, InCameraComponent.GetCameraUniqueId(), DisplayClusterViewportStrings::icvfx::camera, true, OutProjPolicy)
+	return CreateProjectionPolicyICVFX(InConfiguration, InCameraComponent.GetCameraUniqueId(), DisplayClusterViewportStrings::icvfx::camera, DisplayClusterViewportStrings::icvfx::suffix::none, true, OutProjPolicy)
 		&& UpdateCameraProjectionSettingsICVFX(InConfiguration, InCameraComponent, InCameraSettings, OutProjPolicy);
 }
 
-FDisplayClusterViewport* FDisplayClusterViewportConfigurationHelpers_ICVFX::FindViewportICVFX(FDisplayClusterViewportConfiguration& InConfiguration, const FString& InViewportId, const FString& InResourceId)
+FDisplayClusterViewport* FDisplayClusterViewportConfigurationHelpers_ICVFX::FindViewportICVFX(FDisplayClusterViewportConfiguration& InConfiguration, const FString& InViewportId, const FString& InResourceId, const FString& InResourceSuffixId)
 {
 	if (FDisplayClusterViewportManager* ViewportManager = InConfiguration.GetViewportManagerImpl())
 	{
-		TSharedPtr<FDisplayClusterViewport, ESPMode::ThreadSafe> Viewport = ViewportManager->ImplFindViewport(ImplGetNameICVFX(InConfiguration.GetClusterNodeId(), InViewportId, InResourceId));
+		TSharedPtr<FDisplayClusterViewport, ESPMode::ThreadSafe> Viewport = ViewportManager->ImplFindViewport(ImplGetNameICVFX(InConfiguration.GetClusterNodeId(), InViewportId, InResourceId, InResourceSuffixId));
 
 		return Viewport.Get();
 	}
@@ -127,14 +127,14 @@ FDisplayClusterViewport* FDisplayClusterViewportConfigurationHelpers_ICVFX::Find
 	return nullptr;
 }
 
-FDisplayClusterViewport* FDisplayClusterViewportConfigurationHelpers_ICVFX::CreateViewportICVFX(FDisplayClusterViewportConfiguration& InConfiguration, const FString& InViewportId, const FString& InResourceId, TSharedPtr<IDisplayClusterProjectionPolicy, ESPMode::ThreadSafe>& InProjectionPolicy)
+FDisplayClusterViewport* FDisplayClusterViewportConfigurationHelpers_ICVFX::CreateViewportICVFX(FDisplayClusterViewportConfiguration& InConfiguration, const FString& InViewportId, const FString& InResourceId, const FString& InResourceSuffixId, TSharedPtr<IDisplayClusterProjectionPolicy, ESPMode::ThreadSafe>& InProjectionPolicy)
 {
 	check(InProjectionPolicy.IsValid());
 
 	if (FDisplayClusterViewportManager* ViewportManager = InConfiguration.GetViewportManagerImpl())
 	{
 		// Create viewport for new projection policy
-		TSharedPtr<FDisplayClusterViewport, ESPMode::ThreadSafe> NewViewport = ViewportManager->ImplCreateViewport(ImplGetNameICVFX(InConfiguration.GetClusterNodeId(), InViewportId, InResourceId), InProjectionPolicy);
+		TSharedPtr<FDisplayClusterViewport, ESPMode::ThreadSafe> NewViewport = ViewportManager->ImplCreateViewport(ImplGetNameICVFX(InConfiguration.GetClusterNodeId(), InViewportId, InResourceId, InResourceSuffixId), InProjectionPolicy);
 		if (NewViewport.IsValid())
 		{
 			// Gain direct access to internal resources of the NewViewport:
@@ -156,7 +156,7 @@ FDisplayClusterViewport* FDisplayClusterViewportConfigurationHelpers_ICVFX::Crea
 
 FDisplayClusterViewport* FDisplayClusterViewportConfigurationHelpers_ICVFX::FindCameraViewport(FDisplayClusterViewportConfiguration& InConfiguration, UDisplayClusterICVFXCameraComponent& InCameraComponent)
 {
-	return FindViewportICVFX(InConfiguration, InCameraComponent.GetCameraUniqueId(), DisplayClusterViewportStrings::icvfx::camera);
+	return FindViewportICVFX(InConfiguration, InCameraComponent.GetCameraUniqueId(), DisplayClusterViewportStrings::icvfx::camera, DisplayClusterViewportStrings::icvfx::suffix::none);
 }
 
 FDisplayClusterViewport* FDisplayClusterViewportConfigurationHelpers_ICVFX::GetOrCreateCameraViewport(FDisplayClusterViewportConfiguration& InConfiguration, UDisplayClusterICVFXCameraComponent& InCameraComponent, const FDisplayClusterConfigurationICVFX_CameraSettings& InCameraSettings)
@@ -171,13 +171,16 @@ FDisplayClusterViewport* FDisplayClusterViewportConfigurationHelpers_ICVFX::GetO
 			return nullptr;
 		}
 
-		CameraViewport = CreateViewportICVFX(InConfiguration, InCameraComponent.GetCameraUniqueId(), DisplayClusterViewportStrings::icvfx::camera, CameraProjectionPolicy);
+		CameraViewport = CreateViewportICVFX(InConfiguration, InCameraComponent.GetCameraUniqueId(), DisplayClusterViewportStrings::icvfx::camera, DisplayClusterViewportStrings::icvfx::suffix::none, CameraProjectionPolicy);
 		if (CameraViewport == nullptr)
 		{
 			return nullptr;
 		}
 	}
 
+	// Reset runtime flags from prev frame:
+	CameraViewport->ResetRuntimeParameters();
+
 	// Gain direct access to internal resources of the CameraViewport:
 	FDisplayClusterViewport_RenderSettingsICVFX& InOutRenderSettingsICVFX = CameraViewport->GetRenderSettingsICVFXImpl();
 
@@ -190,14 +193,18 @@ FDisplayClusterViewport* FDisplayClusterViewportConfigurationHelpers_ICVFX::GetO
 	return CameraViewport;
 }
 
-void FDisplayClusterViewportConfigurationHelpers_ICVFX::ReuseUVLightCardViewportWithinClusterNode(FDisplayClusterViewport& InUVLightCardViewport)
+void FDisplayClusterViewportConfigurationHelpers_ICVFX::ReuseUVLightCardViewportWithinClusterNode(FDisplayClusterViewport& InUVLightCardViewport, const bool bOverInFrustum)
 {
 	if (FDisplayClusterViewportManager* ViewportManager = InUVLightCardViewport.Configuration->GetViewportManagerImpl())
 	{
+		// Filter lightcard viewport type
+		EDisplayClusterViewportRuntimeICVFXFlags UVLightcardViewportFlags = EDisplayClusterViewportRuntimeICVFXFlags::UVLightcard;
+		UVLightcardViewportFlags |= bOverInFrustum ? EDisplayClusterViewportRuntimeICVFXFlags::OverInFrustum : EDisplayClusterViewportRuntimeICVFXFlags::UnderInFrustum;
+
 		for (const TSharedPtr<FDisplayClusterViewport, ESPMode::ThreadSafe>& ViewportIt : ViewportManager->ImplGetEntireClusterViewports())
 		{
 			if (ViewportIt.IsValid() && ViewportIt != InUVLightCardViewport.AsShared() && !ViewportIt->GetRenderSettings().IsViewportOverridden()
-				&& EnumHasAnyFlags(ViewportIt->GetRenderSettingsICVFX().RuntimeFlags, EDisplayClusterViewportRuntimeICVFXFlags::UVLightcard))
+				&& EnumHasAllFlags(ViewportIt->GetRenderSettingsICVFX().RuntimeFlags, UVLightcardViewportFlags))
 			{
 				if (ViewportIt->IsOpenColorIOEquals(InUVLightCardViewport))
 				{
@@ -258,7 +265,7 @@ TArray<TSharedPtr<FDisplayClusterViewport, ESPMode::ThreadSafe>> FDisplayCluster
 				&& (!ViewportIt->GetRenderSettings().IsViewportOverridden()))
 			{
 				// this is incamera viewport. Check by name
-				const FString RequiredViewportId = ImplGetNameICVFX(ViewportIt->GetClusterNodeId(), InICVFXCameraId, ViewportTypeId);
+				const FString RequiredViewportId = ImplGetNameICVFX(ViewportIt->GetClusterNodeId(), InICVFXCameraId, ViewportTypeId, DisplayClusterViewportStrings::icvfx::suffix::none);
 				if (RequiredViewportId.Equals(ViewportIt->GetId()))
 				{
 					OutViewports.Add(ViewportIt);
@@ -329,24 +336,27 @@ void FDisplayClusterViewportConfigurationHelpers_ICVFX::PreviewReuseChromakeyVie
 FDisplayClusterViewport* FDisplayClusterViewportConfigurationHelpers_ICVFX::GetOrCreateChromakeyViewport(FDisplayClusterViewportConfiguration& InConfiguration, UDisplayClusterICVFXCameraComponent& InCameraComponent, const FDisplayClusterConfigurationICVFX_CameraSettings& InCameraSettings)
 {
 	const FString InICVFXCameraId = InCameraComponent.GetCameraUniqueId();
-	FDisplayClusterViewport* ChromakeyViewport = FindViewportICVFX(InConfiguration, InICVFXCameraId, DisplayClusterViewportStrings::icvfx::chromakey);
+	FDisplayClusterViewport* ChromakeyViewport = FindViewportICVFX(InConfiguration, InICVFXCameraId, DisplayClusterViewportStrings::icvfx::chromakey, DisplayClusterViewportStrings::icvfx::suffix::none);
 
 	// Create new chromakey viewport
 	if (ChromakeyViewport == nullptr)
 	{
 		TSharedPtr<IDisplayClusterProjectionPolicy, ESPMode::ThreadSafe> ChromakeyProjectionPolicy;
-		if (!CreateProjectionPolicyICVFX(InConfiguration, InICVFXCameraId, DisplayClusterViewportStrings::icvfx::chromakey, false, ChromakeyProjectionPolicy))
+		if (!CreateProjectionPolicyICVFX(InConfiguration, InICVFXCameraId, DisplayClusterViewportStrings::icvfx::chromakey, DisplayClusterViewportStrings::icvfx::suffix::none, false, ChromakeyProjectionPolicy))
 		{
 			return nullptr;
 		}
 
-		ChromakeyViewport = CreateViewportICVFX(InConfiguration, InICVFXCameraId, DisplayClusterViewportStrings::icvfx::chromakey, ChromakeyProjectionPolicy);
+		ChromakeyViewport = CreateViewportICVFX(InConfiguration, InICVFXCameraId, DisplayClusterViewportStrings::icvfx::chromakey, DisplayClusterViewportStrings::icvfx::suffix::none, ChromakeyProjectionPolicy);
 		if (ChromakeyViewport == nullptr)
 		{
 			return nullptr;
 		}
 	}
 
+	// Reset runtime flags from prev frame:
+	ChromakeyViewport->ResetRuntimeParameters();
+
 	// Gain direct access to internal resources of the viewport:
 	FDisplayClusterViewport_RenderSettingsICVFX& InOutRenderSettingsICVFX = ChromakeyViewport->GetRenderSettingsICVFXImpl();
 
@@ -359,27 +369,34 @@ FDisplayClusterViewport* FDisplayClusterViewportConfigurationHelpers_ICVFX::GetO
 	return ChromakeyViewport;
 }
 
-FDisplayClusterViewport* FDisplayClusterViewportConfigurationHelpers_ICVFX::GetOrCreateLightcardViewport(FDisplayClusterViewport& BaseViewport)
+FDisplayClusterViewport* FDisplayClusterViewportConfigurationHelpers_ICVFX::GetOrCreateLightcardViewport(FDisplayClusterViewport& BaseViewport, const bool bOverInFrustum)
 {
 	// Create new lightcard viewport
 	const FString ResourceId = DisplayClusterViewportStrings::icvfx::lightcard;
+	const FString ResourceSuffixId = bOverInFrustum ? DisplayClusterViewportStrings::icvfx::suffix::over : DisplayClusterViewportStrings::icvfx::suffix::under;
 
-	FDisplayClusterViewport* LightcardViewport = FindViewportICVFX(*BaseViewport.Configuration, BaseViewport.GetId(), ResourceId);
+	FDisplayClusterViewport* LightcardViewport = FindViewportICVFX(*BaseViewport.Configuration, BaseViewport.GetId(), ResourceId, ResourceSuffixId);
 	if (LightcardViewport == nullptr)
 	{
 		TSharedPtr<IDisplayClusterProjectionPolicy, ESPMode::ThreadSafe> LightcardProjectionPolicy;
-		if (!CreateProjectionPolicyICVFX(*BaseViewport.Configuration, BaseViewport.GetId(), ResourceId, false, LightcardProjectionPolicy))
+		if (!CreateProjectionPolicyICVFX(*BaseViewport.Configuration, BaseViewport.GetId(), ResourceId, ResourceSuffixId, false, LightcardProjectionPolicy))
 		{
 			return nullptr;
 		}
 
-		LightcardViewport = CreateViewportICVFX(*BaseViewport.Configuration, BaseViewport.GetId(), ResourceId, LightcardProjectionPolicy);
+		LightcardViewport = CreateViewportICVFX(*BaseViewport.Configuration, BaseViewport.GetId(), ResourceId, ResourceSuffixId, LightcardProjectionPolicy);
 		if (LightcardViewport == nullptr)
 		{
 			return nullptr;
 		}
 	}
 
+	// Reset runtime flags from prev frame:
+	LightcardViewport->ResetRuntimeParameters(BaseViewport.GetViewportConfigurationData());
+
+	// Also updates the reference to the configuration data of the base viewport.
+	LightcardViewport->SetViewportConfigurationData(BaseViewport.GetViewportConfigurationData());
+
 	// Gain direct access to internal resources of the viewport:
 	FDisplayClusterViewport_RenderSettingsICVFX& InOutRenderSettingsICVFX = LightcardViewport->GetRenderSettingsICVFXImpl();
 
@@ -389,30 +406,37 @@ FDisplayClusterViewport* FDisplayClusterViewportConfigurationHelpers_ICVFX::GetO
 	// Add viewport ICVFX usage as Lightcard
 	EnumAddFlags(InOutRenderSettingsICVFX.RuntimeFlags, EDisplayClusterViewportRuntimeICVFXFlags::Lightcard);
 
+	// Add LC viewport position flag
+	EnumAddFlags(InOutRenderSettingsICVFX.RuntimeFlags, bOverInFrustum ? EDisplayClusterViewportRuntimeICVFXFlags::OverInFrustum : EDisplayClusterViewportRuntimeICVFXFlags::UnderInFrustum);
+
 	return LightcardViewport;
 }
 
-FDisplayClusterViewport* FDisplayClusterViewportConfigurationHelpers_ICVFX::GetOrCreateUVLightcardViewport(FDisplayClusterViewport& BaseViewport)
+FDisplayClusterViewport* FDisplayClusterViewportConfigurationHelpers_ICVFX::GetOrCreateUVLightcardViewport(FDisplayClusterViewport& BaseViewport, const bool bOverInFrustum)
 {
 	// Create new lightcard viewport
 	const FString ResourceId = DisplayClusterViewportStrings::icvfx::uv_lightcard;
+	const FString ResourceSuffixId = bOverInFrustum ? DisplayClusterViewportStrings::icvfx::suffix::over : DisplayClusterViewportStrings::icvfx::suffix::under;
 
-	FDisplayClusterViewport* UVLightcardViewport = FindViewportICVFX(*BaseViewport.Configuration, BaseViewport.GetId(), ResourceId);
+	FDisplayClusterViewport* UVLightcardViewport = FindViewportICVFX(*BaseViewport.Configuration, BaseViewport.GetId(), ResourceId, ResourceSuffixId);
 	if (UVLightcardViewport == nullptr)
 	{
 		TSharedPtr<IDisplayClusterProjectionPolicy, ESPMode::ThreadSafe> UVLightcardProjectionPolicy;
-		if (!CreateProjectionPolicyICVFX(*BaseViewport.Configuration, BaseViewport.GetId(), ResourceId, false, UVLightcardProjectionPolicy))
+		if (!CreateProjectionPolicyICVFX(*BaseViewport.Configuration, BaseViewport.GetId(), ResourceId, ResourceSuffixId, false, UVLightcardProjectionPolicy))
 		{
 			return nullptr;
 		}
 
-		UVLightcardViewport = CreateViewportICVFX(*BaseViewport.Configuration, BaseViewport.GetId(), ResourceId, UVLightcardProjectionPolicy);
+		UVLightcardViewport = CreateViewportICVFX(*BaseViewport.Configuration, BaseViewport.GetId(), ResourceId, ResourceSuffixId, UVLightcardProjectionPolicy);
 		if (UVLightcardViewport == nullptr)
 		{
 			return nullptr;
 		}
 	}
 
+	// Also updates the reference to the configuration data of the base viewport.
+	UVLightcardViewport->SetViewportConfigurationData(BaseViewport.GetViewportConfigurationData());
+
 	// Gain direct access to internal resources of the viewport:
 	FDisplayClusterViewport_RenderSettingsICVFX& InOutRenderSettingsICVFX = UVLightcardViewport->GetRenderSettingsICVFXImpl();
 
@@ -422,6 +446,9 @@ FDisplayClusterViewport* FDisplayClusterViewportConfigurationHelpers_ICVFX::GetO
 	// Add viewport ICVFX usage as Lightcard
 	EnumAddFlags(InOutRenderSettingsICVFX.RuntimeFlags, EDisplayClusterViewportRuntimeICVFXFlags::UVLightcard);
 
+	// Add LC viewport position flag
+	EnumAddFlags(InOutRenderSettingsICVFX.RuntimeFlags, bOverInFrustum ? EDisplayClusterViewportRuntimeICVFXFlags::OverInFrustum : EDisplayClusterViewportRuntimeICVFXFlags::UnderInFrustum);
+
 	return UVLightcardViewport;
 }
 
@@ -430,7 +457,7 @@ bool FDisplayClusterViewportConfigurationHelpers_ICVFX::IsCameraUsed(const FDisp
 	// Check rules for camera settings:
 	if (InCameraSettings.bEnable == false)
 	{
-		// dont use camera if disabled
+		// don't use camera if disabled
 		return false;
 	}
 
@@ -485,9 +512,6 @@ void FDisplayClusterViewportConfigurationHelpers_ICVFX::UpdateCameraViewportSett
 	// Gain direct access to internal resources of the viewport:
 	FDisplayClusterViewport_RenderSettings& InOutRenderSettings = DstViewport.GetRenderSettingsImpl();
 
-	// Reset runtime flags from prev frame:
-	DstViewport.ResetRuntimeParameters();
-
 	// incamera textrure used as overlay
 	InOutRenderSettings.bVisible = false;
 
@@ -557,9 +581,6 @@ void FDisplayClusterViewportConfigurationHelpers_ICVFX::UpdateChromakeyViewportS
 	FDisplayClusterViewport_RenderSettings&               InOutRenderSettings = DstViewport.GetRenderSettingsImpl();
 	FDisplayClusterViewport_RenderSettings& InOutCameraViewportRenderSettings = InCameraViewport.GetRenderSettingsImpl();
 
-	// Reset runtime flags from prev frame:
-	DstViewport.ResetRuntimeParameters();
-
 	// Chromakey used as overlay
 	InOutRenderSettings.bVisible = false;
 
@@ -578,7 +599,7 @@ void FDisplayClusterViewportConfigurationHelpers_ICVFX::UpdateChromakeyViewportS
 		{
 			check(InChromakeyRenderSettings->ShowOnlyList.IsVisibilityListValid());
 
-			FDisplayClusterViewportConfigurationHelpers_Visibility::UpdateShowOnlyList(DstViewport, InChromakeyRenderSettings->ShowOnlyList);
+			FDisplayClusterViewportConfigurationHelpers_Visibility::UpdateShowOnlyList_ICVFX(DstViewport, InChromakeyRenderSettings->ShowOnlyList);
 		}
 
 		DstViewport.UpdateConfiguration_OverlayRenderSettings(InChromakeyRenderSettings->AdvancedRenderSettings);
@@ -708,9 +729,6 @@ void FDisplayClusterViewportConfigurationHelpers_ICVFX::UpdateLightcardViewportS
 	FDisplayClusterViewport_RenderSettings&           InOutBaseViewportRenderSettings = BaseViewport.GetRenderSettingsImpl();
 	FDisplayClusterViewport_RenderSettingsICVFX& InOutBaseViewportRenderSettingsICVFX = BaseViewport.GetRenderSettingsICVFXImpl();
 
-	// Reset runtime flags from prev frame:
-	DstViewport.ResetRuntimeParameters();
-
 	// LIghtcard texture used as overlay
 	InOutRenderSettings.bVisible = false;
 
@@ -739,7 +757,7 @@ void FDisplayClusterViewportConfigurationHelpers_ICVFX::UpdateLightcardViewportS
 		{
 			check(InLightcardSettings.ShowOnlyList.IsVisibilityListValid());
 
-			FDisplayClusterViewportConfigurationHelpers_Visibility::UpdateShowOnlyList(DstViewport, InLightcardSettings.ShowOnlyList);
+			FDisplayClusterViewportConfigurationHelpers_Visibility::UpdateShowOnlyList_ICVFX(DstViewport, InLightcardSettings.ShowOnlyList);
 		}
 
 		DstViewport.UpdateConfiguration_OverlayRenderSettings(InRenderSettings.AdvancedRenderSettings);
@@ -757,23 +775,46 @@ void FDisplayClusterViewportConfigurationHelpers_ICVFX::UpdateLightcardViewportS
 		}
 		else
 		{
-			InOutBaseViewportRenderSettingsICVFX.ICVFX.LightCard.ViewportId = DstViewport.GetId();
+			if (EnumHasAnyFlags(DstViewport.GetRenderSettingsICVFX().RuntimeFlags, EDisplayClusterViewportRuntimeICVFXFlags::UnderInFrustum))
+			{
+				InOutBaseViewportRenderSettingsICVFX.ICVFX.LightCardUnder.ViewportId = DstViewport.GetId();
+			}
+
+			if (EnumHasAnyFlags(DstViewport.GetRenderSettingsICVFX().RuntimeFlags, EDisplayClusterViewportRuntimeICVFXFlags::OverInFrustum))
+			{
+				InOutBaseViewportRenderSettingsICVFX.ICVFX.LightCardOver.ViewportId = DstViewport.GetId();
+			}
+
 		}
 	}
 
 	if (EnumHasAnyFlags(DstViewport.GetRenderSettingsICVFX().RuntimeFlags, EDisplayClusterViewportRuntimeICVFXFlags::UVLightcard))
 	{
-		InOutBaseViewportRenderSettingsICVFX.ICVFX.UVLightCard.ViewportId = DstViewport.GetId();
+		if (EnumHasAnyFlags(DstViewport.GetRenderSettingsICVFX().RuntimeFlags, EDisplayClusterViewportRuntimeICVFXFlags::UnderInFrustum))
+		{
+			InOutBaseViewportRenderSettingsICVFX.ICVFX.UVLightCardUnder.ViewportId = DstViewport.GetId();
+		}
+
+		if (EnumHasAnyFlags(DstViewport.GetRenderSettingsICVFX().RuntimeFlags, EDisplayClusterViewportRuntimeICVFXFlags::OverInFrustum))
+		{
+			InOutBaseViewportRenderSettingsICVFX.ICVFX.UVLightCardOver.ViewportId = DstViewport.GetId();
+		}
 	}
 
 	if (EnumHasAnyFlags(DstViewport.GetRenderSettingsICVFX().RuntimeFlags, EDisplayClusterViewportRuntimeICVFXFlags::Lightcard | EDisplayClusterViewportRuntimeICVFXFlags::UVLightcard))
 	{
-		// Gain direct access to internal settings of the viewport:
-		FDisplayClusterViewport_RenderSettingsICVFX& InOutRenderSettingsICVFX = DstViewport.GetRenderSettingsICVFXImpl();
-
 		// Set the light card gamma used to linearize the light card textures before blending during final composite
 		// The OCIO pass will have already linearized the light card renders, so set the gamma to 1 if using OCIO on the light cards
-		InOutRenderSettingsICVFX.ICVFX.LightCardGamma = DstViewport.GetOpenColorIO().IsValid() ? 1.0 : 2.2;
+		InOutBaseViewportRenderSettingsICVFX.ICVFX.LightCardGamma = DstViewport.GetOpenColorIO().IsValid() ? 1.0 : 2.2;
+
+		if (EnumHasAnyFlags(InOutBaseViewportRenderSettingsICVFX.Flags, EDisplayClusterViewportICVFXFlags::LightcardAlwaysUnder))
+		{
+			InOutBaseViewportRenderSettingsICVFX.ICVFX.OverrideUVLightCardRenderMode = EDisplayClusterShaderParametersICVFX_LightCardRenderMode::Under;
+		}
+		else if (EnumHasAnyFlags(InOutBaseViewportRenderSettingsICVFX.Flags, EDisplayClusterViewportICVFXFlags::LightcardAlwaysOver))
+		{
+			InOutBaseViewportRenderSettingsICVFX.ICVFX.OverrideUVLightCardRenderMode = EDisplayClusterShaderParametersICVFX_LightCardRenderMode::Over;
+		}
 	}
 }
 
diff --git a/Engine/Plugins/Runtime/nDisplay/Source/DisplayCluster/Private/Render/Viewport/Configuration/DisplayClusterViewportConfigurationHelpers_ICVFX.h b/Engine/Plugins/Runtime/nDisplay/Source/DisplayCluster/Private/Render/Viewport/Configuration/DisplayClusterViewportConfigurationHelpers_ICVFX.h
--- a/Engine/Plugins/Runtime/nDisplay/Source/DisplayCluster/Private/Render/Viewport/Configuration/DisplayClusterViewportConfigurationHelpers_ICVFX.h
+++ b/Engine/Plugins/Runtime/nDisplay/Source/DisplayCluster/Private/Render/Viewport/Configuration/DisplayClusterViewportConfigurationHelpers_ICVFX.h
@@ -33,16 +33,16 @@ public:
 	static FDisplayClusterViewport* GetOrCreateChromakeyViewport(FDisplayClusterViewportConfiguration& InConfiguration, UDisplayClusterICVFXCameraComponent& InCameraComponent, const FDisplayClusterConfigurationICVFX_CameraSettings& InCameraSettings);
 	
 	/* Return the new or existing LightCard viewport. */
-	static FDisplayClusterViewport* GetOrCreateLightcardViewport(FDisplayClusterViewport& BaseViewport);
+	static FDisplayClusterViewport* GetOrCreateLightcardViewport(FDisplayClusterViewport& BaseViewport, const bool bOverInFrustum);
 
 	/* Return the new or existing UVLightCard viewport. */
-	static FDisplayClusterViewport* GetOrCreateUVLightcardViewport(FDisplayClusterViewport& BaseViewport);
+	static FDisplayClusterViewport* GetOrCreateUVLightcardViewport(FDisplayClusterViewport& BaseViewport, const bool bOverInFrustum);
 
 	/* Because UVLightCard viewports share the same outer texture (from LightCardManager), they can clone each other.
 	 * Except when the OCIO settings are not equal.
 	 * This function sets the InUVLightCardViewport rendering settings to minimize rendering costs.
 	 */
-	static void ReuseUVLightCardViewportWithinClusterNode(FDisplayClusterViewport& InUVLightCardViewport);
+	static void ReuseUVLightCardViewportWithinClusterNode(FDisplayClusterViewport& InUVLightCardViewport, const bool bOverInFrustum);
 
 	/** Returns all visible InnerCamera viewports. */
 	static TArray<TSharedPtr<FDisplayClusterViewport, ESPMode::ThreadSafe>> GetAllVisibleInnerCameraViewports(FDisplayClusterViewportConfiguration& InConfiguration, bool bGetChromakey = false);
@@ -57,7 +57,7 @@ public:
 
 	static bool IsCameraUsed(const FDisplayClusterConfigurationICVFX_CameraSettings& InCameraSettings);
 
-	static bool CreateProjectionPolicyICVFX(FDisplayClusterViewportConfiguration& InConfiguration, const FString& InViewportId, const FString& InResourceId, bool bIsCameraProjection, TSharedPtr<IDisplayClusterProjectionPolicy, ESPMode::ThreadSafe>& OutProjPolicy);
+	static bool CreateProjectionPolicyICVFX(FDisplayClusterViewportConfiguration& InConfiguration, const FString& InViewportId, const FString& InResourceId, const FString& InSuffixId, bool bIsCameraProjection, TSharedPtr<IDisplayClusterProjectionPolicy, ESPMode::ThreadSafe>& OutProjPolicy);
 	static bool CreateProjectionPolicyCameraICVFX(FDisplayClusterViewportConfiguration& InConfiguration, UDisplayClusterICVFXCameraComponent& CameraComponent, const FDisplayClusterConfigurationICVFX_CameraSettings& InCameraSettings, TSharedPtr<IDisplayClusterProjectionPolicy, ESPMode::ThreadSafe>& OutProjPolicy);
 	static bool UpdateCameraProjectionSettingsICVFX(FDisplayClusterViewportConfiguration& InConfiguration, UDisplayClusterICVFXCameraComponent& CameraComponent, const FDisplayClusterConfigurationICVFX_CameraSettings& InCameraSettings, const TSharedPtr<IDisplayClusterProjectionPolicy, ESPMode::ThreadSafe>& InProjectionPolicy);
 
@@ -69,6 +69,6 @@ public:
 
 	static void UpdateCameraCustomFrustum(const FDisplayClusterConfigurationICVFX_CameraCustomFrustum& InCameraCustomFrustumConfiguration, FDisplayClusterViewport_CustomFrustumSettings& OutCustomFrustumSettings);
 
-	static FDisplayClusterViewport* FindViewportICVFX(FDisplayClusterViewportConfiguration& InConfiguration, const FString& InViewportId, const FString& InResourceId);
-	static FDisplayClusterViewport* CreateViewportICVFX(FDisplayClusterViewportConfiguration& InConfiguration, const FString& InViewportId, const FString& InResourceId, TSharedPtr<IDisplayClusterProjectionPolicy, ESPMode::ThreadSafe>& InProjectionPolicy);
+	static FDisplayClusterViewport* FindViewportICVFX(FDisplayClusterViewportConfiguration& InConfiguration, const FString& InViewportId, const FString& InResourceId, const FString& InResourceSuffixId);
+	static FDisplayClusterViewport* CreateViewportICVFX(FDisplayClusterViewportConfiguration& InConfiguration, const FString& InViewportId, const FString& InResourceId, const FString& InResourceSuffixId, TSharedPtr<IDisplayClusterProjectionPolicy, ESPMode::ThreadSafe>& InProjectionPolicy);
 };
diff --git a/Engine/Plugins/Runtime/nDisplay/Source/DisplayCluster/Private/Render/Viewport/Configuration/DisplayClusterViewportConfigurationHelpers_Tile.cpp b/Engine/Plugins/Runtime/nDisplay/Source/DisplayCluster/Private/Render/Viewport/Configuration/DisplayClusterViewportConfigurationHelpers_Tile.cpp
--- a/Engine/Plugins/Runtime/nDisplay/Source/DisplayCluster/Private/Render/Viewport/Configuration/DisplayClusterViewportConfigurationHelpers_Tile.cpp
+++ b/Engine/Plugins/Runtime/nDisplay/Source/DisplayCluster/Private/Render/Viewport/Configuration/DisplayClusterViewportConfigurationHelpers_Tile.cpp
@@ -143,7 +143,7 @@ FDisplayClusterViewport* FDisplayClusterViewportConfigurationHelpers_Tile::GetOr
 		// Reset runtime flags from prev frame.
 		// Also this function update media states.
 		// Note: Save the new media states after calling ResetRuntimeParameters().
-		TileViewport->ResetRuntimeParameters();
+		TileViewport->ResetRuntimeParameters(InSourceViewport.GetViewportConfigurationData());
 		const EDisplayClusterViewportMediaState SavedTileViewportMediaStates = TileViewport->GetRenderSettings().GetMediaStates();
 
 		// Gain direct access to internal resources of the NewViewport:
diff --git a/Engine/Plugins/Runtime/nDisplay/Source/DisplayCluster/Private/Render/Viewport/Configuration/DisplayClusterViewportConfigurationHelpers_Visibility.cpp b/Engine/Plugins/Runtime/nDisplay/Source/DisplayCluster/Private/Render/Viewport/Configuration/DisplayClusterViewportConfigurationHelpers_Visibility.cpp
--- a/Engine/Plugins/Runtime/nDisplay/Source/DisplayCluster/Private/Render/Viewport/Configuration/DisplayClusterViewportConfigurationHelpers_Visibility.cpp
+++ b/Engine/Plugins/Runtime/nDisplay/Source/DisplayCluster/Private/Render/Viewport/Configuration/DisplayClusterViewportConfigurationHelpers_Visibility.cpp
@@ -14,27 +14,94 @@
 #include "Render/Viewport/DisplayClusterViewportManager.h"
 #include "DisplayClusterConfigurationTypes_Viewport.h"
 
+#include "DisplayClusterLightCardActor.h"
+
 namespace UE::DisplayCluster::Configuration::VisibilityHelpers
 {
-	static inline void ImplCollectActorComponents(AActor& InActor, TSet<FPrimitiveComponentId>& OutAdditionalComponentsList)
+	static inline void ImplCollectActorComponents(AActor& InActor, TSet<FPrimitiveComponentId>& OutComponentsList, const FDisplayClusterViewport* InShowOnlyViewport)
 	{
+		if (InShowOnlyViewport)
+		{
+			if (!FDisplayClusterViewportConfigurationHelpers_Visibility::IsActorVisibleForViewport(*InShowOnlyViewport, InActor))
+			{
+				// Ignore actors that not visible for this viewport.
+				return;
+			}
+		}
+
 		for (const UActorComponent* Component : InActor.GetComponents())
 		{
 			if (const UPrimitiveComponent* PrimComp = Cast<UPrimitiveComponent>(Component))
 			{
-				OutAdditionalComponentsList.Add(PrimComp->GetPrimitiveSceneId());
+				OutComponentsList.Add(PrimComp->GetPrimitiveSceneId());
 			}
 		}
 	}
 
-	static inline void ImplCollectVisibility(FDisplayClusterViewportConfiguration& InConfiguration, const FDisplayClusterConfigurationICVFX_VisibilityList& InVisibilityList, TArray<FName>& OutActorLayerNames, TSet<FPrimitiveComponentId>& OutAdditionalComponentsList)
+	/** Collects actors from the layers of the current world. */
+	static inline void ImplCollectActorsFromLayers(UWorld* InCurrentWorld, const TArray<FActorLayer>& InActorLayers, TArray<TSoftObjectPtr<AActor>>& OutActorsList)
 	{
-		// Collect RootActor components
-		if (InVisibilityList.RootActorComponentNames.Num() > 0)
+		TRACE_CPUPROFILER_EVENT_SCOPE(DisplayClusterViewport_ImplCollectActorsFromLayers);
+
+		TSet<FName> ActorLayerNames;
+		ActorLayerNames.Reserve(InActorLayers.Num());
+
+		// Remove empty names
+		for (const FActorLayer& ActorLayerIt : InActorLayers)
 		{
-			if (ADisplayClusterRootActor* SceneRootActor = InConfiguration.GetRootActor(EDisplayClusterRootActorType::Scene))
+			if (!ActorLayerIt.Name.IsNone())
 			{
-				SceneRootActor->FindPrimitivesByName(InVisibilityList.RootActorComponentNames, OutAdditionalComponentsList);
+				ActorLayerNames.Add(ActorLayerIt.Name);
+			}
+		}
+
+		if (!ActorLayerNames.IsEmpty())
+		{
+			// Iterate over all actors, looking for actors in the specified layers.
+			for (const TWeakObjectPtr<AActor> ActorWeakPtr : FActorRange(InCurrentWorld))
+			{
+				if (ActorWeakPtr.IsValid())
+				{
+					// Search actor on source layers
+					bool bActorFoundOnSourceLayers = false;
+					for (const FName& ActorLayerNameIt : ActorLayerNames)
+					{
+						if (!ActorLayerNameIt.IsNone() && ActorWeakPtr->Layers.Contains(ActorLayerNameIt))
+						{
+							OutActorsList.Add(ActorWeakPtr.Get());
+							break;
+						}
+					}
+				}
+			}
+		}
+	}
+
+	static inline void ImplCollectComponentsFromVisibilityList(FDisplayClusterViewportConfiguration& InConfiguration, const FDisplayClusterConfigurationICVFX_VisibilityList& InVisibilityList, TSet<FPrimitiveComponentId>& OutAdditionalComponentsList, const FDisplayClusterViewport* InShowOnlyViewport = nullptr)
+	{
+		TRACE_CPUPROFILER_EVENT_SCOPE(DisplayCluster_ImplCollectComponentsFromVisibilityList);
+
+		// Collect components from the DCRA
+		if (!InVisibilityList.RootActorComponentNames.IsEmpty())
+		{
+			bool bCanRenderPrimitives = true;
+			if (InShowOnlyViewport)
+			{
+				// If the InShowOnlyViewport argument is defined, check if it is of type lightcard and can be used.
+				if (!FDisplayClusterViewportConfigurationHelpers_Visibility::IsLightcardViewportRenderable(*InShowOnlyViewport))
+				{
+					// Optimization: do not add LC components from DCRA if this viewport is of type lightcard and can be skipped.
+					bCanRenderPrimitives = false;
+				}
+			}
+
+			if (bCanRenderPrimitives)
+			{
+				if (ADisplayClusterRootActor* SceneRootActor = InConfiguration.GetRootActor(EDisplayClusterRootActorType::Scene))
+				{
+					// All DCRA components from the list need to be collected.
+					SceneRootActor->FindPrimitivesByName(InVisibilityList.RootActorComponentNames, OutAdditionalComponentsList, true);
+				}
 			}
 		}
 
@@ -46,7 +113,7 @@ namespace UE::DisplayCluster::Configuration::VisibilityHelpers
 					{
 						if (ActorSOPtrIt->GetWorld() == CurrentWorld)
 						{
-							ImplCollectActorComponents(*ActorSOPtrIt.Get(), OutAdditionalComponentsList);
+							ImplCollectActorComponents(*ActorSOPtrIt.Get(), OutAdditionalComponentsList, InShowOnlyViewport);
 						}
 						else if(CurrentWorld)
 						{
@@ -58,39 +125,33 @@ namespace UE::DisplayCluster::Configuration::VisibilityHelpers
 				}
 			};
 
-		// Collect Actors refs
 		UWorld* CurrentWorld = InConfiguration.GetCurrentWorld();
+
+		// Collect actors from the layers.
+		TArray<TSoftObjectPtr<AActor>> ActorsFromLayers;
+		ImplCollectActorsFromLayers(CurrentWorld, InVisibilityList.ActorLayers, ActorsFromLayers);
+
+		// Collect Actors refs
 		CollectActorRefs(CurrentWorld, InVisibilityList.Actors);
 		CollectActorRefs(CurrentWorld, InVisibilityList.AutoAddedActors);
-
-		// Collect ActorLayers
-		for (const FActorLayer& ActorLayerIt : InVisibilityList.ActorLayers)
-		{
-			if (!ActorLayerIt.Name.IsNone())
-			{
-				OutActorLayerNames.AddUnique(ActorLayerIt.Name);
-			}
-		}
+		CollectActorRefs(CurrentWorld, ActorsFromLayers);
 	}
 };
-using namespace UE::DisplayCluster::Configuration;
 
-void FDisplayClusterViewportConfigurationHelpers_Visibility::UpdateShowOnlyList(FDisplayClusterViewport& DstViewport, const FDisplayClusterConfigurationICVFX_VisibilityList& InVisibilityList)
+void FDisplayClusterViewportConfigurationHelpers_Visibility::UpdateShowOnlyList_ICVFX(FDisplayClusterViewport& DstViewport, const FDisplayClusterConfigurationICVFX_VisibilityList& InVisibilityList)
 {
-	TSet<FPrimitiveComponentId> AdditionalComponentsList;
-	TArray<FName> ActorLayers;
-	VisibilityHelpers::ImplCollectVisibility(*DstViewport.Configuration, InVisibilityList, ActorLayers, AdditionalComponentsList);
+	TSet<FPrimitiveComponentId> ComponentsList;
+	UE::DisplayCluster::Configuration::VisibilityHelpers::ImplCollectComponentsFromVisibilityList(*DstViewport.Configuration, InVisibilityList, ComponentsList, &DstViewport);
 
-	DstViewport.GetVisibilitySettingsImpl().UpdateVisibilitySettings(EDisplayClusterViewport_VisibilityMode::ShowOnly, ActorLayers, AdditionalComponentsList);
+	DstViewport.GetVisibilitySettingsImpl().SetVisibilityModeAndComponentsList(EDisplayClusterViewport_VisibilityMode::ShowOnly, ComponentsList);
 }
 
 void FDisplayClusterViewportConfigurationHelpers_Visibility::AppendHideList_ICVFX(FDisplayClusterViewport& DstViewport, const FDisplayClusterConfigurationICVFX_VisibilityList& InHideList)
 {
-	TSet<FPrimitiveComponentId> AdditionalComponentsList;
-	TArray<FName> ActorLayers;
-	VisibilityHelpers::ImplCollectVisibility(*DstViewport.Configuration, InHideList, ActorLayers, AdditionalComponentsList);
+	TSet<FPrimitiveComponentId> ComponentsList;
+	UE::DisplayCluster::Configuration::VisibilityHelpers::ImplCollectComponentsFromVisibilityList(*DstViewport.Configuration, InHideList, ComponentsList);
 
-	DstViewport.GetVisibilitySettingsImpl().AppendHideList(ActorLayers, AdditionalComponentsList);
+	DstViewport.GetVisibilitySettingsImpl().AppendVisibilityComponentsList(EDisplayClusterViewport_VisibilityMode::Hide, ComponentsList);
 }
 
 void FDisplayClusterViewportConfigurationHelpers_Visibility::UpdateHideList_ICVFX(FDisplayClusterViewportConfiguration& InConfiguration, TArray<TSharedPtr<FDisplayClusterViewport, ESPMode::ThreadSafe>>& DstViewports)
@@ -98,48 +159,98 @@ void FDisplayClusterViewportConfigurationHelpers_Visibility::UpdateHideList_ICVF
 	ADisplayClusterRootActor* ConfigurationRootActor = InConfiguration.GetRootActor(EDisplayClusterRootActorType::Configuration);
 	const FDisplayClusterConfigurationICVFX_StageSettings* StageSettings = InConfiguration.GetStageSettings();
 
-	if (DstViewports.Num() > 0 && ConfigurationRootActor && StageSettings)
+	if (!DstViewports.IsEmpty() && ConfigurationRootActor && StageSettings)
 	{
-		TArray<FName> ActorLayerNames;
-		TSet<FPrimitiveComponentId> AdditionalComponentsList;
+		TSet<FPrimitiveComponentId> ComponentsList;
 
-		VisibilityHelpers::ImplCollectVisibility(InConfiguration, StageSettings->HideList, ActorLayerNames, AdditionalComponentsList);
+		UE::DisplayCluster::Configuration::VisibilityHelpers::ImplCollectComponentsFromVisibilityList(InConfiguration, StageSettings->HideList, ComponentsList);
 
 		// Hide lightcard
-		VisibilityHelpers::ImplCollectVisibility(InConfiguration, StageSettings->Lightcard.ShowOnlyList, ActorLayerNames, AdditionalComponentsList);
+		UE::DisplayCluster::Configuration::VisibilityHelpers::ImplCollectComponentsFromVisibilityList(InConfiguration, StageSettings->Lightcard.ShowOnlyList, ComponentsList);
 
 		// Also hide chromakeys for all cameras
 		TArray<UDisplayClusterICVFXCameraComponent*> ConfigurationRootActorCameras;
 		ConfigurationRootActor->GetComponents(ConfigurationRootActorCameras);
-
 		for (const UDisplayClusterICVFXCameraComponent* ConfigurationCameraIt : ConfigurationRootActorCameras)
 		{
-			if (ConfigurationCameraIt)
+			if (const FDisplayClusterConfigurationICVFX_ChromakeyRenderSettings* ChromakeyRenderSettings = ConfigurationCameraIt ?
+				ConfigurationCameraIt->GetCameraSettingsICVFX().Chromakey.GetChromakeyRenderSettings(*StageSettings) : nullptr)
 			{
-				if (const FDisplayClusterConfigurationICVFX_ChromakeyRenderSettings* ChromakeyRenderSettings = ConfigurationCameraIt->GetCameraSettingsICVFX().Chromakey.GetChromakeyRenderSettings(*StageSettings))
-				{
-					VisibilityHelpers::ImplCollectVisibility(InConfiguration, ChromakeyRenderSettings->ShowOnlyList, ActorLayerNames, AdditionalComponentsList);
-				}
+				UE::DisplayCluster::Configuration::VisibilityHelpers::ImplCollectComponentsFromVisibilityList(InConfiguration, ChromakeyRenderSettings->ShowOnlyList, ComponentsList);
 			}
 		}
 
-		TArray<FName> OuterActorLayerNames;
-		TSet<FPrimitiveComponentId> OuterAdditionalComponentsList;
-		VisibilityHelpers::ImplCollectVisibility(InConfiguration, StageSettings->OuterViewportHideList, OuterActorLayerNames, OuterAdditionalComponentsList);
+		TSet<FPrimitiveComponentId> OuterComponentsList;
+		UE::DisplayCluster::Configuration::VisibilityHelpers::ImplCollectComponentsFromVisibilityList(InConfiguration, StageSettings->OuterViewportHideList, OuterComponentsList);
 
 		// Update hide list for all desired viewports:
 		for (TSharedPtr<FDisplayClusterViewport, ESPMode::ThreadSafe>& ViewportIt : DstViewports)
 		{
 			if (ViewportIt.IsValid())
 			{
-				ViewportIt->GetVisibilitySettingsImpl().UpdateVisibilitySettings(EDisplayClusterViewport_VisibilityMode::Hide, ActorLayerNames, AdditionalComponentsList);
+				ViewportIt->GetVisibilitySettingsImpl().SetVisibilityModeAndComponentsList(EDisplayClusterViewport_VisibilityMode::Hide, ComponentsList);
 
 				// Support additional hide list for outer viewports
 				if (EnumHasAllFlags(ViewportIt->GetRenderSettingsICVFX().RuntimeFlags, EDisplayClusterViewportRuntimeICVFXFlags::Target))
 				{
-					ViewportIt->GetVisibilitySettingsImpl().AppendHideList(OuterActorLayerNames, OuterAdditionalComponentsList);
+					ViewportIt->GetVisibilitySettingsImpl().AppendVisibilityComponentsList(EDisplayClusterViewport_VisibilityMode::Hide, OuterComponentsList);
 				}
 			}
 		}
 	}
 }
+
+bool FDisplayClusterViewportConfigurationHelpers_Visibility::IsLightcardViewportRenderable(const FDisplayClusterViewport& InViewport, const EDisplayClusterConfigurationICVFX_PerLightcardRenderMode PerLightcardRenderMode)
+{
+	const EDisplayClusterViewportRuntimeICVFXFlags RuntimeFlags = InViewport.GetRenderSettingsICVFX().RuntimeFlags;
+	if (EnumHasAnyFlags(RuntimeFlags, EDisplayClusterViewportRuntimeICVFXFlags::Lightcard | EDisplayClusterViewportRuntimeICVFXFlags::UVLightcard))
+	{
+		const FDisplayClusterConfigurationICVFX_StageSettings* StageSettings = InViewport.Configuration->GetStageSettings();
+
+		const UDisplayClusterConfigurationViewport* ViewportConfiguration = nullptr;
+		if (EnumHasAnyFlags(RuntimeFlags, EDisplayClusterViewportRuntimeICVFXFlags::Lightcard))
+		{
+			// The UV-light card is rendered once for all viewports.
+			// This means that here we cannot use the per-viewport LC rules here. These rules are implemented in the ICVFX shader.
+			ViewportConfiguration = InViewport.GetViewportConfigurationData();
+		}
+
+		if (StageSettings)
+		{
+			// Renders all primitives only into the default LC viewport.
+			const EDisplayClusterShaderParametersICVFX_LightCardRenderMode LightCardRenderMode = StageSettings->Lightcard.GetLightCardRenderMode(PerLightcardRenderMode, ViewportConfiguration);
+			switch (LightCardRenderMode)
+			{
+			case EDisplayClusterShaderParametersICVFX_LightCardRenderMode::Under:
+				return EnumHasAnyFlags(RuntimeFlags, EDisplayClusterViewportRuntimeICVFXFlags::UnderInFrustum);
+
+			case EDisplayClusterShaderParametersICVFX_LightCardRenderMode::Over:
+				return EnumHasAnyFlags(RuntimeFlags, EDisplayClusterViewportRuntimeICVFXFlags::OverInFrustum);
+
+			default:
+				break;
+			}
+
+			return false;
+		}
+	}
+
+	return true;
+}
+
+bool FDisplayClusterViewportConfigurationHelpers_Visibility::IsActorVisibleForViewport(const FDisplayClusterViewport& InViewport, AActor& InActor)
+{
+	using namespace UE::DisplayCluster::Configuration::VisibilityHelpers;
+
+	// Get special rules from the lightcard actor.
+	EDisplayClusterConfigurationICVFX_PerLightcardRenderMode PerLightcardRenderMode = EDisplayClusterConfigurationICVFX_PerLightcardRenderMode::Default;
+	if (InActor.IsA<ADisplayClusterLightCardActor>())
+	{
+		if (const ADisplayClusterLightCardActor* LightCardActor = Cast<ADisplayClusterLightCardActor>(&InActor))
+		{
+			PerLightcardRenderMode = LightCardActor->PerLightcardRenderMode;
+		}
+	}
+
+	return IsLightcardViewportRenderable(InViewport, PerLightcardRenderMode);
+}
diff --git a/Engine/Plugins/Runtime/nDisplay/Source/DisplayCluster/Private/Render/Viewport/Configuration/DisplayClusterViewportConfigurationHelpers_Visibility.h b/Engine/Plugins/Runtime/nDisplay/Source/DisplayCluster/Private/Render/Viewport/Configuration/DisplayClusterViewportConfigurationHelpers_Visibility.h
--- a/Engine/Plugins/Runtime/nDisplay/Source/DisplayCluster/Private/Render/Viewport/Configuration/DisplayClusterViewportConfigurationHelpers_Visibility.h
+++ b/Engine/Plugins/Runtime/nDisplay/Source/DisplayCluster/Private/Render/Viewport/Configuration/DisplayClusterViewportConfigurationHelpers_Visibility.h
@@ -3,9 +3,11 @@
 #pragma once
 
 #include "CoreMinimal.h"
+#include "DisplayClusterConfigurationTypes_Enums.h"
 
 class FDisplayClusterViewportConfiguration;
 class FDisplayClusterViewport;
+class AActor;
 struct FDisplayClusterConfigurationICVFX_VisibilityList;
 
 /**
@@ -15,12 +17,21 @@ class FDisplayClusterViewportConfigurationHelpers_Visibility
 {
 public:
 	// Update ShowOnly list for DstViewport
-	static void UpdateShowOnlyList(FDisplayClusterViewport& DstViewport, const FDisplayClusterConfigurationICVFX_VisibilityList& InVisibilityList);
+	static void UpdateShowOnlyList_ICVFX(FDisplayClusterViewport& DstViewport, const FDisplayClusterConfigurationICVFX_VisibilityList& InVisibilityList);
 
 	// Update hide lists for DstViewports (lightcards, chromakey, stage settings hide list, outer viewports hide list, etc)
 	static void UpdateHideList_ICVFX(FDisplayClusterViewportConfiguration& InConfiguration, TArray<TSharedPtr<FDisplayClusterViewport, ESPMode::ThreadSafe>>& DstViewports);
 	
 	// Append to exist hide list (must be call after UpdateHideList_ICVFX)
 	static void AppendHideList_ICVFX(FDisplayClusterViewport& DstViewport, const FDisplayClusterConfigurationICVFX_VisibilityList& InHideList);
+
+	/** Returns true if this actor can be rendered for the viewport.
+	* For Example:
+	*   The LightCard actor has a property that determines which of the two LightCard viewports it should be rendered in.
+	*/
+	static bool IsActorVisibleForViewport(const FDisplayClusterViewport& InViewport, AActor& InActor);
+
+	/** Returns true if the viewport is of type lightcard and can be rendered. */
+	static bool IsLightcardViewportRenderable(const FDisplayClusterViewport& InViewport, const EDisplayClusterConfigurationICVFX_PerLightcardRenderMode PerLightcardRenderMode = EDisplayClusterConfigurationICVFX_PerLightcardRenderMode::Default);
 };
 
diff --git a/Engine/Plugins/Runtime/nDisplay/Source/DisplayCluster/Private/Render/Viewport/Configuration/DisplayClusterViewportConfiguration_ICVFX.cpp b/Engine/Plugins/Runtime/nDisplay/Source/DisplayCluster/Private/Render/Viewport/Configuration/DisplayClusterViewportConfiguration_ICVFX.cpp
--- a/Engine/Plugins/Runtime/nDisplay/Source/DisplayCluster/Private/Render/Viewport/Configuration/DisplayClusterViewportConfiguration_ICVFX.cpp
+++ b/Engine/Plugins/Runtime/nDisplay/Source/DisplayCluster/Private/Render/Viewport/Configuration/DisplayClusterViewportConfiguration_ICVFX.cpp
@@ -42,9 +42,9 @@ static FAutoConsoleVariableRef CVarDisplayClusterEnableAlphaChannelRendering(
 ///////////////////////////////////////////////////////////////////
 // FDisplayClusterViewportConfiguration_ICVFX
 ///////////////////////////////////////////////////////////////////
-bool FDisplayClusterViewportConfiguration_ICVFX::CreateLightcardViewport(FDisplayClusterViewport& BaseViewport)
+bool FDisplayClusterViewportConfiguration_ICVFX::CreateLightcardViewport(FDisplayClusterViewport& BaseViewport, const bool bOverInFrustum)
 {
-	if (FDisplayClusterViewport* LightcardViewport = FDisplayClusterViewportConfigurationHelpers_ICVFX::GetOrCreateLightcardViewport(BaseViewport))
+	if (FDisplayClusterViewport* LightcardViewport = FDisplayClusterViewportConfigurationHelpers_ICVFX::GetOrCreateLightcardViewport(BaseViewport, bOverInFrustum))
 	{
 		// Update lightcard viewport settings
 		FDisplayClusterViewportConfigurationHelpers_ICVFX::UpdateLightcardViewportSetting(*LightcardViewport, BaseViewport);
@@ -58,9 +58,9 @@ bool FDisplayClusterViewportConfiguration_ICVFX::CreateLightcardViewport(FDispla
 	return false;
 }
 
-bool FDisplayClusterViewportConfiguration_ICVFX::CreateUVLightcardViewport(FDisplayClusterViewport& BaseViewport)
+bool FDisplayClusterViewportConfiguration_ICVFX::CreateUVLightcardViewport(FDisplayClusterViewport& BaseViewport, const bool bOverInFrustum)
 {
-	if (FDisplayClusterViewport* UVLightCardViewport = FDisplayClusterViewportConfigurationHelpers_ICVFX::GetOrCreateUVLightcardViewport(BaseViewport))
+	if (FDisplayClusterViewport* UVLightCardViewport = FDisplayClusterViewportConfigurationHelpers_ICVFX::GetOrCreateUVLightcardViewport(BaseViewport, bOverInFrustum))
 	{
 		// Update UV LightCard viewport settings
 		FDisplayClusterViewportConfigurationHelpers_ICVFX::UpdateLightcardViewportSetting(*UVLightCardViewport, BaseViewport);
@@ -69,7 +69,7 @@ bool FDisplayClusterViewportConfiguration_ICVFX::CreateUVLightcardViewport(FDisp
 		UVLightCardViewport->UpdateConfiguration_ProjectionPolicy();
 
 		// Optimize: re-use UVLightCard viewports with equals OCIO
-		FDisplayClusterViewportConfigurationHelpers_ICVFX::ReuseUVLightCardViewportWithinClusterNode(*UVLightCardViewport);
+		FDisplayClusterViewportConfigurationHelpers_ICVFX::ReuseUVLightCardViewportWithinClusterNode(*UVLightCardViewport, bOverInFrustum);
 
 		return true;
 	}
@@ -128,25 +128,41 @@ void FDisplayClusterViewportConfiguration_ICVFX::Update()
 	if (!EnumHasAnyFlags(TargetViewportsFlags, EDisplayClusterViewportICVFXFlags::DisableLightcard))
 	{
 		// UVLightCard must be enabled in LC manager
-		const bool bUVLightCardEnabled = ViewportManager->LightCardManager->IsUVLightCardEnabled();
+		const bool bUVLightCardOverEnabled = ViewportManager->LightCardManager->IsUVLightCardEnabled(EDisplayClusterUVLightCardType::Over);
+		const bool bUVLightCardUnderEnabled = ViewportManager->LightCardManager->IsUVLightCardEnabled(EDisplayClusterUVLightCardType::Under);
+
+		// per-viewport lightcard use-case
+		const EDisplayClusterViewportICVFXFlags LightcardRenderModeFlags = TargetViewportsFlags & EDisplayClusterViewportICVFXFlags::LightcardRenderModeMask;
+		const bool bLightCardOverEnabled  = LightcardRenderModeFlags != EDisplayClusterViewportICVFXFlags::LightcardAlwaysUnder;
+		const bool bLightCardUnderEnabled = LightcardRenderModeFlags != EDisplayClusterViewportICVFXFlags::LightcardAlwaysOver;
 
 		// Allocate and assign lightcard resources
 		const bool bUseLightCard = StageSettings->Lightcard.ShouldUseLightCard(*StageSettings);
-		const bool bUseUVLightCard = bUVLightCardEnabled && StageSettings->Lightcard.ShouldUseUVLightCard(*StageSettings);
+		const bool bUseUVLightCard = StageSettings->Lightcard.ShouldUseUVLightCard(*StageSettings);
 
 		for (const TSharedPtr<FDisplayClusterViewport, ESPMode::ThreadSafe>& TargetIt : TargetViewports)
 		{
 			// only for support targets
 			if (TargetIt.IsValid() && !EnumHasAnyFlags(TargetIt->GetRenderSettingsICVFX().Flags, EDisplayClusterViewportICVFXFlags::DisableLightcard))
 			{
-				if (bUseLightCard)
+				if (bLightCardOverEnabled && bUseLightCard)
+				{
+					CreateLightcardViewport(*TargetIt, true);
+				}
+				
+				if (bLightCardUnderEnabled && bUseLightCard)
+				{
+					CreateLightcardViewport(*TargetIt, false);
+				}
+
+				if (bUVLightCardOverEnabled && bUseUVLightCard)
 				{
-					CreateLightcardViewport(*TargetIt);
+					CreateUVLightcardViewport(*TargetIt, true);
 				}
 
-				if (bUseUVLightCard)
+				if (bUVLightCardUnderEnabled && bUseUVLightCard)
 				{
-					CreateUVLightcardViewport(*TargetIt);
+					CreateUVLightcardViewport(*TargetIt, false);
 				}
 			}
 		}
diff --git a/Engine/Plugins/Runtime/nDisplay/Source/DisplayCluster/Private/Render/Viewport/Configuration/DisplayClusterViewportConfiguration_ICVFX.h b/Engine/Plugins/Runtime/nDisplay/Source/DisplayCluster/Private/Render/Viewport/Configuration/DisplayClusterViewportConfiguration_ICVFX.h
--- a/Engine/Plugins/Runtime/nDisplay/Source/DisplayCluster/Private/Render/Viewport/Configuration/DisplayClusterViewportConfiguration_ICVFX.h
+++ b/Engine/Plugins/Runtime/nDisplay/Source/DisplayCluster/Private/Render/Viewport/Configuration/DisplayClusterViewportConfiguration_ICVFX.h
@@ -43,10 +43,10 @@ private:
 	bool ImplGetTargetViewports(TArray<TSharedPtr<FDisplayClusterViewport, ESPMode::ThreadSafe>>& OutTargets, EDisplayClusterViewportICVFXFlags& OutMergedICVFXFlags);
 
 	/** Create a LightCard viewport for the  BaseViewport. */
-	bool CreateLightcardViewport(FDisplayClusterViewport& BaseViewport);
+	bool CreateLightcardViewport(FDisplayClusterViewport& BaseViewport, const bool bOverInFrustum);
 
 	/** Create an UVLightCard viewport for the BaseViewport. */
-	bool CreateUVLightcardViewport(FDisplayClusterViewport& BaseViewport);
+	bool CreateUVLightcardViewport(FDisplayClusterViewport& BaseViewport, const bool bOverInFrustum);
 
 	/** Update a visibility for the ICVFX viewports and cameras. */
 	void ImplUpdateVisibility();
diff --git a/Engine/Plugins/Runtime/nDisplay/Source/DisplayCluster/Private/Render/Viewport/Configuration/DisplayClusterViewportConfiguration_Viewport.cpp b/Engine/Plugins/Runtime/nDisplay/Source/DisplayCluster/Private/Render/Viewport/Configuration/DisplayClusterViewportConfiguration_Viewport.cpp
--- a/Engine/Plugins/Runtime/nDisplay/Source/DisplayCluster/Private/Render/Viewport/Configuration/DisplayClusterViewportConfiguration_Viewport.cpp
+++ b/Engine/Plugins/Runtime/nDisplay/Source/DisplayCluster/Private/Render/Viewport/Configuration/DisplayClusterViewportConfiguration_Viewport.cpp
@@ -57,6 +57,12 @@ bool FDisplayClusterViewportConfiguration_Viewport::UpdateViewportConfiguration(
 {
 	check(IsInGameThread());
 
+	// Reset runtime flags from prev frame:
+	DstViewport.ResetRuntimeParameters(&ConfigurationViewport);
+
+	// Store viewport configuration data
+	DstViewport.SetViewportConfigurationData(&ConfigurationViewport);
+
 	FDisplayClusterViewportConfigurationHelpers::UpdateBaseViewportSetting(DstViewport, ConfigurationViewport);
 	DstViewport.UpdateConfiguration_ProjectionPolicy(&ConfigurationViewport.ProjectionPolicy);
 
diff --git a/Engine/Plugins/Runtime/nDisplay/Source/DisplayCluster/Private/Render/Viewport/DisplayClusterViewport.cpp b/Engine/Plugins/Runtime/nDisplay/Source/DisplayCluster/Private/Render/Viewport/DisplayClusterViewport.cpp
--- a/Engine/Plugins/Runtime/nDisplay/Source/DisplayCluster/Private/Render/Viewport/DisplayClusterViewport.cpp
+++ b/Engine/Plugins/Runtime/nDisplay/Source/DisplayCluster/Private/Render/Viewport/DisplayClusterViewport.cpp
@@ -316,7 +316,7 @@ void FDisplayClusterViewport::SetupSceneView(uint32 ContextNum, class UWorld* Wo
 	if(RenderSettings.CaptureMode == EDisplayClusterViewportCaptureMode::MoviePipeline)
 	{
 		// Apply visibility settigns to view
-		VisibilitySettings.SetupSceneView(World, InOutView);
+		VisibilitySettings.SetupSceneView(InOutView);
 
 		return;
 	}
@@ -357,7 +357,7 @@ void FDisplayClusterViewport::SetupSceneView(uint32 ContextNum, class UWorld* Wo
 	}
 
 	// Apply visibility settigns to view
-	VisibilitySettings.SetupSceneView(World, InOutView);
+	VisibilitySettings.SetupSceneView(InOutView);
 
 	// Handle Motion blur parameters
 	CameraMotionBlur.SetupSceneView(Contexts[ContextNum], InOutView);
@@ -511,6 +511,12 @@ bool FDisplayClusterViewport::UpdateFrameContexts(const uint32 InStereoViewIndex
 		return false;
 	}
 
+	if (!VisibilitySettings.IsVisible())
+	{
+		// Exclude viewports that are empty from rendering.
+		return false;
+	}
+
 	if (PostRenderSettings.GenerateMips.IsEnabled())
 	{
 		//Check if current projection policy supports this feature
@@ -633,10 +639,13 @@ bool FDisplayClusterViewport::UpdateFrameContexts(const uint32 InStereoViewIndex
 
 		if (FDisplayClusterViewportManager* ViewportManager = Configuration->GetViewportManagerImpl())
 		{
-			if (ViewportManager->LightCardManager->IsUVLightCardEnabled())
+			const EDisplayClusterUVLightCardType UVLightCardType =
+				EnumHasAllFlags(RenderSettingsICVFX.RuntimeFlags, EDisplayClusterViewportRuntimeICVFXFlags::OverInFrustum)
+				? EDisplayClusterUVLightCardType::Over : EDisplayClusterUVLightCardType::Under;
+			if (ViewportManager->LightCardManager->IsUVLightCardEnabled(UVLightCardType))
 			{
 				// Custom viewport size from LC Manager
-				ContextSize = ViewportManager->LightCardManager->GetUVLightCardResourceSize();
+				ContextSize = ViewportManager->LightCardManager->GetUVLightCardResourceSize(UVLightCardType);
 
 				// Size must be not null
 				if (ContextSize.GetMin() > 1)
diff --git a/Engine/Plugins/Runtime/nDisplay/Source/DisplayCluster/Private/Render/Viewport/DisplayClusterViewport.h b/Engine/Plugins/Runtime/nDisplay/Source/DisplayCluster/Private/Render/Viewport/DisplayClusterViewport.h
--- a/Engine/Plugins/Runtime/nDisplay/Source/DisplayCluster/Private/Render/Viewport/DisplayClusterViewport.h
+++ b/Engine/Plugins/Runtime/nDisplay/Source/DisplayCluster/Private/Render/Viewport/DisplayClusterViewport.h
@@ -23,6 +23,8 @@
 
 #include "Templates/SharedPointer.h"
 
+class UDisplayClusterConfigurationViewport;
+
 class FDisplayClusterViewportManager;
 class FDisplayClusterViewportManagerProxy;
 class FDisplayClusterViewportProxy;
@@ -206,7 +208,7 @@ public:
 	void AddReferencedObjects(FReferenceCollector& Collector);
 
 	/** This function MUST always be called before configuring the viewport at the beginning of each frame. */
-	void ResetRuntimeParameters();	
+	void ResetRuntimeParameters(const UDisplayClusterConfigurationViewport* InConfigurationViewport = nullptr);
 
 	// Active view extension for this viewport
 	const TArray<FSceneViewExtensionRef> GatherActiveExtensions(FViewport* InViewport) const;
@@ -377,6 +379,12 @@ public:
 	/** setup viewport remap configuration for this viewport. */
 	bool UpdateConfiguration_ViewportRemap(const struct FDisplayClusterConfigurationViewport_Remap& InRemapConfiguration);
 
+	/** Returns viewport configuration data. */
+	const UDisplayClusterConfigurationViewport* GetViewportConfigurationData() const;
+
+	/** Sets viewport configuration data. */
+	void SetViewportConfigurationData(const UDisplayClusterConfigurationViewport* InConfigurationData);
+
 	/** Support view states for preview. */
 	FSceneViewStateInterface* GetViewState(uint32 ViewIndex);
 
@@ -471,6 +479,9 @@ private:
 	// View states (preview only)
 	TArray<TSharedPtr<FSceneViewStateReference, ESPMode::ThreadSafe>> ViewStates;
 
+	// The soft ptr to the configuration data of this viewport. Updated every frame.
+	TWeakObjectPtr<const UDisplayClusterConfigurationViewport> ConfigurationData;
+
 	// A recurring message in the log will be shown only once
 	mutable EDisplayClusterViewportShowLogMsgOnce ShowLogMsgOnceFlags = EDisplayClusterViewportShowLogMsgOnce::None;
 
diff --git a/Engine/Plugins/Runtime/nDisplay/Source/DisplayCluster/Private/Render/Viewport/DisplayClusterViewportProxy_Resources.cpp b/Engine/Plugins/Runtime/nDisplay/Source/DisplayCluster/Private/Render/Viewport/DisplayClusterViewportProxy_Resources.cpp
--- a/Engine/Plugins/Runtime/nDisplay/Source/DisplayCluster/Private/Render/Viewport/DisplayClusterViewportProxy_Resources.cpp
+++ b/Engine/Plugins/Runtime/nDisplay/Source/DisplayCluster/Private/Render/Viewport/DisplayClusterViewportProxy_Resources.cpp
@@ -275,7 +275,11 @@ bool FDisplayClusterViewportProxy::ImplGetResources_RenderThread(const EDisplayC
 					TSharedPtr<FDisplayClusterViewportLightCardManagerProxy, ESPMode::ThreadSafe> LightCardManager = ViewportManagerProxy->GetLightCardManagerProxy_RenderThread();
 					if (LightCardManager.IsValid())
 					{
-						if (FRHITexture* UVLightCardRHIResource = LightCardManager->GetUVLightCardRHIResource_RenderThread())
+						const EDisplayClusterUVLightCardType UVLightCardType =
+							EnumHasAllFlags(RenderSettingsICVFX.RuntimeFlags, EDisplayClusterViewportRuntimeICVFXFlags::OverInFrustum)
+							? EDisplayClusterUVLightCardType::Over : EDisplayClusterUVLightCardType::Under;
+
+						if (FRHITexture* UVLightCardRHIResource = LightCardManager->GetUVLightCardRHIResource_RenderThread(UVLightCardType))
 						{
 							for (int32 ContextIndex = 0; ContextIndex < Contexts.Num(); ContextIndex++)
 							{
diff --git a/Engine/Plugins/Runtime/nDisplay/Source/DisplayCluster/Private/Render/Viewport/DisplayClusterViewportStrings.h b/Engine/Plugins/Runtime/nDisplay/Source/DisplayCluster/Private/Render/Viewport/DisplayClusterViewportStrings.h
--- a/Engine/Plugins/Runtime/nDisplay/Source/DisplayCluster/Private/Render/Viewport/DisplayClusterViewportStrings.h
+++ b/Engine/Plugins/Runtime/nDisplay/Source/DisplayCluster/Private/Render/Viewport/DisplayClusterViewportStrings.h
@@ -17,9 +17,17 @@ namespace DisplayClusterViewportStrings
 		static constexpr const TCHAR* prefix       = TEXT("icvfx");
 		static constexpr const TCHAR* camera       = TEXT("incamera");
 		static constexpr const TCHAR* chromakey    = TEXT("chromakey");
+
 		static constexpr const TCHAR* lightcard    = TEXT("lightcard");
 		static constexpr const TCHAR* uv_lightcard = TEXT("uv_lightcard");
-		
+
+		// The viewport position to in-frustum.
+		namespace suffix
+		{
+			static constexpr const TCHAR* none = TEXT("");
+			static constexpr const TCHAR* over = TEXT("_over");
+			static constexpr const TCHAR* under = TEXT("_under");
+		}
 	}
 
 	namespace tile
diff --git a/Engine/Plugins/Runtime/nDisplay/Source/DisplayCluster/Private/Render/Viewport/DisplayClusterViewport_Configuration.cpp b/Engine/Plugins/Runtime/nDisplay/Source/DisplayCluster/Private/Render/Viewport/DisplayClusterViewport_Configuration.cpp
--- a/Engine/Plugins/Runtime/nDisplay/Source/DisplayCluster/Private/Render/Viewport/DisplayClusterViewport_Configuration.cpp
+++ b/Engine/Plugins/Runtime/nDisplay/Source/DisplayCluster/Private/Render/Viewport/DisplayClusterViewport_Configuration.cpp
@@ -19,6 +19,17 @@
 ///////////////////////////////////////////////////////////////////////////////////////
 // FDisplayClusterViewport
 ///////////////////////////////////////////////////////////////////////////////////////
+const UDisplayClusterConfigurationViewport* FDisplayClusterViewport::GetViewportConfigurationData() const
+{
+	return ConfigurationData.Get();
+}
+
+/** Sets viewport configuration data. */
+void FDisplayClusterViewport::SetViewportConfigurationData(const UDisplayClusterConfigurationViewport* InConfigurationData)
+{
+	ConfigurationData = InConfigurationData;
+}
+
 void FDisplayClusterViewport::UpdateConfiguration_ProjectionPolicy(const FDisplayClusterConfigurationProjection* InConfigurationProjectionPolicy)
 {
 	// Runtime update projection policy
diff --git a/Engine/Plugins/Runtime/nDisplay/Source/DisplayCluster/Private/Render/Viewport/DisplayClusterViewport_Rules.cpp b/Engine/Plugins/Runtime/nDisplay/Source/DisplayCluster/Private/Render/Viewport/DisplayClusterViewport_Rules.cpp
--- a/Engine/Plugins/Runtime/nDisplay/Source/DisplayCluster/Private/Render/Viewport/DisplayClusterViewport_Rules.cpp
+++ b/Engine/Plugins/Runtime/nDisplay/Source/DisplayCluster/Private/Render/Viewport/DisplayClusterViewport_Rules.cpp
@@ -35,6 +35,7 @@
 #include "UnrealClient.h"
 
 #include "DisplayClusterSceneViewExtensions.h"
+#include "DisplayClusterConfigurationTypes_Viewport.h"
 
 #include "LegacyScreenPercentageDriver.h"
 
@@ -47,7 +48,7 @@
 ///////////////////////////////////////////////////////////////////////////////////////
 // FDisplayClusterViewport
 ///////////////////////////////////////////////////////////////////////////////////////
-void FDisplayClusterViewport::ResetRuntimeParameters()
+void FDisplayClusterViewport::ResetRuntimeParameters(const UDisplayClusterConfigurationViewport* InConfigurationViewport)
 {
 	// Reset runtim flags from prev frame:
 	RenderSettings.BeginUpdateSettings();
@@ -63,8 +64,19 @@ void FDisplayClusterViewport::ResetRuntimeParameters()
 	// Obtain viewport media state from external multicast delegates (This viewport can be used by multiple media).
 	EDisplayClusterViewportMediaState AllMediaStates = EDisplayClusterViewportMediaState::None;
 	IDisplayCluster::Get().GetCallbacks().OnDisplayClusterUpdateViewportMediaState().Broadcast(this, AllMediaStates);
+
 	// Update the media state for the new frame.
 	RenderSettings.AssignMediaStates(AllMediaStates);
+
+	// Read general settings from the configuration
+	if (InConfigurationViewport)
+	{
+		if (const FDisplayClusterConfigurationICVFX_StageSettings* StageSettings = Configuration->GetStageSettings())
+		{
+			// Update base ICVFX settings of viewport.
+			RenderSettingsICVFX.Flags = InConfigurationViewport->GetViewportICVFXFlags(*StageSettings);
+		}
+	}
 }
 
 bool FDisplayClusterViewport::IsInternalViewport() const
diff --git a/Engine/Plugins/Runtime/nDisplay/Source/DisplayCluster/Private/Render/Viewport/DisplayClusterViewport_VisibilitySettings.cpp b/Engine/Plugins/Runtime/nDisplay/Source/DisplayCluster/Private/Render/Viewport/DisplayClusterViewport_VisibilitySettings.cpp
--- a/Engine/Plugins/Runtime/nDisplay/Source/DisplayCluster/Private/Render/Viewport/DisplayClusterViewport_VisibilitySettings.cpp
+++ b/Engine/Plugins/Runtime/nDisplay/Source/DisplayCluster/Private/Render/Viewport/DisplayClusterViewport_VisibilitySettings.cpp
@@ -3,84 +3,37 @@
 #include "Render/Viewport/DisplayClusterViewport_VisibilitySettings.h"
 #include "Render/Viewport/DisplayClusterViewportManager.h"
 #include "Render/Viewport/DisplayClusterViewport.h"
+#include "Render/Viewport/Configuration/DisplayClusterViewportConfigurationHelpers_Visibility.h"
 
 #include "EngineUtils.h"
 #include "SceneView.h"
 
-static void GetPrimitiveComponentsFromLayers(UWorld* World, const TArray<FName>& SourceLayers, TSet<FPrimitiveComponentId>& OutPrimitives,
-                                             const TSet<FPrimitiveComponentId>* InExcludePrimitivesList = nullptr)
+void FDisplayClusterViewport_VisibilitySettings::SetupSceneView(FSceneView& InOutView) const
 {
-	TRACE_CPUPROFILER_EVENT_SCOPE(DCRootActor_GetPrimitiveComponentsFromLayers);
-
-	if (SourceLayers.Num())
-	{
-		// Iterate over all actors, looking for actors in the specified layers.
-		for (const TWeakObjectPtr<AActor> WeakActor : FActorRange(World))
-		{
-			if (const AActor* Actor = WeakActor.Get())
-			{
-				bool bActorFoundOnSourceLayers = false;
-
-				// Search actor on source layers
-				for (const FName& LayerIt : SourceLayers)
-				{
-					if (Actor && Actor->Layers.Contains(LayerIt))
-					{
-						bActorFoundOnSourceLayers = true;
-						break;
-					}
-				}
-
-				if (bActorFoundOnSourceLayers)
-				{
-					// Save all actor components to OutPrimitives
-					for (UActorComponent* Component : Actor->GetComponents())
-					{
-						if (const UPrimitiveComponent* PrimComp = Cast<UPrimitiveComponent>(Component))
-						{
-							if (!InExcludePrimitivesList || !InExcludePrimitivesList->Contains(PrimComp->GetPrimitiveSceneId()))
-							{
-								OutPrimitives.Add(PrimComp->GetPrimitiveSceneId());
-							}
-						}
-					}
-				}
-			}
-		}
-	}
-}
-
-void FDisplayClusterViewport_VisibilitySettings::SetupSceneView(class UWorld* World, FSceneView& InOutView) const
-{
-	check(World);
-	
-	switch (LayersMode)
+	switch (VisibilityMode)
 	{
 	case EDisplayClusterViewport_VisibilityMode::ShowOnly:
 	{
-		if (ActorLayers.Num() > 0 || AdditionalComponentsList.Num() > 0)
+		InOutView.ShowOnlyPrimitives.Emplace();
+		for (const FPrimitiveComponentId& ComponentId : ComponentsList)
 		{
-			InOutView.ShowOnlyPrimitives.Emplace();
-			GetPrimitiveComponentsFromLayers(World, ActorLayers,InOutView.ShowOnlyPrimitives.GetValue(),
-				&RootActorHidePrimitivesList);
-
-			for (const FPrimitiveComponentId& AdditionalPrimitiveId : AdditionalComponentsList)
+			// Except hidden components
+			if (RootActorHidePrimitivesList.Contains(ComponentId))
 			{
-				if (!RootActorHidePrimitivesList.Contains(AdditionalPrimitiveId))
-				{
-					InOutView.ShowOnlyPrimitives.GetValue().Add(AdditionalPrimitiveId);
-				}
+				continue;
 			}
-			return;
+
+			InOutView.ShowOnlyPrimitives->Add(ComponentId);
 		}
-		break;
+
+		return;
 	}
 	case EDisplayClusterViewport_VisibilityMode::Hide:
 	{
-		GetPrimitiveComponentsFromLayers(World, ActorLayers, InOutView.HiddenPrimitives);
-		InOutView.HiddenPrimitives.Append(AdditionalComponentsList);
+		InOutView.HiddenPrimitives.Append(ComponentsList);
 		break;
 	}
+
 	default:
 		break;
 	}
diff --git a/Engine/Plugins/Runtime/nDisplay/Source/DisplayCluster/Private/Render/Viewport/DisplayClusterViewport_VisibilitySettings.h b/Engine/Plugins/Runtime/nDisplay/Source/DisplayCluster/Private/Render/Viewport/DisplayClusterViewport_VisibilitySettings.h
--- a/Engine/Plugins/Runtime/nDisplay/Source/DisplayCluster/Private/Render/Viewport/DisplayClusterViewport_VisibilitySettings.h
+++ b/Engine/Plugins/Runtime/nDisplay/Source/DisplayCluster/Private/Render/Viewport/DisplayClusterViewport_VisibilitySettings.h
@@ -28,46 +28,54 @@ public:
 	// Reset actor layers visibility rules
 	void BeginUpdateSettings()
 	{
-		LayersMode = EDisplayClusterViewport_VisibilityMode::None;
-
-		ActorLayers.Empty();
-		AdditionalComponentsList.Empty();
+		VisibilityMode = EDisplayClusterViewport_VisibilityMode::None;
+		ComponentsList.Empty();
 		RootActorHidePrimitivesList.Empty();
 	}
 
-	bool AppendHideList(const TArray<FName>& InActorLayers, const TSet<FPrimitiveComponentId>& InAdditionalComponentsList)
+	/** Sets the visibility mode and component list. */
+	void SetVisibilityModeAndComponentsList(EDisplayClusterViewport_VisibilityMode InVisibilityMode, const TSet<FPrimitiveComponentId>& InComponentsList)
 	{
-		if (LayersMode == EDisplayClusterViewport_VisibilityMode::Hide)
+
+		VisibilityMode = InVisibilityMode;
+		ComponentsList = InComponentsList;
+	}
+
+	/** Adds components to the list. Must be called after SetVisibilityModeAndComponentsList(). */
+	bool AppendVisibilityComponentsList(EDisplayClusterViewport_VisibilityMode InVisibilityMode, const TSet<FPrimitiveComponentId>& InComponentsList)
+	{
+		if (VisibilityMode == InVisibilityMode)
 		{
-			ActorLayers.Append(InActorLayers);
-			AdditionalComponentsList.Append(InAdditionalComponentsList);
+			ComponentsList.Append(InComponentsList);
+
 			return true;
 		}
 
 		return false;
 	}
 
-	void UpdateVisibilitySettings(EDisplayClusterViewport_VisibilityMode InMode, const TArray<FName>& InActorLayers, const TSet<FPrimitiveComponentId>& InAdditionalComponentsList)
+	void SetRootActorHideList(TSet<FPrimitiveComponentId>& InHidePrimitivesList)
 	{
-
-		LayersMode = InMode;
-		ActorLayers = InActorLayers;
-		AdditionalComponentsList = InAdditionalComponentsList;
+		RootActorHidePrimitivesList = InHidePrimitivesList;
 	}
 
-	void SetRootActorHideList(TSet<FPrimitiveComponentId>& InHidePrimitivesList)
+	/** Returns true if this viewport contains any geometry and can be rendered. */
+	bool IsVisible() const
 	{
-		RootActorHidePrimitivesList = InHidePrimitivesList;
+		if (VisibilityMode == EDisplayClusterViewport_VisibilityMode::ShowOnly && ComponentsList.IsEmpty())
+		{
+			return false;
+		}
+		
+		return true;
 	}
 
-	void SetupSceneView(UWorld* World, FSceneView& InOutView) const;
+	void SetupSceneView(FSceneView& InOutView) const;
 
 private:
-	EDisplayClusterViewport_VisibilityMode LayersMode = EDisplayClusterViewport_VisibilityMode::None;
-	TArray<FName> ActorLayers;
-	TSet<FPrimitiveComponentId> AdditionalComponentsList;
+	EDisplayClusterViewport_VisibilityMode VisibilityMode = EDisplayClusterViewport_VisibilityMode::None;
+	TSet<FPrimitiveComponentId> ComponentsList;
 
 	// Additional hide primitives list from root actor
 	TSet<FPrimitiveComponentId> RootActorHidePrimitivesList;
 };
-
diff --git a/Engine/Plugins/Runtime/nDisplay/Source/DisplayCluster/Private/Render/Viewport/LightCard/DisplayClusterViewportLightCardEnums.h b/Engine/Plugins/Runtime/nDisplay/Source/DisplayCluster/Private/Render/Viewport/LightCard/DisplayClusterViewportLightCardEnums.h
new file mode 100644
--- /dev/null
+++ b/Engine/Plugins/Runtime/nDisplay/Source/DisplayCluster/Private/Render/Viewport/LightCard/DisplayClusterViewportLightCardEnums.h
@@ -0,0 +1,31 @@
+// Copyright Epic Games, Inc. All Rights Reserved.
+
+#pragma once
+
+#include "CoreMinimal.h"
+
+/* Consolidated UV light card rendering mode for the current cluster node. **/
+enum class EDisplayClusterUVLightCardRenderMode : uint8
+{
+	// Use light card render mode from the StageSettings
+	Default,
+
+	// The lightcard will always be displayed only “over the In-Camera”.
+	AlwaysOver,
+
+	// The lightcard will always be displayed only “under the In-Camera”.
+	AlwaysUnder,
+
+	// Lighting maps are not used.
+	Disabled
+};
+
+/** The UV lightcards have 2 type: 'Under' and 'Over'. */
+enum class EDisplayClusterUVLightCardType: uint8
+{
+	// The UV lightcards that are displayed under inner frustum.
+	Under,
+
+	// The UV lightcards that are displayed over inner frustum.
+	Over
+};
diff --git a/Engine/Plugins/Runtime/nDisplay/Source/DisplayCluster/Private/Render/Viewport/LightCard/DisplayClusterViewportLightCardManager.cpp b/Engine/Plugins/Runtime/nDisplay/Source/DisplayCluster/Private/Render/Viewport/LightCard/DisplayClusterViewportLightCardManager.cpp
--- a/Engine/Plugins/Runtime/nDisplay/Source/DisplayCluster/Private/Render/Viewport/LightCard/DisplayClusterViewportLightCardManager.cpp
+++ b/Engine/Plugins/Runtime/nDisplay/Source/DisplayCluster/Private/Render/Viewport/LightCard/DisplayClusterViewportLightCardManager.cpp
@@ -44,8 +44,11 @@ void FDisplayClusterViewportLightCardManager::Release()
 	check(IsInGameThread());
 
 	// Release UVLightCard
-	ReleaseUVLightCardData();
-	ReleaseUVLightCardResource();
+	ReleaseUVLightCardData(EDisplayClusterUVLightCardType::Under);
+	ReleaseUVLightCardResource(EDisplayClusterUVLightCardType::Under);
+
+	ReleaseUVLightCardData(EDisplayClusterUVLightCardType::Over);
+	ReleaseUVLightCardResource(EDisplayClusterUVLightCardType::Over);
 }
 
 ///////////////////////////////////////////////////////////////////////////////////////////////
@@ -55,34 +58,141 @@ void FDisplayClusterViewportLightCardManager::OnHandleStartScene()
 
 void FDisplayClusterViewportLightCardManager::OnHandleEndScene()
 {
-	ReleaseUVLightCardData();
+	ReleaseUVLightCardData(EDisplayClusterUVLightCardType::Under);
+	ReleaseUVLightCardData(EDisplayClusterUVLightCardType::Over);
 }
 
 void FDisplayClusterViewportLightCardManager::RenderFrame()
 {
-	UpdateUVLightCardData();
-	RenderUVLightCard();
+	UpdateUVLightCardData(EDisplayClusterUVLightCardType::Under);
+	RenderUVLightCard(EDisplayClusterUVLightCardType::Under);
+
+	UpdateUVLightCardData(EDisplayClusterUVLightCardType::Over);
+	RenderUVLightCard(EDisplayClusterUVLightCardType::Over);
 }
 
 ///////////////////////////////////////////////////////////////////////////////////////////////
-FIntPoint FDisplayClusterViewportLightCardManager::GetUVLightCardResourceSize() const
+FIntPoint FDisplayClusterViewportLightCardManager::GetUVLightCardResourceSize(const EDisplayClusterUVLightCardType InUVLightCardType) const
 {
+	const TSharedPtr<FDisplayClusterViewportLightCardResource, ESPMode::ThreadSafe>& UVLightCardResource = GetUVLightCardResource(InUVLightCardType);
 	return UVLightCardResource.IsValid() ? UVLightCardResource->GetSizeXY() : FIntPoint(0, 0);
 }
 
-bool FDisplayClusterViewportLightCardManager::IsUVLightCardEnabled() const
+EDisplayClusterUVLightCardRenderMode FDisplayClusterViewportLightCardManager::GetUVLightCardRenderMode() const
+{
+	// Value should be updated.
+	if (const FDisplayClusterViewportManager* ViewportManager = Configuration->GetViewportManagerImpl())
+	{
+		// Collect all ICVFX flags
+		EDisplayClusterViewportICVFXFlags ICVFXFlags = EDisplayClusterViewportICVFXFlags::None;
+		for (const TSharedPtr<IDisplayClusterViewport, ESPMode::ThreadSafe>& ViewportIt : ViewportManager->GetEntireClusterViewports())
+		{
+			if (ViewportIt.IsValid())
+			{
+				ICVFXFlags |= ViewportIt->GetRenderSettingsICVFX().Flags;
+			}
+		}
+
+		if (EnumHasAnyFlags(ICVFXFlags, EDisplayClusterViewportICVFXFlags::DisableLightcard))
+		{
+			return EDisplayClusterUVLightCardRenderMode::Disabled;
+		}
+		else
+		{
+			// Returns the consolidated light card rendering mode for this cluster node:
+			const EDisplayClusterViewportICVFXFlags LightcardRenderModeFlags = ICVFXFlags & EDisplayClusterViewportICVFXFlags::LightcardRenderModeMask;
+			if (LightcardRenderModeFlags == EDisplayClusterViewportICVFXFlags::LightcardAlwaysUnder)
+			{
+				// The lightcard will always be displayed only “under the In-Camera” for this cluster node.
+				return EDisplayClusterUVLightCardRenderMode::AlwaysUnder;
+			}
+			else if (LightcardRenderModeFlags == EDisplayClusterViewportICVFXFlags::LightcardAlwaysOver)
+			{
+				// The lightcard will always be displayed only “over the In-Camera” for this cluster node.
+				return EDisplayClusterUVLightCardRenderMode::AlwaysOver;
+			}
+		}
+	}
+	
+	return EDisplayClusterUVLightCardRenderMode::Default;
+}
+
+bool FDisplayClusterViewportLightCardManager::IsUVLightCardEnabled(const EDisplayClusterUVLightCardType InUVLightCardType) const
 {
+	switch (GetUVLightCardRenderMode())
+	{
+	case EDisplayClusterUVLightCardRenderMode::Disabled:
+		return false;
+
+	case EDisplayClusterUVLightCardRenderMode::AlwaysOver:
+		if (InUVLightCardType != EDisplayClusterUVLightCardType::Over)
+		{
+			// Force to render only over
+			return false;
+		}
+		break;
+
+	case EDisplayClusterUVLightCardRenderMode::AlwaysUnder:
+		if (InUVLightCardType != EDisplayClusterUVLightCardType::Under)
+		{
+			// Force to render only under
+			return false;
+		}
+		break;
+
+	default:
+		break;
+	}
+
+	const TArray<UPrimitiveComponent*>& UVLightCardPrimitiveComponents = GetUVLightCardPrimitiveComponents(InUVLightCardType);
 	return !UVLightCardPrimitiveComponents.IsEmpty();
 }
 
-void FDisplayClusterViewportLightCardManager::ReleaseUVLightCardData()
+void FDisplayClusterViewportLightCardManager::ReleaseUVLightCardData(const EDisplayClusterUVLightCardType InUVLightCardType)
 {
+	TArray<UPrimitiveComponent*>& UVLightCardPrimitiveComponents = GetUVLightCardPrimitiveComponents(InUVLightCardType);
 	UVLightCardPrimitiveComponents.Empty();
 }
 
-void FDisplayClusterViewportLightCardManager::UpdateUVLightCardData()
+void FDisplayClusterViewportLightCardManager::UpdateUVLightCardData(const EDisplayClusterUVLightCardType InUVLightCardType)
 {
-	ReleaseUVLightCardData();
+	ReleaseUVLightCardData(InUVLightCardType);
+
+	// Special use-case - when all viewports force to use only over or under, we have to ignore per-light card mode.
+	bool bEnablePerLightcardRenderMode = true;
+	switch (GetUVLightCardRenderMode())
+	{
+	case EDisplayClusterUVLightCardRenderMode::Disabled:
+		return;
+
+	case EDisplayClusterUVLightCardRenderMode::AlwaysOver:
+		bEnablePerLightcardRenderMode = false;
+		if (InUVLightCardType != EDisplayClusterUVLightCardType::Over)
+		{
+			// Force to render only over
+			return;
+		}
+		break;
+
+	case EDisplayClusterUVLightCardRenderMode::AlwaysUnder:
+		bEnablePerLightcardRenderMode = false;
+		if (InUVLightCardType != EDisplayClusterUVLightCardType::Under)
+		{
+			// Force to render only under
+			return;
+		}
+		break;
+
+	default:
+		break;
+	}
+
+
+	const FDisplayClusterConfigurationICVFX_StageSettings* StageSettings = Configuration->GetStageSettings();
+	if (!StageSettings)
+	{
+		return;
+	}
 
 	/** The list of UV light card actors that are referenced by the root actor */
 	TArray<ADisplayClusterLightCardActor*> UVLightCardActors;
@@ -96,7 +206,22 @@ void FDisplayClusterViewportLightCardManager::UpdateUVLightCardData()
 		{
 			if (LightCard->bIsUVLightCard)
 			{
-				UVLightCardActors.Add(LightCard);
+				if (bEnablePerLightcardRenderMode)
+				{
+					// Per-lightcard rules:
+					const EDisplayClusterShaderParametersICVFX_LightCardRenderMode LightCardRenderMode = StageSettings->Lightcard.GetLightCardRenderMode(LightCard->PerLightcardRenderMode, nullptr);
+					const bool bLightCardActorOver = (LightCardRenderMode == EDisplayClusterShaderParametersICVFX_LightCardRenderMode::Over);
+					const bool bLightCardTypeOver = InUVLightCardType == EDisplayClusterUVLightCardType::Over;
+					if (bLightCardTypeOver == bLightCardActorOver)
+					{
+						UVLightCardActors.Add(LightCard);
+					}
+				}
+				else
+				{
+					// Render all UVLC to the one RTT
+					UVLightCardActors.Add(LightCard);
+				}
 			}
 		}
 	}
@@ -112,6 +237,7 @@ void FDisplayClusterViewportLightCardManager::UpdateUVLightCardData()
 		LightCardMeshComponents.Empty(LightCardMeshComponents.Num());
 		LightCard->GetLightCardMeshComponents(LightCardMeshComponents);
 
+		TArray<UPrimitiveComponent*>& UVLightCardPrimitiveComponents = GetUVLightCardPrimitiveComponents(InUVLightCardType);
 		for (UMeshComponent* LightCardMeshComp : LightCardMeshComponents)
 		{
 			if (LightCardMeshComp && LightCardMeshComp->SceneProxy == nullptr)
@@ -122,50 +248,54 @@ void FDisplayClusterViewportLightCardManager::UpdateUVLightCardData()
 	}
 }
 
-void FDisplayClusterViewportLightCardManager::CreateUVLightCardResource(const FIntPoint& InResourceSize)
+void FDisplayClusterViewportLightCardManager::CreateUVLightCardResource(const FIntPoint& InResourceSize, const EDisplayClusterUVLightCardType InUVLightCardType)
 {
+	TSharedPtr<FDisplayClusterViewportLightCardResource, ESPMode::ThreadSafe>& UVLightCardResource = GetUVLightCardResource(InUVLightCardType);
 	UVLightCardResource = MakeShared<FDisplayClusterViewportLightCardResource>(InResourceSize);
-	LightCardManagerProxy->UpdateUVLightCardResource(UVLightCardResource);
+	LightCardManagerProxy->UpdateUVLightCardResource(UVLightCardResource, InUVLightCardType);
 }
 
-void FDisplayClusterViewportLightCardManager::ReleaseUVLightCardResource()
+void FDisplayClusterViewportLightCardManager::ReleaseUVLightCardResource(const EDisplayClusterUVLightCardType InUVLightCardType)
 {
+	TSharedPtr<FDisplayClusterViewportLightCardResource, ESPMode::ThreadSafe>& UVLightCardResource = GetUVLightCardResource(InUVLightCardType);
 	if (UVLightCardResource.IsValid())
 	{
-		LightCardManagerProxy->ReleaseUVLightCardResource();
+		LightCardManagerProxy->ReleaseUVLightCardResource(InUVLightCardType);
 	}
 
 	UVLightCardResource.Reset();
 }
 
-void FDisplayClusterViewportLightCardManager::UpdateUVLightCardResource()
+void FDisplayClusterViewportLightCardManager::UpdateUVLightCardResource(const EDisplayClusterUVLightCardType InUVLightCardType)
 {
 	const uint32 UVLightCardTextureSize = CVarUVLightCardTextureSize.GetValueOnGameThread();
 	const FIntPoint UVLightCardResourceSize = FIntPoint(UVLightCardTextureSize, UVLightCardTextureSize);
 
+	TSharedPtr<FDisplayClusterViewportLightCardResource, ESPMode::ThreadSafe>& UVLightCardResource = GetUVLightCardResource(InUVLightCardType);
 	if (UVLightCardResource.IsValid())
 	{
 		if (UVLightCardResource->GetSizeXY() != UVLightCardResourceSize)
 		{
-			ReleaseUVLightCardResource();
+			ReleaseUVLightCardResource(InUVLightCardType);
 		}
 	}
 
 	if (!UVLightCardResource.IsValid())
 	{
-		CreateUVLightCardResource(UVLightCardResourceSize);
+		CreateUVLightCardResource(UVLightCardResourceSize, InUVLightCardType);
 	}
 }
 
-void FDisplayClusterViewportLightCardManager::RenderUVLightCard()
+void FDisplayClusterViewportLightCardManager::RenderUVLightCard(const EDisplayClusterUVLightCardType InUVLightCardType)
 {
 	// Render UV LightCard:
 	FDisplayClusterViewportManager* ViewportManager = Configuration->GetViewportManagerImpl();
 	UWorld* CurrentWorld = Configuration->GetCurrentWorld();
-	if (IsUVLightCardEnabled() && CurrentWorld && ViewportManager)
+	if (IsUVLightCardEnabled(InUVLightCardType) && CurrentWorld && ViewportManager)
 	{
-		UpdateUVLightCardResource();
+		UpdateUVLightCardResource(InUVLightCardType);
 
+		TSharedPtr<FDisplayClusterViewportLightCardResource, ESPMode::ThreadSafe>& UVLightCardResource = GetUVLightCardResource(InUVLightCardType);
 		if (UVLightCardResource.IsValid())
 		{
 			FDisplayClusterShaderParameters_UVLightCards UVLightCardParameters;
@@ -175,6 +305,7 @@ void FDisplayClusterViewportLightCardManager::RenderUVLightCard()
 
 			// Store any components that were invisible but forced to be visible so they can be set back to invisible after the render
 			TArray<UPrimitiveComponent*> ComponentsToUnload;
+			const TArray<UPrimitiveComponent*>& UVLightCardPrimitiveComponents = GetUVLightCardPrimitiveComponents(InUVLightCardType);
 			for (UPrimitiveComponent* PrimitiveComponent : UVLightCardPrimitiveComponents)
 			{
 				// Set the component's visibility to true and force it to generate its scene proxies
@@ -191,7 +322,7 @@ void FDisplayClusterViewportLightCardManager::RenderUVLightCard()
 				}
 			}
 
-			LightCardManagerProxy->RenderUVLightCard(CurrentWorld->Scene, UVLightCardParameters);
+			LightCardManagerProxy->RenderUVLightCard(CurrentWorld->Scene, UVLightCardParameters, InUVLightCardType);
 
 			for (UPrimitiveComponent* LoadedComponent : ComponentsToUnload)
 			{
@@ -202,7 +333,7 @@ void FDisplayClusterViewportLightCardManager::RenderUVLightCard()
 	}
 	else
 	{
-		ReleaseUVLightCardResource();
+		ReleaseUVLightCardResource(InUVLightCardType);
 	}
 }
 
diff --git a/Engine/Plugins/Runtime/nDisplay/Source/DisplayCluster/Private/Render/Viewport/LightCard/DisplayClusterViewportLightCardManager.h b/Engine/Plugins/Runtime/nDisplay/Source/DisplayCluster/Private/Render/Viewport/LightCard/DisplayClusterViewportLightCardManager.h
--- a/Engine/Plugins/Runtime/nDisplay/Source/DisplayCluster/Private/Render/Viewport/LightCard/DisplayClusterViewportLightCardManager.h
+++ b/Engine/Plugins/Runtime/nDisplay/Source/DisplayCluster/Private/Render/Viewport/LightCard/DisplayClusterViewportLightCardManager.h
@@ -5,7 +5,7 @@
 #include "CoreMinimal.h"
 
 #include "DisplayClusterViewportLightCardManagerProxy.h"
-#include "DisplayClusterViewportLightCardResource.h"
+#include "Containers/DisplayClusterShader_Enums.h"
 
 #include "UObject/GCObject.h"
 
@@ -32,10 +32,10 @@ public:
 
 public:
 	/** Return true if UV LightCard is used in this frame. */
-	bool IsUVLightCardEnabled() const;
+	bool IsUVLightCardEnabled(const EDisplayClusterUVLightCardType InUVLightCardType) const;
 
 	/** Get UV LightCard texture size. */
-	FIntPoint GetUVLightCardResourceSize() const;
+	FIntPoint GetUVLightCardResourceSize(const EDisplayClusterUVLightCardType InUVLightCardType) const;
 
 public:
 	/** Handle StartScene event: created and update internal resources. */
@@ -49,22 +49,51 @@ public:
 
 private:
 	/** Render UVLightCard */
-	void RenderUVLightCard();
+	void RenderUVLightCard(const EDisplayClusterUVLightCardType InUVLightCardType);
 
 	/** Update the UV light card map texture */
-	void UpdateUVLightCardResource();
+	void UpdateUVLightCardResource(const EDisplayClusterUVLightCardType InUVLightCardType);
 
 	/** Releases the UV light card map texture */
-	void ReleaseUVLightCardResource();
+	void ReleaseUVLightCardResource(const EDisplayClusterUVLightCardType InUVLightCardType);
 
 	/** Create the UV light card map texture */
-	void CreateUVLightCardResource(const FIntPoint& InResourceSize);
+	void CreateUVLightCardResource(const FIntPoint& InResourceSize, const EDisplayClusterUVLightCardType InUVLightCardType);
 
 	/** Update UVLightCard data game thread*/
-	void UpdateUVLightCardData();
+	void UpdateUVLightCardData(const EDisplayClusterUVLightCardType InUVLightCardType);
 
 	/** Release UVLightCard data game thread*/
-	void ReleaseUVLightCardData();
+	void ReleaseUVLightCardData(const EDisplayClusterUVLightCardType InUVLightCardType);
+
+private:
+
+
+	/** Returns the consolidated UV-light card rendering mode for the current cluster node.
+	* If the light card override mode is the same for all viewports of the current cluster node, return this value.
+	* Or returns 'None' if they do not.
+	*/
+	EDisplayClusterUVLightCardRenderMode GetUVLightCardRenderMode() const;
+
+	TSharedPtr<FDisplayClusterViewportLightCardResource, ESPMode::ThreadSafe>& GetUVLightCardResource(const EDisplayClusterUVLightCardType InUVLightCardType)
+	{
+		return (InUVLightCardType == EDisplayClusterUVLightCardType::Over) ? UVLightCardOverResource : UVLightCardUnderResource;
+	}
+
+	const TSharedPtr<FDisplayClusterViewportLightCardResource, ESPMode::ThreadSafe>& GetUVLightCardResource(const EDisplayClusterUVLightCardType InUVLightCardType) const
+	{
+		return (InUVLightCardType == EDisplayClusterUVLightCardType::Over) ? UVLightCardOverResource : UVLightCardUnderResource;
+	}
+
+	TArray<UPrimitiveComponent*>& GetUVLightCardPrimitiveComponents(const EDisplayClusterUVLightCardType InUVLightCardType)
+	{
+		return (InUVLightCardType == EDisplayClusterUVLightCardType::Over) ? UVLightCardOverPrimitiveComponents : UVLightCardUnderPrimitiveComponents;
+	}
+
+	const TArray<UPrimitiveComponent*>& GetUVLightCardPrimitiveComponents(const EDisplayClusterUVLightCardType InUVLightCardType) const
+	{
+		return (InUVLightCardType == EDisplayClusterUVLightCardType::Over) ? UVLightCardOverPrimitiveComponents : UVLightCardUnderPrimitiveComponents;
+	}
 
 public:
 	// Configuration of the current cluster node
@@ -75,9 +104,14 @@ public:
 
 private:
 	/** A list of primitive components that have been added to the preview scene for rendering in the current frame */
-	TArray<UPrimitiveComponent*> UVLightCardPrimitiveComponents;
+	TArray<UPrimitiveComponent*> UVLightCardUnderPrimitiveComponents;
+
+	/** A list of primitive components that have been added to the preview scene for rendering in the current frame */
+	TArray<UPrimitiveComponent*> UVLightCardOverPrimitiveComponents;
 
 	/** The render target to which the UV light card map is rendered */
-	TSharedPtr<FDisplayClusterViewportLightCardResource, ESPMode::ThreadSafe> UVLightCardResource;
+	TSharedPtr<FDisplayClusterViewportLightCardResource, ESPMode::ThreadSafe> UVLightCardUnderResource;
 
+	/** The render target to which the UV light card map is rendered */
+	TSharedPtr<FDisplayClusterViewportLightCardResource, ESPMode::ThreadSafe> UVLightCardOverResource;
 };
diff --git a/Engine/Plugins/Runtime/nDisplay/Source/DisplayCluster/Private/Render/Viewport/LightCard/DisplayClusterViewportLightCardManagerProxy.cpp b/Engine/Plugins/Runtime/nDisplay/Source/DisplayCluster/Private/Render/Viewport/LightCard/DisplayClusterViewportLightCardManagerProxy.cpp
--- a/Engine/Plugins/Runtime/nDisplay/Source/DisplayCluster/Private/Render/Viewport/LightCard/DisplayClusterViewportLightCardManagerProxy.cpp
+++ b/Engine/Plugins/Runtime/nDisplay/Source/DisplayCluster/Private/Render/Viewport/LightCard/DisplayClusterViewportLightCardManagerProxy.cpp
@@ -13,49 +13,52 @@
 ///////////////////////////////////////////////////////////////////////////////////////////////
 FDisplayClusterViewportLightCardManagerProxy::~FDisplayClusterViewportLightCardManagerProxy()
 {
-	ImplReleaseUVLightCardResource_RenderThread();
+	ImplReleaseUVLightCardResource_RenderThread(EDisplayClusterUVLightCardType::Under);
+	ImplReleaseUVLightCardResource_RenderThread(EDisplayClusterUVLightCardType::Over);
 }
 
 ///////////////////////////////////////////////////////////////////////////////////////////////
-FRHITexture* FDisplayClusterViewportLightCardManagerProxy::GetUVLightCardRHIResource_RenderThread() const
+FRHITexture* FDisplayClusterViewportLightCardManagerProxy::GetUVLightCardRHIResource_RenderThread(const EDisplayClusterUVLightCardType InUVLightCardType) const
 {
+	const TSharedPtr<FDisplayClusterViewportLightCardResource, ESPMode::ThreadSafe>& UVLightCardMapResource = GetUVLightCardMapResource(InUVLightCardType);
 	return UVLightCardMapResource.IsValid() ? UVLightCardMapResource->GetTextureRHI().GetReference() : nullptr;
 }
 
-void FDisplayClusterViewportLightCardManagerProxy::UpdateUVLightCardResource(const TSharedPtr<FDisplayClusterViewportLightCardResource, ESPMode::ThreadSafe>& InUVLightCardMapResource)
+void FDisplayClusterViewportLightCardManagerProxy::UpdateUVLightCardResource(const TSharedPtr<FDisplayClusterViewportLightCardResource, ESPMode::ThreadSafe>& InUVLightCardMapResource, const EDisplayClusterUVLightCardType InUVLightCardType)
 {
 	ENQUEUE_RENDER_COMMAND(DisplayClusterViewportLightCardManagerProxy_UpdateUVLightCardResource)(
-		[InProxyData = SharedThis(this), UVLightCardMapResource = InUVLightCardMapResource](FRHICommandListImmediate& RHICmdList)
+		[InProxyData = SharedThis(this), NewUVLightCardMapResource = InUVLightCardMapResource, InUVLightCardType](FRHICommandListImmediate& RHICmdList)
 		{
-			InProxyData->ImplUpdateUVLightCardResource_RenderThread(UVLightCardMapResource);
+			InProxyData->ImplUpdateUVLightCardResource_RenderThread(NewUVLightCardMapResource, InUVLightCardType);
 		});
 }
 
-void FDisplayClusterViewportLightCardManagerProxy::ReleaseUVLightCardResource()
+void FDisplayClusterViewportLightCardManagerProxy::ReleaseUVLightCardResource(const EDisplayClusterUVLightCardType InUVLightCardType)
 {
 	ENQUEUE_RENDER_COMMAND(DisplayClusterViewportLightCardManagerProxy_ReleaseUVLightCardResource)(
-		[InProxyData = SharedThis(this)](FRHICommandListImmediate& RHICmdList)
+		[InProxyData = SharedThis(this), InUVLightCardType](FRHICommandListImmediate& RHICmdList)
 		{
-			InProxyData->ImplReleaseUVLightCardResource_RenderThread();
+			InProxyData->ImplReleaseUVLightCardResource_RenderThread(InUVLightCardType);
 		});
 }
 
-void FDisplayClusterViewportLightCardManagerProxy::RenderUVLightCard(FSceneInterface* InScene, const FDisplayClusterShaderParameters_UVLightCards& InParameters) const
+void FDisplayClusterViewportLightCardManagerProxy::RenderUVLightCard(FSceneInterface* InScene, const FDisplayClusterShaderParameters_UVLightCards& InParameters, const EDisplayClusterUVLightCardType InUVLightCardType) const
 {
 	UE::RenderCommandPipe::FSyncScope SyncScope;
 
 	ENQUEUE_RENDER_COMMAND(DisplayClusterViewportLightCardManagerProxy_RenderUVLightCard)(
-		[InProxyData = SharedThis(this), Scene = InScene, Parameters = InParameters](FRHICommandListImmediate& RHICmdList)
+		[InProxyData = SharedThis(this), Scene = InScene, Parameters = InParameters, InUVLightCardType](FRHICommandListImmediate& RHICmdList)
 		{
-			InProxyData->ImplRenderUVLightCard_RenderThread(RHICmdList, Scene, Parameters);
+			InProxyData->ImplRenderUVLightCard_RenderThread(RHICmdList, Scene, Parameters, InUVLightCardType);
 		});
 }
 
-void FDisplayClusterViewportLightCardManagerProxy::ImplUpdateUVLightCardResource_RenderThread(const TSharedPtr<FDisplayClusterViewportLightCardResource, ESPMode::ThreadSafe>& InUVLightCardMapResource)
+void FDisplayClusterViewportLightCardManagerProxy::ImplUpdateUVLightCardResource_RenderThread(const TSharedPtr<FDisplayClusterViewportLightCardResource, ESPMode::ThreadSafe>& InUVLightCardMapResource, const EDisplayClusterUVLightCardType InUVLightCardType)
 {
+	TSharedPtr<FDisplayClusterViewportLightCardResource, ESPMode::ThreadSafe>& UVLightCardMapResource = GetUVLightCardMapResource(InUVLightCardType);
 	if (UVLightCardMapResource != InUVLightCardMapResource)
 	{
-		ImplReleaseUVLightCardResource_RenderThread();
+		ImplReleaseUVLightCardResource_RenderThread(InUVLightCardType);
 
 		// Update resource ptr
 		UVLightCardMapResource = InUVLightCardMapResource;
@@ -66,19 +69,20 @@ void FDisplayClusterViewportLightCardManagerProxy::ImplUpdateUVLightCardResource
 	}
 }
 
-void FDisplayClusterViewportLightCardManagerProxy::ImplReleaseUVLightCardResource_RenderThread()
+void FDisplayClusterViewportLightCardManagerProxy::ImplReleaseUVLightCardResource_RenderThread(const EDisplayClusterUVLightCardType InUVLightCardType)
 {
 	// Release the texture's resources and delete the texture object from the rendering thread
+	TSharedPtr<FDisplayClusterViewportLightCardResource, ESPMode::ThreadSafe>& UVLightCardMapResource = GetUVLightCardMapResource(InUVLightCardType);
 	if (UVLightCardMapResource.IsValid())
 	{
 		UVLightCardMapResource->ReleaseResource();
-
 		UVLightCardMapResource.Reset();
 	}
 }
 
-void FDisplayClusterViewportLightCardManagerProxy::ImplRenderUVLightCard_RenderThread(FRHICommandListImmediate& RHICmdList, FSceneInterface* InSceneInterface, const FDisplayClusterShaderParameters_UVLightCards& InParameters) const
+void FDisplayClusterViewportLightCardManagerProxy::ImplRenderUVLightCard_RenderThread(FRHICommandListImmediate& RHICmdList, FSceneInterface* InSceneInterface, const FDisplayClusterShaderParameters_UVLightCards& InParameters, const EDisplayClusterUVLightCardType InUVLightCardType) const
 {
+	const TSharedPtr<FDisplayClusterViewportLightCardResource, ESPMode::ThreadSafe>& UVLightCardMapResource = GetUVLightCardMapResource(InUVLightCardType);
 	if (InParameters.PrimitivesToRender.Num() && UVLightCardMapResource.IsValid())
 	{
 		IDisplayClusterShaders& ShadersAPI = IDisplayClusterShaders::Get();
diff --git a/Engine/Plugins/Runtime/nDisplay/Source/DisplayCluster/Private/Render/Viewport/LightCard/DisplayClusterViewportLightCardManagerProxy.h b/Engine/Plugins/Runtime/nDisplay/Source/DisplayCluster/Private/Render/Viewport/LightCard/DisplayClusterViewportLightCardManagerProxy.h
--- a/Engine/Plugins/Runtime/nDisplay/Source/DisplayCluster/Private/Render/Viewport/LightCard/DisplayClusterViewportLightCardManagerProxy.h
+++ b/Engine/Plugins/Runtime/nDisplay/Source/DisplayCluster/Private/Render/Viewport/LightCard/DisplayClusterViewportLightCardManagerProxy.h
@@ -4,6 +4,7 @@
 
 #include "CoreMinimal.h"
 
+#include "DisplayClusterViewportLightCardEnums.h"
 #include "DisplayClusterViewportLightCardResource.h"
 
 #include "RHIResources.h"
@@ -22,24 +23,38 @@ public:
 	virtual ~FDisplayClusterViewportLightCardManagerProxy();
 
 	/** Update UVLightCard resource. */
-	void UpdateUVLightCardResource(const TSharedPtr<FDisplayClusterViewportLightCardResource, ESPMode::ThreadSafe>& InUVLightCardMapResource);
+	void UpdateUVLightCardResource(const TSharedPtr<FDisplayClusterViewportLightCardResource, ESPMode::ThreadSafe>& InUVLightCardMapResource, const EDisplayClusterUVLightCardType InUVLightCardType);
 
 	/** Release UVLightCard resource. */
-	void ReleaseUVLightCardResource();
+	void ReleaseUVLightCardResource(const EDisplayClusterUVLightCardType InUVLightCardType);
 
 	/** Render UVLightCard. */
-	void RenderUVLightCard(FSceneInterface* InSceneInterface, const FDisplayClusterShaderParameters_UVLightCards& InParameters) const;
+	void RenderUVLightCard(FSceneInterface* InSceneInterface, const FDisplayClusterShaderParameters_UVLightCards& InParameters, const EDisplayClusterUVLightCardType InUVLightCardType) const;
 
 	/** Get current UVLightCard RHI resource on rendering thread. */
-	FRHITexture* GetUVLightCardRHIResource_RenderThread() const;
+	FRHITexture* GetUVLightCardRHIResource_RenderThread(const EDisplayClusterUVLightCardType InUVLightCardType) const;
 
 protected:
-	void ImplUpdateUVLightCardResource_RenderThread(const TSharedPtr<FDisplayClusterViewportLightCardResource, ESPMode::ThreadSafe>& InUVLightCardMapResource);
-	void ImplReleaseUVLightCardResource_RenderThread();
+	void ImplUpdateUVLightCardResource_RenderThread(const TSharedPtr<FDisplayClusterViewportLightCardResource, ESPMode::ThreadSafe>& InUVLightCardMapResource, const EDisplayClusterUVLightCardType InUVLightCardType);
+	void ImplReleaseUVLightCardResource_RenderThread(const EDisplayClusterUVLightCardType InUVLightCardType);
 
-	void ImplRenderUVLightCard_RenderThread(FRHICommandListImmediate& RHICmdList, FSceneInterface* InScene, const FDisplayClusterShaderParameters_UVLightCards& InParameters) const;
+	/** Render UV-lightcards to texture (the layer is defined by the bOverInFrustum parameter.). */
+	void ImplRenderUVLightCard_RenderThread(FRHICommandListImmediate& RHICmdList, FSceneInterface* InScene, const FDisplayClusterShaderParameters_UVLightCards& InParameters, const EDisplayClusterUVLightCardType InUVLightCardType) const;
+
+	TSharedPtr<FDisplayClusterViewportLightCardResource, ESPMode::ThreadSafe>& GetUVLightCardMapResource(const EDisplayClusterUVLightCardType InUVLightCardType)
+	{
+		return (InUVLightCardType == EDisplayClusterUVLightCardType::Over) ? UVLightCardOverMapResource: UVLightCardUnderMapResource;
+	}
+
+	const TSharedPtr<FDisplayClusterViewportLightCardResource, ESPMode::ThreadSafe>& GetUVLightCardMapResource(const EDisplayClusterUVLightCardType InUVLightCardType) const
+	{
+		return (InUVLightCardType == EDisplayClusterUVLightCardType::Over) ? UVLightCardOverMapResource : UVLightCardUnderMapResource;
+	}
 
 private:
 	/** The render thread copy of the pointer to the UV light card map */
-	TSharedPtr<FDisplayClusterViewportLightCardResource, ESPMode::ThreadSafe> UVLightCardMapResource;
+	TSharedPtr<FDisplayClusterViewportLightCardResource, ESPMode::ThreadSafe> UVLightCardUnderMapResource;
+
+	/** The render thread copy of the pointer to the UV light card map */
+	TSharedPtr<FDisplayClusterViewportLightCardResource, ESPMode::ThreadSafe> UVLightCardOverMapResource;
 };
diff --git a/Engine/Plugins/Runtime/nDisplay/Source/DisplayCluster/Public/DisplayClusterLightCardActor.h b/Engine/Plugins/Runtime/nDisplay/Source/DisplayCluster/Public/DisplayClusterLightCardActor.h
--- a/Engine/Plugins/Runtime/nDisplay/Source/DisplayCluster/Public/DisplayClusterLightCardActor.h
+++ b/Engine/Plugins/Runtime/nDisplay/Source/DisplayCluster/Public/DisplayClusterLightCardActor.h
@@ -4,6 +4,7 @@
 
 #include "Components/DisplayClusterLabelConfiguration.h"
 #include "StageActor/IDisplayClusterStageActor.h"
+#include "DisplayClusterConfigurationTypes_Enums.h"
 
 #include "GameFramework/Actor.h"
 
@@ -239,6 +240,10 @@ public:
 	UPROPERTY(EditAnywhere, BlueprintReadWrite, Interp, Category = "Orientation", meta = (EditCondition = "!bIsUVLightCard", HideEditConditionToggle, EditConditionHides))
 	bool bAlwaysFlushToWall;
 
+	/** Specify how to render this Light Card Actor in relation to the inner frustum. */
+	UPROPERTY(EditAnywhere, BlueprintReadWrite, Interp, Category = "Appearance", meta = (DisplayName = "Blending Mode"))
+	EDisplayClusterConfigurationICVFX_PerLightcardRenderMode PerLightcardRenderMode;
+
 	UPROPERTY(EditAnywhere, BlueprintReadWrite, Interp, Category = "Appearance")
 	EDisplayClusterLightCardMask Mask;
 
diff --git a/Engine/Plugins/Runtime/nDisplay/Source/DisplayCluster/Public/DisplayClusterRootActor.h b/Engine/Plugins/Runtime/nDisplay/Source/DisplayCluster/Public/DisplayClusterRootActor.h
--- a/Engine/Plugins/Runtime/nDisplay/Source/DisplayCluster/Public/DisplayClusterRootActor.h
+++ b/Engine/Plugins/Runtime/nDisplay/Source/DisplayCluster/Public/DisplayClusterRootActor.h
@@ -107,7 +107,17 @@ public:
 
 	// Return hidden in game privitives set
 	bool GetHiddenInGamePrimitives(TSet<FPrimitiveComponentId>& OutPrimitives);
-	bool FindPrimitivesByName(const TArray<FString>& InNames, TSet<FPrimitiveComponentId>& OutPrimitives);
+
+	/** Collect UActorComponent primitives by names.
+	* Note: Components must have bHiddenInGame set to true.
+	* 
+	* @param InNames       - The names of the primitives to be searched for.
+	* @param OutPrimitives - (out) An array with the primitives that were found.
+	* @param bForceHide    - (opt) if true, collects components with any value of the bHiddenInGame property.
+	* 
+	* @return true if any primitives were found.
+	*/
+	bool FindPrimitivesByName(const TArray<FString>& InNames, TSet<FPrimitiveComponentId>& OutPrimitives, bool bForceHide = false);
 
 	bool IsBlueprint() const;
 
@@ -393,7 +403,7 @@ private:
 
 private:
 	template <typename TComp>
-	void GetTypedPrimitives(TSet<FPrimitiveComponentId>& OutPrimitives, const TArray<FString>* InCompNames = nullptr, bool bCollectChildrenVisualizationComponent = true) const;
+	void GetTypedPrimitives(TSet<FPrimitiveComponentId>& OutPrimitives, const TArray<FString>* InCompNames = nullptr, bool bForceHide = false, bool bCollectChildrenVisualizationComponent = true) const;
 
 public:
 	/** Set the priority for inner frustum rendering if there is any overlap when enabling multiple ICVFX cameras. */
diff --git a/Engine/Plugins/Runtime/nDisplay/Source/DisplayCluster/Public/Render/Viewport/Containers/DisplayClusterViewport_EnumsICVFX.h b/Engine/Plugins/Runtime/nDisplay/Source/DisplayCluster/Public/Render/Viewport/Containers/DisplayClusterViewport_EnumsICVFX.h
--- a/Engine/Plugins/Runtime/nDisplay/Source/DisplayCluster/Public/Render/Viewport/Containers/DisplayClusterViewport_EnumsICVFX.h
+++ b/Engine/Plugins/Runtime/nDisplay/Source/DisplayCluster/Public/Render/Viewport/Containers/DisplayClusterViewport_EnumsICVFX.h
@@ -25,13 +25,25 @@ enum class EDisplayClusterViewportICVFXFlags : uint8
 
 	// Disable lightcard render to this viewport
 	DisableLightcard = 1 << 4,
+
+	// lightcard render always under the InCamera
+	LightcardAlwaysUnder = 1 << 5,
+
+	// lightcard render always over the InCamera
+	LightcardAlwaysOver = 1 << 6,
+
+	// Lightcard rendering mode is determined from the stage settings.
+	LightcardUseStageSettings = 1 << 7,
+
+	// Mask that get only flags that define the light card render mode.
+	LightcardRenderModeMask = LightcardAlwaysUnder | LightcardAlwaysOver | LightcardUseStageSettings,
 };
 ENUM_CLASS_FLAGS(EDisplayClusterViewportICVFXFlags);
 
 /**
  * This flag raised only from icvfx manager.
 */
-enum class EDisplayClusterViewportRuntimeICVFXFlags: uint8
+enum class EDisplayClusterViewportRuntimeICVFXFlags: uint16
 {
 	None = 0,
 
@@ -44,10 +56,14 @@ enum class EDisplayClusterViewportRuntimeICVFXFlags: uint8
 	Lightcard   = 1 << 3,
 	UVLightcard = 1 << 4,
 
-	// This viewport used as internal icvfx composing resource (created and deleted inside icvfx logic)
-	InternalResource = 1 << 5,
+	// Additional flags marking the position of the viewport relative to the in-frustum.
+	OverInFrustum = 1 << 5,
+	UnderInFrustum = 1 << 6,
+
+	// This viewport used as internal icvfx compositing resource (created and deleted inside icvfx logic)
+	InternalResource = 1 << 14,
 
 	// Mark unused icvfx dynamic viewports
-	Unused = 1 << 6,
+	Unused = 1 << 15,
 };
 ENUM_CLASS_FLAGS(EDisplayClusterViewportRuntimeICVFXFlags);
diff --git a/Engine/Plugins/Runtime/nDisplay/Source/DisplayClusterConfiguration/Private/DisplayClusterConfigurationTypes_ICVFX.cpp b/Engine/Plugins/Runtime/nDisplay/Source/DisplayClusterConfiguration/Private/DisplayClusterConfigurationTypes_ICVFX.cpp
--- a/Engine/Plugins/Runtime/nDisplay/Source/DisplayClusterConfiguration/Private/DisplayClusterConfigurationTypes_ICVFX.cpp
+++ b/Engine/Plugins/Runtime/nDisplay/Source/DisplayClusterConfiguration/Private/DisplayClusterConfigurationTypes_ICVFX.cpp
@@ -564,7 +564,7 @@ bool FDisplayClusterConfigurationICVFX_LightcardSettings::ShouldUseLightCard(con
 {
 	if (!bEnable)
 	{
-		// dont use lightcard if disabled
+		// Don't use the lightcard if it is disabled
 		return false;
 	}
 
@@ -587,6 +587,86 @@ bool FDisplayClusterConfigurationICVFX_LightcardSettings::ShouldUseUVLightCard(c
 	return ShouldUseLightCard(InStageSettings);
 }
 
+EDisplayClusterShaderParametersICVFX_LightCardRenderMode FDisplayClusterConfigurationICVFX_LightcardSettings::GetLightCardRenderModeOverride(const UDisplayClusterConfigurationViewport* InViewportConfiguration) const
+{
+	if (!bEnable || (InViewportConfiguration && !InViewportConfiguration->ICVFX.bAllowICVFX))
+	{
+		// When ICVFX is disabled we don't override lightcards rendering mode
+		return EDisplayClusterShaderParametersICVFX_LightCardRenderMode::None;
+	}
+
+	if (InViewportConfiguration && InViewportConfiguration->ICVFX.LightcardRenderMode != EDisplayClusterConfigurationICVFX_OverrideLightcardRenderMode::Default)
+	{
+		// Use overridden values from the viewport:
+		switch (InViewportConfiguration->ICVFX.LightcardRenderMode)
+		{
+		case EDisplayClusterConfigurationICVFX_OverrideLightcardRenderMode::Over:
+			return EDisplayClusterShaderParametersICVFX_LightCardRenderMode::Over;
+
+		case EDisplayClusterConfigurationICVFX_OverrideLightcardRenderMode::Under:
+			return EDisplayClusterShaderParametersICVFX_LightCardRenderMode::Under;
+
+		default:
+			break;
+		}
+	}
+
+	return EDisplayClusterShaderParametersICVFX_LightCardRenderMode::None;
+}
+
+EDisplayClusterShaderParametersICVFX_LightCardRenderMode FDisplayClusterConfigurationICVFX_LightcardSettings::GetLightCardRenderMode(const EDisplayClusterConfigurationICVFX_PerLightcardRenderMode InPerLightcardRenderMode, const UDisplayClusterConfigurationViewport* InViewportConfiguration) const
+{
+	if (!bEnable || (InViewportConfiguration && !InViewportConfiguration->ICVFX.bAllowICVFX))
+	{
+		// When ICVFX is disabled we don't render lightcards
+		return EDisplayClusterShaderParametersICVFX_LightCardRenderMode::None;
+	}
+
+	if (InViewportConfiguration && InViewportConfiguration->ICVFX.LightcardRenderMode != EDisplayClusterConfigurationICVFX_OverrideLightcardRenderMode::Default)
+	{
+		// Use overridden values from the viewport:
+		switch (InViewportConfiguration->ICVFX.LightcardRenderMode)
+		{
+		case EDisplayClusterConfigurationICVFX_OverrideLightcardRenderMode::Over:
+			return EDisplayClusterShaderParametersICVFX_LightCardRenderMode::Over;
+
+		case EDisplayClusterConfigurationICVFX_OverrideLightcardRenderMode::Under:
+			return EDisplayClusterShaderParametersICVFX_LightCardRenderMode::Under;
+
+		default:
+			break;
+		}
+
+		return EDisplayClusterShaderParametersICVFX_LightCardRenderMode::None;
+	}
+
+	// Per-lightcard render mode:
+	switch (InPerLightcardRenderMode)
+	{
+	case EDisplayClusterConfigurationICVFX_PerLightcardRenderMode::Under:
+		return EDisplayClusterShaderParametersICVFX_LightCardRenderMode::Under;
+
+	case EDisplayClusterConfigurationICVFX_PerLightcardRenderMode::Over:
+		return EDisplayClusterShaderParametersICVFX_LightCardRenderMode::Over;
+
+	default:
+		break;
+	}
+
+	// Use global lightcard settings:
+	switch (Blendingmode)
+	{
+	case EDisplayClusterConfigurationICVFX_LightcardRenderMode::Under:
+		return EDisplayClusterShaderParametersICVFX_LightCardRenderMode::Under;
+
+	default:
+		break;
+	};
+
+	// By default, lightcards are rendered in “Over” mode.
+	return EDisplayClusterShaderParametersICVFX_LightCardRenderMode::Over;
+}
+
 ///////////////////////////////////////////////////////////////////////////////////////
 // FDisplayClusterConfigurationICVFX_ChromakeyRenderSettings
 ///////////////////////////////////////////////////////////////////////////////////////
diff --git a/Engine/Plugins/Runtime/nDisplay/Source/DisplayClusterConfiguration/Private/DisplayClusterConfigurationTypes_Viewport.cpp b/Engine/Plugins/Runtime/nDisplay/Source/DisplayClusterConfiguration/Private/DisplayClusterConfigurationTypes_Viewport.cpp
--- a/Engine/Plugins/Runtime/nDisplay/Source/DisplayClusterConfiguration/Private/DisplayClusterConfigurationTypes_Viewport.cpp
+++ b/Engine/Plugins/Runtime/nDisplay/Source/DisplayClusterConfiguration/Private/DisplayClusterConfigurationTypes_Viewport.cpp
@@ -5,60 +5,19 @@
 #include "DisplayClusterConfigurationUtils.h"
 
 ///////////////////////////////////////////////////////////////////////////////////////
-// FDisplayClusterConfigurationViewport_ICVFX
+// UDisplayClusterConfigurationViewport
 ///////////////////////////////////////////////////////////////////////////////////////
-EDisplayClusterShaderParametersICVFX_LightCardRenderMode FDisplayClusterConfigurationViewport_ICVFX::GetLightCardRenderMode(const FDisplayClusterConfigurationICVFX_StageSettings& InStageSettings) const
-{
-	// Note: please check if this rule is valid
-	if (!bAllowICVFX || !InStageSettings.Lightcard.bEnable)
-	{
-		// When ICVFX is disabled we dont render lightcards
-		return EDisplayClusterShaderParametersICVFX_LightCardRenderMode::None;
-	}
-
-	if (LightcardRenderMode != EDisplayClusterConfigurationICVFX_OverrideLightcardRenderMode::Default)
-	{
-		// Use overridden values from the viewport:
-		switch (LightcardRenderMode)
-		{
-		case EDisplayClusterConfigurationICVFX_OverrideLightcardRenderMode::Over:
-			return EDisplayClusterShaderParametersICVFX_LightCardRenderMode::Over;
-
-		case EDisplayClusterConfigurationICVFX_OverrideLightcardRenderMode::Under:
-			return EDisplayClusterShaderParametersICVFX_LightCardRenderMode::Under;
-
-		default:
-			break;
-		}
-
-		return EDisplayClusterShaderParametersICVFX_LightCardRenderMode::None;
-	}
-
-	// Use global lightcard settings:
-	switch (InStageSettings.Lightcard.Blendingmode)
-	{
-	case EDisplayClusterConfigurationICVFX_LightcardRenderMode::Under:
-		return EDisplayClusterShaderParametersICVFX_LightCardRenderMode::Under;
-
-	default:
-		break;
-	};
-
-	// By default lightcard render in 'Over' mode
-	return EDisplayClusterShaderParametersICVFX_LightCardRenderMode::Over;
-}
-
-EDisplayClusterViewportICVFXFlags FDisplayClusterConfigurationViewport_ICVFX::GetViewportICVFXFlags(const FDisplayClusterConfigurationICVFX_StageSettings& InStageSettings) const
+EDisplayClusterViewportICVFXFlags UDisplayClusterConfigurationViewport::GetViewportICVFXFlags(const FDisplayClusterConfigurationICVFX_StageSettings& InStageSettings) const
 {
 	EDisplayClusterViewportICVFXFlags OutFlags = EDisplayClusterViewportICVFXFlags::None;
-	if (bAllowICVFX)
+	if (ICVFX.bAllowICVFX)
 	{
 		EnumAddFlags(OutFlags, EDisplayClusterViewportICVFXFlags::Enable);
 	}
 
 	// Override camera render mode
-	EDisplayClusterConfigurationICVFX_OverrideCameraRenderMode UsedCameraRenderMode = CameraRenderMode;
-	if (!bAllowInnerFrustum || !InStageSettings.bEnableInnerFrustums)
+	EDisplayClusterConfigurationICVFX_OverrideCameraRenderMode UsedCameraRenderMode = ICVFX.CameraRenderMode;
+	if (!ICVFX.bAllowInnerFrustum || !InStageSettings.bEnableInnerFrustums)
 	{
 		UsedCameraRenderMode = EDisplayClusterConfigurationICVFX_OverrideCameraRenderMode::Disabled;
 	}
@@ -85,11 +44,29 @@ EDisplayClusterViewportICVFXFlags FDisplayClusterConfigurationViewport_ICVFX::Ge
 	}
 
 	// Disable lightcards rendering
-	const EDisplayClusterShaderParametersICVFX_LightCardRenderMode LightCardRenderMode = GetLightCardRenderMode(InStageSettings);
+	const EDisplayClusterShaderParametersICVFX_LightCardRenderMode LightCardRenderMode = InStageSettings.Lightcard.GetLightCardRenderMode(EDisplayClusterConfigurationICVFX_PerLightcardRenderMode::Default, this);
 	if (LightCardRenderMode == EDisplayClusterShaderParametersICVFX_LightCardRenderMode::None)
 	{
 		EnumAddFlags(OutFlags, EDisplayClusterViewportICVFXFlags::DisableLightcard);
 	}
 
+	// Per-viewport lightcard
+	const EDisplayClusterShaderParametersICVFX_LightCardRenderMode LightCardRenderModeOverride = InStageSettings.Lightcard.GetLightCardRenderModeOverride(this);
+	switch (LightCardRenderModeOverride)
+	{
+	case EDisplayClusterShaderParametersICVFX_LightCardRenderMode::Over:
+		EnumAddFlags(OutFlags, EDisplayClusterViewportICVFXFlags::LightcardAlwaysOver);
+		break;
+
+	case EDisplayClusterShaderParametersICVFX_LightCardRenderMode::Under:
+		EnumAddFlags(OutFlags, EDisplayClusterViewportICVFXFlags::LightcardAlwaysUnder);
+		break;
+
+	default:
+		EnumAddFlags(OutFlags, EDisplayClusterViewportICVFXFlags::LightcardUseStageSettings);
+		break;
+	}
+
+
 	return OutFlags;
 }
diff --git a/Engine/Plugins/Runtime/nDisplay/Source/DisplayClusterConfiguration/Public/DisplayClusterConfigurationTypes_Enums.h b/Engine/Plugins/Runtime/nDisplay/Source/DisplayClusterConfiguration/Public/DisplayClusterConfigurationTypes_Enums.h
--- a/Engine/Plugins/Runtime/nDisplay/Source/DisplayClusterConfiguration/Public/DisplayClusterConfigurationTypes_Enums.h
+++ b/Engine/Plugins/Runtime/nDisplay/Source/DisplayClusterConfiguration/Public/DisplayClusterConfigurationTypes_Enums.h
@@ -84,6 +84,7 @@ enum class EDisplayClusterConfigurationICVFX_ChromakeySource : uint8
 	ChromakeyRenderTexture UMETA(DisplayName = "Chromakey Render Texture"),
 };
 
+/** How to render Light Cards in relation to the inner frustum. */
 UENUM()
 enum class EDisplayClusterConfigurationICVFX_LightcardRenderMode : uint8
 {
@@ -94,6 +95,21 @@ enum class EDisplayClusterConfigurationICVFX_LightcardRenderMode : uint8
 	Under   UMETA(DisplayName = "Lightcard Under Frustum"),
 };
 
+/** How to render a Light Card Actor in relation to the inner frustum. */
+UENUM()
+enum class EDisplayClusterConfigurationICVFX_PerLightcardRenderMode : uint8
+{
+	/** It is determined by nDisplay settings. */
+	Default UMETA(DisplayName = "Default"),
+
+	/** Render Light Card Actor over the inner frustum.  Light Card Actor can be directly visible in camera. */
+	Over    UMETA(DisplayName = "Lightcard Over Frustum"),
+
+	/** Render Light Card Actor under the inner frustum. Light Card Actor will not be directly visible in camera. */
+	Under   UMETA(DisplayName = "Lightcard Under Frustum"),
+};
+
+/** Override the lightcard rendering order for the viewport. */
 UENUM()
 enum class EDisplayClusterConfigurationICVFX_OverrideLightcardRenderMode : uint8
 {
@@ -103,10 +119,10 @@ enum class EDisplayClusterConfigurationICVFX_OverrideLightcardRenderMode : uint8
 	// Disable lightcard rendering for this viewport
 	Disabled    UMETA(DisplayName = "Disabled"),
 
-	// Render incamera frame over lightcard for this viewport
+	/** Always render the In-Camera frame over lightcards for this viewport. */
 	Over    UMETA(DisplayName = "Lightcard Over Frustum"),
 
-	// Over lightcard over incamera frame  for this viewport
+	/** Always render the In-Camera frame under lightcards for this viewport. */
 	Under   UMETA(DisplayName = "Lightcard Under Frustum"),
 };
 
diff --git a/Engine/Plugins/Runtime/nDisplay/Source/DisplayClusterConfiguration/Public/DisplayClusterConfigurationTypes_ICVFX.h b/Engine/Plugins/Runtime/nDisplay/Source/DisplayClusterConfiguration/Public/DisplayClusterConfigurationTypes_ICVFX.h
--- a/Engine/Plugins/Runtime/nDisplay/Source/DisplayClusterConfiguration/Public/DisplayClusterConfigurationTypes_ICVFX.h
+++ b/Engine/Plugins/Runtime/nDisplay/Source/DisplayClusterConfiguration/Public/DisplayClusterConfigurationTypes_ICVFX.h
@@ -23,6 +23,7 @@
 #include "DisplayClusterConfigurationTypes_ICVFX.generated.h"
 
 class UCineCameraComponent;
+class UDisplayClusterConfigurationViewport;
 
 USTRUCT(Blueprintable)
 struct DISPLAYCLUSTERCONFIGURATION_API FDisplayClusterConfigurationICVFX_LightcardCustomOCIO
@@ -116,8 +117,10 @@ public:
 	bool IsVisibilityListValid() const;
 
 public:
-	/** Actor Layers */
-	UPROPERTY(EditAnywhere, BlueprintReadWrite, Category = NDisplay, meta = (DisplayName = "Layers", ToolTip = "{LayersTooltip}"))
+	/** Actor layers.
+	* Warning: This option has a very expensive performance cost, especially on big projects.
+	*/
+	UPROPERTY(EditAnywhere, BlueprintReadWrite, Category = NDisplay, meta = (DisplayName = "Layers"))
 	TArray<FActorLayer> ActorLayers;
 
 	/** Actor references */
@@ -416,6 +419,12 @@ public:
 	/* Returns true if the use of the UVLightCard is allowed */
 	bool ShouldUseUVLightCard(const FDisplayClusterConfigurationICVFX_StageSettings& InStageSettings) const;
 
+	/** Get lightcard render mode. */
+	EDisplayClusterShaderParametersICVFX_LightCardRenderMode GetLightCardRenderMode(const EDisplayClusterConfigurationICVFX_PerLightcardRenderMode InPerLightcardRenderMode, const UDisplayClusterConfigurationViewport* InViewportConfiguration = nullptr) const;
+
+	/** Get LC render mode override. */
+	EDisplayClusterShaderParametersICVFX_LightCardRenderMode GetLightCardRenderModeOverride(const UDisplayClusterConfigurationViewport* InViewportConfiguration) const;
+
 public:
 	/** Enable Light Cards */
 	UPROPERTY(EditAnywhere, BlueprintReadWrite, Category = "Light Cards", meta = (DisplayName = "Enable Light Cards"))
diff --git a/Engine/Plugins/Runtime/nDisplay/Source/DisplayClusterConfiguration/Public/DisplayClusterConfigurationTypes_Viewport.h b/Engine/Plugins/Runtime/nDisplay/Source/DisplayClusterConfiguration/Public/DisplayClusterConfigurationTypes_Viewport.h
--- a/Engine/Plugins/Runtime/nDisplay/Source/DisplayClusterConfiguration/Public/DisplayClusterConfigurationTypes_Viewport.h
+++ b/Engine/Plugins/Runtime/nDisplay/Source/DisplayClusterConfiguration/Public/DisplayClusterConfigurationTypes_Viewport.h
@@ -31,11 +31,19 @@ struct DISPLAYCLUSTERCONFIGURATION_API FDisplayClusterConfigurationViewport_ICVF
 	GENERATED_BODY()
 
 public:
-	/** Get ligthcard render mode for this viewport. */
-	EDisplayClusterShaderParametersICVFX_LightCardRenderMode GetLightCardRenderMode(const FDisplayClusterConfigurationICVFX_StageSettings& InStageSettings) const;
+	/** Get lightcard render mode for this viewport. */
+	UE_DEPRECATED(5.5, "This function has been moved to FDisplayClusterConfigurationICVFX_LightcardSettings.")
+	EDisplayClusterShaderParametersICVFX_LightCardRenderMode GetLightCardRenderMode(const FDisplayClusterConfigurationICVFX_StageSettings& InStageSettings) const
+	{
+		return EDisplayClusterShaderParametersICVFX_LightCardRenderMode::None;
+	}
 
 	/** Get ICVFX settings flags for viewport*/
-	EDisplayClusterViewportICVFXFlags GetViewportICVFXFlags(const FDisplayClusterConfigurationICVFX_StageSettings& InStageSettings) const;
+	UE_DEPRECATED(5.5, "This function has been moved to UDisplayClusterConfigurationViewport.")
+	EDisplayClusterViewportICVFXFlags GetViewportICVFXFlags(const FDisplayClusterConfigurationICVFX_StageSettings& InStageSettings) const
+	{
+		return EDisplayClusterViewportICVFXFlags::None;
+	}
 
 public:
 	/** Enable in-camera VFX for this Viewport (works only with supported Projection Policies) */
@@ -149,6 +157,9 @@ public:
 		return bAllowRendering;
 	}
 
+	/** Get ICVFX settings flags for viewport*/
+	EDisplayClusterViewportICVFXFlags GetViewportICVFXFlags(const FDisplayClusterConfigurationICVFX_StageSettings& InStageSettings) const;
+
 #if WITH_EDITOR
 	/** Enable the preview texture. Only should be called by the object managing the preview texture state. */
 	void EnablePreviewTexture();
diff --git a/Engine/Plugins/Runtime/nDisplay/Source/DisplayClusterShaders/Private/Shaders/DisplayClusterShadersWarpblend_ICVFX.cpp b/Engine/Plugins/Runtime/nDisplay/Source/DisplayClusterShaders/Private/Shaders/DisplayClusterShadersWarpblend_ICVFX.cpp
--- a/Engine/Plugins/Runtime/nDisplay/Source/DisplayClusterShaders/Private/Shaders/DisplayClusterShadersWarpblend_ICVFX.cpp
+++ b/Engine/Plugins/Runtime/nDisplay/Source/DisplayClusterShaders/Private/Shaders/DisplayClusterShadersWarpblend_ICVFX.cpp
@@ -96,17 +96,20 @@ enum class EIcvfxPassRenderPass : uint8
 	InnerFrustumChromakeyOverlapIterator,
 };
 
-enum class EIcvfxPassRenderState : uint32
+enum class EIcvfxPassRenderState : uint8
 {
 	None = 0,
+
 	BlendDisabled = 1 << 0,
 
-	LightCardOver = 1 << 1,
+	LightCardOver  = 1 << 1,
 	LightCardUnder = 1 << 2,
-	UVLightCardOver = 1 << 3,
-	UVLightCardUnder = 1 << 4,
 
-	EnableInnerFrustumChromakeyOverlap = 1 << 5,
+	UVLightCardOver   = 1 << 3,
+	UVLightCardUnder  = 1 << 4,
+	UVLightCardMerged = 1 << 5,
+
+	EnableInnerFrustumChromakeyOverlap = 1 << 6,
 };
 ENUM_CLASS_FLAGS(EIcvfxPassRenderState);
 
@@ -122,18 +125,26 @@ enum class EIcvfxShaderType : uint8
 	Invalid,
 };
 
+/**
+* Special enums for the shader parameter 'LightCardMode'
+*/
+enum class EShaderParameterLightCardMode : uint32
+{
+	None = 0,
+	LIGHTCARD_UNDER = 1 << 0,
+	LIGHTCARD_OVER = 1 << 1,
+	UVLIGHTCARD_UNDER = 1 << 2,
+	UVLIGHTCARD_OVER = 1 << 3,
+	UVLIGHTCARD_MERGE = 1 << 4
+};
+ENUM_CLASS_FLAGS(EShaderParameterLightCardMode);
+
 namespace IcvfxShaderPermutation
 {
 	// Shared permutation for picp warp
 	class FIcvfxShaderViewportInput      : SHADER_PERMUTATION_BOOL("VIEWPORT_INPUT");
 	class FIcvfxShaderViewportInputAlpha : SHADER_PERMUTATION_BOOL("VIEWPORT_INPUT_ALPHA");
 
-	class FIcvfxShaderLightCardUnder     : SHADER_PERMUTATION_BOOL("LIGHTCARD_UNDER");
-	class FIcvfxShaderLightCardOver      : SHADER_PERMUTATION_BOOL("LIGHTCARD_OVER");
-
-	class FIcvfxShaderUVLightCardUnder : SHADER_PERMUTATION_BOOL("UVLIGHTCARD_UNDER");
-	class FIcvfxShaderUVLightCardOver : SHADER_PERMUTATION_BOOL("UVLIGHTCARD_OVER");
-
 	class FIcvfxShaderInnerCamera         : SHADER_PERMUTATION_BOOL("INNER_CAMERA");
 	class FIcvfxShaderInnerCameraOverlap  : SHADER_PERMUTATION_BOOL("INNER_CAMERA_OVERLAP");
 
@@ -153,11 +164,6 @@ namespace IcvfxShaderPermutation
 		FIcvfxShaderViewportInput,
 		FIcvfxShaderViewportInputAlpha,
 
-		FIcvfxShaderLightCardUnder,
-		FIcvfxShaderLightCardOver,
-		FIcvfxShaderUVLightCardUnder,
-		FIcvfxShaderUVLightCardOver,
-
 		FIcvfxShaderInnerCamera,
 		FIcvfxShaderInnerCameraOverlap,
 
@@ -174,12 +180,6 @@ namespace IcvfxShaderPermutation
 	{
 		if (!PermutationVector.Get<FIcvfxShaderMeshWarp>())
 		{
-			// UVLightCard require UV_Chromakey (Mesh Warp)
-			if (PermutationVector.Get<FIcvfxShaderUVLightCardUnder>() || PermutationVector.Get<FIcvfxShaderUVLightCardOver>())
-			{
-				return false;
-			}
-
 			// Chromakey marker require UV_Chromakey (Mesh Warp)
 			if (PermutationVector.Get<FIcvfxShaderChromakeyMarker>())
 			{
@@ -202,27 +202,8 @@ namespace IcvfxShaderPermutation
 			return false;
 		}
 
-		// LightCard can be only 'over' or 'under'
-		if ((PermutationVector.Get<FIcvfxShaderLightCardUnder>() || PermutationVector.Get<FIcvfxShaderUVLightCardUnder>()) && (PermutationVector.Get<FIcvfxShaderUVLightCardOver>() || PermutationVector.Get<FIcvfxShaderLightCardOver>()))
-		{
-			return false;
-		}
-
 		if (!PermutationVector.Get<FIcvfxShaderViewportInput>())
 		{
-			if (PermutationVector.Get<FIcvfxShaderLightCardUnder>() || PermutationVector.Get<FIcvfxShaderUVLightCardUnder>())
-			{
-				return false;
-			}
-
-			if (PermutationVector.Get<FIcvfxShaderLightCardOver>() || PermutationVector.Get<FIcvfxShaderUVLightCardOver>())
-			{
-				if (PermutationVector.Get<FIcvfxShaderInnerCamera>())
-				{
-					return false;
-				}
-			}
-
 			if (PermutationVector.Get<FIcvfxShaderViewportInputAlpha>())
 			{
 				return false;
@@ -233,7 +214,9 @@ namespace IcvfxShaderPermutation
 		{
 			// All innner camera vectors
 			if (PermutationVector.Get<FIcvfxShaderInnerCameraOverlap>()
-			|| PermutationVector.Get<FIcvfxShaderChromakey>() || PermutationVector.Get<FIcvfxShaderChromakeyFrameColor>() || PermutationVector.Get<FIcvfxShaderChromakeyMarker>())
+			|| PermutationVector.Get<FIcvfxShaderChromakey>()
+			|| PermutationVector.Get<FIcvfxShaderChromakeyFrameColor>()
+			|| PermutationVector.Get<FIcvfxShaderChromakeyMarker>())
 			{
 				return false;
 			}
@@ -273,8 +256,10 @@ BEGIN_SHADER_PARAMETER_STRUCT(FIcvfxPixelShaderParameters, )
 	SHADER_PARAMETER_TEXTURE(Texture2D, ChromakeyCameraTexture)
 	SHADER_PARAMETER_TEXTURE(Texture2D, ChromakeyMarkerTexture)
 
-	SHADER_PARAMETER_TEXTURE(Texture2D, LightCardTexture)
-	SHADER_PARAMETER_TEXTURE(Texture2D, UVLightCardTexture)
+	SHADER_PARAMETER_TEXTURE(Texture2D, LightCardUnderTexture)
+	SHADER_PARAMETER_TEXTURE(Texture2D, LightCardOverTexture)
+	SHADER_PARAMETER_TEXTURE(Texture2D, UVLightCardUnderTexture)
+	SHADER_PARAMETER_TEXTURE(Texture2D, UVLightCardOverTexture)
 
 	SHADER_PARAMETER_SAMPLER(SamplerState, InputSampler)
 	SHADER_PARAMETER_SAMPLER(SamplerState, WarpMapSampler)
@@ -285,13 +270,16 @@ BEGIN_SHADER_PARAMETER_STRUCT(FIcvfxPixelShaderParameters, )
 	SHADER_PARAMETER_SAMPLER(SamplerState, ChromakeyCameraSampler)
 	SHADER_PARAMETER_SAMPLER(SamplerState, ChromakeyMarkerSampler)
 
-	SHADER_PARAMETER_SAMPLER(SamplerState, LightCardSampler)
-	SHADER_PARAMETER_SAMPLER(SamplerState, UVLightCardSampler)
+	SHADER_PARAMETER_SAMPLER(SamplerState, LightCardUnderSampler)
+	SHADER_PARAMETER_SAMPLER(SamplerState, LightCardOverSampler)
+	SHADER_PARAMETER_SAMPLER(SamplerState, UVLightCardUnderSampler)
+	SHADER_PARAMETER_SAMPLER(SamplerState, UVLightCardOverSampler)
 
 	SHADER_PARAMETER(FMatrix44f, ViewportTextureProjectionMatrix)
 	SHADER_PARAMETER(FMatrix44f, OverlayProjectionMatrix)
 	SHADER_PARAMETER(FMatrix44f, InnerCameraProjectionMatrix)
 
+	SHADER_PARAMETER(int32, LightCardMode)
 	SHADER_PARAMETER(float, LightCardGamma)
 
 	SHADER_PARAMETER(float, AlphaMapGammaEmbedded)
@@ -418,7 +406,30 @@ public:
 			EnumAddFlags(RenderState, EIcvfxPassRenderState::UVLightCardUnder);
 		}
 
-		if (ICVFXParameters.CameraOverlappingRenderMode != EDisplayClusterShaderParametersICVFX_CameraOverlappingRenderMode::None)
+		// Merge only when used 2 UVLC textures.
+		if (ICVFXParameters.OverrideUVLightCardRenderMode != EDisplayClusterShaderParametersICVFX_LightCardRenderMode::None
+			&& EnumHasAllFlags(RenderState, EIcvfxPassRenderState::UVLightCardUnder | EIcvfxPassRenderState::UVLightCardOver))
+		{
+			EnumAddFlags(RenderState, EIcvfxPassRenderState::UVLightCardMerged);
+
+			// disable overridden use cases
+			switch (ICVFXParameters.OverrideUVLightCardRenderMode)
+			{
+			case EDisplayClusterShaderParametersICVFX_LightCardRenderMode::Under:
+				EnumRemoveFlags(RenderState, EIcvfxPassRenderState::UVLightCardOver);
+				break;
+
+			case EDisplayClusterShaderParametersICVFX_LightCardRenderMode::Over:
+				EnumRemoveFlags(RenderState, EIcvfxPassRenderState::UVLightCardUnder);
+				break;
+
+			default:
+				break;
+			}
+		}
+
+		// Camera overlap work only for multi-cameras setup
+		if (ICVFXParameters.CameraOverlappingRenderMode != EDisplayClusterShaderParametersICVFX_CameraOverlappingRenderMode::None && ICVFXParameters.IsMultiCamerasUsed())
 		{
 			EnumAddFlags(RenderState, EIcvfxPassRenderState::EnableInnerFrustumChromakeyOverlap);
 		}
@@ -482,6 +493,19 @@ private:
 		return EIcvfxShaderType::Invalid;
 	}
 
+	inline void EnumAddLightCardModeFlags(FIcvfxRenderPassData& RenderPassData, const EShaderParameterLightCardMode InFlags)
+	{
+		EShaderParameterLightCardMode NewLightCardMode = (EShaderParameterLightCardMode)RenderPassData.PSParameters.LightCardMode;
+		EnumAddFlags(NewLightCardMode, InFlags);
+		RenderPassData.PSParameters.LightCardMode = (int32)NewLightCardMode;
+	}
+
+	inline bool EnumHasAnyLightCardModeFlags(FIcvfxRenderPassData& RenderPassData, const EShaderParameterLightCardMode InFlags) const
+	{
+		EShaderParameterLightCardMode LightCardMode = (EShaderParameterLightCardMode)RenderPassData.PSParameters.LightCardMode;
+		return EnumHasAnyFlags(LightCardMode, InFlags);
+	}
+
 public:
 
 	bool GetViewportParameters(FIcvfxRenderPassData& RenderPassData)
@@ -594,6 +618,26 @@ public:
 		return false;
 	}
 
+	void InitializeLightCardParameters(FIcvfxRenderPassData& RenderPassData)
+	{
+		RenderPassData.PSParameters.LightCardMode = (int32)EShaderParameterLightCardMode::None;
+
+		// Since we do not use a permutation vector for light cards, the pixel shader will always contain code with light card textures.
+		// To avoid shader compiler errors, we must set all textures even if they are not used by the logic.
+
+		RenderPassData.PSParameters.LightCardUnderSampler = TStaticSamplerState<SF_Trilinear, AM_Clamp, AM_Clamp, AM_Clamp>::GetRHI();
+		RenderPassData.PSParameters.LightCardUnderTexture = WarpBlendParameters.Src.Texture;
+
+		RenderPassData.PSParameters.LightCardOverSampler = TStaticSamplerState<SF_Trilinear, AM_Clamp, AM_Clamp, AM_Clamp>::GetRHI();
+		RenderPassData.PSParameters.LightCardOverTexture = WarpBlendParameters.Src.Texture;
+
+		RenderPassData.PSParameters.UVLightCardUnderSampler = TStaticSamplerState<SF_Trilinear, AM_Clamp, AM_Clamp, AM_Clamp>::GetRHI();
+		RenderPassData.PSParameters.UVLightCardUnderTexture = WarpBlendParameters.Src.Texture;
+
+		RenderPassData.PSParameters.UVLightCardOverSampler = TStaticSamplerState<SF_Trilinear, AM_Clamp, AM_Clamp, AM_Clamp>::GetRHI();
+		RenderPassData.PSParameters.UVLightCardOverTexture = WarpBlendParameters.Src.Texture;
+	}
+
 	bool GetLightCardParameters(FIcvfxRenderPassData& RenderPassData)
 	{
 		RenderPassData.PSParameters.OverlayProjectionMatrix = FMatrix44f(LocalUVMatrix);
@@ -604,21 +648,20 @@ public:
 			if (EnumHasAnyFlags(RenderState, EIcvfxPassRenderState::LightCardUnder))
 			{
 				// Rendering in one pass along with the first camera
-				RenderPassData.PSPermutationVector.Set<IcvfxShaderPermutation::FIcvfxShaderLightCardUnder>(true);
+				EnumAddLightCardModeFlags(RenderPassData, EShaderParameterLightCardMode::LIGHTCARD_UNDER);
 			}
-			else if (EnumHasAnyFlags(RenderState, EIcvfxPassRenderState::LightCardOver) && !ICVFXParameters.IsMultiCamerasUsed())
+
+			if (EnumHasAnyFlags(RenderState, EIcvfxPassRenderState::LightCardOver) && !ICVFXParameters.IsMultiCamerasUsed())
 			{
 				// If we have only one camera, the top LC layer is rendered in one pass
-				RenderPassData.PSPermutationVector.Set<IcvfxShaderPermutation::FIcvfxShaderLightCardOver>(true);
+				EnumAddLightCardModeFlags(RenderPassData, EShaderParameterLightCardMode::LIGHTCARD_OVER);
 			}
 		break;
 
 		case EIcvfxPassRenderPass::LightCardOver:
-			check(ICVFXParameters.IsMultiCamerasUsed());
-
 			if (EnumHasAllFlags(RenderState, EIcvfxPassRenderState::LightCardOver))
 			{
-				RenderPassData.PSPermutationVector.Set<IcvfxShaderPermutation::FIcvfxShaderLightCardOver>(true);
+				EnumAddLightCardModeFlags(RenderPassData, EShaderParameterLightCardMode::LIGHTCARD_OVER);
 			}
 			break;
 
@@ -626,13 +669,13 @@ public:
 			break;
 		}
 
-		if(RenderPassData.PSPermutationVector.Get<IcvfxShaderPermutation::FIcvfxShaderLightCardUnder>()
-		|| RenderPassData.PSPermutationVector.Get<IcvfxShaderPermutation::FIcvfxShaderLightCardOver>())
+		if (EnumHasAnyLightCardModeFlags(RenderPassData, EShaderParameterLightCardMode::LIGHTCARD_UNDER | EShaderParameterLightCardMode::LIGHTCARD_OVER))
 		{
-			RenderPassData.PSParameters.LightCardSampler = TStaticSamplerState<SF_Trilinear, AM_Clamp, AM_Clamp, AM_Clamp>::GetRHI();
-			RenderPassData.PSParameters.LightCardTexture = ICVFXParameters.LightCard.Texture;
 			RenderPassData.PSParameters.LightCardGamma = ICVFXParameters.LightCardGamma;
 
+			RenderPassData.PSParameters.LightCardUnderTexture = ICVFXParameters.LightCardUnder.Texture ? ICVFXParameters.LightCardUnder.Texture : ICVFXParameters.LightCardOver.Texture;
+			RenderPassData.PSParameters.LightCardOverTexture  = ICVFXParameters.LightCardOver.Texture  ? ICVFXParameters.LightCardOver.Texture  : ICVFXParameters.LightCardUnder.Texture;
+
 			return true;
 		}
 
@@ -647,26 +690,31 @@ public:
 			return false;
 		}
 
+		if (EnumHasAnyFlags(RenderState, EIcvfxPassRenderState::UVLightCardMerged))
+		{
+			EnumAddLightCardModeFlags(RenderPassData, EShaderParameterLightCardMode::UVLIGHTCARD_MERGE);
+		}
+
 		switch (RenderPass)
 		{
 		case EIcvfxPassRenderPass::Base:
 			if (EnumHasAnyFlags(RenderState, EIcvfxPassRenderState::UVLightCardUnder))
 			{
-				RenderPassData.PSPermutationVector.Set<IcvfxShaderPermutation::FIcvfxShaderUVLightCardUnder>(true);
+				// Rendering in one pass along with the first camera
+				EnumAddLightCardModeFlags(RenderPassData, EShaderParameterLightCardMode::UVLIGHTCARD_UNDER);
 			}
-			else if (EnumHasAnyFlags(RenderState, EIcvfxPassRenderState::UVLightCardOver) && !ICVFXParameters.IsMultiCamerasUsed())
+
+			if (EnumHasAnyFlags(RenderState, EIcvfxPassRenderState::UVLightCardOver) && !ICVFXParameters.IsMultiCamerasUsed())
 			{
-				// Render in one pass for single camera
-				RenderPassData.PSPermutationVector.Set<IcvfxShaderPermutation::FIcvfxShaderUVLightCardOver>(true);
+				// If we have only one camera, the top LC layer is rendered in one pass
+				EnumAddLightCardModeFlags(RenderPassData, EShaderParameterLightCardMode::UVLIGHTCARD_OVER);
 			}
 			break;
 
 		case EIcvfxPassRenderPass::LightCardOver:
-			check(ICVFXParameters.IsMultiCamerasUsed());
-
 			if (EnumHasAllFlags(RenderState, EIcvfxPassRenderState::UVLightCardOver))
 			{
-				RenderPassData.PSPermutationVector.Set<IcvfxShaderPermutation::FIcvfxShaderUVLightCardOver>(true);
+				EnumAddLightCardModeFlags(RenderPassData, EShaderParameterLightCardMode::UVLIGHTCARD_OVER);
 			}
 			break;
 
@@ -674,13 +722,13 @@ public:
 			break;
 		}
 
-		if (RenderPassData.PSPermutationVector.Get<IcvfxShaderPermutation::FIcvfxShaderUVLightCardUnder>()
-			|| RenderPassData.PSPermutationVector.Get<IcvfxShaderPermutation::FIcvfxShaderUVLightCardOver>())
+		if (EnumHasAnyLightCardModeFlags(RenderPassData, EShaderParameterLightCardMode::UVLIGHTCARD_UNDER | EShaderParameterLightCardMode::UVLIGHTCARD_OVER))
 		{
-			RenderPassData.PSParameters.UVLightCardSampler = TStaticSamplerState<SF_Trilinear, AM_Clamp, AM_Clamp, AM_Clamp>::GetRHI();
-			RenderPassData.PSParameters.UVLightCardTexture = ICVFXParameters.UVLightCard.Texture;
 			RenderPassData.PSParameters.LightCardGamma = ICVFXParameters.LightCardGamma;
 
+			RenderPassData.PSParameters.UVLightCardUnderTexture = ICVFXParameters.UVLightCardUnder.Texture ? ICVFXParameters.UVLightCardUnder.Texture : ICVFXParameters.UVLightCardOver.Texture;
+			RenderPassData.PSParameters.UVLightCardOverTexture  = ICVFXParameters.UVLightCardOver.Texture  ? ICVFXParameters.UVLightCardOver.Texture  : ICVFXParameters.UVLightCardUnder.Texture;
+
 			return true;
 		}
 
@@ -722,7 +770,7 @@ public:
 		int32 OverlapCameraNum = 0;
 		EDisplayClusterShaderParametersICVFX_ChromakeySource OutChromakeySource = EDisplayClusterShaderParametersICVFX_ChromakeySource::Unknown;
 
-		if (EnumHasAnyFlags(RenderState, EIcvfxPassRenderState::EnableInnerFrustumChromakeyOverlap) && ICVFXParameters.IsMultiCamerasUsed())
+		if (EnumHasAnyFlags(RenderState, EIcvfxPassRenderState::EnableInnerFrustumChromakeyOverlap))
 		{
 			switch (RenderPass)
 			{
@@ -902,8 +950,12 @@ public:
 				GetWarpBlendParameters(RenderPassData);
 			}
 
-			GetLightCardParameters(RenderPassData);
-			GetUVLightCardParameters(RenderPassData);
+			// Configure Lightcards:
+			{
+				InitializeLightCardParameters(RenderPassData);
+				GetLightCardParameters(RenderPassData);
+				GetUVLightCardParameters(RenderPassData);
+			}
 
 			if (GetCameraParameters(RenderPassData))
 			{
@@ -1191,17 +1243,17 @@ public:
 			break;
 
 		case EIcvfxPassRenderPass::Base:
-			if (!ICVFXParameters.IsMultiCamerasUsed())
+			if (ICVFXParameters.IsMultiCamerasUsed())
 			{
-				// Render one camera in single pass
-				RenderPass = EIcvfxPassRenderPass::None;
-				return false;
+				// multi-camera rendering
+				RenderPass = EIcvfxPassRenderPass::InnerFrustumIterator;
+				ActiveCameraIndex = 1;
+				break;
 			}
 
-			// render multi-cam
-			RenderPass = EIcvfxPassRenderPass::InnerFrustumIterator;
-			ActiveCameraIndex = 1;
-			break;
+			// Render in single pass completed
+			RenderPass = EIcvfxPassRenderPass::None;
+			return false;
 
 		case EIcvfxPassRenderPass::InnerFrustumIterator:
 			if (!ICVFXParameters.IsCameraUsed(++ActiveCameraIndex))
diff --git a/Engine/Plugins/Runtime/nDisplay/Source/DisplayClusterShaders/Public/ShaderParameters/DisplayClusterShaderParameters_ICVFX.h b/Engine/Plugins/Runtime/nDisplay/Source/DisplayClusterShaders/Public/ShaderParameters/DisplayClusterShaderParameters_ICVFX.h
--- a/Engine/Plugins/Runtime/nDisplay/Source/DisplayClusterShaders/Public/ShaderParameters/DisplayClusterShaderParameters_ICVFX.h
+++ b/Engine/Plugins/Runtime/nDisplay/Source/DisplayClusterShaders/Public/ShaderParameters/DisplayClusterShaderParameters_ICVFX.h
@@ -60,22 +60,22 @@ public:
 public:
 	inline bool IsLightCardOverUsed() const
 	{
-		return LightCard.IsValid() && LightCardMode == EDisplayClusterShaderParametersICVFX_LightCardRenderMode::Over;
+		return LightCardOver.IsValid();
 	}
 
 	inline bool IsLightCardUnderUsed() const
 	{
-		return LightCard.IsValid() && LightCardMode == EDisplayClusterShaderParametersICVFX_LightCardRenderMode::Under;
+		return LightCardUnder.IsValid();
 	}
 
 	inline bool IsUVLightCardOverUsed() const
 	{
-		return UVLightCard.IsValid() && LightCardMode == EDisplayClusterShaderParametersICVFX_LightCardRenderMode::Over;
+		return UVLightCardOver.IsValid();
 	}
 
 	inline bool IsUVLightCardUnderUsed() const
 	{
-		return UVLightCard.IsValid() && LightCardMode == EDisplayClusterShaderParametersICVFX_LightCardRenderMode::Under;
+		return UVLightCardUnder.IsValid();
 	}
 
 	inline bool IsCameraUsed(int32 CameraIndex) const
@@ -93,9 +93,14 @@ public:
 		return Cameras.Num() > 1;
 	}
 
-	inline bool IsValid()
+	inline bool IsValid() const
 	{
-		return LightCard.IsValid() || UVLightCard.IsValid() || IsAnyCameraUsed();
+		if (LightCardOver.IsValid() || UVLightCardOver.IsValid() || LightCardUnder.IsValid() || UVLightCardUnder.IsValid())
+		{
+			return true;
+		}
+
+		return IsAnyCameraUsed();
 	}
 
 public:
@@ -103,9 +108,12 @@ public:
 	{
 		Cameras.Empty();
 
-		UVLightCard.Reset();
-		LightCard.Reset();
+		LightCardUnder.Reset();
+		LightCardOver.Reset();
+		UVLightCardUnder.Reset();
+		UVLightCardOver.Reset();
 
+		OverrideUVLightCardRenderMode = EDisplayClusterShaderParametersICVFX_LightCardRenderMode::None;
 		LightCardGamma = 2.2;
 	}
 
@@ -118,9 +126,13 @@ public:
 
 		CameraOverlappingRenderMode = InParameters.CameraOverlappingRenderMode;
 
-		UVLightCard = InParameters.UVLightCard;
-		LightCard   = InParameters.LightCard;
-		LightCardMode = InParameters.LightCardMode;
+		LightCardUnder = InParameters.LightCardUnder;
+		LightCardOver = InParameters.LightCardOver;
+
+		UVLightCardUnder = InParameters.UVLightCardUnder;
+		UVLightCardOver = InParameters.UVLightCardOver;
+		OverrideUVLightCardRenderMode = InParameters.OverrideUVLightCardRenderMode;
+
 		LightCardGamma = InParameters.LightCardGamma;
 	}
 
@@ -278,14 +290,24 @@ public:
 	template <typename Predicate>
 	void IterateViewportResourcesByPredicate(Predicate Pred)
 	{
-		if (LightCard.IsDefined())
+		if (LightCardUnder.IsDefined())
+		{
+			::Invoke(Pred, LightCardUnder);
+		}
+
+		if (LightCardOver.IsDefined())
+		{
+			::Invoke(Pred, LightCardOver);
+		}
+
+		if (UVLightCardUnder.IsDefined())
 		{
-			::Invoke(Pred, LightCard);
+			::Invoke(Pred, UVLightCardUnder);
 		}
 
-		if (UVLightCard.IsDefined())
+		if (UVLightCardOver.IsDefined())
 		{
-			::Invoke(Pred, UVLightCard);
+			::Invoke(Pred, UVLightCardOver);
 		}
 
 		for (FCameraSettings& CameraIt : Cameras)
@@ -328,10 +350,15 @@ public:
 	// Rendering mode for overlapping areas of camera projections
 	EDisplayClusterShaderParametersICVFX_CameraOverlappingRenderMode CameraOverlappingRenderMode = EDisplayClusterShaderParametersICVFX_CameraOverlappingRenderMode::None;
 
-	// LightCard settings
-	FDisplayClusterShaderParametersICVFX_ViewportResource    UVLightCard;
-	FDisplayClusterShaderParametersICVFX_ViewportResource    LightCard;
-	EDisplayClusterShaderParametersICVFX_LightCardRenderMode LightCardMode = EDisplayClusterShaderParametersICVFX_LightCardRenderMode::Under;
+	// LightCard resources
+	FDisplayClusterShaderParametersICVFX_ViewportResource    LightCardUnder;
+	FDisplayClusterShaderParametersICVFX_ViewportResource    LightCardOver;
+
+	// UV LightCard resources
+	FDisplayClusterShaderParametersICVFX_ViewportResource    UVLightCardUnder;
+	FDisplayClusterShaderParametersICVFX_ViewportResource    UVLightCardOver;
+	// Force UVLC to render merged (per-viewport LC)
+	EDisplayClusterShaderParametersICVFX_LightCardRenderMode OverrideUVLightCardRenderMode = EDisplayClusterShaderParametersICVFX_LightCardRenderMode::None;
 
 	// The gamma that the light card renders have been encoded with, used to linearize during final composite
 	float LightCardGamma = 2.2;
