nDisplay: reworked the UI related to ViewPoint components.
- reworked components: UDisplayClusterCameraComponent, UDisplayClusterInFrustumFitCameraComponent.
- renamed "View Origin" -> "View Point"
- reworked UI for DCRA, ViewPoint
- All DCRA preview options have been moved into one category and regrouped.

diff --git a/Engine/Plugins/Runtime/nDisplay/Source/DisplayCluster/Private/Game/EngineClasses/Scene/DisplayClusterCameraComponent.cpp b/Engine/Plugins/Runtime/nDisplay/Source/DisplayCluster/Private/Game/EngineClasses/Scene/DisplayClusterCameraComponent.cpp
--- a/Engine/Plugins/Runtime/nDisplay/Source/DisplayCluster/Private/Game/EngineClasses/Scene/DisplayClusterCameraComponent.cpp
+++ b/Engine/Plugins/Runtime/nDisplay/Source/DisplayCluster/Private/Game/EngineClasses/Scene/DisplayClusterCameraComponent.cpp
@@ -10,6 +10,8 @@
 
 #include "DisplayClusterRootActor.h"
 
+#include "CineCameraComponent.h"
+
 #include "Engine/CollisionProfile.h"
 #include "Engine/Texture2D.h"
 #include "UObject/ConstructorHelpers.h"
@@ -33,7 +35,14 @@ UDisplayClusterCameraComponent::UDisplayClusterCameraComponent(const FObjectInit
 	{
 		ConstructorHelpers::FObjectFinderOptional<UTexture2D> SpriteTextureObject = TEXT("/nDisplay/Icons/S_nDisplayViewOrigin");
 		SpriteTexture = SpriteTextureObject.Get();
+
 	}
+
+	PrimaryComponentTick.bCanEverTick = true;
+	PrimaryComponentTick.TickGroup = TG_PostUpdateWork;
+	PrimaryComponentTick.bStartWithTickEnabled = true;
+	bAutoActivate = true;
+
 #endif
 }
 
@@ -51,13 +60,20 @@ void UDisplayClusterCameraComponent::ApplyViewPointComponentPostProcessesToViewp
 	// Get the same component from DCRA that is used as the configuration source. Then this component can also be used as a configuration data source.
 	const UDisplayClusterCameraComponent& CfgCameraComponent = GetMatchingComponentFromRootActor(InViewport->GetConfiguration(), EDisplayClusterRootActorType::Configuration, *this);
 
+	const bool bICVFXCameraBeingUsed = CfgCameraComponent.IsICVFXCameraBeingUsed();
+	const bool bUseTargetCamera =
+		CfgCameraComponent.IsActiveEngineCameraBeingUsed()
+		|| bICVFXCameraBeingUsed
+		|| CfgCameraComponent.IsExternalCameraBeingUsed();
+
 	// Setup Outer Viewport postprocessing
-	if (CfgCameraComponent.bEnableOuterViewportCamera)
+	if (bUseTargetCamera)
 	{
-		const EDisplayClusterViewportCameraPostProcessFlags CameraPostProcessingFlags = CfgCameraComponent.OuterViewportPostProcessSettings.GetCameraPostProcessFlags();
+		const EDisplayClusterViewportCameraPostProcessFlags CameraPostProcessingFlags = CfgCameraComponent.GetCameraPostProcessFlags();
 
 		// Also, if we are referencing the ICVFXCamera component, use the special ICVFX PostProcess from it.
-		if (UDisplayClusterICVFXCameraComponent* SceneICVFXCameraComponent = GetRootActorComponentByName<UDisplayClusterICVFXCameraComponent>(InViewport->GetConfiguration(), EDisplayClusterRootActorType::Scene, CfgCameraComponent.OuterViewportCameraName))
+		if (UDisplayClusterICVFXCameraComponent* SceneICVFXCameraComponent
+			= !bICVFXCameraBeingUsed ? nullptr : GetRootActorComponentByName<UDisplayClusterICVFXCameraComponent>(InViewport->GetConfiguration(), EDisplayClusterRootActorType::Scene, CfgCameraComponent.ICVFXCameraComponentName))
 		{
 			// Use PostProcess from the ICVFXCamera
 			// This function also uses PostProcess from the parent CineCamera class.
@@ -67,7 +83,7 @@ void UDisplayClusterCameraComponent::ApplyViewPointComponentPostProcessesToViewp
 		{
 			// Use post-processing settings from Camera/CineCamera or from the active game camera.
 			FMinimalViewInfo CustomViewInfo;
-			if (CfgCameraComponent.GetOuterViewportCameraDesiredViewInternal(InViewport->GetConfiguration(), CustomViewInfo))
+			if (CfgCameraComponent.GetTargetCameraDesiredViewInternal(InViewport->GetConfiguration(), CustomViewInfo))
 			{
 				// Applies a filter to the post-processing settings.
 				FDisplayClusterViewportConfigurationHelpers_Postprocess::FilterPostProcessSettings(CustomViewInfo.PostProcessSettings, CameraPostProcessingFlags);
@@ -79,18 +95,25 @@ void UDisplayClusterCameraComponent::ApplyViewPointComponentPostProcessesToViewp
 	}
 }
 
-UCameraComponent* UDisplayClusterCameraComponent::GetOuterViewportCameraComponent(const IDisplayClusterViewportConfiguration& InViewportConfiguration) const
+UCameraComponent* UDisplayClusterCameraComponent::GetTargetCameraComponent(const IDisplayClusterViewportConfiguration& InViewportConfiguration) const
 {
 	using namespace UE::DisplayClusterViewportHelpers;
 
-	// Get the same component from DCRA that is used as the configuration source. Then this component can also be used as a configuration data source.
 	const UDisplayClusterCameraComponent& CfgCameraComponent = GetMatchingComponentFromRootActor(InViewportConfiguration, EDisplayClusterRootActorType::Configuration, *this);
-	if (CfgCameraComponent.bEnableOuterViewportCamera)
+
+	// 1. Active engine camera
+	if (CfgCameraComponent.IsActiveEngineCameraBeingUsed())
+	{
+		return nullptr;
+	}
+
+	// 2. ICVFX camera component
+	if (CfgCameraComponent.IsICVFXCameraBeingUsed())
 	{
-		if(UCameraComponent* SceneCameraComponent = GetRootActorComponentByName<UCameraComponent>(InViewportConfiguration, EDisplayClusterRootActorType::Scene, CfgCameraComponent.OuterViewportCameraName))
+		if (UCameraComponent* SceneCameraComponent = GetRootActorComponentByName<UCameraComponent>(InViewportConfiguration, EDisplayClusterRootActorType::Scene, CfgCameraComponent.ICVFXCameraComponentName))
 		{
 			// If we use the ICVFX camera component, we must use GetActualCineCameraComponent() to get the actual camera.
-			if(SceneCameraComponent->IsA<UDisplayClusterICVFXCameraComponent>())
+			if (SceneCameraComponent->IsA<UDisplayClusterICVFXCameraComponent>())
 			{
 				if (UDisplayClusterICVFXCameraComponent* ICVFXCameraComponent = Cast<UDisplayClusterICVFXCameraComponent>(SceneCameraComponent))
 				{
@@ -99,67 +122,87 @@ UCameraComponent* UDisplayClusterCameraComponent::GetOuterViewportCameraComponen
 						// Use referenced camera as the source of Camera PP and CineCamera CustomNearClippingPlane
 						return ExtCineCameraComponent;
 					}
-					}
 				}
+			}
 
 			return SceneCameraComponent;
 		}
 	}
 
-	return nullptr;
+	// 3. External camera actor
+	return CfgCameraComponent.GetExternalCineCameraActorComponent();
 }
 
-bool UDisplayClusterCameraComponent::GetOuterViewportCameraDesiredViewInternal(const IDisplayClusterViewportConfiguration& InViewportConfiguration, FMinimalViewInfo& InOutViewInfo, float* OutCustomNCP) const
+bool UDisplayClusterCameraComponent::IsViewPointOverrideCameraPosition() const
 {
-	using namespace UE::DisplayClusterViewportHelpers;
+	// If the ICFX camera component is used, it can override the viewpoint position.
+	if (IsICVFXCameraBeingUsed())
+	{
+		return !bUseICVFXCameraComponentTracking;
+	}
+
+	// By default, ViewPoint is always used as the camera.
+	return true;
+}
 
-	bool bViewInfoFound = false;
+bool UDisplayClusterCameraComponent::GetTargetCameraDesiredViewInternal(const IDisplayClusterViewportConfiguration& InViewportConfiguration, FMinimalViewInfo& InOutViewInfo, float* OutCustomNCP) const
+{
+	using namespace UE::DisplayClusterViewportHelpers;
 
 	// Get the same component from DCRA that is used as the configuration source. Then this component can also be used as a configuration data source.
 	const UDisplayClusterCameraComponent& CfgCameraComponent = GetMatchingComponentFromRootActor(InViewportConfiguration, EDisplayClusterRootActorType::Configuration, *this);
-	if (CfgCameraComponent.bEnableOuterViewportCamera)
-	{
-		const EDisplayClusterViewportCameraPostProcessFlags CameraPostProcessingFlags = CfgCameraComponent.OuterViewportPostProcessSettings.GetCameraPostProcessFlags();
 
-		const bool bUseCameraPostprocess = EnumHasAnyFlags(CameraPostProcessingFlags, EDisplayClusterViewportCameraPostProcessFlags::EnablePostProcess);
+	const EDisplayClusterViewportCameraPostProcessFlags CameraPostProcessingFlags = CfgCameraComponent.GetCameraPostProcessFlags();
+	const bool bUseCameraPostprocess = EnumHasAnyFlags(CameraPostProcessingFlags, EDisplayClusterViewportCameraPostProcessFlags::EnablePostProcess);
 
-		float* OutCustomNearClippingPlane = OutCustomNCP;
-		if (!EnumHasAnyFlags(CameraPostProcessingFlags, EDisplayClusterViewportCameraPostProcessFlags::EnableNearClippingPlane))
-		{
-			// Ignore NCP from the custom camera.
-			OutCustomNearClippingPlane = nullptr;
-		}
+	float* OutCustomNearClippingPlane = OutCustomNCP;
+	if (!EnumHasAnyFlags(CameraPostProcessingFlags, EDisplayClusterViewportCameraPostProcessFlags::EnableNearClippingPlane))
+	{
+		// Ignore NCP from the custom camera.
+		OutCustomNearClippingPlane = nullptr;
+	}
 
-		if (UCameraComponent* SceneCameraComponent = GetOuterViewportCameraComponent(InViewportConfiguration))
+	// 1. Active engine camera
+	if (CfgCameraComponent.IsActiveEngineCameraBeingUsed())
+	{
+		// Get PostProcess from the Game camera.
+		if (IDisplayClusterViewport::GetPlayerCameraView(InViewportConfiguration.GetCurrentWorld(), bUseCameraPostprocess, InOutViewInfo))
 		{
-			if (IDisplayClusterViewport::GetCameraComponentView(SceneCameraComponent, InViewportConfiguration.GetRootActorWorldDeltaSeconds(), bUseCameraPostprocess, InOutViewInfo, OutCustomNearClippingPlane))
+			if (IsViewPointOverrideCameraPosition())
 			{
-				bViewInfoFound = true;
+				// Use this component as a camera
+				InOutViewInfo.Location = GetComponentLocation();
+				InOutViewInfo.Rotation = GetComponentRotation();
 			}
+
+			return true;
 		}
-		// Get PostProcess from the Game camera.
-		else if (IDisplayClusterViewport::GetPlayerCameraView(InViewportConfiguration.GetCurrentWorld(), bUseCameraPostprocess, InOutViewInfo))
+	}
+	else
+	{
+		if (UCameraComponent* SceneCameraComponent = GetTargetCameraComponent(InViewportConfiguration))
 		{
-			
-			bViewInfoFound = true;
-		}
+			if (IDisplayClusterViewport::GetCameraComponentView(SceneCameraComponent, InViewportConfiguration.GetRootActorWorldDeltaSeconds(), bUseCameraPostprocess, InOutViewInfo, OutCustomNearClippingPlane))
+			{
+				if (IsViewPointOverrideCameraPosition())
+				{
+					// Use this component as a camera
+					InOutViewInfo.Location = GetComponentLocation();
+					InOutViewInfo.Rotation = GetComponentRotation();
+				}
 
-		if (bViewInfoFound && !CfgCameraComponent.bFollowOuterViewportCamera)
-		{
-			// Use this component as a camera
-			InOutViewInfo.Location = GetComponentLocation();
-			InOutViewInfo.Rotation = GetComponentRotation();
+				return true;
+			}
 		}
-
-		// The default camera is not found, so we can't use the custom camera view.
 	}
 
-	return bViewInfoFound;
+	// The target camera is not found, so we can't use the custom camera view.
+	return false;
 }
 
 void UDisplayClusterCameraComponent::GetDesiredView(IDisplayClusterViewportConfiguration& InViewportConfiguration, FMinimalViewInfo& InOutViewInfo, float* OutCustomNearClippingPlane)
 {
-	if (GetOuterViewportCameraDesiredViewInternal(InViewportConfiguration, InOutViewInfo, OutCustomNearClippingPlane))
+	if (GetTargetCameraDesiredViewInternal(InViewportConfiguration, InOutViewInfo, OutCustomNearClippingPlane))
 	{
 		return;
 	}
@@ -182,7 +225,7 @@ void UDisplayClusterCameraComponent::GetDesiredView(IDisplayClusterViewportConfi
 void UDisplayClusterCameraComponent::GetEyePosition(const IDisplayClusterViewportConfiguration& InViewportConfiguration, FVector& OutViewLocation, FRotator& OutViewRotation)
 {
 	FMinimalViewInfo ViewInfo;
-	if (GetOuterViewportCameraDesiredViewInternal(InViewportConfiguration, ViewInfo))
+	if (GetTargetCameraDesiredViewInternal(InViewportConfiguration, ViewInfo))
 	{
 		OutViewLocation = ViewInfo.Location;
 		OutViewRotation = ViewInfo.Rotation;
@@ -226,7 +269,7 @@ void UDisplayClusterCameraComponent::OnRegister()
 				SpriteComponent->SetMobility(EComponentMobility::Movable);
 				SpriteComponent->Sprite = SpriteTexture;
 				SpriteComponent->SpriteInfo.Category = TEXT("NDisplayViewOrigin");
-				SpriteComponent->SpriteInfo.DisplayName = NSLOCTEXT("DisplayClusterCameraComponent", "NDisplayViewOriginSpriteInfo", "nDisplay View Origin");
+				SpriteComponent->SpriteInfo.DisplayName = NSLOCTEXT("DisplayClusterCameraComponent", "NDisplayViewOriginSpriteInfo", "nDisplay View Point");
 				SpriteComponent->SetCollisionProfileName(UCollisionProfile::NoCollision_ProfileName);
 				SpriteComponent->bHiddenInGame = true;
 				SpriteComponent->bIsScreenSizeScaled = true;
@@ -244,6 +287,26 @@ void UDisplayClusterCameraComponent::OnRegister()
 }
 
 #if WITH_EDITOR
+bool UDisplayClusterCameraComponent::CanEditChange(const FProperty* InProperty) const
+{
+	// If other logic prevents editing, we want to respect that
+	bool bIsEditable = Super::CanEditChange(InProperty);
+	if (bIsEditable && InProperty)
+	{
+		// Can we edit tracking of the ICVFX camera component
+		const FName PropertyName = InProperty->GetFName();
+		if(PropertyName == GET_MEMBER_NAME_CHECKED(UDisplayClusterCameraComponent, bUseICVFXCameraComponentTracking)
+		|| PropertyName == GET_MEMBER_NAME_CHECKED(UDisplayClusterCameraComponent, bEnableICVFXDepthOfFieldCompensation)
+		|| PropertyName == GET_MEMBER_NAME_CHECKED(UDisplayClusterCameraComponent, bEnableICVFXColorGrading)
+		|| PropertyName == GET_MEMBER_NAME_CHECKED(UDisplayClusterCameraComponent, bEnableICVFXMotionBlur))
+		{
+			return IsICVFXCameraBeingUsed();
+		}
+	}
+
+	return bIsEditable;
+}
+
 void UDisplayClusterCameraComponent::PostEditChangeProperty(FPropertyChangedEvent& PropertyChangedEvent)
 {
 	Super::PostEditChangeProperty(PropertyChangedEvent);
@@ -264,7 +327,22 @@ void UDisplayClusterCameraComponent::RefreshVisualRepresentation()
 }
 #endif
 
-EDisplayClusterViewportCameraPostProcessFlags FDisplayClusterCameraComponent_OuterViewportPostProcessSettings::GetCameraPostProcessFlags() const
+bool UDisplayClusterCameraComponent::IsActiveEngineCameraBeingUsed() const
+{
+	return TargetCameraType == EDisplayClusterTargetCameraType::ActiveEngineCamera;
+}
+
+bool UDisplayClusterCameraComponent::IsICVFXCameraBeingUsed() const
+{
+	return TargetCameraType == EDisplayClusterTargetCameraType::ICVFXCameraComponent && !ICVFXCameraComponentName.IsEmpty();
+}
+
+bool UDisplayClusterCameraComponent::IsExternalCameraBeingUsed() const
+{
+	return TargetCameraType == EDisplayClusterTargetCameraType::ExternalCineCameraActor && !IsICVFXCameraBeingUsed() && ExternalCineCameraActor.IsValid();
+}
+
+EDisplayClusterViewportCameraPostProcessFlags UDisplayClusterCameraComponent::GetCameraPostProcessFlags() const
 {
 	EDisplayClusterViewportCameraPostProcessFlags OutPostProcessFlags = EDisplayClusterViewportCameraPostProcessFlags::None;
 
@@ -301,3 +379,56 @@ EDisplayClusterViewportCameraPostProcessFlags FDisplayClusterCameraComponent_Out
 	return OutPostProcessFlags;
 }
 
+UCameraComponent* UDisplayClusterCameraComponent::GetExternalCineCameraActorComponent() const
+{
+	if (IsExternalCameraBeingUsed())
+	{
+		if (ACineCameraActor* CineCamera = ExternalCineCameraActor.Get())
+		{
+			return CineCamera->GetCameraComponent();
+		}
+	}
+
+	return nullptr;
+}
+
+#if WITH_EDITOR
+UCameraComponent* UDisplayClusterCameraComponent::GetEditorPreviewCameraComponent()
+{
+	using namespace UE::DisplayClusterViewportHelpers;
+
+	if (IsICVFXCameraBeingUsed())
+	{
+		if (UDisplayClusterICVFXCameraComponent* ICVFXCameraComponent = GetOwnerRootActorComponentByName<UDisplayClusterICVFXCameraComponent>(this, ICVFXCameraComponentName))
+		{
+			return ICVFXCameraComponent;
+		}
+	}
+	else if (UCameraComponent* CameraComponent = GetExternalCineCameraActorComponent())
+	{
+		return CameraComponent;
+	}
+
+	return nullptr;
+}
+
+bool UDisplayClusterCameraComponent::GetEditorPreviewInfo(float DeltaTime, FMinimalViewInfo& ViewOut)
+{
+	if (UCameraComponent* CameraComponent = GetEditorPreviewCameraComponent())
+	{
+		return CameraComponent->GetEditorPreviewInfo(DeltaTime, ViewOut);
+	}
+
+	return false;
+}
+
+TSharedPtr<SWidget> UDisplayClusterCameraComponent::GetCustomEditorPreviewWidget()
+{
+	if (UCameraComponent* CameraComponent = GetEditorPreviewCameraComponent())
+	{
+		return CameraComponent->GetCustomEditorPreviewWidget();
+	}
+
+	return nullptr;
+}
+#endif
\ No newline at end of file
diff --git a/Engine/Plugins/Runtime/nDisplay/Source/DisplayCluster/Private/Game/EngineClasses/Scene/DisplayClusterLabelComponent.cpp b/Engine/Plugins/Runtime/nDisplay/Source/DisplayCluster/Private/Game/EngineClasses/Scene/DisplayClusterLabelComponent.cpp
--- a/Engine/Plugins/Runtime/nDisplay/Source/DisplayCluster/Private/Game/EngineClasses/Scene/DisplayClusterLabelComponent.cpp
+++ b/Engine/Plugins/Runtime/nDisplay/Source/DisplayCluster/Private/Game/EngineClasses/Scene/DisplayClusterLabelComponent.cpp
@@ -153,7 +153,7 @@ void UDisplayClusterLabelComponent::UpdateWidgetComponent()
 			
 				WidgetComponent->SetWorldRotation(Rotation);
 
-				// Keep the label the same size regardless of distance to the view origin.
+				// Keep the label the same size regardless of distance to the view point.
 				const float Distance = (DestinationPoint - LightCardLocation).Length();
 
 				const float BaseLabelScale = 0.0025f;
diff --git a/Engine/Plugins/Runtime/nDisplay/Source/DisplayCluster/Private/Game/EngineClasses/Scene/DisplayClusterLightCardActor.cpp b/Engine/Plugins/Runtime/nDisplay/Source/DisplayCluster/Private/Game/EngineClasses/Scene/DisplayClusterLightCardActor.cpp
--- a/Engine/Plugins/Runtime/nDisplay/Source/DisplayCluster/Private/Game/EngineClasses/Scene/DisplayClusterLightCardActor.cpp
+++ b/Engine/Plugins/Runtime/nDisplay/Source/DisplayCluster/Private/Game/EngineClasses/Scene/DisplayClusterLightCardActor.cpp
@@ -567,19 +567,19 @@ void ADisplayClusterLightCardActor::UpdateLightCardPositionToRootActor()
 		{
 			if (ConfigData->StageSettings.Lightcard.bEnable)
 			{
-				// Find a view origin to use as the transform "anchor" of each light card. At the moment, assume the first view origin component
-				// found is the correct view origin (the same assumption is made in the light card editor). If no view origin is found, use the root component
-				const USceneComponent* ViewOriginComponent = RootActor->GetRootComponent();
+				// Find a origin to use as the transform "anchor" of each light card. At the moment, assume the first view point component
+				// found is the correct view point (the same assumption is made in the light card editor). If no view point is found, use the root component
+				const USceneComponent* ViewPointComponent = RootActor->GetRootComponent();
 
-				TArray<UDisplayClusterCameraComponent*> ViewOriginComponents;
-				RootActor->GetComponents(ViewOriginComponents);
+				TArray<UDisplayClusterCameraComponent*> ViewPointComponents;
+				RootActor->GetComponents(ViewPointComponents);
 
-				if (ViewOriginComponents.Num())
+				if (ViewPointComponents.Num())
 				{
-					ViewOriginComponent = ViewOriginComponents[0];
+					ViewPointComponent = ViewPointComponents[0];
 				}
 		
-				const FVector Location = ViewOriginComponent ? ViewOriginComponent->GetComponentLocation() : RootActor->GetActorLocation();
+				const FVector Location = ViewPointComponent ? ViewPointComponent->GetComponentLocation() : RootActor->GetActorLocation();
 				const FVector RelativeLocation = RootActor->GetTransform().InverseTransformPosition(Location);
 				const FRotator Rotation = RootActor->GetActorRotation();
 				
diff --git a/Engine/Plugins/Runtime/nDisplay/Source/DisplayCluster/Private/Game/EngineClasses/Scene/DisplayClusterRootActor.cpp b/Engine/Plugins/Runtime/nDisplay/Source/DisplayCluster/Private/Game/EngineClasses/Scene/DisplayClusterRootActor.cpp
--- a/Engine/Plugins/Runtime/nDisplay/Source/DisplayCluster/Private/Game/EngineClasses/Scene/DisplayClusterRootActor.cpp
+++ b/Engine/Plugins/Runtime/nDisplay/Source/DisplayCluster/Private/Game/EngineClasses/Scene/DisplayClusterRootActor.cpp
@@ -1243,6 +1243,15 @@ void ADisplayClusterRootActor::TickPreviewRenderer()
 		NewPreviewSettings.bPreviewEnable = false;
 	}
 
+#if WITH_EDITOR
+	// The StageIsosphereComponent can only be seen if DCRA preview is enabled.
+	const bool bShowStageIsosphere = NewPreviewSettings.bPreviewEnable && bPreviewStageGeometryMesh;
+	if (StageIsosphereComponent && StageIsosphereComponent->GetVisibleFlag() != bShowStageIsosphere)
+	{
+		StageIsosphereComponent->SetVisibility(bShowStageIsosphere);
+	}
+#endif
+
 	// Update entire cluster preview rendering
 	if (IDisplayClusterViewportManager* ViewportManager = bEnablePreviewInScene ? GetOrCreateViewportManager() : GetViewportManager())
 	{
diff --git a/Engine/Plugins/Runtime/nDisplay/Source/DisplayCluster/Private/Game/EngineClasses/Scene/DisplayClusterRootActorEditor.cpp b/Engine/Plugins/Runtime/nDisplay/Source/DisplayCluster/Private/Game/EngineClasses/Scene/DisplayClusterRootActorEditor.cpp
--- a/Engine/Plugins/Runtime/nDisplay/Source/DisplayCluster/Private/Game/EngineClasses/Scene/DisplayClusterRootActorEditor.cpp
+++ b/Engine/Plugins/Runtime/nDisplay/Source/DisplayCluster/Private/Game/EngineClasses/Scene/DisplayClusterRootActorEditor.cpp
@@ -715,10 +715,6 @@ void ADisplayClusterRootActor::PostEditChangeProperty(FPropertyChangedEvent& Pro
 		ResetEntireClusterPreviewRendering();
 		bReinitializeActor = false;
 	}
-	else if (PropertyName == GET_MEMBER_NAME_CHECKED(ADisplayClusterRootActor, bPreviewStageGeometryMesh))
-	{
-		StageIsosphereComponent->SetVisibility(bPreviewStageGeometryMesh);
-	}
 
 	if (bReinitializeActor)
 	{
@@ -755,7 +751,7 @@ void ADisplayClusterRootActor::OnEndObjectMovement(UObject& InObject)
 	{
 		if (SceneComponent->GetOwner() == this)
 		{
-			// Check to see if the object being moved is a part of the stage's geometry or a view origin
+			// Check to see if the object being moved is a part of the stage's geometry or a view point
 			TArray<FString> ProjectionMeshNames;
 			if (UDisplayClusterConfigurationData* Config = GetConfigData())
 			{
@@ -764,9 +760,9 @@ void ADisplayClusterRootActor::OnEndObjectMovement(UObject& InObject)
 
 			const bool bIsProjectionMesh = SceneComponent->IsA<UStaticMeshComponent>() && ProjectionMeshNames.Contains(SceneComponent->GetName());
 			const bool bIsScreen = SceneComponent->IsA<UDisplayClusterScreenComponent>();
-			const bool bIsViewOrigin = SceneComponent->IsA<UDisplayClusterCameraComponent>();
+			const bool bIsViewPoint = SceneComponent->IsA<UDisplayClusterCameraComponent>();
 
-			if (bIsProjectionMesh || bIsScreen || bIsViewOrigin)
+			if (bIsProjectionMesh || bIsScreen || bIsViewPoint)
 			{
 				StageGeometryComponent->Invalidate();
 			}
diff --git a/Engine/Plugins/Runtime/nDisplay/Source/DisplayCluster/Private/Render/Viewport/Configuration/DisplayClusterViewportConfiguration_Viewport.h b/Engine/Plugins/Runtime/nDisplay/Source/DisplayCluster/Private/Render/Viewport/Configuration/DisplayClusterViewportConfiguration_Viewport.h
--- a/Engine/Plugins/Runtime/nDisplay/Source/DisplayCluster/Private/Render/Viewport/Configuration/DisplayClusterViewportConfiguration_Viewport.h
+++ b/Engine/Plugins/Runtime/nDisplay/Source/DisplayCluster/Private/Render/Viewport/Configuration/DisplayClusterViewportConfiguration_Viewport.h
@@ -52,6 +52,6 @@ public:
 	// The viewport instance
 	TSharedPtr<FDisplayClusterViewport, ESPMode::ThreadSafe> Viewport;
 
-	// Pointer to the view origin component, used by this viewport instance.
+	// Pointer to the view point component, used by this viewport instance.
 	UDisplayClusterCameraComponent* ViewPointCameraComponent = nullptr;
 };
diff --git a/Engine/Plugins/Runtime/nDisplay/Source/DisplayCluster/Public/Components/DisplayClusterCameraComponent.h b/Engine/Plugins/Runtime/nDisplay/Source/DisplayCluster/Public/Components/DisplayClusterCameraComponent.h
--- a/Engine/Plugins/Runtime/nDisplay/Source/DisplayCluster/Public/Components/DisplayClusterCameraComponent.h
+++ b/Engine/Plugins/Runtime/nDisplay/Source/DisplayCluster/Public/Components/DisplayClusterCameraComponent.h
@@ -22,6 +22,7 @@ class IDisplayClusterViewportConfiguration;
 class IDisplayClusterViewportPreview;
 class IDisplayClusterViewport;
 class UCameraComponent;
+class ACineCameraActor;
 struct FMinimalViewInfo;
 
 UENUM()
@@ -32,47 +33,28 @@ enum class EDisplayClusterEyeStereoOffset : uint8
 	Right UMETA(DisplayName = "Right Eye"),
 };
 
-/**
-* Specifies the parameters to be used from the specified camera.
-*/
-USTRUCT(BlueprintType)
-struct DISPLAYCLUSTER_API FDisplayClusterCameraComponent_OuterViewportPostProcessSettings
+/** The type of camera used. */
+UENUM()
+enum class EDisplayClusterTargetCameraType: int32
 {
-	GENERATED_BODY()
+	// No camera used
+	None  UMETA(DisplayName = "Disabled"),
 
-	/** Decodes parameters into flags. */
-	EDisplayClusterViewportCameraPostProcessFlags GetCameraPostProcessFlags() const;
+	// Use active engine camera
+	ActiveEngineCamera UMETA(DisplayName = "Active Engine Camera"),
 
-public:
-	/** Use the NearClippingPlane value from the specified cine camera. */
-	UPROPERTY(EditAnywhere, BlueprintReadWrite, Category = "Advanced", meta = (DisplayName = "Use Custom Near Clipping Plane"))
-	bool bEnableNearClippingPlane = false;
-
-	/** Use the PP settings from the specified camera. */
-	UPROPERTY(EditAnywhere, BlueprintReadWrite, Category = "Advanced", meta = (DisplayName = "Use Post Process"))
-	bool bEnablePostProcess = true;
-
-	/** Enable the DoF PP settings from the specified camera. */
-	UPROPERTY(EditAnywhere, BlueprintReadWrite, Category = "Advanced", meta = (DisplayName = "Use Depth Of Field"))
-	bool bEnableDepthOfField = false;
-
-	/** Use the DC Depth-Of-Field settings from the specified ICVFX camera. */
-	UPROPERTY(EditAnywhere, BlueprintReadWrite, Category = "Advanced", meta = (DisplayName = "Use ICVFX Depth Of Field Compensation"))
-	bool bEnableICVFXDepthOfFieldCompensation = false;
+	// Use camera component from the root actor
+	ICVFXCameraComponent UMETA(DisplayName = "ICVFX Camera"),
 
-	/** Use the DC ColorGrading from the specified ICVFX camera. */
-	UPROPERTY(EditAnywhere, BlueprintReadWrite, Category = "Advanced", meta = (DisplayName = "Use ICVFX Color Grading"))
-	bool bEnableICVFXColorGrading = true;
-
-	/** Use the DC Motion Blur settings from the specified ICVFX camera. */
-	UPROPERTY(EditAnywhere, BlueprintReadWrite, Category = "Advanced", meta = (DisplayName = "Use ICVFX Motion Blur"))
-	bool bEnableICVFXMotionBlur = false;
+	// Use the CineCamera actor reference.
+	ExternalCineCameraActor UMETA(DisplayName = "Camera Actor"),
 };
 
+
 /**
  * 3D point in space used to render nDisplay viewports from
  */
-UCLASS(ClassGroup = (DisplayCluster), HideCategories = (Navigation, AssetUserData, LOD, Physics, Cooking, Activation, Tags, Gizmo, Collision, ComponentReplication, Events, Sockets, ComponentTick), meta = (BlueprintSpawnableComponent, DisplayName = "NDisplay View Origin"))
+UCLASS(ClassGroup = (DisplayCluster), HideCategories = (Navigation, AssetUserData, LOD, Physics, Cooking, Activation, Tags, Gizmo, Collision, ComponentReplication, Events, Sockets, ComponentTick, Rendering), meta = (BlueprintSpawnableComponent, DisplayName = "NDisplay View Point"))
 class DISPLAYCLUSTER_API UDisplayClusterCameraComponent
 	: public USceneComponent
 	, public IDisplayClusterComponent
@@ -155,11 +137,19 @@ public:
 	* 
 	* return nullptr if the camera is not in use.
 	*/
-	virtual UCameraComponent* GetOuterViewportCameraComponent(const IDisplayClusterViewportConfiguration& InViewportConfiguration) const;
+	virtual UCameraComponent* GetTargetCameraComponent(const IDisplayClusterViewportConfiguration& InViewportConfiguration) const;
+
+	/** Returns true if the given viewpoint position should be used as a camera position. */
+	virtual bool IsViewPointOverrideCameraPosition() const;
 
 protected:
-	/** Get Outer Viewport Camera view. */
-	virtual bool GetOuterViewportCameraDesiredViewInternal(const IDisplayClusterViewportConfiguration& InViewportConfiguration, FMinimalViewInfo& InOutViewInfo, float* OutCustomNearClippingPlane = nullptr) const;
+	/** Get view from the referenced camera. */
+	virtual bool GetTargetCameraDesiredViewInternal(const IDisplayClusterViewportConfiguration& InViewportConfiguration, FMinimalViewInfo& InOutViewInfo, float* OutCustomNearClippingPlane = nullptr) const;
+
+#if WITH_EDITOR
+	virtual bool GetEditorPreviewInfo(float DeltaTime, FMinimalViewInfo& ViewOut) override;
+	virtual TSharedPtr<SWidget> GetCustomEditorPreviewWidget() override;
+#endif
 
 public:
 	/**
@@ -239,6 +229,28 @@ public:
 		StereoOffset = InStereoOffset;
 	}
 
+public:
+	/** Decodes parameters into flags. */
+	EDisplayClusterViewportCameraPostProcessFlags GetCameraPostProcessFlags() const;
+
+	/** Returns the camera component from the 'ExternalCameraActor' parameter. */
+	virtual UCameraComponent* GetExternalCineCameraActorComponent() const;
+
+protected:
+	/** Returns true if active engine camera is used. */
+	virtual bool IsActiveEngineCameraBeingUsed() const;
+
+	/** Returns true if the `ICVFXCameraComponentName` parameter is used. */
+	virtual bool IsICVFXCameraBeingUsed() const;
+
+	/** Returns true if the `ExternalCineCameraActor` parameter is used. */
+	virtual bool IsExternalCameraBeingUsed() const;
+
+#if WITH_EDITOR
+	/** Return used camera component. */
+	virtual UCameraComponent* GetEditorPreviewCameraComponent();
+#endif
+
 public:
 #if WITH_EDITOR
 	// Begin IDisplayClusterComponent
@@ -253,6 +265,7 @@ public:
 
 	// Begin UObject
 #if WITH_EDITOR
+	virtual bool CanEditChange(const FProperty* InProperty) const override;
 	virtual void PostEditChangeProperty(FPropertyChangedEvent& PropertyChangedEvent) override;
 #endif
 	// End UObject
@@ -285,22 +298,45 @@ protected:
 #endif
 
 public:
-	/** Use the post process from the specified camera. This applies to all viewports that use this viewpoint. */
-	UPROPERTY(EditAnywhere, Category = "Outer Viewport Post Process", meta = (DisplayName = "Use Outer Viewport Camera"))
-	bool bEnableOuterViewportCamera = false;
-
-	/** The viewpoint location follows the camera location. */
-	UPROPERTY(EditAnywhere, Category = "Outer Viewport Post Process", meta = (DisplayName = "Follow Outer Viewport Camera"))
-	bool bFollowOuterViewportCamera = false;
-
-	/** The name of the camera component that is used as the PP source.
-	* (An empty string means that the active game camera is used). */
-	UPROPERTY(EditAnywhere, Category = "Outer Viewport Post Process", meta = (DisplayName = "Outer Viewport Camera Name"))
-	FString OuterViewportCameraName;
-
-	/** Additional settings that control how PP will be used. */
-	UPROPERTY(EditAnywhere, Category = "Outer Viewport Post Process", meta = (DisplayName = "Post Process Settings"))
-	FDisplayClusterCameraComponent_OuterViewportPostProcessSettings OuterViewportPostProcessSettings;
+	/** Type of source camera used. */
+	UPROPERTY(EditAnywhere, BlueprintReadWrite, Category = "Camera Post Process", meta = (DisplayName = "Source Camera"))
+	EDisplayClusterTargetCameraType TargetCameraType = EDisplayClusterTargetCameraType::None;
+
+	/** The name of the camera component that is used as the PP source. */
+	UPROPERTY(EditAnywhere, BlueprintReadWrite, Category = "Camera Post Process", meta = (DisplayName = "ICVFX Camera", EditCondition = "TargetCameraType==EDisplayClusterTargetCameraType::ICVFXCameraComponent", EditConditionHides))
+	FString ICVFXCameraComponentName;
+
+	UPROPERTY(EditAnywhere, BlueprintReadWrite, Category = "Camera Post Process", meta = (DisplayName = "Follow ICVFX Camera", EditCondition = "TargetCameraType==EDisplayClusterTargetCameraType::ICVFXCameraComponent", EditConditionHides))
+	uint8 bUseICVFXCameraComponentTracking: 1 = 0;
+
+	/** Use a specific actor camera instead of a game camera. */
+	UPROPERTY(EditAnywhere, BlueprintReadWrite, Category = "Camera Post Process", meta = (DisplayName = "Camera Actor", EditCondition = "TargetCameraType==EDisplayClusterTargetCameraType::ExternalCineCameraActor", EditConditionHides))
+	TSoftObjectPtr<ACineCameraActor> ExternalCineCameraActor;
+
+public:
+	/** Use the NearClippingPlane value from the specified cine camera. */
+	UPROPERTY(EditAnywhere, BlueprintReadWrite, Category = "Camera Post Process", AdvancedDisplay, meta = (DisplayName = "Include Custom Near Clipping Plane", EditCondition = "TargetCameraType!=EDisplayClusterTargetCameraType::None", EditConditionHides))
+	uint8 bEnableNearClippingPlane : 1 = 1;
+
+	/** Use the PP settings from the specified camera. */
+	UPROPERTY(EditAnywhere, BlueprintReadWrite, Category = "Camera Post Process", AdvancedDisplay, meta = (DisplayName = "Include Camera Post Process", EditCondition = "TargetCameraType!=EDisplayClusterTargetCameraType::None", EditConditionHides))
+	uint8 bEnablePostProcess : 1 = 1;
+
+	/** Enable the DoF PP settings from the specified camera. */
+	UPROPERTY(EditAnywhere, BlueprintReadWrite, Category = "Camera Post Process", AdvancedDisplay, meta = (DisplayName = "Include Depth Of Field", EditCondition = "TargetCameraType!=EDisplayClusterTargetCameraType::None", EditConditionHides))
+	uint8 bEnableDepthOfField : 1 = 0;
+
+	/** Use the DC Depth-Of-Field settings from the specified ICVFX camera. */
+	UPROPERTY(EditAnywhere, BlueprintReadWrite, Category = "Camera Post Process", AdvancedDisplay, meta = (DisplayName = "Include Inner Frustum Depth Of Field Compensation", EditCondition = "TargetCameraType==EDisplayClusterTargetCameraType::ICVFXCameraComponent", EditConditionHides))
+	uint8 bEnableICVFXDepthOfFieldCompensation : 1 = 0;
+
+	/** Use the DC ColorGrading from the specified ICVFX camera. */
+	UPROPERTY(EditAnywhere, BlueprintReadWrite, Category = "Camera Post Process", meta = (DisplayName = "Include Inner Frustum Color Grading", EditCondition = "TargetCameraType==EDisplayClusterTargetCameraType::ICVFXCameraComponent", EditConditionHides))
+	uint8 bEnableICVFXColorGrading : 1 = 1;
+
+	/** Use the DC Motion Blur settings from the specified ICVFX camera. */
+	UPROPERTY(EditAnywhere, BlueprintReadWrite, Category = "Camera Post Process", AdvancedDisplay, meta = (DisplayName = "Include Inner Frustum Motion Blur", EditCondition = "TargetCameraType==EDisplayClusterTargetCameraType::ICVFXCameraComponent", EditConditionHides))
+	uint8 bEnableICVFXMotionBlur : 1 = 0;
 
 private:
 	UPROPERTY(EditAnywhere, Category = "Stereo")
diff --git a/Engine/Plugins/Runtime/nDisplay/Source/DisplayCluster/Public/Components/DisplayClusterICVFXCameraComponent.h b/Engine/Plugins/Runtime/nDisplay/Source/DisplayCluster/Public/Components/DisplayClusterICVFXCameraComponent.h
--- a/Engine/Plugins/Runtime/nDisplay/Source/DisplayCluster/Public/Components/DisplayClusterICVFXCameraComponent.h
+++ b/Engine/Plugins/Runtime/nDisplay/Source/DisplayCluster/Public/Components/DisplayClusterICVFXCameraComponent.h
@@ -88,7 +88,7 @@ public:
 	virtual void OnRegister() override;
 
 	/** Sets new depth of field parameters and updates the dynamically generated compensation LUT if needed */
-	UFUNCTION(BlueprintCallable, Category = "In Camera VFX")
+	UFUNCTION(BlueprintCallable, Category = "ICVFX Camera")
 	void SetDepthOfFieldParameters(const FDisplayClusterConfigurationICVFX_CameraDepthOfField& NewDepthOfFieldParams);
 
 private:
diff --git a/Engine/Plugins/Runtime/nDisplay/Source/DisplayCluster/Public/Components/DisplayClusterLabelComponent.h b/Engine/Plugins/Runtime/nDisplay/Source/DisplayCluster/Public/Components/DisplayClusterLabelComponent.h
--- a/Engine/Plugins/Runtime/nDisplay/Source/DisplayCluster/Public/Components/DisplayClusterLabelComponent.h
+++ b/Engine/Plugins/Runtime/nDisplay/Source/DisplayCluster/Public/Components/DisplayClusterLabelComponent.h
@@ -16,7 +16,7 @@ class UUserWidget;
 class UWidgetComponent;
 
 /**
- * A label component specific to nDisplay. Displays a widget with a consistent scale facing the root actor view origin.
+ * A label component specific to nDisplay. Displays a widget with a consistent scale facing the root actor view point.
  * Visible only in scene capture when in the editor or in 3d space on the wall.
  *
  * The component needs to be transient so its settings aren't saved, but should still transact over multi-user.
diff --git a/Engine/Plugins/Runtime/nDisplay/Source/DisplayCluster/Public/Components/DisplayClusterStageGeometryComponent.h b/Engine/Plugins/Runtime/nDisplay/Source/DisplayCluster/Public/Components/DisplayClusterStageGeometryComponent.h
--- a/Engine/Plugins/Runtime/nDisplay/Source/DisplayCluster/Public/Components/DisplayClusterStageGeometryComponent.h
+++ b/Engine/Plugins/Runtime/nDisplay/Source/DisplayCluster/Public/Components/DisplayClusterStageGeometryComponent.h
@@ -64,8 +64,8 @@ public:
 	/**
 	 * Gets the distance and normal vector (in radial space) of the stage's geometry in the specified world direction
 	 * @param InDirection The direction in world coordinates to query the stage geometry map
-	 * @param OutDistance The distance from the stage's default view origin to the nearest stage geometry in the specified direction
-	 * @param OutNormal The normal vector in radial space (with x axis pointing in the direction of the stage's default view origin) of the nearest stage geometry in the specified direction
+	 * @param OutDistance The distance from the stage's default view point to the nearest stage geometry in the specified direction
+	 * @param OutNormal The normal vector in radial space (with x axis pointing in the direction of the stage's default view point) of the nearest stage geometry in the specified direction
 	 * @return true if the stage geometry map was successfully queried, otherwise false
 	 */
 	UFUNCTION(BlueprintCallable, Category = "NDisplay")
@@ -80,7 +80,7 @@ public:
 	/**
 	 * Morphs the specified procedural mesh to match the stage's geometry map
 	 * @param InProceduralMeshComponent - The mesh component to morph
-	 * @param bSyncMeshLocation - When true, the procedural mesh's world location will be synced to the view origin's world location
+	 * @param bSyncMeshLocation - When true, the procedural mesh's world location will be synced to the view point's world location
 	 */
 	bool MorphProceduralMesh(UProceduralMeshComponent* InProceduralMeshComponent, bool bSyncMeshLocation = false);
 
diff --git a/Engine/Plugins/Runtime/nDisplay/Source/DisplayCluster/Public/DisplayClusterRootActor.h b/Engine/Plugins/Runtime/nDisplay/Source/DisplayCluster/Public/DisplayClusterRootActor.h
--- a/Engine/Plugins/Runtime/nDisplay/Source/DisplayCluster/Public/DisplayClusterRootActor.h
+++ b/Engine/Plugins/Runtime/nDisplay/Source/DisplayCluster/Public/DisplayClusterRootActor.h
@@ -48,7 +48,7 @@ struct FDisplayClusterRootActorPropertyOverrides;
 /**
  * VR root. This contains nDisplay VR hierarchy in the game.
  */
-UCLASS(HideCategories=(Replication, Collision, Input, Actor, HLOD, Cooking, Physics, Activation, AssetUserData, ActorTick, Advanced, DataLayers, Events), meta=(DisplayName = "nDisplay Root Actor"))
+UCLASS(HideCategories=(Replication, Collision, Input, Actor, HLOD, Cooking, Physics, Activation, AssetUserData, ActorTick, Advanced, DataLayers, Events), meta = (DisplayName = "nDisplay Root Actor"))
 class DISPLAYCLUSTER_API ADisplayClusterRootActor
 	: public AActor
 {
@@ -203,7 +203,7 @@ public:
 	ULineBatchComponent* GetLineBatchComponent() const;
 
 	/**
-	 * Get the view origin most commonly used by viewports in this cluster.
+	 * Get the view point most commonly used by viewports in this cluster.
 	 * If no viewports override the camera, this returns the default camera, or if there isn't one, the actor's root component.
 	 */
 	UFUNCTION(BlueprintCallable, Category = "NDisplay|DCRA")
@@ -407,7 +407,7 @@ private:
 
 public:
 	/** Set the priority for inner frustum rendering if there is any overlap when enabling multiple ICVFX cameras. */
-	UPROPERTY(EditInstanceOnly, EditFixedSize, Category = "In Camera VFX", meta = (TitleProperty = "Name", DisplayAfter = "ViewportAllowInnerFrustumRef"))
+	UPROPERTY(EditInstanceOnly, EditFixedSize, Category = "In Camera VFX", meta = (TitleProperty = "Name", DisplayAfter = "ShowInnerFrustumOverlapsRef"))
 	TArray<FDisplayClusterComponentRef> InnerFrustumPriority;
 
 	/**
@@ -442,31 +442,29 @@ public:
 // EDITOR RELATED SETTINGS
 //////////////////////////////////////////////////////////////////////////////////////////////
 public:
-	/** Render this DCRA in game for Standalone/Package builds. */
-	UPROPERTY(EditAnywhere, BlueprintReadWrite, Category = "Preview In Game", meta = (DisplayName = "Enable Preview in Game"))
-	bool bPreviewInGameEnable = false;
-
-	UPROPERTY(EditAnywhere, BlueprintReadWrite, Category = "Preview In Game", meta = (DisplayName = "Render Preview Frustum in Game"))
-	bool bPreviewInGameRenderFrustum = false;
 
 	/** Render the scene and display it as a preview on the nDisplay root actor in the editor.  This will impact editor performance. */
-	UPROPERTY(EditAnywhere, BlueprintReadWrite, Category = "Editor Preview", meta = (DisplayName = "Enable Editor Preview"))
+	UPROPERTY(EditAnywhere, BlueprintReadWrite, Category = "Preview", meta = (DisplayName = "Enable Editor Preview"))
 	bool bPreviewEnable = true;
 	
+	/** Render this DCRA in game for Standalone/Package builds. */
+	UPROPERTY(EditAnywhere, BlueprintReadWrite, Category = "Preview", AdvancedDisplay, meta = (DisplayName = "Enable Preview in Game"))
+	bool bPreviewInGameEnable = false;
+
 	/** Adjust resolution scaling for the editor preview. */
-	UPROPERTY(EditAnywhere, BlueprintReadWrite, Category = "Editor Preview", meta = (DisplayName = "Preview Screen Percentage", ClampMin = "0.05", UIMin = "0.05", ClampMax = "1", UIMax = "1"))
+	UPROPERTY(EditAnywhere, BlueprintReadWrite, Category = "Preview", meta = (DisplayName = "Preview Screen Percentage", ClampMin = "0.05", UIMin = "0.05", ClampMax = "1", UIMax = "1"))
 	float PreviewRenderTargetRatioMult = 0.25;
 
 	/** Enable PostProcess for preview. */
-	UPROPERTY(EditAnywhere, BlueprintReadWrite, Category = "Editor Preview", meta = (DisplayName = "Enable Post Process"), BlueprintSetter = SetPreviewEnablePostProcess)
+	UPROPERTY(EditAnywhere, BlueprintReadWrite, Category = "Preview", meta = (DisplayName = "Enable Post Process"), BlueprintSetter = SetPreviewEnablePostProcess)
 	bool bPreviewEnablePostProcess = false;
 
 	/** Show overlay material on the preview mesh when preview rendering is enabled (UMeshComponent::OverlayMaterial). */
-	UPROPERTY(EditAnywhere, BlueprintReadWrite, Category = "Editor Preview", meta = (DisplayName = "Enable Preview Overlay"))
+	UPROPERTY(EditAnywhere, BlueprintReadWrite, Category = "Preview", AdvancedDisplay, meta = (DisplayName = "Enable Preview Overlay"))
 	bool bPreviewEnableOverlayMaterial = true;
 
 	/** Configure the root actor for Techvis rendering with preview components. */
-	UPROPERTY(EditAnywhere, BlueprintReadWrite, Category = "Editor Preview")
+	UPROPERTY(EditAnywhere, BlueprintReadWrite, Category = "Preview", AdvancedDisplay)
 	bool bEnablePreviewTechvis = false;
 
 	/** Enable the use of a preview mesh for the preview for this DCRA. */
@@ -482,55 +480,59 @@ public:
 	EDisplayClusterConfigurationRootActorPreviewSettingsSource PreviewSetttingsSource = EDisplayClusterConfigurationRootActorPreviewSettingsSource::RootActor;
 
 	/** Freeze preview render.  This will impact editor performance. */
-	UPROPERTY(EditAnywhere, BlueprintReadWrite, Category = "Editor Preview", meta = (DisplayName = "Freeze Editor Preview"))
+	UPROPERTY(EditAnywhere, BlueprintReadWrite, Category = "Preview", meta = (DisplayName = "Freeze Editor Preview"))
 	bool bFreezePreviewRender = false;
 
 	/** Render ICVFX Frustums */
-	UPROPERTY(EditAnywhere, BlueprintReadWrite, Category = "Editor Preview", meta = (DisplayName = "Enable Camera Frustums"))
+	UPROPERTY(EditAnywhere, BlueprintReadWrite, Category = "Preview", meta = (DisplayName = "Enable Camera Frustums"))
 	bool bPreviewICVFXFrustums = false;
 
 	/** Render ICVFX Frustums */
-	UPROPERTY(EditAnywhere, BlueprintReadWrite, Category = "Editor Preview", meta = (DisplayName = "Camera Frustum Distance"))
+	UPROPERTY(EditAnywhere, BlueprintReadWrite, Category = "Preview", AdvancedDisplay, meta = (DisplayName = "Camera Frustum Distance"))
 	float PreviewICVFXFrustumsFarDistance = 1000.0f;
 
+	/** Enables the use of 'bPreviewICVFXFrustums' in the game. */
+	UPROPERTY()
+	bool bPreviewInGameRenderFrustum = true;
+
 #if WITH_EDITORONLY_DATA
 	/** When the MRQ is rendered, this flag is raised. */
 	UPROPERTY(Transient, NonTransactional)
 	bool bMoviePipelineRenderPass = false;
 
 	/** Selectively preview a specific viewport or show all/none. */
-	UPROPERTY(EditAnywhere, BlueprintReadWrite, Category = "Editor Preview", meta = (DisplayName = "Preview Node"))
+	UPROPERTY(EditAnywhere, BlueprintReadWrite, Category = "Preview", meta = (DisplayName = "Preview Node"))
 	FString PreviewNodeId = DisplayClusterConfigurationStrings::gui::preview::PreviewNodeNone;
 
 	/** Render Mode */
-	UPROPERTY(EditAnywhere, BlueprintReadWrite, Category = "Editor Preview", meta = (DisplayName = "Render Mode"))
+	UPROPERTY(EditAnywhere, BlueprintReadWrite, Category = "Preview", meta = (DisplayName = "Render Mode"))
 	EDisplayClusterConfigurationRenderMode RenderMode = EDisplayClusterConfigurationRenderMode::Mono;
 #endif
 
 	/** Tick Per Frame */
-	UPROPERTY(EditAnywhere, BlueprintReadWrite, Category = "Editor Preview", AdvancedDisplay, meta = (DisplayName = "Tick Per Frame", ClampMin = "1", UIMin = "1", ClampMax = "200", UIMax = "200"))
+	UPROPERTY(EditAnywhere, BlueprintReadWrite, Category = "Preview", AdvancedDisplay, meta = (DisplayName = "Tick Per Frame", ClampMin = "1", UIMin = "1", ClampMax = "200", UIMax = "200"))
 	int TickPerFrame = 1;
 
 	/** Max amount of Viewports Per Frame */
-	UPROPERTY(EditAnywhere, BlueprintReadWrite, Category = "Editor Preview", AdvancedDisplay, meta = (DisplayName = "Viewports Per Frame", ClampMin = "1", UIMin = "1", ClampMax = "200", UIMax = "200"))
+	UPROPERTY(EditAnywhere, BlueprintReadWrite, Category = "Preview", AdvancedDisplay, meta = (DisplayName = "Viewports Per Frame", ClampMin = "1", UIMin = "1", ClampMax = "200", UIMax = "200"))
 	int ViewportsPerFrame = 1;
 
 	/** The maximum dimension of any internal texture for preview. Use less memory for large preview viewports */
-	UPROPERTY(EditAnywhere, BlueprintReadWrite, Category = "Editor Preview", AdvancedDisplay, meta = (DisplayName = "Preview Texture Max Size", ClampMin = "64", UIMin = "64", ClampMax = "4096", UIMax = "4096"))
+	UPROPERTY(EditAnywhere, BlueprintReadWrite, Category = "Preview", AdvancedDisplay, meta = (DisplayName = "Preview Texture Max Size", ClampMin = "64", UIMin = "64", ClampMax = "4096", UIMax = "4096"))
 	int PreviewMaxTextureDimension = 2048;
 
 	/** The included display device nDisplay provides by default */
-	UPROPERTY(VisibleDefaultsOnly, Category = "Editor Preview", DisplayName = "Basic Display Device")
+	UPROPERTY(VisibleDefaultsOnly, Category = "Preview", DisplayName = "Basic Display Device")
 	TObjectPtr<UDisplayClusterDisplayDeviceBaseComponent> BasicDisplayDeviceComponent;
 
 	/** Select the default display device class to use when a viewport doesn't have one assigned */
-	UPROPERTY(EditDefaultsOnly, Category = "Editor Preview", DisplayName = "Default Display Device")
+	UPROPERTY(EditDefaultsOnly, Category = "Preview", AdvancedDisplay, DisplayName = "Default Display Device")
 	FName DefaultDisplayDeviceName;
 
 
 #if WITH_EDITORONLY_DATA
 	/** Toggles the visibility of the stage's geometry mesh, a smooth, continuous mesh generated and processed from the stage's geometry */
-	UPROPERTY(EditInstanceOnly, Category = "Editor Preview", AdvancedDisplay)
+	UPROPERTY(EditInstanceOnly, Category = "Preview", AdvancedDisplay)
 	bool bPreviewStageGeometryMesh = false;
 #endif
 
diff --git a/Engine/Plugins/Runtime/nDisplay/Source/DisplayCluster/Public/Render/Viewport/Misc/DisplayClusterViewportHelpers.h b/Engine/Plugins/Runtime/nDisplay/Source/DisplayCluster/Public/Render/Viewport/Misc/DisplayClusterViewportHelpers.h
--- a/Engine/Plugins/Runtime/nDisplay/Source/DisplayCluster/Public/Render/Viewport/Misc/DisplayClusterViewportHelpers.h
+++ b/Engine/Plugins/Runtime/nDisplay/Source/DisplayCluster/Public/Render/Viewport/Misc/DisplayClusterViewportHelpers.h
@@ -13,6 +13,7 @@ namespace UE::DisplayClusterViewportHelpers
 {
 	/** Find a component by name from a RootActor of the specified type.
 	*
+	* @param InConfiguration - current configuration
 	* @param InRootActorType - the search is performed within the DCRA of the given type.
 	* @param InComponentName - the name of the component
 	*/
@@ -35,6 +36,36 @@ namespace UE::DisplayClusterViewportHelpers
 		return nullptr;
 	}
 
+	/** Find a component by name from a RootActor of the specified type.
+	*
+	* @param ComponentOfRootActor
+	* @param InComponentName - the name of the component
+	*/
+	template<class TComp>
+	TComp* GetOwnerRootActorComponentByName(USceneComponent* ComponentOfRootActor, const FString& InComponentName)
+	{
+		if (InComponentName.IsEmpty())
+		{
+			return nullptr;
+		}
+
+		if (AActor* OwnerActor = ComponentOfRootActor->GetOwner())
+		{
+			if (OwnerActor->IsA<ADisplayClusterRootActor>())
+			{
+				if (ADisplayClusterRootActor* OwnerRootActor = Cast<ADisplayClusterRootActor>(OwnerActor))
+				{
+					if (TComp* OutComponent = OwnerRootActor->GetComponentByName<TComp>(InComponentName))
+					{
+						return OutComponent;
+					}
+				}
+			}
+		}
+
+		return nullptr;
+	}
+
 	/** Return the same component (by class and name) from a DCRA of the specified type.
 	* If no other component exists, or if the current DCRA type matches the requested one,
 	* the input component is returned.
diff --git a/Engine/Plugins/Runtime/nDisplay/Source/DisplayClusterConfiguration/Public/DisplayClusterConfigurationStrings.h b/Engine/Plugins/Runtime/nDisplay/Source/DisplayClusterConfiguration/Public/DisplayClusterConfigurationStrings.h
--- a/Engine/Plugins/Runtime/nDisplay/Source/DisplayClusterConfiguration/Public/DisplayClusterConfigurationStrings.h
+++ b/Engine/Plugins/Runtime/nDisplay/Source/DisplayClusterConfiguration/Public/DisplayClusterConfigurationStrings.h
@@ -33,13 +33,19 @@ namespace DisplayClusterConfigurationStrings
 		static constexpr const TCHAR* TileCategory                 = TEXT("Tile Rendering");
 		static constexpr const TCHAR* OverrideCategory             = TEXT("Texture Replacement");
 		static constexpr const TCHAR* ViewportsCategory            = TEXT("Viewports");
-		static constexpr const TCHAR* ICVFXCategory                = TEXT("Inner Frustum");
-		static constexpr const TCHAR* ICVFXCameraCategory          = TEXT("ICVFX Camera");
+		
+		static constexpr const TCHAR* InCameraVFXCategory  = TEXT("In Camera VFX");
+		static constexpr const TCHAR* InnerFrustumCategory = TEXT("Inner Frustum");
+		static constexpr const TCHAR* ICVFXCameraCategory  = TEXT("ICVFX Camera");
+
 		static constexpr const TCHAR* ConfigurationCategory        = TEXT("Configuration");
-		static constexpr const TCHAR* PreviewCategory              = TEXT("Editor Preview");
-		static constexpr const TCHAR* PreviewInGameCategory        = TEXT("Preview In Game");
+		static constexpr const TCHAR* PreviewCategory              = TEXT("Preview");
 		static constexpr const TCHAR* AdvancedCategory             = TEXT("Advanced");
 		static constexpr const TCHAR* TextureShareCategory         = TEXT("Texture Share");
+
+		static constexpr const TCHAR* ViewPointStereoCategory              = TEXT("Stereo");
+		static constexpr const TCHAR* ViewPointCameraPostProcessCategory   = TEXT("Camera Post Process");
+		static constexpr const TCHAR* ViewPointInFrustumProjectionCategory = TEXT("Frustum Fit");
 	}
 
 	// Command line arguments
diff --git a/Engine/Plugins/Runtime/nDisplay/Source/DisplayClusterConfiguration/Public/DisplayClusterConfigurationTypes_Viewport.h b/Engine/Plugins/Runtime/nDisplay/Source/DisplayClusterConfiguration/Public/DisplayClusterConfigurationTypes_Viewport.h
--- a/Engine/Plugins/Runtime/nDisplay/Source/DisplayClusterConfiguration/Public/DisplayClusterConfigurationTypes_Viewport.h
+++ b/Engine/Plugins/Runtime/nDisplay/Source/DisplayClusterConfiguration/Public/DisplayClusterConfigurationTypes_Viewport.h
@@ -189,8 +189,8 @@ public:
 	UPROPERTY(EditAnywhere, BlueprintReadWrite, Category = "Configuration", meta = (DisplayName = "Enable Viewport"))
 	bool bAllowRendering = true;
 
-	/** Reference to the nDisplay View Origin */
-	UPROPERTY(EditAnywhere, BlueprintReadWrite, Category = "Configuration", meta = (DisplayName = "View Origin"))
+	/** Reference to the nDisplay View Point*/
+	UPROPERTY(EditAnywhere, BlueprintReadWrite, Category = "Configuration", meta = (DisplayName = "View Point"))
 	FString Camera;
 
 	/** Specify your Projection Policy Settings */
diff --git a/Engine/Plugins/Runtime/nDisplay/Source/DisplayClusterConfigurator/Private/DisplayClusterConfiguratorModule.cpp b/Engine/Plugins/Runtime/nDisplay/Source/DisplayClusterConfigurator/Private/DisplayClusterConfiguratorModule.cpp
--- a/Engine/Plugins/Runtime/nDisplay/Source/DisplayClusterConfigurator/Private/DisplayClusterConfiguratorModule.cpp
+++ b/Engine/Plugins/Runtime/nDisplay/Source/DisplayClusterConfigurator/Private/DisplayClusterConfiguratorModule.cpp
@@ -35,6 +35,7 @@
 #include "Blueprints/DisplayClusterBlueprint.h"
 #include "Components/DisplayClusterScreenComponent.h"
 #include "Components/DisplayClusterCameraComponent.h"
+#include "Components/DisplayClusterInFrustumFitCameraComponent.h"
 #include "Components/DisplayClusterICVFXCameraComponent.h"
 #include "Misc/DisplayClusterObjectRef.h"
 #include "DisplayClusterRootActor.h"
@@ -267,8 +268,8 @@ void FDisplayClusterConfiguratorModule::RegisterSectionMappings()
 	}
 	{
 		const TSharedRef<FPropertySection> Section = PropertyModule.FindOrCreateSection(ADisplayClusterRootActor::StaticClass()->GetFName(),
-			DisplayClusterConfigurationStrings::categories::ICVFXCategory, LOCTEXT("In-Camera VFX", "In-Camera VFX"));
-		Section->AddCategory(DisplayClusterConfigurationStrings::categories::ICVFXCategory);
+			DisplayClusterConfigurationStrings::categories::InCameraVFXCategory, LOCTEXT("InCameraVFXCategoryLabel", "In-Camera VFX"));
+		Section->AddCategory(DisplayClusterConfigurationStrings::categories::InCameraVFXCategory);
 	}
 	{
 		const TSharedRef<FPropertySection> Section = PropertyModule.FindOrCreateSection(ADisplayClusterRootActor::StaticClass()->GetFName(),
@@ -287,20 +288,15 @@ void FDisplayClusterConfiguratorModule::RegisterSectionMappings()
 	}
 	{
 		const TSharedRef<FPropertySection> Section = PropertyModule.FindOrCreateSection(ADisplayClusterRootActor::StaticClass()->GetFName(),
-			DisplayClusterConfigurationStrings::categories::PreviewCategory, LOCTEXT("Editor Preview", "Editor Preview"));
+			DisplayClusterConfigurationStrings::categories::PreviewCategory, LOCTEXT("Preview", "Preview"));
 		Section->AddCategory(DisplayClusterConfigurationStrings::categories::PreviewCategory);
-	}
-	{
-		const TSharedRef<FPropertySection> Section = PropertyModule.FindOrCreateSection(ADisplayClusterRootActor::StaticClass()->GetFName(),
-			DisplayClusterConfigurationStrings::categories::PreviewInGameCategory, LOCTEXT("Preview In Game", "Preview In Game"));
-		Section->AddCategory(DisplayClusterConfigurationStrings::categories::PreviewInGameCategory);
-	}
+	}	
 
 	// ICVFX Component
 	{
 		const TSharedRef<FPropertySection> Section = PropertyModule.FindOrCreateSection(UDisplayClusterICVFXCameraComponent::StaticClass()->GetFName(),
-		DisplayClusterConfigurationStrings::categories::ICVFXCategory, LOCTEXT("InnerFrustumSectionLabel", "Inner Frustum"));
-		Section->AddCategory(DisplayClusterConfigurationStrings::categories::ICVFXCategory);
+		DisplayClusterConfigurationStrings::categories::InnerFrustumCategory, LOCTEXT("InnerFrustumCategoryLabel", "Inner Frustum"));
+		Section->AddCategory(DisplayClusterConfigurationStrings::categories::InnerFrustumCategory);
 	}
 	{
 		const TSharedRef<FPropertySection> Section = PropertyModule.FindOrCreateSection(UDisplayClusterICVFXCameraComponent::StaticClass()->GetFName(),
@@ -327,6 +323,26 @@ void FDisplayClusterConfiguratorModule::RegisterSectionMappings()
 			DisplayClusterConfigurationStrings::categories::ChromaKeyCategory, LOCTEXT("Chromakey", "Chromakey"));
 		Section->AddCategory(DisplayClusterConfigurationStrings::categories::ChromaKeyCategory);
 	}
+
+	// ViewPoint
+	{
+		const TSharedRef<FPropertySection> Section = PropertyModule.FindOrCreateSection(UDisplayClusterCameraComponent::StaticClass()->GetFName(),
+			DisplayClusterConfigurationStrings::categories::ViewPointStereoCategory, LOCTEXT("ViewPointStereoSectionLabel", "Stereo"));
+		Section->AddCategory(DisplayClusterConfigurationStrings::categories::ViewPointStereoCategory);
+	}
+	{
+		const TSharedRef<FPropertySection> Section = PropertyModule.FindOrCreateSection(UDisplayClusterCameraComponent::StaticClass()->GetFName(),
+			DisplayClusterConfigurationStrings::categories::ViewPointCameraPostProcessCategory, LOCTEXT("ViewPointCameraPostProcessSectionLabel", "Camera Post Process"));
+		Section->AddCategory(DisplayClusterConfigurationStrings::categories::ViewPointCameraPostProcessCategory);
+	}
+
+	// InFrustum ViewPoint
+	{
+		const TSharedRef<FPropertySection> Section = PropertyModule.FindOrCreateSection(UDisplayClusterInFrustumFitCameraComponent::StaticClass()->GetFName(),
+			DisplayClusterConfigurationStrings::categories::ViewPointInFrustumProjectionCategory, LOCTEXT("ViewPointInFrustumProjectionSectionLabel", "Frustum Fit"));
+		Section->AddCategory(DisplayClusterConfigurationStrings::categories::ViewPointInFrustumProjectionCategory);
+	}
+
 }
 
 void FDisplayClusterConfiguratorModule::UnregisterSectionMappings()
@@ -335,12 +351,12 @@ void FDisplayClusterConfiguratorModule::UnregisterSectionMappings()
 	{
 		FPropertyEditorModule& PropertyModule = FModuleManager::GetModuleChecked<FPropertyEditorModule>("PropertyEditor");
 		PropertyModule.RemoveSection(ADisplayClusterRootActor::StaticClass()->GetFName(), DisplayClusterConfigurationStrings::categories::ViewportsCategory);
-		PropertyModule.RemoveSection(ADisplayClusterRootActor::StaticClass()->GetFName(), DisplayClusterConfigurationStrings::categories::ICVFXCategory);
+		PropertyModule.RemoveSection(ADisplayClusterRootActor::StaticClass()->GetFName(), DisplayClusterConfigurationStrings::categories::InCameraVFXCategory);
 		PropertyModule.RemoveSection(ADisplayClusterRootActor::StaticClass()->GetFName(), DisplayClusterConfigurationStrings::categories::ColorGradingCategory);
 		PropertyModule.RemoveSection(ADisplayClusterRootActor::StaticClass()->GetFName(), DisplayClusterConfigurationStrings::categories::OCIOCategory);
 		PropertyModule.RemoveSection(ADisplayClusterRootActor::StaticClass()->GetFName(), DisplayClusterConfigurationStrings::categories::LightcardCategory);
 
-		PropertyModule.RemoveSection(UDisplayClusterICVFXCameraComponent::StaticClass()->GetFName(), DisplayClusterConfigurationStrings::categories::ICVFXCategory);
+		PropertyModule.RemoveSection(UDisplayClusterICVFXCameraComponent::StaticClass()->GetFName(), DisplayClusterConfigurationStrings::categories::InnerFrustumCategory);
 		PropertyModule.RemoveSection(UDisplayClusterICVFXCameraComponent::StaticClass()->GetFName(), DisplayClusterConfigurationStrings::categories::CameraColorGradingCategory);
 		PropertyModule.RemoveSection(UDisplayClusterICVFXCameraComponent::StaticClass()->GetFName(), DisplayClusterConfigurationStrings::categories::OCIOCategory);
 		PropertyModule.RemoveSection(UDisplayClusterICVFXCameraComponent::StaticClass()->GetFName(), DisplayClusterConfigurationStrings::categories::ChromaKeyCategory);
diff --git a/Engine/Plugins/Runtime/nDisplay/Source/DisplayClusterConfigurator/Private/MPCDI/DisplayClusterConfiguratorMPCDIImporter.cpp b/Engine/Plugins/Runtime/nDisplay/Source/DisplayClusterConfigurator/Private/MPCDI/DisplayClusterConfiguratorMPCDIImporter.cpp
--- a/Engine/Plugins/Runtime/nDisplay/Source/DisplayClusterConfigurator/Private/MPCDI/DisplayClusterConfiguratorMPCDIImporter.cpp
+++ b/Engine/Plugins/Runtime/nDisplay/Source/DisplayClusterConfigurator/Private/MPCDI/DisplayClusterConfiguratorMPCDIImporter.cpp
@@ -36,10 +36,10 @@ bool FDisplayClusterConfiguratorMPCDIImporter::ImportMPCDIIntoBlueprint(const FS
 
 	// The ViewPoint component cannot be used as a ViewOrigin. They have different purposes.
 	// The default ViewOrigin is the RootComponent.
-	const FName ViewOriginComponentName = !InParams.ViewOriginComponentName.IsNone() ? InParams.ViewOriginComponentName : TEXT("RootComponent");
+	const FName OriginComponentName = !InParams.OriginComponentName.IsNone() ? InParams.OriginComponentName : TEXT("RootComponent");
 
-	USCS_Node* ViewOriginNode = nullptr;
-	UDisplayClusterCameraComponent* ViewOriginComponent = Cast<UDisplayClusterCameraComponent>(FindBlueprintSceneComponent(InBlueprint, ViewOriginComponentName, &ViewOriginNode));
+	USCS_Node* OriginNode = nullptr;
+	UDisplayClusterCameraComponent* OriginComponent = Cast<UDisplayClusterCameraComponent>(FindBlueprintSceneComponent(InBlueprint, OriginComponentName, &OriginNode));
 
 	FIPv4Address CurrentIPAddress = InParams.HostStartingIPAddress;
 	for (const TPair<FString, TMap<FString, FDisplayClusterWarpMPCDIAttributes>>& Buffer : MPCDIFile)
@@ -56,7 +56,7 @@ bool FDisplayClusterConfiguratorMPCDIImporter::ImportMPCDIIntoBlueprint(const FS
 					if (ScreenNode)
 					{
 						UDisplayClusterScreenComponent* ScreenComponent = CastChecked<UDisplayClusterScreenComponent>(ScreenNode->GetActualComponentTemplate(InBlueprint->GetGeneratedClass()));
-						ConfigureScreenComponentFrom2DProfileRegion(ScreenComponent, ViewOriginComponent, Region.Value, InParams);
+						ConfigureScreenComponentFrom2DProfileRegion(ScreenComponent, OriginComponent, Region.Value, InParams);
 
 						// If a new screen had to be created, it needs to be parented to a valid parent component
 						if (!bFoundExistingScreen)
@@ -160,7 +160,7 @@ USCS_Node* FDisplayClusterConfiguratorMPCDIImporter::FindOrCreateScreenNodeForRe
 
 void FDisplayClusterConfiguratorMPCDIImporter::ConfigureScreenComponentFrom2DProfileRegion(
 	UDisplayClusterScreenComponent* InScreenComponent,
-	UDisplayClusterCameraComponent* InViewOriginComponent,
+	USceneComponent* InOriginComponent,
 	const FDisplayClusterWarpMPCDIAttributes& InAttributes,
 	const FDisplayClusterConfiguratorMPCDIImporterParams& InParams)
 {
@@ -171,16 +171,16 @@ void FDisplayClusterConfiguratorMPCDIImporter::ConfigureScreenComponentFrom2DPro
 		return;
 	}
 
-	if (InViewOriginComponent)
+	if (InOriginComponent)
 	{
-		// We want to position the screen component to be in front of the view origin it is assigned to, so we must compute
-		// the transforms from view origin to root and from root to screen component to correct position
+		// We want to position the screen component to be in front of the view point it is assigned to, so we must compute
+		// the transforms from view point to root and from root to screen component to correct position
 
-		// Blueprint components don't have world transforms, so compute the view origin to root transform manually
-		FTransform ViewOriginTransform = FTransform::Identity;
-		for (const USceneComponent* Comp = InViewOriginComponent; Comp != nullptr; Comp = Comp->GetAttachParent())
+		// Blueprint components don't have world transforms, so compute the view point to root transform manually
+		FTransform ViewPointTransform = FTransform::Identity;
+		for (const USceneComponent* Comp = InOriginComponent; Comp != nullptr; Comp = Comp->GetAttachParent())
 		{
-			ViewOriginTransform *= Comp->GetRelativeTransform();
+			ViewPointTransform *= Comp->GetRelativeTransform();
 		}
 
 		// The screen may also be parented to a non-root component, so also compute its root transform
@@ -190,8 +190,8 @@ void FDisplayClusterConfiguratorMPCDIImporter::ConfigureScreenComponentFrom2DPro
 			ScreenTransform *= Comp->GetRelativeTransform();
 		}
 
-		FTransform ViewOriginToScreen = ViewOriginTransform * ScreenTransform.Inverse();
-		ScreenPosition = ViewOriginToScreen.TransformPositionNoScale(ScreenPosition);
+		FTransform ViewPointToScreen = ViewPointTransform * ScreenTransform.Inverse();
+		ScreenPosition = ViewPointToScreen.TransformPositionNoScale(ScreenPosition);
 	}
 
 	InScreenComponent->SetRelativeLocation(ScreenPosition);
@@ -239,9 +239,9 @@ void FDisplayClusterConfiguratorMPCDIImporter::ConfigureViewportFromRegion(UDisp
 	FDisplayClusterConfiguratorViewportViewModel ViewportViewModel(InViewport);
 	ViewportViewModel.SetRegion(FDisplayClusterConfigurationRectangle(0, 0, InAttributes.Region.Resolution.X, InAttributes.Region.Resolution.Y));
 
-	if (InParams.ViewOriginComponentName != NAME_None)
+	if (InParams.ViewPointComponentName != NAME_None)
 	{
-		ViewportViewModel.SetCamera(InParams.ViewOriginComponentName.ToString());
+		ViewportViewModel.SetCamera(InParams.ViewPointComponentName.ToString());
 	}
 }
 
diff --git a/Engine/Plugins/Runtime/nDisplay/Source/DisplayClusterConfigurator/Private/MPCDI/DisplayClusterConfiguratorMPCDIImporter.h b/Engine/Plugins/Runtime/nDisplay/Source/DisplayClusterConfigurator/Private/MPCDI/DisplayClusterConfiguratorMPCDIImporter.h
--- a/Engine/Plugins/Runtime/nDisplay/Source/DisplayClusterConfigurator/Private/MPCDI/DisplayClusterConfiguratorMPCDIImporter.h
+++ b/Engine/Plugins/Runtime/nDisplay/Source/DisplayClusterConfigurator/Private/MPCDI/DisplayClusterConfiguratorMPCDIImporter.h
@@ -19,8 +19,11 @@ struct FDisplayClusterConfiguratorMPCDIImporterParams
 	/** The name of the component to make the parent of the MPCDI geometry */
 	FName ParentComponentName = NAME_None;
 
-	/** The name of the view origin component to link to the MPCDI geometry. If none is supplied, the default view origin is used */
-	FName ViewOriginComponentName = NAME_None;
+	/** The name of the view point component to link to the viewports. If none is supplied, the default view point is used */
+	FName ViewPointComponentName = NAME_None;
+
+	/** The name of the origin component to link to the MPCDI geometry. If none is supplied, the root component is used */
+	FName OriginComponentName = NAME_None;
 
 	/** The IP address to configure the first cluster node with */
 	FIPv4Address HostStartingIPAddress = FIPv4Address::InternalLoopback;
@@ -58,11 +61,11 @@ private:
 	 * Configures a screen component to match the configuration of an MPCDI region
 	 * 
 	 * @param InScreenComponent - The component to configure
-	 * @param InViewOriginComponent - The view origin component to configure the screens from
-	 * @param InAttributes - The attributes of the region to configure
-	 * @param InParams - The import parameters 
+	 * @param InOriginComponent - The Origin component to configure the screens from
+	 * @param InAttributes      - The attributes of the region to configure
+	 * @param InParams          - The import parameters
 	 **/
-	static void ConfigureScreenComponentFrom2DProfileRegion(UDisplayClusterScreenComponent* InScreenComponent, UDisplayClusterCameraComponent* InViewOriginComponent, const FDisplayClusterWarpMPCDIAttributes& InAttributes, const FDisplayClusterConfiguratorMPCDIImporterParams& InParams);
+	static void ConfigureScreenComponentFrom2DProfileRegion(UDisplayClusterScreenComponent* InScreenComponent, USceneComponent* InOriginComponent, const FDisplayClusterWarpMPCDIAttributes& InAttributes, const FDisplayClusterConfiguratorMPCDIImporterParams& InParams);
 
 	/**
 	 * Attempts to find an existing viewport configuration for the specified region, and if one is not found, creates a new viewport configuration
diff --git a/Engine/Plugins/Runtime/nDisplay/Source/DisplayClusterConfigurator/Private/Views/Details/Cluster/DisplayClusterConfiguratorViewportDetailsCustomization.cpp b/Engine/Plugins/Runtime/nDisplay/Source/DisplayClusterConfigurator/Private/Views/Details/Cluster/DisplayClusterConfiguratorViewportDetailsCustomization.cpp
--- a/Engine/Plugins/Runtime/nDisplay/Source/DisplayClusterConfigurator/Private/Views/Details/Cluster/DisplayClusterConfiguratorViewportDetailsCustomization.cpp
+++ b/Engine/Plugins/Runtime/nDisplay/Source/DisplayClusterConfigurator/Private/Views/Details/Cluster/DisplayClusterConfiguratorViewportDetailsCustomization.cpp
@@ -22,7 +22,7 @@ void FDisplayClusterConfiguratorViewportDetailsCustomization::CustomizeDetails(I
 	Super::CustomizeDetails(InLayoutBuilder);
 
 	ConfigurationViewportPtr = nullptr;
-	NoneOption = MakeShared<FString>("None");
+	DefaultViewPointComponentNameOption = MakeShared<FString>("DefaultViewPoint");
 
 	// Set config data pointer
 	UDisplayClusterConfigurationData* ConfigurationData = GetConfigData();
@@ -48,7 +48,7 @@ void FDisplayClusterConfiguratorViewportDetailsCustomization::CustomizeDetails(I
 
 	if (!bDisplayICVFX)
 	{
-		IDetailCategoryBuilder& Category = InLayoutBuilder.EditCategory(TEXT("In Camera VFX"));
+		IDetailCategoryBuilder& Category = InLayoutBuilder.EditCategory(DisplayClusterConfigurationStrings::categories::InCameraVFXCategory);
 		Category.SetCategoryVisibility(false);
 	}
 
@@ -151,12 +151,6 @@ void FDisplayClusterConfiguratorViewportDetailsCustomization::ResetCameraOptions
 		// Default sort isn't compatible with TSharedPtr<FString>.
 		return *A < *B;
 	});
-
-	// Add None option
-	if (!ConfigurationViewport->Camera.IsEmpty())
-	{
-		CameraOptions.Add(NoneOption);
-	}
 }
 
 TSharedRef<SWidget> FDisplayClusterConfiguratorViewportDetailsCustomization::CreateCustomCameraWidget()
@@ -194,7 +188,7 @@ void FDisplayClusterConfiguratorViewportDetailsCustomization::OnCameraSelected(T
 		UDisplayClusterConfigurationViewport* ConfigurationViewport = ConfigurationViewportPtr.Get();
 		check(ConfigurationViewport != nullptr);
 		// Handle empty case
-		if (InCamera->Equals(*NoneOption.Get()))
+		if (InCamera->Equals(*DefaultViewPointComponentNameOption.Get()))
 		{
 			CameraHandle->SetValue(TEXT(""));
 		}
@@ -214,7 +208,8 @@ FText FDisplayClusterConfiguratorViewportDetailsCustomization::GetSelectedCamera
 	FString SelectedOption = ConfigurationViewportPtr.Get()->Camera;
 	if (SelectedOption.IsEmpty())
 	{
-		SelectedOption = *NoneOption.Get();
+		SelectedOption = *DefaultViewPointComponentNameOption.Get();
 	}
+
 	return FText::FromString(SelectedOption);
-}
\ No newline at end of file
+}
diff --git a/Engine/Plugins/Runtime/nDisplay/Source/DisplayClusterConfigurator/Private/Views/Details/Cluster/DisplayClusterConfiguratorViewportDetailsCustomization.h b/Engine/Plugins/Runtime/nDisplay/Source/DisplayClusterConfigurator/Private/Views/Details/Cluster/DisplayClusterConfiguratorViewportDetailsCustomization.h
--- a/Engine/Plugins/Runtime/nDisplay/Source/DisplayClusterConfigurator/Private/Views/Details/Cluster/DisplayClusterConfiguratorViewportDetailsCustomization.h
+++ b/Engine/Plugins/Runtime/nDisplay/Source/DisplayClusterConfigurator/Private/Views/Details/Cluster/DisplayClusterConfiguratorViewportDetailsCustomization.h
@@ -56,8 +56,8 @@ private:
 	/** The property handle for the Camera property of the UDisplayClusterConfigurationViewport object */
 	TSharedPtr<IPropertyHandle> CameraHandle;
 
-	/** A cached pointer to the "None" option that is added to the list of options in the dropdown menu */
-	TSharedPtr<FString>	NoneOption;
+	/** A cached pointer to the "DefaultViewPoint" option that is added to the list of options in the dropdown menu */
+	TSharedPtr<FString> DefaultViewPointComponentNameOption;
 
 	/** A weak reference to the UDisplayClusterConfigurationViewport object being edited by the details panel */
 	TWeakObjectPtr<UDisplayClusterConfigurationViewport> ConfigurationViewportPtr;
diff --git a/Engine/Plugins/Runtime/nDisplay/Source/DisplayClusterConfigurator/Private/Views/Details/Components/DisplayClusterCameraComponentDetailsCustomization.cpp b/Engine/Plugins/Runtime/nDisplay/Source/DisplayClusterConfigurator/Private/Views/Details/Components/DisplayClusterCameraComponentDetailsCustomization.cpp
--- a/Engine/Plugins/Runtime/nDisplay/Source/DisplayClusterConfigurator/Private/Views/Details/Components/DisplayClusterCameraComponentDetailsCustomization.cpp
+++ b/Engine/Plugins/Runtime/nDisplay/Source/DisplayClusterConfigurator/Private/Views/Details/Components/DisplayClusterCameraComponentDetailsCustomization.cpp
@@ -5,6 +5,7 @@
 #include "Components/DisplayClusterCameraComponent.h"
 #include "DisplayClusterRootActor.h"
 #include "Camera/CameraComponent.h"
+#include "Components/DisplayClusterICVFXCameraComponent.h"
 
 #include "DisplayClusterProjectionStrings.h"
 #include "DisplayClusterConfiguratorLog.h"
@@ -16,7 +17,9 @@
 #include "DetailLayoutBuilder.h"
 #include "DetailCategoryBuilder.h"
 #include "DetailWidgetRow.h"
+#include "PropertyCustomizationHelpers.h"
 #include "PropertyHandle.h"
+#include "ScopedTransaction.h"
 
 #define LOCTEXT_NAMESPACE "DisplayClusterCameraComponentDetailsCustomization"
 
@@ -42,9 +45,9 @@ void FDisplayClusterCameraComponentDetailsCustomization::CustomizeDetails(IDetai
 
 	if (EditedObject.IsValid())
 	{
-		NoneOption = MakeShared<FString>("Active Engine Camera");
+		NoneOption = MakeShared<FString>("None");
 
-		CameraHandle = InLayoutBuilder.GetProperty(GET_MEMBER_NAME_CHECKED(UDisplayClusterCameraComponent, OuterViewportCameraName));
+		CameraHandle = InLayoutBuilder.GetProperty(GET_MEMBER_NAME_CHECKED(UDisplayClusterCameraComponent, ICVFXCameraComponentName));
 		check(CameraHandle->IsValidHandle());
 
 		RebuildCameraOptions();
@@ -64,17 +67,30 @@ void FDisplayClusterCameraComponentDetailsCustomization::CustomizeDetails(IDetai
 	}
 }
 
+UClass* FDisplayClusterCameraComponentDetailsCustomization::GetCameraComponentClass() const
+{
+	return UDisplayClusterICVFXCameraComponent::StaticClass();
+}
+
 void FDisplayClusterCameraComponentDetailsCustomization::RebuildCameraOptions()
 {
 	CameraOptions.Reset();
-	UDisplayClusterCameraComponent* DestCameraComponent = EditedObject.Get();
+
+	// The EditedObject may become invalid when component destoryed (DCRA rebuild).
+	if (!EditedObject.IsValid())
+	{
+		// Just do nothing
+		return;
+	}
+
+	const UDisplayClusterCameraComponent* DestCameraComponent = EditedObject.Get();
 	check(DestCameraComponent != nullptr);
 
-	AActor* RootActor = GetRootActor();
+	const AActor* RootActor = GetRootActor();
 	check(RootActor);
 
 	TArray<UActorComponent*> ActorComponents;
-	RootActor->GetComponents(UCameraComponent::StaticClass(), ActorComponents);
+	RootActor->GetComponents(GetCameraComponentClass(), ActorComponents);
 	for (UActorComponent* ActorComponent : ActorComponents)
 	{
 		const FString ComponentName = ActorComponent->GetName();
@@ -89,7 +105,7 @@ void FDisplayClusterCameraComponentDetailsCustomization::RebuildCameraOptions()
 	});
 
 	// Add None option
-	if (!DestCameraComponent->OuterViewportCameraName.IsEmpty())
+	if (!DestCameraComponent->ICVFXCameraComponentName.IsEmpty())
 	{
 		CameraOptions.Add(NoneOption);
 	}
@@ -97,23 +113,95 @@ void FDisplayClusterCameraComponentDetailsCustomization::RebuildCameraOptions()
 
 TSharedRef<SWidget> FDisplayClusterCameraComponentDetailsCustomization::CreateCustomCameraWidget()
 {
-	if (CameraComboBox.IsValid())
+	if (CameraComboBoxWidged.IsValid())
 	{
-		return CameraComboBox.ToSharedRef();
+		return CameraComboBoxWidged.ToSharedRef();
 	}
 
-	return SAssignNew(CameraComboBox, SDisplayClusterConfigurationSearchableComboBox)
-		.OptionsSource(&CameraOptions)
-		.OnGenerateWidget(this, &FDisplayClusterCameraComponentDetailsCustomization::MakeCameraOptionComboWidget)
-		.OnSelectionChanged(this, &FDisplayClusterCameraComponentDetailsCustomization::OnCameraSelected)
-		.ContentPadding(2)
-		.MaxListHeight(200.0f)
-		.Content()
+	CameraComboBoxWidged = SNew(SHorizontalBox)
+		+ SHorizontalBox::Slot()
+		.HAlign(HAlign_Fill)
+		.FillWidth(VAlign_Fill)
+		.FillWidth(1.f)
+		[
+			SAssignNew(CameraComboBox, SDisplayClusterConfigurationSearchableComboBox)
+				.OptionsSource(&CameraOptions)
+				.OnGenerateWidget(this, &FDisplayClusterCameraComponentDetailsCustomization::MakeCameraOptionComboWidget)
+				.OnSelectionChanged(this, &FDisplayClusterCameraComponentDetailsCustomization::OnCameraSelected)
+				.ContentPadding(2)
+				.MaxListHeight(200.0f)
+				.Content()
+				[
+					SNew(STextBlock)
+						.Text(this, &FDisplayClusterCameraComponentDetailsCustomization::GetSelectedCameraText)
+						.Font(IDetailLayoutBuilder::GetDetailFont())
+				]
+		]
+		/* Add browse to component button */
+		+ SHorizontalBox::Slot()
+		.HAlign(HAlign_Fill)
+		.FillWidth(VAlign_Fill)
+		.AutoWidth()
 		[
-			SNew(STextBlock)
-			.Text(this, &FDisplayClusterCameraComponentDetailsCustomization::GetSelectedCameraText)
-			.Font(IDetailLayoutBuilder::GetDetailFont())
+			PropertyCustomizationHelpers::MakeBrowseButton(
+				FSimpleDelegate::CreateSP(this, &FDisplayClusterCameraComponentDetailsCustomization::OnSelectComponentButton)
+				, TAttribute<FText>(this, &FDisplayClusterCameraComponentDetailsCustomization::GetSelectComponentButtonTooltipText)
+				, /** IsEnabled */true
+				, /** IsActor */true)
 		];
+
+	return CameraComboBoxWidged.ToSharedRef();
+}
+
+
+FText FDisplayClusterCameraComponentDetailsCustomization::GetSelectComponentButtonTooltipText() const
+{
+	const FText CameraName = GetSelectedCameraText();
+	if (!CameraName.ToString().Equals(*NoneOption.Get()))
+	{
+		return FText::Format(LOCTEXT("SelectCameraComponent", "Select '{0}' camera component in the Root Actor"), CameraName);
+	}
+
+	return FText::GetEmpty();
+}
+
+void FDisplayClusterCameraComponentDetailsCustomization::OnSelectComponentButton() const
+{
+	if (!GEditor)
+	{
+		return;
+	}
+
+	const FText CameraName = GetSelectedCameraText();
+	if (!CameraName.ToString().Equals(*NoneOption.Get()))
+	{
+		// The component in the DCRA to be selected.
+		UActorComponent* RootActorComponent = nullptr;
+
+		const AActor* RootActor = GetRootActor();
+		check(RootActor);
+
+		TArray<UActorComponent*> ActorComponents;
+		RootActor->GetComponents(GetCameraComponentClass(), ActorComponents);
+		for (UActorComponent* ActorComponent : ActorComponents)
+		{
+			const FString ComponentName = ActorComponent->GetName();
+			if (ComponentName == CameraName.ToString())
+			{
+				RootActorComponent = ActorComponent;
+				break;
+			}
+		}
+
+		if (IsValid(RootActorComponent))
+		{
+			const FScopedTransaction Transaction(NSLOCTEXT("UnrealEd", "ClickingOnElements", "Clicking on Elements"));
+
+			constexpr bool bNotifySelectionChanged = true;
+			GEditor->SelectNone(!bNotifySelectionChanged, /** DeselectBSP */true);
+			GEditor->SelectComponent(RootActorComponent, /** IsSelected */true, bNotifySelectionChanged);
+		}
+	}
 }
 
 TSharedRef<SWidget> FDisplayClusterCameraComponentDetailsCustomization::MakeCameraOptionComboWidget(TSharedPtr<FString> InItem)
@@ -146,11 +234,22 @@ void FDisplayClusterCameraComponentDetailsCustomization::OnCameraSelected(TShare
 
 FText FDisplayClusterCameraComponentDetailsCustomization::GetSelectedCameraText() const
 {
-	FString SelectedOption = EditedObject.Get()->OuterViewportCameraName;
+	// The EditedObject may become invalid when component destoryed (DCRA rebuild).
+	if (!EditedObject.IsValid())
+	{
+		// Just do nothing
+		return FText::FromString(*NoneOption.Get());
+	}
+
+	const UDisplayClusterCameraComponent* DestCameraComponent = EditedObject.Get();
+	check(DestCameraComponent != nullptr);
+
+	FString SelectedOption = DestCameraComponent->ICVFXCameraComponentName;
 	if (SelectedOption.IsEmpty())
 	{
 		SelectedOption = *NoneOption.Get();
 	}
+
 	return FText::FromString(SelectedOption);
 }
 
diff --git a/Engine/Plugins/Runtime/nDisplay/Source/DisplayClusterConfigurator/Private/Views/Details/Components/DisplayClusterCameraComponentDetailsCustomization.h b/Engine/Plugins/Runtime/nDisplay/Source/DisplayClusterConfigurator/Private/Views/Details/Components/DisplayClusterCameraComponentDetailsCustomization.h
--- a/Engine/Plugins/Runtime/nDisplay/Source/DisplayClusterConfigurator/Private/Views/Details/Components/DisplayClusterCameraComponentDetailsCustomization.h
+++ b/Engine/Plugins/Runtime/nDisplay/Source/DisplayClusterConfigurator/Private/Views/Details/Components/DisplayClusterCameraComponentDetailsCustomization.h
@@ -49,6 +49,16 @@ private:
 	/** Gets the text to display for the currently selected camera */
 	FText GetSelectedCameraText() const;
 
+	/** When select component button is pressed. */
+	void OnSelectComponentButton() const;
+
+	/** Return select component button tooltip text. */
+	FText GetSelectComponentButtonTooltipText() const;
+
+	/** Return camera component class. */
+	UClass* GetCameraComponentClass() const;
+
+
 private:
 	/** A weak reference to the UDisplayClusterCameraComponent object being edited by the details panel */
 	TWeakObjectPtr<UDisplayClusterCameraComponent> EditedObject;
@@ -68,4 +78,7 @@ private:
 
 	/** The combo box that is being displayed in the details panel for the Camera property */
 	TSharedPtr<SDisplayClusterConfigurationSearchableComboBox> CameraComboBox;
+
+	/** The widged that is being displayed in the details panel for the Camera property */
+	TSharedPtr<SWidget> CameraComboBoxWidged;
 };
diff --git a/Engine/Plugins/Runtime/nDisplay/Source/DisplayClusterConfigurator/Private/Views/Details/Components/DisplayClusterICVFXCameraComponentDetailsCustomization.cpp b/Engine/Plugins/Runtime/nDisplay/Source/DisplayClusterConfigurator/Private/Views/Details/Components/DisplayClusterICVFXCameraComponentDetailsCustomization.cpp
--- a/Engine/Plugins/Runtime/nDisplay/Source/DisplayClusterConfigurator/Private/Views/Details/Components/DisplayClusterICVFXCameraComponentDetailsCustomization.cpp
+++ b/Engine/Plugins/Runtime/nDisplay/Source/DisplayClusterConfigurator/Private/Views/Details/Components/DisplayClusterICVFXCameraComponentDetailsCustomization.cpp
@@ -26,7 +26,7 @@ namespace DisplayClusterICVFXCameraComponentDetailsCustomizationUtils
 		{
 			TEXT("Variable"),
 			TEXT("TransformCommon"),
-			DisplayClusterConfigurationStrings::categories::ICVFXCategory,
+			DisplayClusterConfigurationStrings::categories::InnerFrustumCategory,
 			DisplayClusterConfigurationStrings::categories::ICVFXCameraCategory,
 			DisplayClusterConfigurationStrings::categories::CameraColorGradingCategory,
 			DisplayClusterConfigurationStrings::categories::OCIOCategory,
diff --git a/Engine/Plugins/Runtime/nDisplay/Source/DisplayClusterConfigurator/Private/Views/Details/DisplayClusterRootActorDetailsCustomization.cpp b/Engine/Plugins/Runtime/nDisplay/Source/DisplayClusterConfigurator/Private/Views/Details/DisplayClusterRootActorDetailsCustomization.cpp
--- a/Engine/Plugins/Runtime/nDisplay/Source/DisplayClusterConfigurator/Private/Views/Details/DisplayClusterRootActorDetailsCustomization.cpp
+++ b/Engine/Plugins/Runtime/nDisplay/Source/DisplayClusterConfigurator/Private/Views/Details/DisplayClusterRootActorDetailsCustomization.cpp
@@ -32,7 +32,7 @@ namespace DisplayClusterRootActorDetailsCustomizationUtils
 		{
 			TEXT("TransformCommon"),
 			DisplayClusterConfigurationStrings::categories::ViewportsCategory,
-			DisplayClusterConfigurationStrings::categories::ICVFXCategory,
+			DisplayClusterConfigurationStrings::categories::InCameraVFXCategory,
 			DisplayClusterConfigurationStrings::categories::ColorGradingCategory,
 			DisplayClusterConfigurationStrings::categories::OCIOCategory,
 			DisplayClusterConfigurationStrings::categories::ChromaKeyCategory,
@@ -83,6 +83,18 @@ void FDisplayClusterRootActorDetailsCustomization::CustomizeDetails(IDetailLayou
 
 	ForcePropertyWindowRebuildHandle = FEditorSupportDelegates::ForcePropertyWindowRebuild.AddSP(this, &FDisplayClusterRootActorDetailsCustomization::OnForcePropertyWindowRebuild);
 
+	// Hide Categories:
+	InLayoutBuilder.HideCategory(DisplayClusterConfigurationStrings::categories::ViewPointStereoCategory);
+	InLayoutBuilder.HideCategory(DisplayClusterConfigurationStrings::categories::ViewPointCameraPostProcessCategory);
+	InLayoutBuilder.HideCategory(DisplayClusterConfigurationStrings::categories::ViewPointInFrustumProjectionCategory);
+
+	// For DCRA in a scene, hide more categories:
+	if (!IsRunningForBlueprintEditor())
+	{
+		InLayoutBuilder.HideCategory(DisplayClusterConfigurationStrings::categories::ConfigurationCategory);
+		InLayoutBuilder.HideCategory(DisplayClusterConfigurationStrings::categories::DefaultCategory);
+	}
+
 	InLayoutBuilder.SortCategories(DisplayClusterRootActorDetailsCustomizationUtils::SortCategories);
 	
 	/*
@@ -111,7 +123,7 @@ void FDisplayClusterRootActorDetailsCustomization::CustomizeDetails(IDetailLayou
 	}
 
 	// Manually label the ICVFX category to properly format it to have the dash in "In-Camera"
-	InLayoutBuilder.EditCategory(DisplayClusterConfigurationStrings::categories::ICVFXCategory, LOCTEXT("ICVFXCategoryLabel", "In-Camera VFX"));
+	InLayoutBuilder.EditCategory(DisplayClusterConfigurationStrings::categories::InCameraVFXCategory, LOCTEXT("InCameraVFXCategoryLabel", "In-Camera VFX"));
 
 	// Customize the PreviewNodeId property to be a dropdown filled with the nodes configured on the root actor.
 	{
diff --git a/Engine/Plugins/Runtime/nDisplay/Source/DisplayClusterLightCardEditor/Private/DisplayClusterLightCardEditorCommands.cpp b/Engine/Plugins/Runtime/nDisplay/Source/DisplayClusterLightCardEditor/Private/DisplayClusterLightCardEditorCommands.cpp
--- a/Engine/Plugins/Runtime/nDisplay/Source/DisplayClusterLightCardEditor/Private/DisplayClusterLightCardEditorCommands.cpp
+++ b/Engine/Plugins/Runtime/nDisplay/Source/DisplayClusterLightCardEditor/Private/DisplayClusterLightCardEditorCommands.cpp
@@ -10,9 +10,9 @@ void FDisplayClusterLightCardEditorCommands::RegisterCommands()
 	UI_COMMAND(ResetCamera, "Reset Camera", "Resets the camera to focus on the mesh", EUserInterfaceActionType::Button, FInputChord());
 	UI_COMMAND(FrameSelection, "Frame Selection", "Frames the selected items in the viewport", EUserInterfaceActionType::Button, FInputChord(EKeys::F));
 
-	UI_COMMAND(PerspectiveProjection, "Perspective", "A perspective projection from the stage's view origin", EUserInterfaceActionType::RadioButton, FInputChord());
-	UI_COMMAND(OrthographicProjection, "Orthographic", "An orthographic projection from the stage's view origin", EUserInterfaceActionType::RadioButton, FInputChord());
-	UI_COMMAND(AzimuthalProjection, "Dome", "A hemispherical projection from the stage's view origin", EUserInterfaceActionType::RadioButton, FInputChord());
+	UI_COMMAND(PerspectiveProjection, "Perspective", "A perspective projection from the stage's view point", EUserInterfaceActionType::RadioButton, FInputChord());
+	UI_COMMAND(OrthographicProjection, "Orthographic", "An orthographic projection from the stage's view point", EUserInterfaceActionType::RadioButton, FInputChord());
+	UI_COMMAND(AzimuthalProjection, "Dome", "A hemispherical projection from the stage's view point", EUserInterfaceActionType::RadioButton, FInputChord());
 	UI_COMMAND(UVProjection, "UV", "A UV projection that flattens the stage's meshes based on their UV coordinates", EUserInterfaceActionType::RadioButton, FInputChord());
 
 	UI_COMMAND(ViewOrientationTop, "Top", "Orient the view to look at the top of the stage", EUserInterfaceActionType::Button, FInputChord(EModifierKey::Alt, EKeys::J));
diff --git a/Engine/Plugins/Runtime/nDisplay/Source/DisplayClusterLightCardEditor/Private/Viewport/DisplayClusterLightCardEditorViewportClient.cpp b/Engine/Plugins/Runtime/nDisplay/Source/DisplayClusterLightCardEditor/Private/Viewport/DisplayClusterLightCardEditorViewportClient.cpp
--- a/Engine/Plugins/Runtime/nDisplay/Source/DisplayClusterLightCardEditor/Private/Viewport/DisplayClusterLightCardEditorViewportClient.cpp
+++ b/Engine/Plugins/Runtime/nDisplay/Source/DisplayClusterLightCardEditor/Private/Viewport/DisplayClusterLightCardEditorViewportClient.cpp
@@ -849,7 +849,7 @@ void FDisplayClusterLightCardEditorViewportClient::CreateDrawnLightCard(const TA
 	const FVector ViewOrigin = View->ViewLocation;
 	const FVector OriginOffset = ProjectionMode == EDisplayClusterMeshProjectionType::UV ? ViewOrigin : FVector::ZeroVector;
 
-	TArray<FVector> MouseWorldDirections; // directions from view origin
+	TArray<FVector> MouseWorldDirections; // directions from view point
 
 	FVector ScreenOrigin;
 	FVector ScreenDirection;
@@ -1305,7 +1305,7 @@ void FDisplayClusterLightCardEditorViewportClient::BeginCameraMovement(bool bHas
 	// A little bit of hackery here. FEditorViewportClient doesn't provide many ways to change how camera movement happens (e.g. what keys or mouse movements do what),
 	// but it does call BeginCameraMovement before it actually applies any of the movement stored in CameraUserImpulseData, so we can perform our own camera movement 
 	// impulses here. For most projection modes, we simply zero out any camera translation impulses since the user isn't generally allowed to move the camera away from
-	// the view origin. The UV projection mode is the exception, where the user is allowed to pan the camera around
+	// the view point. The UV projection mode is the exception, where the user is allowed to pan the camera around
 	if (ProjectionMode == EDisplayClusterMeshProjectionType::UV)
 	{
 		bool bIsLeftRightKeyPressed = false;
@@ -2757,26 +2757,26 @@ AActor* FDisplayClusterLightCardEditorViewportClient::TraceScreenForActor(const
 			{
 				if (UDisplayClusterConfigurationViewport* CfgViewport = FindViewportForPrimitiveComponent(ScreenHitResult.Component.Get()))
 				{
-					FString ViewOriginName = CfgViewport->Camera;
-					UDisplayClusterCameraComponent* ViewOrigin = nullptr;
+					const FString ViewPointName = CfgViewport->Camera;
+					UDisplayClusterCameraComponent* ViewPointComponent= nullptr;
 
-					// If the view origin name is empty, use the first found view origin in the root actor
-					if (ViewOriginName.IsEmpty())
+					// If the view point name is empty, use the first found view point in the root actor
+					if (ViewPointName.IsEmpty())
 					{
-						ViewOrigin = RootActorProxy->GetDefaultCamera();
+						ViewPointComponent = RootActorProxy->GetDefaultCamera();
 					}
 					else
 					{
-						ViewOrigin = RootActorProxy->GetComponentByName<UDisplayClusterCameraComponent>(ViewOriginName);
+						ViewPointComponent = RootActorProxy->GetComponentByName<UDisplayClusterCameraComponent>(ViewPointName);
 					}
 
-					if (ViewOrigin)
+					if (ViewPointComponent)
 					{
-						const FVector ViewOriginRayStart = ViewOrigin->GetComponentLocation();
-						const FVector ViewOriginRayEnd = ViewOriginRayStart + (ScreenHitResult.Location - ViewOriginRayStart) * HALF_WORLD_MAX;
+						const FVector ViewPointRayStart = ViewPointComponent->GetComponentLocation();
+						const FVector ViewPointRayEnd = ViewPointRayStart + (ScreenHitResult.Location - ViewPointRayStart) * HALF_WORLD_MAX;
 
 						TArray<FHitResult> HitResults;
-						if (PreviewWorld->LineTraceMultiByObjectType(HitResults, ViewOriginRayStart, ViewOriginRayEnd, FCollisionObjectQueryParams(FCollisionObjectQueryParams::InitType::AllObjects), TraceParams))
+						if (PreviewWorld->LineTraceMultiByObjectType(HitResults, ViewPointRayStart, ViewPointRayEnd, FCollisionObjectQueryParams(FCollisionObjectQueryParams::InitType::AllObjects), TraceParams))
 						{
 							for (FHitResult& HitResult : HitResults)
 							{
diff --git a/Engine/Plugins/Runtime/nDisplay/Source/DisplayClusterLightCardEditor/Private/Viewport/DisplayClusterLightCardEditorViewportClient.h b/Engine/Plugins/Runtime/nDisplay/Source/DisplayClusterLightCardEditor/Private/Viewport/DisplayClusterLightCardEditorViewportClient.h
--- a/Engine/Plugins/Runtime/nDisplay/Source/DisplayClusterLightCardEditor/Private/Viewport/DisplayClusterLightCardEditorViewportClient.h
+++ b/Engine/Plugins/Runtime/nDisplay/Source/DisplayClusterLightCardEditor/Private/Viewport/DisplayClusterLightCardEditorViewportClient.h
@@ -142,7 +142,7 @@ public:
 	/** Moves specified card to desired coordinates. Actual radius will be based on flush constraint and Actor's RadialOffset.
 	 *
 	 * @param Actor The actor that we are moving
-	 * @param SphericalCoords specifies desired location of actor in spherical coordinates with respect to view origin.
+	 * @param SphericalCoords specifies desired location of actor in spherical coordinates with respect to view point.
 	 * 
 	*/
 	void MoveActorTo(const FDisplayClusterWeakStageActorPtr& Actor, const FDisplayClusterLightCardEditorHelper::FSphericalCoordinates& SphericalCoords) const;
@@ -383,7 +383,7 @@ private:
 	/** The viewport type (perspective or orthogonal) to use when rendering the viewport. Separate from ViewportType since ViewportType also determines input functionality */
 	ELevelViewportType RenderViewportType = LVT_Perspective;
 
-	/** The component of the root actor that is acting as the projection origin. Can be either the root component (stage origin) or a view origin component */
+	/** The component of the root actor that is acting as the projection origin. Can be either the root component (stage origin) or a view point component */
 	TWeakObjectPtr<USceneComponent> ProjectionOriginComponent;
 
 	/** Stores each projecion mode's view configuration separately so they can be restored when the projection mode is selected  */
diff --git a/Engine/Plugins/Runtime/nDisplay/Source/DisplayClusterLightCardEditor/Private/Viewport/DisplayClusterLightCardEditorWidget.h b/Engine/Plugins/Runtime/nDisplay/Source/DisplayClusterLightCardEditor/Private/Viewport/DisplayClusterLightCardEditorWidget.h
--- a/Engine/Plugins/Runtime/nDisplay/Source/DisplayClusterLightCardEditor/Private/Viewport/DisplayClusterLightCardEditorWidget.h
+++ b/Engine/Plugins/Runtime/nDisplay/Source/DisplayClusterLightCardEditor/Private/Viewport/DisplayClusterLightCardEditorWidget.h
@@ -77,10 +77,10 @@ private:
 	/** Gets the color for the specified axis */
 	FLinearColor GetAxisColor(EAxisList::Type Axis) const;
 
-	/** Calculates a scalar to use to keep lengths and lines a fixed size on the screen regardless of DPI, viewport size, FOV, or distance from the view origin */
+	/** Calculates a scalar to use to keep lengths and lines a fixed size on the screen regardless of DPI, viewport size, FOV, or distance from the view point. */
 	float GetLengthScreenScalar(const FSceneView* View, const FDisplayClusterLightCardEditorViewportClient* ViewportClient, const FVector& Origin) const;
 
-	/** Calculates a scalar to use to keep sizes a fixed size on the screen regardless of DPI, viewport size, FOV, or distance from the view origin */
+	/** Calculates a scalar to use to keep sizes a fixed size on the screen regardless of DPI, viewport size, FOV, or distance from the view point. */
 	float GetSizeScreenScalar(const FSceneView* View, const FDisplayClusterLightCardEditorViewportClient* ViewportClient) const;
 
 private:
diff --git a/Engine/Plugins/Runtime/nDisplay/Source/DisplayClusterLightCardEditorShaders/Public/DisplayClusterMeshProjectionRenderer.h b/Engine/Plugins/Runtime/nDisplay/Source/DisplayClusterLightCardEditorShaders/Public/DisplayClusterMeshProjectionRenderer.h
--- a/Engine/Plugins/Runtime/nDisplay/Source/DisplayClusterLightCardEditorShaders/Public/DisplayClusterMeshProjectionRenderer.h
+++ b/Engine/Plugins/Runtime/nDisplay/Source/DisplayClusterLightCardEditorShaders/Public/DisplayClusterMeshProjectionRenderer.h
@@ -73,7 +73,7 @@ struct FDisplayClusterMeshProjectionTypeSettings
 	/** The size of the plane the UVs are projected to, in view space */
 	float UVProjectionPlaneSize = 100.0f;
 
-	/** The distance from the view origin of the plane the UVs are projected to, in view space */
+	/** The distance from the view point of the plane the UVs are projected to, in view space */
 	float UVProjectionPlaneDistance = 100.0f;
 
 	/** A translation offset for the UV projection plane */
diff --git a/Engine/Plugins/Runtime/nDisplay/Source/DisplayClusterScenePreview/Private/DisplayClusterLightCardEditorHelper.cpp b/Engine/Plugins/Runtime/nDisplay/Source/DisplayClusterScenePreview/Private/DisplayClusterLightCardEditorHelper.cpp
--- a/Engine/Plugins/Runtime/nDisplay/Source/DisplayClusterScenePreview/Private/DisplayClusterLightCardEditorHelper.cpp
+++ b/Engine/Plugins/Runtime/nDisplay/Source/DisplayClusterScenePreview/Private/DisplayClusterLightCardEditorHelper.cpp
@@ -430,7 +430,7 @@ void FDisplayClusterLightCardEditorHelper::DragUVActors(
 
 void FDisplayClusterLightCardEditorHelper::VerifyAndFixActorOrigin(const FDisplayClusterWeakStageActorPtr& Actor)
 {
-	// Center actor on the current view origin, let it keep its current world placement 
+	// Center actor on the current view point, let it keep its current world placement 
 	// (but not its spin/yaw/pitch since that will be happen later using the cache)
 
 	const ADisplayClusterRootActor* RootActor = UpdateRootActor();
@@ -481,7 +481,7 @@ void FDisplayClusterLightCardEditorHelper::VerifyAndFixActorOrigin(const FDispla
 }
 
 bool FDisplayClusterLightCardEditorHelper::CalculateNormalAndPositionInDirection(
-	const FVector& InViewOrigin,
+	const FVector& InOrigin,
 	const FVector& InDirection,
 	FVector& OutWorldPosition,
 	FVector& OutRelativeNormal,
@@ -499,8 +499,8 @@ bool FDisplayClusterLightCardEditorHelper::CalculateNormalAndPositionInDirection
 
 		const float UVProjectionPlaneDistance = ADisplayClusterLightCardActor::UVPlaneDefaultDistance;
 
-		const FPlane UVProjectionPlane(InViewOrigin + FVector::ForwardVector * UVProjectionPlaneDistance, -FVector::ForwardVector);
-		const FVector PlaneIntersection = FMath::RayPlaneIntersection(InViewOrigin, InDirection, UVProjectionPlane);
+		const FPlane UVProjectionPlane(InOrigin + FVector::ForwardVector * UVProjectionPlaneDistance, -FVector::ForwardVector);
+		const FVector PlaneIntersection = FMath::RayPlaneIntersection(InOrigin, InDirection, UVProjectionPlane);
 
 		OutRelativeNormal = UVProjectionPlane.GetNormal();
 		OutWorldPosition = PlaneIntersection;
@@ -513,7 +513,7 @@ bool FDisplayClusterLightCardEditorHelper::CalculateNormalAndPositionInDirection
 		Distance = CalculateFinalLightCardDistance(Distance, InDesiredDistanceFromFlush);
 
 		// Calculate world position
-		OutWorldPosition = InViewOrigin + Distance * InDirection;
+		OutWorldPosition = InOrigin + Distance * InDirection;
 	}
 
 	return true;
@@ -531,13 +531,13 @@ bool FDisplayClusterLightCardEditorHelper::CalculateOriginAndDirectionFromPixelP
 			return false;
 		}
 
-		// For orthogonal projections, PixelToWorld does not return the view origin or a direction from the view origin. Use TraceScreenRay
-		// to find a useful direction away from the view origin to use
+		// For orthogonal projections, PixelToWorld does not return the view point or a direction from the view point. Use TraceScreenRay
+		// to find a useful direction away from the origin to use
 
-		const FVector ViewOrigin = SceneView.ViewLocation;
+		const FVector NewOrigin = SceneView.ViewLocation;
 
-		OutDirection = TraceScreenRay(OutOrigin + OriginOffset, OutDirection, ViewOrigin);
-		OutOrigin = ViewOrigin;
+		OutDirection = TraceScreenRay(OutOrigin + OriginOffset, OutDirection, NewOrigin);
+		OutOrigin = NewOrigin;
 	}
 
 	return true;
@@ -1393,7 +1393,7 @@ FVector FDisplayClusterLightCardEditorHelper::TraceScreenRay(const FVector& Orth
 	}
 	else
 	{
-		// First, trace against the stage actor to see if the screen ray hits it; if so, simply return the direction from the view origin to this hit point
+		// First, trace against the stage actor to see if the screen ray hits it; if so, simply return the direction from the view point to this hit point
 		FVector HitLocation = FVector::ZeroVector;
 		if (TraceStage(RayStart, RayEnd, HitLocation))
 		{
@@ -1413,7 +1413,7 @@ FVector FDisplayClusterLightCardEditorHelper::TraceScreenRay(const FVector& Orth
 			}
 			else
 			{
-				// If the screen ray does not hit the stage or the normal map mesh, then simply use the closest point on the ray to the view origin
+				// If the screen ray does not hit the stage or the normal map mesh, then simply use the closest point on the ray to the view point
 				const FVector ClosestPoint = OrthogonalOrigin + ((ViewOrigin - OrthogonalOrigin) | OrthogonalDirection) * OrthogonalDirection;
 
 				Direction = (ClosestPoint - ViewOrigin).GetSafeNormal();
diff --git a/Engine/Plugins/Runtime/nDisplay/Source/DisplayClusterScenePreview/Public/DisplayClusterLightCardEditorHelper.h b/Engine/Plugins/Runtime/nDisplay/Source/DisplayClusterScenePreview/Public/DisplayClusterLightCardEditorHelper.h
--- a/Engine/Plugins/Runtime/nDisplay/Source/DisplayClusterScenePreview/Public/DisplayClusterLightCardEditorHelper.h
+++ b/Engine/Plugins/Runtime/nDisplay/Source/DisplayClusterScenePreview/Public/DisplayClusterLightCardEditorHelper.h
@@ -129,7 +129,7 @@ public:
 	 * Moves specified actors to desired coordinates. Actual radius will be based on flush constraint and actor's RadialOffset.
 	 *
 	 * @param Actors The actors that we are moving
-	 * @param SphericalCoords The desired location of the actors in spherical coordinates with respect to view origin.
+	 * @param SphericalCoords The desired location of the actors in spherical coordinates with respect to origin.
 	*/
 	DISPLAYCLUSTERSCENEPREVIEW_API void MoveActorsTo(const TArray<FDisplayClusterWeakStageActorPtr>& Actors, const FSphericalCoordinates& SphericalCoords);
 	
@@ -160,24 +160,24 @@ public:
 	DISPLAYCLUSTERSCENEPREVIEW_API void DragUVActors(const TArray<FDisplayClusterWeakStageActorPtr>& Actors, const FIntPoint& PixelPos, const FSceneView& SceneView,
 		const FVector& DragWidgetOffset, EAxisList::Type DragAxis, FDisplayClusterWeakStageActorPtr PrimaryActor = nullptr);
 
-	/** Ensures that the actor root component is at the same location as the projection/view origin */
+	/** Ensures that the actor root component is at the same location as the projection/origin */
 	DISPLAYCLUSTERSCENEPREVIEW_API void VerifyAndFixActorOrigin(const FDisplayClusterWeakStageActorPtr& Actor);
 
 	/**
-	 * Calculates the relative normal vector and world position in the specified direction from the given view origin.
+	 * Calculates the relative normal vector and world position in the specified direction from the given origin.
 	 * 
-	 * @param InViewOrigin The origin point of the view.
+	 * @param InOrigin The origin point of the view.
 	 * @param InDirection The direction in which check the normal.
 	 * @param OutWorldPosition The world-space coordinates of the calculated position.
 	 * @param OutRelativeNormal The normal vector at the calculated position.
 	 * @param InDesiredDistanceFromFlush The desired flush distance from the initial calculated position.
 	 * @return true if the position was found, or false if the normal maps need to be updated first.
 	 */
-	DISPLAYCLUSTERSCENEPREVIEW_API bool CalculateNormalAndPositionInDirection(const FVector& InViewOrigin, const FVector& InDirection, FVector& OutWorldPosition,
+	DISPLAYCLUSTERSCENEPREVIEW_API bool CalculateNormalAndPositionInDirection(const FVector& InOrigin, const FVector& InDirection, FVector& OutWorldPosition,
 		FVector& OutRelativeNormal, double InDesiredDistanceFromFlush = 0.);
 
 	/**
-	 * Calculates the desired direction from the view origin given a pixel position within the view.
+	 * Calculates the desired direction from the origin given a pixel position within the view.
 	 * 
 	 * @param InPixelPos The desired position within the view in pixels.
 	 * @param InSceneView The view that was displayed when the pixel position was chosen.
@@ -335,7 +335,7 @@ private:
 	/** Performs a ray trace against the stage's geometry, and returns the hit point. */
 	bool TraceStage(const FVector& RayStart, const FVector& RayEnd, FVector& OutHitLocation);
 
-	/** Traces the world geometry to find the best direction vector from the view origin to a valid point in space using a screen ray. Requires valid normal maps. */
+	/** Traces the world geometry to find the best direction vector from the origin to a valid point in space using a screen ray. Requires valid normal maps. */
 	FVector TraceScreenRay(const FVector& RayOrigin, const FVector& RayDirection, const FVector& ViewOrigin);
 
 	/** Calculates the final distance from the origin of a light card, given its flush distance and a desired offset */
@@ -384,7 +384,7 @@ private:
 	/** The level instance root actor of the DisplayCluster being controlled, if RootActor is a proxy. */
 	TWeakObjectPtr<ADisplayClusterRootActor> LevelInstanceRootActor;
 
-	/** The component of the root actor that is acting as the projection origin. Can be either the root component (stage origin) or a view origin component */
+	/** The component of the root actor that is acting as the projection origin. Can be either the root component (stage origin) or a view point component */
 	TWeakObjectPtr<USceneComponent> ProjectionOriginComponent;
 
 	/** The projection mode of the view this is helping. */
diff --git a/Engine/Plugins/Runtime/nDisplay/Source/DisplayClusterWarp/Private/Components/DisplayClusterInFrustumFitCameraComponent.cpp b/Engine/Plugins/Runtime/nDisplay/Source/DisplayClusterWarp/Private/Components/DisplayClusterInFrustumFitCameraComponent.cpp
--- a/Engine/Plugins/Runtime/nDisplay/Source/DisplayClusterWarp/Private/Components/DisplayClusterInFrustumFitCameraComponent.cpp
+++ b/Engine/Plugins/Runtime/nDisplay/Source/DisplayClusterWarp/Private/Components/DisplayClusterInFrustumFitCameraComponent.cpp
@@ -126,14 +126,15 @@ bool UDisplayClusterInFrustumFitCameraComponent::IsEnabled() const
 	return bEnableCameraProjection;
 }
 
-UCameraComponent* UDisplayClusterInFrustumFitCameraComponent::GetExternalCameraComponent() const
+bool UDisplayClusterInFrustumFitCameraComponent::IsICVFXCameraBeingUsed() const
 {
-	if (ACineCameraActor* CineCamera = ExternalCameraActor.Get())
+	// When using InFrustum projection, ignore the camera component from the parent class
+	if (IsEnabled())
 	{
-		return CineCamera->GetCameraComponent();
+		return false;
 	}
 
-	return nullptr;
+	return Super::IsICVFXCameraBeingUsed();
 }
 
 void UDisplayClusterInFrustumFitCameraComponent::TickComponent(float DeltaTime, ELevelTick TickType, FActorComponentTickFunction* ThisTickFunction)
@@ -151,30 +152,6 @@ void UDisplayClusterInFrustumFitCameraComponent::TickComponent(float DeltaTime,
 	}
 }
 
-void UDisplayClusterInFrustumFitCameraComponent::GetDesiredView(IDisplayClusterViewportConfiguration& InViewportConfiguration, FMinimalViewInfo& InOutViewInfo, float* OutCustomNearClippingPlane)
-{
-	const UDisplayClusterInFrustumFitCameraComponent& ConfigurationCameraComponent = GetConfigurationInFrustumFitCameraComponent(InViewportConfiguration);
-	if (ConfigurationCameraComponent.IsEnabled())
-	{
-		UCameraComponent* CameraComponent = ConfigurationCameraComponent.GetExternalCameraComponent();
-		const float DeltaSeconds = InViewportConfiguration.GetRootActorWorldDeltaSeconds(EDisplayClusterRootActorType::Scene);
-		if (IDisplayClusterViewport::GetCameraComponentView(CameraComponent, DeltaSeconds, ConfigurationCameraComponent.bUseCameraPostprocess, InOutViewInfo, OutCustomNearClippingPlane))
-		{
-			// 1. Use external camera for rendering
-			return;
-		}
-
-		if (IDisplayClusterViewport::GetPlayerCameraView(InViewportConfiguration.GetCurrentWorld(), ConfigurationCameraComponent.bUseCameraPostprocess, InOutViewInfo))
-		{
-			// 2. Use active game camera
-			return;
-		}
-	}
-
-	// use default logic
-	return UDisplayClusterCameraComponent::GetDesiredView(InViewportConfiguration, InOutViewInfo, OutCustomNearClippingPlane);
-}
-
 bool UDisplayClusterInFrustumFitCameraComponent::ShouldUseEntireClusterViewports(IDisplayClusterViewportManager* InViewportManager) const
 {
 	 const UDisplayClusterInFrustumFitCameraComponent& ConfigurationCameraComponent = InViewportManager ? GetConfigurationInFrustumFitCameraComponent(InViewportManager->GetConfiguration()) : *this;
@@ -215,33 +192,47 @@ void UDisplayClusterInFrustumFitCameraComponent::OnRegister()
 		if (SpriteComponent)
 		{
 			SpriteComponent->SpriteInfo.Category = TEXT("NDisplayCameraViewOrigin");
-			SpriteComponent->SpriteInfo.DisplayName = NSLOCTEXT("DisplayClusterInFrustumFitCameraComponent", "DisplayClusterInFrustumFitCameraComponentSpriteInfo", "nDisplay InFrustumFit View Origin");
+			SpriteComponent->SpriteInfo.DisplayName = NSLOCTEXT("DisplayClusterInFrustumFitCameraComponent", "DisplayClusterInFrustumFitCameraComponentSpriteInfo", "nDisplay InFrustumFit View Point");
 		}
 	}
 
 	RefreshVisualRepresentation();
 #endif
-
 }
 
 #if WITH_EDITOR
-bool UDisplayClusterInFrustumFitCameraComponent::GetEditorPreviewInfo(float DeltaTime, FMinimalViewInfo& ViewOut)
+bool UDisplayClusterInFrustumFitCameraComponent::CanEditChange(const FProperty* InProperty) const
 {
-	if (UCameraComponent* CameraComponent = GetExternalCameraComponent())
+	// If other logic prevents editing, we want to respect that
+	bool bIsEditable = Super::CanEditChange(InProperty);
+	if (bIsEditable && InProperty)
 	{
-		return CameraComponent->GetEditorPreviewInfo(DeltaTime, ViewOut);
+		// InFrustum projection requires external cinecamera
+		const FName PropertyName = InProperty->GetFName();
+		if (PropertyName == GET_MEMBER_NAME_CHECKED(UDisplayClusterCameraComponent, ICVFXCameraComponentName))
+		{
+			return !IsEnabled();
+		}
 	}
 
-	return false;
+	return bIsEditable;
 }
 
-TSharedPtr<SWidget> UDisplayClusterInFrustumFitCameraComponent::GetCustomEditorPreviewWidget()
+void UDisplayClusterInFrustumFitCameraComponent::PostEditChangeProperty(FPropertyChangedEvent& PropertyChangedEvent)
 {
-	if (UCameraComponent* CameraComponent = GetExternalCameraComponent())
+	const FName PropertyName = (PropertyChangedEvent.Property != nullptr) ? PropertyChangedEvent.Property->GetFName() : NAME_None;
+
+	if (PropertyName == GET_MEMBER_NAME_CHECKED(UDisplayClusterCameraComponent, TargetCameraType)
+	|| PropertyName == GET_MEMBER_NAME_CHECKED(UDisplayClusterInFrustumFitCameraComponent, bEnableCameraProjection))
 	{
-		return CameraComponent->GetCustomEditorPreviewWidget();
+		if (IsEnabled() && TargetCameraType == EDisplayClusterTargetCameraType::ICVFXCameraComponent)
+		{
+			// When using projection, the internal target camera type cannot be used.
+			// We have to switch to the external target camera type.
+			TargetCameraType = EDisplayClusterTargetCameraType::None;
+		}
 	}
 
-	return nullptr;
+	Super::PostEditChangeProperty(PropertyChangedEvent);
 }
 #endif
diff --git a/Engine/Plugins/Runtime/nDisplay/Source/DisplayClusterWarp/Private/Policy/InFrustumFit/DisplayClusterWarpInFrustumFitPolicy.cpp b/Engine/Plugins/Runtime/nDisplay/Source/DisplayClusterWarp/Private/Policy/InFrustumFit/DisplayClusterWarpInFrustumFitPolicy.cpp
--- a/Engine/Plugins/Runtime/nDisplay/Source/DisplayClusterWarp/Private/Policy/InFrustumFit/DisplayClusterWarpInFrustumFitPolicy.cpp
+++ b/Engine/Plugins/Runtime/nDisplay/Source/DisplayClusterWarp/Private/Policy/InFrustumFit/DisplayClusterWarpInFrustumFitPolicy.cpp
@@ -271,7 +271,7 @@ void FDisplayClusterWarpInFrustumFitPolicy::OnUpdateDisplayDeviceMeshAndMaterial
 
 		const FMatrix CameraBasis = FRotationMatrix::Make(WarpData.WarpProjection.CameraRotation).Inverse();
 
-		// Compute the relative transform from the view origin to the geometry
+		// Compute the relative transform from the origin to the geometry
 		FTransform RelativeTransform = FTransform(WarpData.WarpContext.MeshToStageMatrix * WarpData.Local2World.Inverse());
 		RelativeTransform.ScaleTranslation(Scale);
 
diff --git a/Engine/Plugins/Runtime/nDisplay/Source/DisplayClusterWarp/Private/WarpBlend/Math/DisplayClusterWarpBlendMath_FrustumProjection.cpp b/Engine/Plugins/Runtime/nDisplay/Source/DisplayClusterWarp/Private/WarpBlend/Math/DisplayClusterWarpBlendMath_FrustumProjection.cpp
--- a/Engine/Plugins/Runtime/nDisplay/Source/DisplayClusterWarp/Private/WarpBlend/Math/DisplayClusterWarpBlendMath_FrustumProjection.cpp
+++ b/Engine/Plugins/Runtime/nDisplay/Source/DisplayClusterWarp/Private/WarpBlend/Math/DisplayClusterWarpBlendMath_FrustumProjection.cpp
@@ -130,7 +130,8 @@ void FDisplayClusterWarpBlendMath_Frustum::ImplCalcViewProjectionMatrices()
 		WarpData.Local2World = FRotationMatrix::MakeFromXY(NewX, FVector(0.f, 1.f, 0.f));
 	}
 
-	WarpData.Local2World.SetOrigin(WarpData.GeometryWarpProjection.EyeLocation); // Finally set view origin to eye location
+	// Finally, set the eye location for the view matrix.
+	WarpData.Local2World.SetOrigin(WarpData.GeometryWarpProjection.EyeLocation);
 }
 
 void FDisplayClusterWarpBlendMath_Frustum::InitializeWarpProjectionData()
diff --git a/Engine/Plugins/Runtime/nDisplay/Source/DisplayClusterWarp/Public/Blueprints/DisplayClusterWarpBlueprint_Enums.h b/Engine/Plugins/Runtime/nDisplay/Source/DisplayClusterWarp/Public/Blueprints/DisplayClusterWarpBlueprint_Enums.h
--- a/Engine/Plugins/Runtime/nDisplay/Source/DisplayClusterWarp/Public/Blueprints/DisplayClusterWarpBlueprint_Enums.h
+++ b/Engine/Plugins/Runtime/nDisplay/Source/DisplayClusterWarp/Public/Blueprints/DisplayClusterWarpBlueprint_Enums.h
@@ -29,6 +29,6 @@ enum class EDisplayClusterWarpCameraViewTarget
 	/** The camera will point in the direction of the geometric center of the stage's geometry */
 	GeometricCenter UMETA(DisplayName = "Geometric Center"),
 
-	/** The camera will point in the same direction as the frustum fit view origin */
-	MatchViewOrigin UMETA(DisplayName = "Match View Origin")
+	/** The camera will point in the same direction as the frustum fit view point. */
+	MatchViewOrigin UMETA(DisplayName = "Match View Point")
 };
diff --git a/Engine/Plugins/Runtime/nDisplay/Source/DisplayClusterWarp/Public/Components/DisplayClusterInFrustumFitCameraComponent.h b/Engine/Plugins/Runtime/nDisplay/Source/DisplayClusterWarp/Public/Components/DisplayClusterInFrustumFitCameraComponent.h
--- a/Engine/Plugins/Runtime/nDisplay/Source/DisplayClusterWarp/Public/Components/DisplayClusterInFrustumFitCameraComponent.h
+++ b/Engine/Plugins/Runtime/nDisplay/Source/DisplayClusterWarp/Public/Components/DisplayClusterInFrustumFitCameraComponent.h
@@ -16,7 +16,7 @@ class IDisplayClusterWarpPolicy;
  * 3D point in space used to project the camera view onto a group of nDisplay viewports.
  * Support projection policies: mpcdi/pfm 2d/a3d, mesh.
  */
-UCLASS(ClassGroup = (DisplayCluster), meta = (BlueprintSpawnableComponent, DisplayName = "NDisplay In Frustum Fit View Origin"))
+UCLASS(ClassGroup = (DisplayCluster), meta = (BlueprintSpawnableComponent, DisplayName = "NDisplay Frustum Fit View Point"))
 class DISPLAYCLUSTERWARP_API UDisplayClusterInFrustumFitCameraComponent : public UDisplayClusterCameraComponent
 {
 	GENERATED_BODY()
@@ -29,55 +29,58 @@ public:
 	virtual void TickComponent(float DeltaTime, ELevelTick TickType, FActorComponentTickFunction* ThisTickFunction) override;
 	//~~~End UActorComponent
 
+	//~Begin UObject
+#if WITH_EDITOR
+	virtual bool CanEditChange(const FProperty* InProperty) const override;
+	virtual void PostEditChangeProperty(FPropertyChangedEvent& PropertyChangedEvent) override;
+#endif
+	//~~~End UObject
+
 	//~Begin UDisplayClusterCameraComponent
-	virtual void GetDesiredView(IDisplayClusterViewportConfiguration& InViewportConfiguration, FMinimalViewInfo& InOutViewInfo, float* OutCustomNearClippingPlane = nullptr) override;
+	//virtual void GetDesiredView(IDisplayClusterViewportConfiguration& InViewportConfiguration, FMinimalViewInfo& InOutViewInfo, float* OutCustomNearClippingPlane = nullptr) override;
 	virtual bool ShouldUseEntireClusterViewports(IDisplayClusterViewportManager* InViewportManager) const override;
 	virtual IDisplayClusterWarpPolicy* GetWarpPolicy(IDisplayClusterViewportManager* InViewportManager) override;
 
 	virtual TObjectPtr<UMaterial> GetDisplayDeviceMaterial(const EDisplayClusterDisplayDeviceMeshType InMeshType, const EDisplayClusterDisplayDeviceMaterialType InMaterialType) const override;
 	virtual void OnUpdateDisplayDeviceMeshAndMaterialInstance(IDisplayClusterViewportPreview& InViewportPreview, const EDisplayClusterDisplayDeviceMeshType InMeshType, const EDisplayClusterDisplayDeviceMaterialType InMaterialType, UMeshComponent* InMeshComponent, UMaterialInstanceDynamic* InMeshMaterialInstance) const override;
-	//~~End UDisplayClusterCameraComponent
 
-#if WITH_EDITOR
-	virtual bool GetEditorPreviewInfo(float DeltaTime, FMinimalViewInfo& ViewOut) override;
-	virtual TSharedPtr<SWidget> GetCustomEditorPreviewWidget() override;
-#endif
+protected:
+	virtual bool IsICVFXCameraBeingUsed() const override;
+	//~~End UDisplayClusterCameraComponent
 
+public:
 	/** Return component that used for configuration. */
 	const UDisplayClusterInFrustumFitCameraComponent& GetConfigurationInFrustumFitCameraComponent(IDisplayClusterViewportConfiguration& InViewportConfiguration) const;
 
 private:
-	/** Return external camera component. */
-	UCameraComponent* GetExternalCameraComponent() const;
-
 	/** true, if camera projection is used. */
 	bool IsEnabled() const;
 
 public:
 	/** Camera projection mode is used. */
-	UPROPERTY(EditAnywhere, Category = "Projection")
+	UPROPERTY(EditAnywhere, Category = "Frustum Fit", meta = (DisplayName = "Enable Frustum Fit"))
 	bool bEnableCameraProjection = true;
 
 	/** Enable special rendering mode for all viewports using this viewpoint. */
-	UPROPERTY(EditAnywhere, Category = "Projection", meta = (EditCondition = "bEnableCameraProjection"))
+	UPROPERTY(EditAnywhere, Category = "Frustum Fit", meta = (DisplayName = "Frustum Fit Mode", EditCondition = "bEnableCameraProjection", EditConditionHides))
 	EDisplayClusterWarpCameraProjectionMode CameraProjectionMode = EDisplayClusterWarpCameraProjectionMode::Fit;
 
 	/** Indicates which camera facing mode is used when frustum fitting the stage geometry */
-	UPROPERTY(EditAnywhere, Category = "Projection", meta = (EditCondition = "bEnableCameraProjection"))
+	UPROPERTY(EditAnywhere, Category = "Frustum Fit", meta = (DisplayName = "Frustum Fit Target", EditCondition = "bEnableCameraProjection", EditConditionHides))
 	EDisplayClusterWarpCameraViewTarget CameraViewTarget = EDisplayClusterWarpCameraViewTarget::GeometricCenter;
 
-	/** Use a specific actor camera instead of a game camera. */
-	UPROPERTY(EditAnywhere, Category = "Projection", meta = (EditCondition = "bEnableCameraProjection"))
+	/** Show additional warped preview meshes before the camera. */
+	UPROPERTY(EditAnywhere, Category = "Frustum Fit", meta = (DisplayName = "Show Frustum Fit Preview", EditCondition = "bEnableCameraProjection", EditConditionHides))
+	bool bShowPreviewFrustumFit = false;
+
+	UE_DEPRECATED(5.5, "Use the camera settings from the UDisplayClusterCameraComponent instead")
+	UPROPERTY(meta = (DeprecatedProperty, DeprecationMessage = "Use the Camera from the UDisplayClusterCameraComponent instead"))
 	TSoftObjectPtr<ACineCameraActor> ExternalCameraActor;
 
-	/** Allows you to override the PostProcess of the viewport from the camera. */
-	UPROPERTY(EditAnywhere, Category = "Projection", meta = (EditCondition = "bEnableCameraProjection"))
+	UE_DEPRECATED(5.5, "Use the bEnablePostProcess from the UDisplayClusterCameraComponent instead")
+	UPROPERTY(meta = (DeprecatedProperty, DeprecationMessage = "Use the bEnablePostProcess from the UDisplayClusterCameraComponent instead"))
 	bool bUseCameraPostprocess = false;
 
-	/** Show additional warped preview meshes before the camera. */
-	UPROPERTY(EditAnywhere, Category = "Projection", meta = (EditCondition = "bEnableCameraProjection"))
-	bool bShowPreviewFrustumFit = false;
-
 private:
 	// a unique type of warp policy for this component
 	// this policy class knows the properties of the component and implements the corresponding logic
