nDisplay: fixed crash in UDisplayClusterICVFXCameraComponent
- added nullptr check for UObject references to ExternalCameraActor.

diff --git a/Engine/Plugins/Runtime/nDisplay/Source/DisplayCluster/Private/Game/EngineClasses/Scene/DisplayClusterICVFXCameraComponent.cpp b/Engine/Plugins/Runtime/nDisplay/Source/DisplayCluster/Private/Game/EngineClasses/Scene/DisplayClusterICVFXCameraComponent.cpp
--- a/Engine/Plugins/Runtime/nDisplay/Source/DisplayCluster/Private/Game/EngineClasses/Scene/DisplayClusterICVFXCameraComponent.cpp
+++ b/Engine/Plugins/Runtime/nDisplay/Source/DisplayCluster/Private/Game/EngineClasses/Scene/DisplayClusterICVFXCameraComponent.cpp
@@ -87,10 +87,10 @@ void UDisplayClusterICVFXCameraComponent::GetCameraView(float DeltaTime, FMinima
 		return;
 	}
 
-	if (CameraSettings.ExternalCameraActor.IsValid())
+	if (UCineCameraComponent* ExternalCineCameraComponent = CameraSettings.GetExternalCineCameraComponent())
 	{
 		// Get ViewInfo from external CineCamera
-		CameraSettings.ExternalCameraActor->GetCineCameraComponent()->GetCameraView(DeltaTime, InOutViewInfo);
+		ExternalCineCameraComponent->GetCameraView(DeltaTime, InOutViewInfo);
 	}
 	else
 	{
@@ -103,12 +103,9 @@ void UDisplayClusterICVFXCameraComponent::GetCameraView(float DeltaTime, FMinima
 
 UCineCameraComponent* UDisplayClusterICVFXCameraComponent::GetActualCineCameraComponent()
 {
-	if (UCineCameraComponent* ExternalCineCameraComponent = CameraSettings.ExternalCameraActor.IsValid() ? CameraSettings.ExternalCameraActor->GetCineCameraComponent() : nullptr)
-	{
-		return ExternalCineCameraComponent;
-	}
+	UCineCameraComponent* ExternalCineCameraComponent = CameraSettings.GetExternalCineCameraComponent();
 
-	return this;
+	return ExternalCineCameraComponent ? ExternalCineCameraComponent : this;
 }
 
 FString UDisplayClusterICVFXCameraComponent::GetCameraUniqueId() const
@@ -119,15 +116,17 @@ FString UDisplayClusterICVFXCameraComponent::GetCameraUniqueId() const
 #if WITH_EDITOR
 bool UDisplayClusterICVFXCameraComponent::GetEditorPreviewInfo(float DeltaTime, FMinimalViewInfo& ViewOut)
 {
-	return CameraSettings.ExternalCameraActor.IsValid() ?
-		CameraSettings.ExternalCameraActor->GetCineCameraComponent()->GetEditorPreviewInfo(DeltaTime, ViewOut) :
+	UCineCameraComponent* ExternalCineCameraComponent = CameraSettings.GetExternalCineCameraComponent();
+	return ExternalCineCameraComponent ?
+		ExternalCineCameraComponent->GetEditorPreviewInfo(DeltaTime, ViewOut) :
 		UCameraComponent::GetEditorPreviewInfo(DeltaTime, ViewOut);
 }
 
 TSharedPtr<SWidget> UDisplayClusterICVFXCameraComponent::GetCustomEditorPreviewWidget()
 {
-	return CameraSettings.ExternalCameraActor.IsValid() ?
-		CameraSettings.ExternalCameraActor->GetCineCameraComponent()->GetCustomEditorPreviewWidget() :
+	UCineCameraComponent* ExternalCineCameraComponent = CameraSettings.GetExternalCineCameraComponent();
+	return ExternalCineCameraComponent ?
+		ExternalCineCameraComponent->GetCustomEditorPreviewWidget() :
 		UCameraComponent::GetCustomEditorPreviewWidget();
 }
 #endif
@@ -146,10 +145,10 @@ void UDisplayClusterICVFXCameraComponent::TickComponent(float DeltaTime, ELevelT
 		{
 			FVector CameraLocation = FVector::ZeroVector;
 			FVector CameraDirection = FVector::XAxisVector;
-			if (CameraSettings.ExternalCameraActor.IsValid())
+			if (ACineCameraActor* ExternalCineCameraActor = CameraSettings.GetExternalCineCameraActor())
 			{
-				CameraLocation = CameraSettings.ExternalCameraActor->GetActorLocation();
-				CameraDirection = CameraSettings.ExternalCameraActor->GetActorRotation().RotateVector(FVector::XAxisVector);
+				CameraLocation = ExternalCineCameraActor->GetActorLocation();
+				CameraDirection = ExternalCineCameraActor->GetActorRotation().RotateVector(FVector::XAxisVector);
 			}
 			else
 			{
@@ -333,35 +332,42 @@ void UDisplayClusterICVFXCameraComponent::PostEditChangeProperty(FPropertyChange
 void UDisplayClusterICVFXCameraComponent::UpdateICVFXPreviewState()
 {
 	// handle frustum visibility
-	if (CameraSettings.ExternalCameraActor.IsValid())
+	if (ACineCameraActor* ExternalCineCameraActor = CameraSettings.GetExternalCineCameraActor())
 	{
-		ACineCameraActor* CineCamera = CameraSettings.ExternalCameraActor.Get();
-		CineCamera->GetCineCameraComponent()->bDrawFrustumAllowed = false;
+		UCineCameraComponent* ExternalCineCameraComponent = ExternalCineCameraActor->GetCineCameraComponent();
+		if (IsValid(ExternalCineCameraComponent))
+		{
+			ExternalCineCameraComponent->bDrawFrustumAllowed = false;
+		}
 
-		UDrawFrustumComponent* DrawFustumComponent = Cast<UDrawFrustumComponent>(CineCamera->GetComponentByClass(UDrawFrustumComponent::StaticClass()));
-		if (DrawFustumComponent != nullptr)
+		UDrawFrustumComponent* DrawFustumComponent = Cast<UDrawFrustumComponent>(ExternalCineCameraActor->GetComponentByClass(UDrawFrustumComponent::StaticClass()));
+		if (IsValid(DrawFustumComponent))
 		{
 			DrawFustumComponent->bFrustumEnabled = false;
 			DrawFustumComponent->MarkRenderStateDirty();
 		}
 
-		if (ProxyMeshComponent)
+		if (IsValid(ProxyMeshComponent))
 		{
 			ProxyMeshComponent->DestroyComponent();
 			ProxyMeshComponent = nullptr;
 		}
 	}
 
-
 	// restore frustum visibility if reference was changed
-	if (ExternalCameraCachedValue.IsValid())
+	if (ACineCameraActor* ExternalCineCameraCachedActor = ExternalCameraCachedValue.Get())
 	{
-		ACineCameraActor* CineCamera = ExternalCameraCachedValue.Get();
-		UDrawFrustumComponent* DrawFustumComponent = Cast<UDrawFrustumComponent>(CineCamera->GetComponentByClass(UDrawFrustumComponent::StaticClass()));
-		DrawFustumComponent->bFrustumEnabled = true;
-		DrawFustumComponent->MarkRenderStateDirty();
+		if (IsValid(ExternalCineCameraCachedActor))
+		{
+			UDrawFrustumComponent* DrawFustumComponent = Cast<UDrawFrustumComponent>(ExternalCineCameraCachedActor->GetComponentByClass(UDrawFrustumComponent::StaticClass()));
+			if (IsValid(DrawFustumComponent))
+			{
+				DrawFustumComponent->bFrustumEnabled = true;
+				DrawFustumComponent->MarkRenderStateDirty();
+			}
+		}
 
 		ExternalCameraCachedValue.Reset();
 	}
 }
-#endif
\ No newline at end of file
+#endif
diff --git a/Engine/Plugins/Runtime/nDisplay/Source/DisplayClusterConfiguration/Private/DisplayClusterConfigurationTypes_ICVFX.cpp b/Engine/Plugins/Runtime/nDisplay/Source/DisplayClusterConfiguration/Private/DisplayClusterConfigurationTypes_ICVFX.cpp
--- a/Engine/Plugins/Runtime/nDisplay/Source/DisplayClusterConfiguration/Private/DisplayClusterConfigurationTypes_ICVFX.cpp
+++ b/Engine/Plugins/Runtime/nDisplay/Source/DisplayClusterConfiguration/Private/DisplayClusterConfigurationTypes_ICVFX.cpp
@@ -6,6 +6,7 @@
 #include "IDisplayCluster.h"
 #include "Camera/CameraTypes.h"
 #include "CineCameraComponent.h"
+#include "CineCameraActor.h"
 
 namespace UE::DisplayClusterConfiguration::ICVFX
 {
@@ -196,6 +197,31 @@ FDisplayClusterConfigurationICVFX_CameraSettings::FDisplayClusterConfigurationIC
 	AllNodesColorGrading.bEnableEntireClusterColorGrading = true;
 }
 
+ACineCameraActor* FDisplayClusterConfigurationICVFX_CameraSettings::GetExternalCineCameraActor() const
+{
+	ACineCameraActor* ExternalCineCameraActor = ExternalCameraActor.Get();
+	if (IsValid(ExternalCineCameraActor))
+	{
+		return ExternalCineCameraActor;
+	}
+
+	return nullptr;
+}
+
+UCineCameraComponent* FDisplayClusterConfigurationICVFX_CameraSettings::GetExternalCineCameraComponent() const
+{
+	if (ACineCameraActor* ExternalCineCameraActor = GetExternalCineCameraActor())
+	{
+		UCineCameraComponent* ExternalCineCameraComponent = ExternalCineCameraActor->GetCineCameraComponent();
+		if (IsValid(ExternalCineCameraComponent))
+		{
+			return ExternalCineCameraComponent;
+		}
+	}
+
+	return nullptr;
+}
+
 bool FDisplayClusterConfigurationICVFX_CameraSettings::IsICVFXEnabled(const UDisplayClusterConfigurationData& InConfigurationData, const FString& InClusterNodeId) const
 {
 	// When rendering offscreen, we have an extended logic for camera rendering activation
diff --git a/Engine/Plugins/Runtime/nDisplay/Source/DisplayClusterConfiguration/Public/DisplayClusterConfigurationTypes_ICVFX.h b/Engine/Plugins/Runtime/nDisplay/Source/DisplayClusterConfiguration/Public/DisplayClusterConfigurationTypes_ICVFX.h
--- a/Engine/Plugins/Runtime/nDisplay/Source/DisplayClusterConfiguration/Public/DisplayClusterConfigurationTypes_ICVFX.h
+++ b/Engine/Plugins/Runtime/nDisplay/Source/DisplayClusterConfiguration/Public/DisplayClusterConfigurationTypes_ICVFX.h
@@ -733,6 +733,12 @@ public:
 	**/
 	bool IsICVFXEnabled(const class UDisplayClusterConfigurationData& InConfigurationData, const FString& InClusterNodeId) const;
 
+	/** Returns the CineCameraActor from the 'ExternalCameraActor' parameter. */
+	ACineCameraActor* GetExternalCineCameraActor() const;
+
+	/** Returns the CineCameraComponent from the 'ExternalCameraActor' parameter. */
+	UCineCameraComponent* GetExternalCineCameraComponent() const;
+
 public:
 	FDisplayClusterConfigurationICVFX_CameraSettings();
 
