Tagged plugin RDG passes with FRDGAsyncTask

diff --git a/Engine/Plugins/FX/Niagara/Source/Niagara/Private/NiagaraDataInterfaceLandscape.cpp b/Engine/Plugins/FX/Niagara/Source/Niagara/Private/NiagaraDataInterfaceLandscape.cpp
--- a/Engine/Plugins/FX/Niagara/Source/Niagara/Private/NiagaraDataInterfaceLandscape.cpp
+++ b/Engine/Plugins/FX/Niagara/Source/Niagara/Private/NiagaraDataInterfaceLandscape.cpp
@@ -1325,7 +1325,7 @@ bool UNiagaraDataInterfaceLandscape::SimCacheReadFrame(UObject* StorageObject, i
 				// Issues compute commands.
 				ERDGPassFlags::Compute,
 				// This is deferred until Execute. May execute in parallel with other passes.
-				[PassParameters, ComputeShader, NumThreadGroups](FRHIComputeCommandList& RHICmdList)
+				[PassParameters, ComputeShader, NumThreadGroups](FRDGAsyncTask, FRHIComputeCommandList& RHICmdList)
 				{
 					FComputeShaderUtils::Dispatch(RHICmdList, ComputeShader, *PassParameters, NumThreadGroups);
 				});
diff --git a/Engine/Plugins/FX/Niagara/Source/Niagara/Private/NiagaraDataInterfaceRenderTargetVolume.cpp b/Engine/Plugins/FX/Niagara/Source/Niagara/Private/NiagaraDataInterfaceRenderTargetVolume.cpp
--- a/Engine/Plugins/FX/Niagara/Source/Niagara/Private/NiagaraDataInterfaceRenderTargetVolume.cpp
+++ b/Engine/Plugins/FX/Niagara/Source/Niagara/Private/NiagaraDataInterfaceRenderTargetVolume.cpp
@@ -872,7 +872,7 @@ bool UNiagaraDataInterfaceRenderTargetVolume::SimCacheReadFrame(UObject* Storage
 							// Issues compute commands.
 							ERDGPassFlags::Compute,
 							// This is deferred until Execute. May execute in parallel with other passes.
-							[PassParameters, ComputeShader, NumThreadGroups](FRHIComputeCommandList& RHICmdList)
+							[PassParameters, ComputeShader, NumThreadGroups](FRDGAsyncTask, FRHIComputeCommandList& RHICmdList)
 							{
 								FComputeShaderUtils::Dispatch(RHICmdList, ComputeShader, *PassParameters, NumThreadGroups);
 							});
diff --git a/Engine/Plugins/FX/Niagara/Source/Niagara/Private/NiagaraGpuReadbackManager.cpp b/Engine/Plugins/FX/Niagara/Source/Niagara/Private/NiagaraGpuReadbackManager.cpp
--- a/Engine/Plugins/FX/Niagara/Source/Niagara/Private/NiagaraGpuReadbackManager.cpp
+++ b/Engine/Plugins/FX/Niagara/Source/Niagara/Private/NiagaraGpuReadbackManager.cpp
@@ -86,7 +86,7 @@ void FNiagaraGpuReadbackManager::EnqueueReadback(FRDGBuilder& GraphBuilder, FRDG
 
 	AddReadbackBufferPass(
 		GraphBuilder, RDG_EVENT_NAME("NiagaraReadback"), Buffer,
-		[Buffer, ReadbackBuffer=Readback.StagingBuffers[0].Key, Offset, NumBytes, Fence=Readback.Fence](FRHICommandList& RHICmdList)
+		[Buffer, ReadbackBuffer=Readback.StagingBuffers[0].Key, Offset, NumBytes, Fence=Readback.Fence](FRDGAsyncTask, FRHICommandList& RHICmdList)
 		{
 			RHICmdList.CopyToStagingBuffer(Buffer->GetRHI(), ReadbackBuffer, Offset, NumBytes);
 			RHICmdList.WriteGPUFence(Fence);
@@ -117,7 +117,7 @@ void FNiagaraGpuReadbackManager::EnqueueReadbacks(FRDGBuilder& GraphBuilder, TCo
 		const bool bWriteFence = i == Buffers.Num() - 1;
 		AddReadbackBufferPass(
 			GraphBuilder, RDG_EVENT_NAME("NiagaraReadback"), Buffer,
-			[Buffer, ReadbackBuffer, Offset=0, NumBytes, Fence=bWriteFence ? Readback.Fence : nullptr](FRHICommandListImmediate& RHICmdList)
+			[Buffer, ReadbackBuffer, Offset=0, NumBytes, Fence=bWriteFence ? Readback.Fence : nullptr](FRDGAsyncTask, FRHICommandList& RHICmdList)
 			{
 				RHICmdList.CopyToStagingBuffer(Buffer->GetRHI(), ReadbackBuffer, Offset, NumBytes);
 				if (Fence != nullptr)
@@ -154,7 +154,7 @@ void FNiagaraGpuReadbackManager::EnqueueReadbacks(FRDGBuilder& GraphBuilder, TCo
 		const bool bWriteFence = i == BufferRequests.Num() - 1;
 		AddReadbackBufferPass(
 			GraphBuilder, RDG_EVENT_NAME("NiagaraReadback"), Buffer,
-			[Buffer, ReadbackBuffer, Offset = 0, NumBytes, Fence=bWriteFence ? Readback.Fence : nullptr](FRHICommandListImmediate& RHICmdList)
+			[Buffer, ReadbackBuffer, Offset = 0, NumBytes, Fence=bWriteFence ? Readback.Fence : nullptr](FRDGAsyncTask, FRHICommandList& RHICmdList)
 			{
 				RHICmdList.CopyToStagingBuffer(Buffer->GetRHI(), ReadbackBuffer, Offset, NumBytes);
 				if (Fence != nullptr)
diff --git a/Engine/Plugins/Runtime/ComputeFramework/Source/ComputeFramework/Private/ComputeGraphWorker.cpp b/Engine/Plugins/Runtime/ComputeFramework/Source/ComputeFramework/Private/ComputeGraphWorker.cpp
--- a/Engine/Plugins/Runtime/ComputeFramework/Source/ComputeFramework/Private/ComputeGraphWorker.cpp
+++ b/Engine/Plugins/Runtime/ComputeFramework/Source/ComputeFramework/Private/ComputeGraphWorker.cpp
@@ -342,10 +342,8 @@ void FComputeGraphTaskWorker::SubmitWork(FRDGBuilder& GraphBuilder, FName InExec
 	}
 
 	// Release any graph resources at the end of graph execution.
-	GraphBuilder.AddPass(
-		{},
-		ERDGPassFlags::None,
-		[this, InExecutionGroupName](FRHICommandList&)
+	GraphBuilder.AddPostExecuteCallback(
+		[this, InExecutionGroupName]
 		{
 			GraphInvocationsPerGroup.FindChecked(InExecutionGroupName).Reset();
 		});
diff --git a/Engine/Plugins/Runtime/HairStrands/Source/HairStrandsCore/Private/GroomTextureBuilder.cpp b/Engine/Plugins/Runtime/HairStrands/Source/HairStrandsCore/Private/GroomTextureBuilder.cpp
--- a/Engine/Plugins/Runtime/HairStrands/Source/HairStrandsCore/Private/GroomTextureBuilder.cpp
+++ b/Engine/Plugins/Runtime/HairStrands/Source/HairStrandsCore/Private/GroomTextureBuilder.cpp
@@ -927,7 +927,7 @@ static void InternalGenerateHairStrandsTextures(
 			RDG_EVENT_NAME("HairStrands::TexturePS"),
 			ParametersPS,
 			ERDGPassFlags::Raster,
-			[ParametersPS, VertexShader, PixelShader, InMeshIndexBuffer, VertexCount, PrimitiveCount, IndexBaseIndex, VertexBaseIndex, OutputResolution, OutTileCoord, TileOffsetInPixels](FRHICommandList& RHICmdList)
+			[ParametersPS, VertexShader, PixelShader, InMeshIndexBuffer, VertexCount, PrimitiveCount, IndexBaseIndex, VertexBaseIndex, OutputResolution, OutTileCoord, TileOffsetInPixels](FRDGAsyncTask, FRHICommandList& RHICmdList)
 			{
 				FHairStrandsTextureVS::FParameters ParametersVS;
 				ParametersVS.OutputResolution	= ParametersPS->OutputResolution;
diff --git a/Engine/Plugins/Runtime/HairStrands/Source/HairStrandsCore/Private/HairStrandsInterpolation.cpp b/Engine/Plugins/Runtime/HairStrands/Source/HairStrandsCore/Private/HairStrandsInterpolation.cpp
--- a/Engine/Plugins/Runtime/HairStrands/Source/HairStrandsCore/Private/HairStrandsInterpolation.cpp
+++ b/Engine/Plugins/Runtime/HairStrands/Source/HairStrandsCore/Private/HairStrandsInterpolation.cpp
@@ -1057,7 +1057,7 @@ void AddHairCardsDeformationPass(
 		RDG_EVENT_NAME("HairStrands::CardsDeformation(%s)", bSupportDynamicMesh ? TEXT("Dynamic") : TEXT("Static")),
 		Parameters,
 		ERDGPassFlags::Compute,
-		[Parameters, ComputeShader, DispatchCount, CardsRestPositionBuffer, CardsRestTangentBuffer, bManualFetch](FRHIComputeCommandList& RHICmdList)
+		[Parameters, ComputeShader, DispatchCount, CardsRestPositionBuffer, CardsRestTangentBuffer, bManualFetch](FRDGAsyncTask, FRHIComputeCommandList& RHICmdList)
 		{
 			// On platforms not supporting manual vertex fetching, ensure the resources are in 'VerteOrIndexBuffer' state after position/normals update
 			if (!bManualFetch)
diff --git a/Engine/Plugins/Runtime/HairStrands/Source/HairStrandsCore/Private/HairStrandsMeshProjection.cpp b/Engine/Plugins/Runtime/HairStrands/Source/HairStrandsCore/Private/HairStrandsMeshProjection.cpp
--- a/Engine/Plugins/Runtime/HairStrands/Source/HairStrandsCore/Private/HairStrandsMeshProjection.cpp
+++ b/Engine/Plugins/Runtime/HairStrands/Source/HairStrandsCore/Private/HairStrandsMeshProjection.cpp
@@ -142,7 +142,7 @@ void AddSkinUpdatePass(
 		ParametersMetadata,
 		Parameters,
 		ERDGPassFlags::Compute,
-		[ParametersMetadata, Parameters, ComputeShader, Sections, bPrevPosition](FRHIComputeCommandList& RHICmdList)
+		[ParametersMetadata, Parameters, ComputeShader, Sections, bPrevPosition](FRDGAsyncTask, FRHIComputeCommandList& RHICmdList)
 		{
 			for (const FSkinUpdateSection& Section : Sections)
 			{
@@ -816,7 +816,7 @@ static void AddFollicleMaskPass(
 		RDG_EVENT_NAME("HairStrands::FollicleMask"),
 		Parameters,
 		ERDGPassFlags::Raster,
-		[Parameters, ParametersVS, ParametersPS, VertexShader, PixelShader, OutputResolution](FRHICommandList& RHICmdList)
+		[Parameters, ParametersVS, ParametersPS, VertexShader, PixelShader, OutputResolution](FRDGAsyncTask, FRHICommandList& RHICmdList)
 	{
 
 		RHICmdList.SetViewport(0, 0, 0.0f, OutputResolution.X, OutputResolution.Y, 1.0f);
@@ -877,7 +877,7 @@ static void AddFollicleMaskPass(
 		RDG_EVENT_NAME("HairStrands::FollicleMask"),
 		Parameters,
 		ERDGPassFlags::Raster,
-		[Parameters, ParametersVS, ParametersPS, VertexShader, PixelShader, OutputResolution](FRHICommandList& RHICmdList)
+		[Parameters, ParametersVS, ParametersPS, VertexShader, PixelShader, OutputResolution](FRDGAsyncTask, FRHICommandList& RHICmdList)
 		{
 
 			RHICmdList.SetViewport(0, 0, 0.0f, OutputResolution.X, OutputResolution.Y, 1.0f);
@@ -1008,7 +1008,7 @@ void AddComputeMipsPass(
 			RDG_EVENT_NAME("HairStrandsComputeVoxelMip"),
 			Parameters,
 			ERDGPassFlags::Compute,
-			[Parameters, ComputeShader, TargetResolution](FRHICommandList& RHICmdList)
+			[Parameters, ComputeShader, TargetResolution](FRDGAsyncTask, FRHICommandList& RHICmdList)
 		{
 			const FIntVector GroupCount = FComputeShaderUtils::GetGroupCount(FIntPoint(TargetResolution, TargetResolution), FIntPoint(8, 8));
 			FComputeShaderUtils::Dispatch(RHICmdList, ComputeShader, *Parameters, GroupCount);
diff --git a/Engine/Plugins/Runtime/nDisplay/Source/DisplayCluster/Private/Render/Viewport/DisplayClusterViewportProxy_Resources.cpp b/Engine/Plugins/Runtime/nDisplay/Source/DisplayCluster/Private/Render/Viewport/DisplayClusterViewportProxy_Resources.cpp
--- a/Engine/Plugins/Runtime/nDisplay/Source/DisplayCluster/Private/Render/Viewport/DisplayClusterViewportProxy_Resources.cpp
+++ b/Engine/Plugins/Runtime/nDisplay/Source/DisplayCluster/Private/Render/Viewport/DisplayClusterViewportProxy_Resources.cpp
@@ -62,7 +62,7 @@ namespace UE::DisplayCluster::ViewportProxy
 	static const int32 DisplayClusterViewportProxyResourcesOverrideRecursionDepthMax = 4;
 
 	template<class TScreenPixelShader>
-	void ResampleCopyTextureImpl_RenderThread(FRHICommandListImmediate& RHICmdList, FRHITexture* SrcTexture, FRHITexture* DstTexture, const FIntRect& SrcRect, const FIntRect& DstRect, const EDisplayClusterTextureCopyMode InCopyMode = EDisplayClusterTextureCopyMode::RGBA)
+	void ResampleCopyTextureImpl_RenderThread(FRHICommandList& RHICmdList, FRHITexture* SrcTexture, FRHITexture* DstTexture, const FIntRect& SrcRect, const FIntRect& DstRect, const EDisplayClusterTextureCopyMode InCopyMode = EDisplayClusterTextureCopyMode::RGBA)
 	{
 		// Texture format mismatch, use a shader to do the copy.
 		// #todo-renderpasses there's no explicit resolve here? Do we need one?
@@ -520,7 +520,7 @@ bool FDisplayClusterViewportProxy::CopyResource_RenderThread(FRDGBuilder& GraphB
 			RDG_EVENT_NAME("DisplayClusterViewportProxy_CopyResource(%s)", *GetId()),
 			PassParameters,
 			ERDGPassFlags::Copy | ERDGPassFlags::NeverCull,
-			[InCopyMode, InContextNum, InSrcTextureRef, InSrcRect, DestTextureRef, DestRect](FRHICommandListImmediate& RHICmdList)
+			[InCopyMode, InContextNum, InSrcTextureRef, InSrcRect, DestTextureRef, DestRect](FRDGAsyncTask, FRHICommandList& RHICmdList)
 			{
 				ResampleCopyTextureImpl_RenderThread<FScreenPS>(RHICmdList, InSrcTextureRef->GetRHI(), DestTextureRef->GetRHI(), InSrcRect, DestRect, InCopyMode);
 			});
@@ -554,7 +554,7 @@ bool FDisplayClusterViewportProxy::CopyResource_RenderThread(FRDGBuilder& GraphB
 			RDG_EVENT_NAME("DisplayClusterViewportProxy_CopyResource(%s)", *GetId()),
 			PassParameters,
 			ERDGPassFlags::Copy | ERDGPassFlags::NeverCull,
-			[InCopyMode, InContextNum, SrcTextureRef, SrcRect, InDestTextureRef, InDestRect](FRHICommandListImmediate& RHICmdList)
+			[InCopyMode, InContextNum, SrcTextureRef, SrcRect, InDestTextureRef, InDestRect](FRDGAsyncTask, FRHICommandList& RHICmdList)
 			{
 				ResampleCopyTextureImpl_RenderThread<FScreenPS>(RHICmdList, SrcTextureRef->GetRHI(), InDestTextureRef->GetRHI(), SrcRect, InDestRect, InCopyMode);
 			});
diff --git a/Engine/Plugins/Runtime/nDisplay/Source/DisplayClusterLightCardEditorShaders/Private/DisplayClusterMeshProjectionRenderer.cpp b/Engine/Plugins/Runtime/nDisplay/Source/DisplayClusterLightCardEditorShaders/Private/DisplayClusterMeshProjectionRenderer.cpp
--- a/Engine/Plugins/Runtime/nDisplay/Source/DisplayClusterLightCardEditorShaders/Private/DisplayClusterMeshProjectionRenderer.cpp
+++ b/Engine/Plugins/Runtime/nDisplay/Source/DisplayClusterLightCardEditorShaders/Private/DisplayClusterMeshProjectionRenderer.cpp
@@ -1064,7 +1064,7 @@ void FDisplayClusterMeshProjectionRenderer::RenderColorOutput(FRDGBuilder& Graph
 			RDG_EVENT_NAME("MeshProjectionRenderer::CopyColorTexture"),
 			ScreenPassParameters,
 			ERDGPassFlags::Raster,
-			[View, ScreenPassVS, CopyPixelShader, RegionViewport, ScreenPassParameters, DefaultBlendState](FRHICommandList& RHICmdList)
+			[View, ScreenPassVS, CopyPixelShader, RegionViewport, ScreenPassParameters, DefaultBlendState](FRDGAsyncTask, FRHICommandList& RHICmdList)
 		{
 			DrawScreenPass(
 				RHICmdList,
@@ -1121,7 +1121,7 @@ void FDisplayClusterMeshProjectionRenderer::RenderHitProxyOutput(FRDGBuilder& Gr
 			RDG_EVENT_NAME("MeshProjectionRenderer::CopyHitProxyTexture"),
 			ScreenPassParameters,
 			ERDGPassFlags::Raster,
-			[View, ScreenPassVS, CopyPixelShader, RegionViewport, ScreenPassParameters, DefaultBlendState](FRHICommandList& RHICmdList)
+			[View, ScreenPassVS, CopyPixelShader, RegionViewport, ScreenPassParameters, DefaultBlendState](FRDGAsyncTask, FRHICommandList& RHICmdList)
 		{
 			DrawScreenPass(
 				RHICmdList,
@@ -1455,7 +1455,7 @@ void FDisplayClusterMeshProjectionRenderer::AddNormalsFilterPass(FRDGBuilder& Gr
 			RDG_EVENT_NAME("MeshProjectionRenderer::OutputNormals"),
 			ScreenPassParameters,
 			ERDGPassFlags::Raster,
-			[View, ScreenPassVS, OutputNormalsPS, InputViewport, OutputViewport, ScreenPassParameters, DefaultBlendState](FRHICommandList& RHICmdList)
+			[View, ScreenPassVS, OutputNormalsPS, InputViewport, OutputViewport, ScreenPassParameters, DefaultBlendState](FRDGAsyncTask, FRHICommandList& RHICmdList)
 		{
 			DrawScreenPass(
 				RHICmdList,
@@ -1536,7 +1536,7 @@ void FDisplayClusterMeshProjectionRenderer::AddSelectionOutlineScreenPass(FRDGBu
 		RDG_EVENT_NAME("MeshProjectionRenderer::SelectionScreen"),
 		ScreenPassParameters,
 		ERDGPassFlags::Raster,
-		[View, ScreenPassVS, SelectionOutlinePS, OutputViewport, ScreenPassParameters, DefaultBlendState](FRHICommandList& RHICmdList)
+		[View, ScreenPassVS, SelectionOutlinePS, OutputViewport, ScreenPassParameters, DefaultBlendState](FRDGAsyncTask, FRHICommandList& RHICmdList)
 	{
 		DrawScreenPass(
 			RHICmdList,
@@ -1962,7 +1962,7 @@ void FDisplayClusterMeshProjectionRenderer::RenderColorOutputs(FRDGBuilder& Grap
 			RDG_EVENT_NAME("MeshProjectionRenderer::CopyColorTexture"),
 			ScreenPassParameters,
 			ERDGPassFlags::Raster,
-			[View = Views[0], ScreenPassVS, CopyPixelShader, RegionViewport, ScreenPassParameters, DefaultBlendState](FRHICommandList& RHICmdList)
+			[View = Views[0], ScreenPassVS, CopyPixelShader, RegionViewport, ScreenPassParameters, DefaultBlendState](FRDGAsyncTask, FRHICommandList& RHICmdList)
 			{
 				DrawScreenPass(
 					RHICmdList,
diff --git a/Engine/Plugins/VirtualProduction/CameraCalibrationCore/Source/CameraCalibrationCore/Private/CalibratedMapProcessor.cpp b/Engine/Plugins/VirtualProduction/CameraCalibrationCore/Source/CameraCalibrationCore/Private/CalibratedMapProcessor.cpp
--- a/Engine/Plugins/VirtualProduction/CameraCalibrationCore/Source/CameraCalibrationCore/Private/CalibratedMapProcessor.cpp
+++ b/Engine/Plugins/VirtualProduction/CameraCalibrationCore/Source/CameraCalibrationCore/Private/CalibratedMapProcessor.cpp
@@ -214,7 +214,7 @@ void FCalibratedMapProcessor::ExecuteJob(TSharedPtr<FDerivedDistortionDataJob> J
             			RDG_EVENT_NAME("LensFileSTMapConversion"),
             			Parameters,
             			ERDGPassFlags::Compute,
-            			[Parameters, ComputeShader](FRHICommandList& RHICmdList)
+            			[Parameters, ComputeShader](FRDGAsyncTask, FRHICommandList& RHICmdList)
 						{
             				const FIntVector GroupCount = FComputeShaderUtils::GetGroupCount(Parameters->TextureSize, FComputeShaderUtils::kGolden2DGroupSize);
             				FComputeShaderUtils::Dispatch(RHICmdList, ComputeShader, *Parameters, GroupCount);
diff --git a/Engine/Plugins/VirtualProduction/CameraCalibrationCore/Source/CameraCalibrationCore/Private/LensDistortionSceneViewExtension.cpp b/Engine/Plugins/VirtualProduction/CameraCalibrationCore/Source/CameraCalibrationCore/Private/LensDistortionSceneViewExtension.cpp
--- a/Engine/Plugins/VirtualProduction/CameraCalibrationCore/Source/CameraCalibrationCore/Private/LensDistortionSceneViewExtension.cpp
+++ b/Engine/Plugins/VirtualProduction/CameraCalibrationCore/Source/CameraCalibrationCore/Private/LensDistortionSceneViewExtension.cpp
@@ -300,7 +300,7 @@ void FLensDistortionSceneViewExtension::InvertDistortionMap_RenderThread(FRDGBui
 		RDG_EVENT_NAME("InvertDistortionDisplacementMap"),
 		PassParameters,
 		ERDGPassFlags::Raster,
-		[PassParameters, Output](FRHICommandList& RHICmdList)
+		[PassParameters, Output](FRDGAsyncTask, FRHICommandList& RHICmdList)
 		{
 			const int32 NumSquares = FMath::Clamp(CVarLensDistortionInvertGridDensity.GetValueOnRenderThread(), 64, 255);
 			const FIntPoint GridDimensions = FIntPoint(NumSquares);
