nDisplay: Hid some of the extra options in the view point related to using icvfx camera properties.

For now the UX will be that those properties will be controlled from the camera itself instead of having enable/disable in the viewpoint.
Values will now default to on.
Changed follow camera default to on as well.

#rnx

diff --git a/Engine/Plugins/Runtime/nDisplay/Source/DisplayCluster/Public/Components/DisplayClusterCameraComponent.h b/Engine/Plugins/Runtime/nDisplay/Source/DisplayCluster/Public/Components/DisplayClusterCameraComponent.h
--- a/Engine/Plugins/Runtime/nDisplay/Source/DisplayCluster/Public/Components/DisplayClusterCameraComponent.h
+++ b/Engine/Plugins/Runtime/nDisplay/Source/DisplayCluster/Public/Components/DisplayClusterCameraComponent.h
@@ -307,36 +307,37 @@ public:
 	FString ICVFXCameraComponentName;
 
 	UPROPERTY(EditAnywhere, BlueprintReadWrite, Category = "Camera Post Process", meta = (DisplayName = "Follow ICVFX Camera", EditCondition = "TargetCameraType==EDisplayClusterTargetCameraType::ICVFXCameraComponent", EditConditionHides))
-	uint8 bUseICVFXCameraComponentTracking: 1 = 0;
+	uint8 bUseICVFXCameraComponentTracking: 1 = 1;
 
 	/** Use a specific actor camera instead of a game camera. */
 	UPROPERTY(EditAnywhere, BlueprintReadWrite, Category = "Camera Post Process", meta = (DisplayName = "Camera Actor", EditCondition = "TargetCameraType==EDisplayClusterTargetCameraType::ExternalCineCameraActor", EditConditionHides))
 	TSoftObjectPtr<ACineCameraActor> ExternalCineCameraActor;
 
 public:
-	/** Use the NearClippingPlane value from the specified cine camera. */
-	UPROPERTY(EditAnywhere, BlueprintReadWrite, Category = "Camera Post Process", AdvancedDisplay, meta = (DisplayName = "Include Custom Near Clipping Plane", EditCondition = "TargetCameraType!=EDisplayClusterTargetCameraType::None", EditConditionHides))
-	uint8 bEnableNearClippingPlane : 1 = 1;
 
 	/** Use the PP settings from the specified camera. */
-	UPROPERTY(EditAnywhere, BlueprintReadWrite, Category = "Camera Post Process", AdvancedDisplay, meta = (DisplayName = "Include Camera Post Process", EditCondition = "TargetCameraType!=EDisplayClusterTargetCameraType::None", EditConditionHides))
+	UPROPERTY(EditAnywhere, BlueprintReadWrite, Category = "Camera Post Process", meta = (DisplayName = "Include Post Process Settings", EditCondition = "TargetCameraType!=EDisplayClusterTargetCameraType::None", EditConditionHides))
 	uint8 bEnablePostProcess : 1 = 1;
 
+	/** Use the NearClippingPlane value from the specified cine camera. */
+	UPROPERTY(EditAnywhere, BlueprintReadWrite, Category = "Camera Post Process", meta = (DisplayName = "Include Custom Near Clipping Plane", EditCondition = "TargetCameraType!=EDisplayClusterTargetCameraType::None && TargetCameraType!=EDisplayClusterTargetCameraType::ICVFXCameraComponent", EditConditionHides))
+	uint8 bEnableNearClippingPlane : 1 = 1;
+
 	/** Enable the DoF PP settings from the specified camera. */
-	UPROPERTY(EditAnywhere, BlueprintReadWrite, Category = "Camera Post Process", AdvancedDisplay, meta = (DisplayName = "Include Depth Of Field", EditCondition = "TargetCameraType!=EDisplayClusterTargetCameraType::None", EditConditionHides))
-	uint8 bEnableDepthOfField : 1 = 0;
+	UPROPERTY(EditAnywhere, BlueprintReadWrite, Category = "Camera Post Process", meta = (DisplayName = "Include Depth Of Field", EditCondition = "TargetCameraType!=EDisplayClusterTargetCameraType::None && TargetCameraType!=EDisplayClusterTargetCameraType::ICVFXCameraComponent", EditConditionHides))
+	uint8 bEnableDepthOfField : 1 = 1;
 
 	/** Use the DC Depth-Of-Field settings from the specified ICVFX camera. */
-	UPROPERTY(EditAnywhere, BlueprintReadWrite, Category = "Camera Post Process", AdvancedDisplay, meta = (DisplayName = "Include Inner Frustum Depth Of Field Compensation", EditCondition = "TargetCameraType==EDisplayClusterTargetCameraType::ICVFXCameraComponent", EditConditionHides))
-	uint8 bEnableICVFXDepthOfFieldCompensation : 1 = 0;
+	UPROPERTY(BlueprintReadWrite, Category = "Camera Post Process", meta = (DisplayName = "Include Inner Frustum Depth Of Field Compensation", EditCondition = "TargetCameraType==EDisplayClusterTargetCameraType::ICVFXCameraComponent", EditConditionHides))
+	uint8 bEnableICVFXDepthOfFieldCompensation : 1 = 1;
 
 	/** Use the DC ColorGrading from the specified ICVFX camera. */
 	UPROPERTY(EditAnywhere, BlueprintReadWrite, Category = "Camera Post Process", meta = (DisplayName = "Include Inner Frustum Color Grading", EditCondition = "TargetCameraType==EDisplayClusterTargetCameraType::ICVFXCameraComponent", EditConditionHides))
 	uint8 bEnableICVFXColorGrading : 1 = 1;
 
 	/** Use the DC Motion Blur settings from the specified ICVFX camera. */
-	UPROPERTY(EditAnywhere, BlueprintReadWrite, Category = "Camera Post Process", AdvancedDisplay, meta = (DisplayName = "Include Inner Frustum Motion Blur", EditCondition = "TargetCameraType==EDisplayClusterTargetCameraType::ICVFXCameraComponent", EditConditionHides))
-	uint8 bEnableICVFXMotionBlur : 1 = 0;
+	UPROPERTY(BlueprintReadWrite, Category = "Camera Post Process", meta = (DisplayName = "Include Inner Frustum Motion Blur", EditCondition = "TargetCameraType==EDisplayClusterTargetCameraType::ICVFXCameraComponent", EditConditionHides))
+	uint8 bEnableICVFXMotionBlur : 1 = 1;
 
 private:
 	UPROPERTY(EditAnywhere, Category = "Stereo")
