[nDisplay] Minor GUI fixes/updates in media configuration - basic GUI changes for docs

- This CL makes some changes to the GUI required for docs (screenshots)
- It doesn't handle all the requests from the corresponding jira

#rnx

diff --git a/Engine/Plugins/Runtime/nDisplay/Source/DisplayClusterConfiguration/Public/DisplayClusterConfigurationTypes_Media.h b/Engine/Plugins/Runtime/nDisplay/Source/DisplayClusterConfiguration/Public/DisplayClusterConfigurationTypes_Media.h
--- a/Engine/Plugins/Runtime/nDisplay/Source/DisplayClusterConfiguration/Public/DisplayClusterConfigurationTypes_Media.h
+++ b/Engine/Plugins/Runtime/nDisplay/Source/DisplayClusterConfiguration/Public/DisplayClusterConfigurationTypes_Media.h
@@ -303,7 +303,7 @@ public:
 	TArray<FDisplayClusterConfigurationMediaTiledOutputGroup> TiledMediaOutputGroups;
 
 	/** Allows the receviers to apply their own OCIO transformations (per node OCIO override). Requires the media to support FloatRGBA.*/
-	UPROPERTY(EditAnywhere, BlueprintReadWrite, Category = "Media", meta = (DisplayName = "Receivers apply OCIO"))
+	UPROPERTY(EditAnywhere, BlueprintReadWrite, Category = "Media", meta = (DisplayName = "Receivers Apply OCIO"))
 	bool bLateOCIOPass = false;
 
 public:
diff --git a/Engine/Plugins/Runtime/nDisplay/Source/DisplayClusterConfigurator/Private/Views/Details/Components/DisplayClusterICVFXCameraComponentDetailsCustomization.cpp b/Engine/Plugins/Runtime/nDisplay/Source/DisplayClusterConfigurator/Private/Views/Details/Components/DisplayClusterICVFXCameraComponentDetailsCustomization.cpp
--- a/Engine/Plugins/Runtime/nDisplay/Source/DisplayClusterConfigurator/Private/Views/Details/Components/DisplayClusterICVFXCameraComponentDetailsCustomization.cpp
+++ b/Engine/Plugins/Runtime/nDisplay/Source/DisplayClusterConfigurator/Private/Views/Details/Components/DisplayClusterICVFXCameraComponentDetailsCustomization.cpp
@@ -41,7 +41,7 @@ namespace DisplayClusterICVFXCameraComponentDetailsCustomizationUtils
 		{
 			int32 CurrentSortOrder = Pair.Value->GetSortOrder();
 
-			int32 DesiredSortOrder;
+			int32 DesiredSortOrder = 0;
 			if (CategoryOrder.Find(Pair.Key, DesiredSortOrder))
 			{
 				CurrentSortOrder = DesiredSortOrder;
diff --git a/Engine/Plugins/Runtime/nDisplay/Source/DisplayClusterConfigurator/Private/Views/Details/Media/DisplayClusterConfiguratorICVFXMediaCustomization.cpp b/Engine/Plugins/Runtime/nDisplay/Source/DisplayClusterConfigurator/Private/Views/Details/Media/DisplayClusterConfiguratorICVFXMediaCustomization.cpp
--- a/Engine/Plugins/Runtime/nDisplay/Source/DisplayClusterConfigurator/Private/Views/Details/Media/DisplayClusterConfiguratorICVFXMediaCustomization.cpp
+++ b/Engine/Plugins/Runtime/nDisplay/Source/DisplayClusterConfigurator/Private/Views/Details/Media/DisplayClusterConfiguratorICVFXMediaCustomization.cpp
@@ -40,7 +40,7 @@ void FDisplayClusterConfiguratorICVFXMediaCustomization::CustomizeChildren(TShar
 	check(SplitTypeHandle && SplitTypeHandle->IsValidHandle());
 
 	// Layout property
-	TilesLayoutHandle = GET_CHILD_HANDLE(FDisplayClusterConfigurationMediaICVFX, TiledSplitLayout);
+	TSharedPtr<IPropertyHandle> TilesLayoutHandle = GET_CHILD_HANDLE(FDisplayClusterConfigurationMediaICVFX, TiledSplitLayout);
 	check(TilesLayoutHandle && TilesLayoutHandle->IsValidHandle());
 
 	// Separate groups specific for every split type available
@@ -65,7 +65,7 @@ void FDisplayClusterConfiguratorICVFXMediaCustomization::CustomizeChildren(TShar
 	TSharedPtr<IPropertyUtilities> PropertyUtils = InCustomizationUtils.GetPropertyUtilities();
 	check(PropertyUtils);
 
-	// Setup details update on frustum type change
+	// Set details update request on frustum type change
 	SplitTypeHandle->SetOnPropertyValueChanged(
 		FSimpleDelegate::CreateLambda([PropertyUtils]()
 		{
@@ -100,20 +100,39 @@ void FDisplayClusterConfiguratorICVFXMediaCustomization::CustomizeChildren(TShar
 		}
 	}
 
-	// Create all the property widgets
-	FDisplayClusterConfiguratorBaseTypeCustomization::CustomizeChildren(InPropertyHandle, InChildBuilder, InCustomizationUtils);
-
-	// Create 'setup' button
-	if (SplitTypeValue == EDisplayClusterConfigurationMediaSplitType::UniformTiles)
+	// Finally, build the panel
+	if (ShouldShowChildren(InPropertyHandle))
 	{
-		AddSetupButton(InChildBuilder);
-	}
+		const bool bUsingTiles = (SplitTypeValue == EDisplayClusterConfigurationMediaSplitType::UniformTiles);
+
+		uint32 NumChildren = 0;
+		InPropertyHandle->GetNumChildren(NumChildren);
+
+		// For each child property, build its own layout
+		for (uint32 ChildIndex = 0; ChildIndex < NumChildren; ++ChildIndex)
+		{
+			TSharedPtr<IPropertyHandle> ChildHandle = InPropertyHandle->GetChildHandle(ChildIndex);
+			if (ChildHandle && ChildHandle->IsValidHandle() && !ChildHandle->IsCustomized())
+			{
+				FText ChildTooltip = ApplySubstitutions(ChildHandle->GetToolTipText());
+				ChildHandle->SetToolTipText(ChildTooltip);
 
-	// Create 'reset' button at the bottom
-	AddResetButton(InChildBuilder);
+				InChildBuilder.AddProperty(ChildHandle.ToSharedRef());
+
+				// Insert tile configuration button after the split type combobox
+				if (bUsingTiles && ChildHandle->IsSamePropertyNode(SplitTypeHandle))
+				{
+					AddConfigureTilesButton(InChildBuilder);
+				}
+			}
+		}
+
+		// Create 'reset' button at the bottom
+		AddResetButton(InChildBuilder);
+	}
 }
 
-void FDisplayClusterConfiguratorICVFXMediaCustomization::AddSetupButton(IDetailChildrenBuilder& InChildBuilder)
+void FDisplayClusterConfiguratorICVFXMediaCustomization::AddConfigureTilesButton(IDetailChildrenBuilder& InChildBuilder)
 {
 	InChildBuilder.AddCustomRow(FText::GetEmpty())
 		.WholeRowContent()
@@ -123,11 +142,11 @@ void FDisplayClusterConfiguratorICVFXMediaCustomization::AddSetupButton(IDetailC
 			[
 				SNew(SButton)
 				.HAlign(HAlign_Center)
-				.OnClicked(this, &FDisplayClusterConfiguratorICVFXMediaCustomization::OnSetupButtonClicked)
+				.OnClicked(this, &FDisplayClusterConfiguratorICVFXMediaCustomization::OnConfigureTilesButtonClicked)
 				[
 					SNew(STextBlock)
 					.Font(FAppStyle::GetFontStyle(TEXT("PropertyWindow.NormalFont")))
-					.Text(LOCTEXT("SetupTilesButtonTitle", "Setup Tiles"))
+					.Text(LOCTEXT("ConfigureTilesButtonTitle", "Configure Tiles"))
 				]
 			]
 		];
@@ -147,13 +166,13 @@ void FDisplayClusterConfiguratorICVFXMediaCustomization::AddResetButton(IDetailC
 				[
 					SNew(STextBlock)
 					.Font(FAppStyle::GetFontStyle(TEXT("PropertyWindow.NormalFont")))
-					.Text(LOCTEXT("ResetToDefaultsButtonTitle", "Reset To Defaults"))
+					.Text(LOCTEXT("ResetToDefaultButtonTitle", "Reset Media Input and Output to Default"))
 				]
 			]
 		];
 }
 
-FReply FDisplayClusterConfiguratorICVFXMediaCustomization::OnSetupButtonClicked()
+FReply FDisplayClusterConfiguratorICVFXMediaCustomization::OnConfigureTilesButtonClicked()
 {
 	// We're in camera tiles customization so let's get the camera component
 	UDisplayClusterICVFXCameraComponent* ICVFXCamera = Cast<UDisplayClusterICVFXCameraComponent>(EditingObject.Get());
@@ -172,7 +191,7 @@ FReply FDisplayClusterConfiguratorICVFXMediaCustomization::OnSetupButtonClicked(
 	// Nothing to do if no cluster nodes available
 	if (ConfigData->Cluster->Nodes.Num() < 1)
 	{
-		FMessageDialog::Open(EAppMsgType::Ok, LOCTEXT("MessageNoClusterNodesAvailable", "There are no cluster nodes assigned. Can't setup media."));
+		FMessageDialog::Open(EAppMsgType::Ok, LOCTEXT("MessageNoClusterNodesAvailable", "There are no cluster nodes assigned. Can't configure media."));
 		return FReply::Handled();
 	}
 
diff --git a/Engine/Plugins/Runtime/nDisplay/Source/DisplayClusterConfigurator/Private/Views/Details/Media/DisplayClusterConfiguratorICVFXMediaCustomization.h b/Engine/Plugins/Runtime/nDisplay/Source/DisplayClusterConfigurator/Private/Views/Details/Media/DisplayClusterConfiguratorICVFXMediaCustomization.h
--- a/Engine/Plugins/Runtime/nDisplay/Source/DisplayClusterConfigurator/Private/Views/Details/Media/DisplayClusterConfiguratorICVFXMediaCustomization.h
+++ b/Engine/Plugins/Runtime/nDisplay/Source/DisplayClusterConfigurator/Private/Views/Details/Media/DisplayClusterConfiguratorICVFXMediaCustomization.h
@@ -31,14 +31,14 @@ protected:
 
 private:
 
-	/** Builds setup button widget. */
-	void AddSetupButton(IDetailChildrenBuilder& ChildBuilder);
+	/** Builds the button widget for tiles configuration. */
+	void AddConfigureTilesButton(IDetailChildrenBuilder& ChildBuilder);
 
 	/** Builds reset button widget. */
 	void AddResetButton(IDetailChildrenBuilder& ChildBuilder);
 
-	/** Handles setup button clicks. */
-	FReply OnSetupButtonClicked();
+	/** Handles configure tiles button clicks. */
+	FReply OnConfigureTilesButtonClicked();
 
 	/** Handles reset button clicks. */
 	FReply OnResetButtonClicked();
@@ -50,9 +50,4 @@ private:
 
 	/** Marks package as dirty */
 	void MarkDirty();
-
-private:
-
-	/** Tiles layout property. */
-	TSharedPtr<IPropertyHandle> TilesLayoutHandle;
 };
