nDisplay: UDisplayClusterCameraComponent options must be properly handled in logic when hidden in the UI.

Note: Some of the extra options have been hidden in CL 36098231, but they require to be handled.

diff --git a/Engine/Plugins/Runtime/nDisplay/Source/DisplayCluster/Private/Game/EngineClasses/Scene/DisplayClusterCameraComponent.cpp b/Engine/Plugins/Runtime/nDisplay/Source/DisplayCluster/Private/Game/EngineClasses/Scene/DisplayClusterCameraComponent.cpp
--- a/Engine/Plugins/Runtime/nDisplay/Source/DisplayCluster/Private/Game/EngineClasses/Scene/DisplayClusterCameraComponent.cpp
+++ b/Engine/Plugins/Runtime/nDisplay/Source/DisplayCluster/Private/Game/EngineClasses/Scene/DisplayClusterCameraComponent.cpp
@@ -351,27 +351,32 @@ EDisplayClusterViewportCameraPostProcessFlags UDisplayClusterCameraComponent::Ge
 		EnumAddFlags(OutPostProcessFlags, EDisplayClusterViewportCameraPostProcessFlags::EnablePostProcess);
 	}
 
-	if (bEnableDepthOfField)
+	// If an ICVFX camera is used, DoF is always enabled
+	if (bEnableDepthOfField || IsICVFXCameraBeingUsed())
 	{
 		EnumAddFlags(OutPostProcessFlags, EDisplayClusterViewportCameraPostProcessFlags::EnableDepthOfField);
 	}
 
-	if (bEnableNearClippingPlane)
+	// If an ICVFX camera is used, custom NCP is always enabled
+	if (bEnableNearClippingPlane || IsICVFXCameraBeingUsed())
 	{
 		EnumAddFlags(OutPostProcessFlags, EDisplayClusterViewportCameraPostProcessFlags::EnableNearClippingPlane);
 	}
 
-	if (bEnableICVFXColorGrading)
+	// This option requires an ICVFX camera.
+	if (bEnableICVFXColorGrading && IsICVFXCameraBeingUsed())
 	{
 		EnumAddFlags(OutPostProcessFlags, EDisplayClusterViewportCameraPostProcessFlags::EnableICVFXColorGrading);
 	}
 
-	if (bEnableICVFXMotionBlur)
+	// This option requires an ICVFX camera.
+	if (bEnableICVFXMotionBlur && IsICVFXCameraBeingUsed())
 	{
 		EnumAddFlags(OutPostProcessFlags, EDisplayClusterViewportCameraPostProcessFlags::EnableICVFXMotionBlur);
 	}
 
-	if (bEnableICVFXDepthOfFieldCompensation)
+	// This option requires an ICVFX camera.
+	if (bEnableICVFXDepthOfFieldCompensation && IsICVFXCameraBeingUsed())
 	{
 		EnumAddFlags(OutPostProcessFlags, EDisplayClusterViewportCameraPostProcessFlags::EnableICVFXDepthOfFieldCompensation);
 	}
