Derived Data Cache, Switchboard.  Use the project pak instead of the install engine pak..  A crash will still occur on second run within the editor but at least we avoid issues with installed engine pak.  JIRA issue will remain open.

-Fix: -run=DerivedDataCache -fill -DDC=CreateInstalledEnginePak crashes editor

diff --git a/Engine/Plugins/Runtime/nDisplay/Source/DisplayClusterFillDerivedDataCache/Private/DisplayClusterFillDerivedDataCacheWorker.cpp b/Engine/Plugins/Runtime/nDisplay/Source/DisplayClusterFillDerivedDataCache/Private/DisplayClusterFillDerivedDataCacheWorker.cpp
--- a/Engine/Plugins/Runtime/nDisplay/Source/DisplayClusterFillDerivedDataCache/Private/DisplayClusterFillDerivedDataCacheWorker.cpp
+++ b/Engine/Plugins/Runtime/nDisplay/Source/DisplayClusterFillDerivedDataCache/Private/DisplayClusterFillDerivedDataCacheWorker.cpp
@@ -45,14 +45,14 @@ FDisplayClusterFillDerivedDataCacheWorker::FDisplayClusterFillDerivedDataCacheWo
 	CurrentExecutableName = FPlatformProcess::ExecutablePath();
 	const FString ProjectPath = FPaths::SetExtension(
 		FPaths::Combine(FPaths::ProjectDir(), FApp::GetProjectName()),".uproject");
-	Arguments = ProjectPath + " " + GetDdcCommandletParams();
+	Arguments = FString::Printf(TEXT("\"%s\" %s"), *ProjectPath, *GetDdcCommandletParams());
 	
 	UE_LOG(LogDisplayClusterFillDerivedDataCache, Display, TEXT("Running commandlet: %s %s"), *CurrentExecutableName, *Arguments);
 
 	uint32 ProcessID;
 	const bool bLaunchDetached = true;
-	const bool bLaunchHidden = true;
-	const bool bLaunchReallyHidden = true;
+	const bool bLaunchHidden = false;
+	const bool bLaunchReallyHidden = false;
 	ProcessHandle = FPlatformProcess::CreateProc(
 		*CurrentExecutableName, *Arguments, bLaunchDetached, bLaunchHidden, bLaunchReallyHidden, &ProcessID,
 		0, nullptr, WritePipe, ReadPipe);
@@ -189,7 +189,7 @@ void FDisplayClusterFillDerivedDataCacheWorker::RegexParseForCompilationProgress
 
 FString FDisplayClusterFillDerivedDataCacheWorker::GetDdcCommandletParams() const
 {
-	return FString::Printf(TEXT("-run=DerivedDataCache %s -fill -DDC=CreateInstalledEnginePak"), *GetTargetPlatformParams());
+	return FString::Printf(TEXT("-run=DerivedDataCache %s -fill -DDC=CreateInstalledProjectPak"), *GetTargetPlatformParams());
 }
 
 FString FDisplayClusterFillDerivedDataCacheWorker::GetTargetPlatformParams() const
diff --git a/Engine/Plugins/VirtualProduction/Switchboard/Source/Switchboard/switchboard/devices/unreal/plugin_unreal.py b/Engine/Plugins/VirtualProduction/Switchboard/Source/Switchboard/switchboard/devices/unreal/plugin_unreal.py
--- a/Engine/Plugins/VirtualProduction/Switchboard/Source/Switchboard/switchboard/devices/unreal/plugin_unreal.py
+++ b/Engine/Plugins/VirtualProduction/Switchboard/Source/Switchboard/switchboard/devices/unreal/plugin_unreal.py
@@ -2232,7 +2232,7 @@ class DeviceUnreal(Device):
         current_level_valid = current_level_only and current_level != None and current_level != DEFAULT_MAP_TEXT
         map_name = ' Map=' + current_level if current_level_valid else ''
 
-        args = CONFIG.UPROJECT_PATH.get_value(self.name) + map_name + " -run=DerivedDataCache -TargetPlatform=" + platform_string + " -fill -DDC=CreateInstalledEnginePak"
+        args = "\"" + CONFIG.UPROJECT_PATH.get_value(self.name) + map_name + "\"" + " -run=DerivedDataCache -TargetPlatform=" + platform_string + " -fill -DDC=CreateInstalledProjectPak"
         LOGGER.info('Filling ddc with arguments: ' + args)
 
         puuid, msg = message_protocol.create_start_process_message(
