Display: Overrides the DoF PPS received from CineCamera to produce the expected result in nDisplay viewports.

DoF parameters are defined in CineCamera. The sensor size and squeeze factor set in CineCamera cannot be directly used for nDisplay viewports for several reasons:
- Overscan: changes the FOV (sensor size) and/or RTT size.
- Tile rendering: changes the FOV and RTT size.
- Outers: they use a proprietary FOV that dynamically changes depending on the ViewPoint location, and an RTT size that is independent of the sensor size.

Added new CVars:
- "nDisplay.render.postprocess.override.enable" - Enable postprocess overrides for nDisplay
- "nDisplay.render.postprocess.override.InCameraVFX" - Enable post-processing override for ICVFX Camera viewport
- "nDisplay.render.postprocess.override.OutersVFX" - Enable postprocess override for ICVFX Outer viewports. Disabled by default as some issues have not been resolved in this CL.

Note: Set some CVars from "r.DOF.Kernel..." to 0.1 (or even higher) to keep the correct DoF radius:
"r.DOF.Kernel.MaxForegroundRadius 0.1"
"r.DOF.Kernel.MaxBackgroundRadius 0.1"

diff --git a/Engine/Plugins/Runtime/nDisplay/Source/DisplayCluster/Private/Render/Viewport/DisplayClusterViewport.h b/Engine/Plugins/Runtime/nDisplay/Source/DisplayCluster/Private/Render/Viewport/DisplayClusterViewport.h
--- a/Engine/Plugins/Runtime/nDisplay/Source/DisplayCluster/Private/Render/Viewport/DisplayClusterViewport.h
+++ b/Engine/Plugins/Runtime/nDisplay/Source/DisplayCluster/Private/Render/Viewport/DisplayClusterViewport.h
@@ -119,7 +119,7 @@ public:
 	virtual bool    CalculateView(const uint32 InContextNum, FVector& InOutViewLocation, FRotator& InOutViewRotation, const float WorldToMeters) override;
 	virtual bool    GetProjectionMatrix(const uint32 InContextNum, FMatrix& OutPrjMatrix)  override;
 
-	virtual bool SetupViewPoint(FMinimalViewInfo& InOutViewInfo) override;
+	virtual bool SetupViewPoint(const uint32 InContextNum, FMinimalViewInfo& InOutViewInfo) override;
 	virtual float GetStereoEyeOffsetDistance(const uint32 InContextNum) override;
 	virtual class UDisplayClusterCameraComponent* GetViewPointCameraComponent(const EDisplayClusterRootActorType InRootActorType) const override;
 	virtual class UDisplayClusterDisplayDeviceBaseComponent* GetDisplayDeviceComponent(const EDisplayClusterRootActorType InRootActorType) const override;
diff --git a/Engine/Plugins/Runtime/nDisplay/Source/DisplayCluster/Private/Render/Viewport/DisplayClusterViewportManagerViewPointExtension.cpp b/Engine/Plugins/Runtime/nDisplay/Source/DisplayCluster/Private/Render/Viewport/DisplayClusterViewportManagerViewPointExtension.cpp
--- a/Engine/Plugins/Runtime/nDisplay/Source/DisplayCluster/Private/Render/Viewport/DisplayClusterViewportManagerViewPointExtension.cpp
+++ b/Engine/Plugins/Runtime/nDisplay/Source/DisplayCluster/Private/Render/Viewport/DisplayClusterViewportManagerViewPointExtension.cpp
@@ -26,8 +26,9 @@ bool FDisplayClusterViewportManagerViewPointExtension::IsActive() const
 
 void FDisplayClusterViewportManagerViewPointExtension::SetupViewPoint(APlayerController* Player, FMinimalViewInfo& InOutViewInfo)
 {
-	if (IDisplayClusterViewport* DCViewport = IsActive() ? Configuration->GetViewportManager()->FindViewport(CurrentStereoViewIndex) : nullptr)
+	uint32 ContextNum = 0;
+	if (IDisplayClusterViewport* DCViewport = IsActive() ? Configuration->GetViewportManager()->FindViewport(CurrentStereoViewIndex, &ContextNum) : nullptr)
 	{
-		DCViewport->SetupViewPoint(InOutViewInfo);
+		DCViewport->SetupViewPoint(ContextNum, InOutViewInfo);
 	}
 }
diff --git a/Engine/Plugins/Runtime/nDisplay/Source/DisplayCluster/Private/Render/Viewport/DisplayClusterViewport_CustomPostProcessSettings.cpp b/Engine/Plugins/Runtime/nDisplay/Source/DisplayCluster/Private/Render/Viewport/DisplayClusterViewport_CustomPostProcessSettings.cpp
--- a/Engine/Plugins/Runtime/nDisplay/Source/DisplayCluster/Private/Render/Viewport/DisplayClusterViewport_CustomPostProcessSettings.cpp
+++ b/Engine/Plugins/Runtime/nDisplay/Source/DisplayCluster/Private/Render/Viewport/DisplayClusterViewport_CustomPostProcessSettings.cpp
@@ -6,14 +6,97 @@
 #include "Render/Viewport/Configuration/DisplayClusterViewportConfigurationHelpers_Postprocess.h"
 #include "DisplayClusterConfigurationTypes_Postprocess.h"
 
-int32 GDisplayClusterPostProcessConfigureForViewport = 1;
-static FAutoConsoleVariableRef CVarDisplayClusterPostProcessConfigureForViewport(
-	TEXT("nDisplay.render.postprocess.ConfigureForViewport"),
-	GDisplayClusterPostProcessConfigureForViewport,
-	TEXT("Enable changes to some postprocessing parameters depending on the viewport context. (DoF, etc.) (0 to disable).\n"),
+// Override post-processing for nDisplay is allowed by default.
+int32 GDisplayClusterPostProcessOverrideEnable = 1;
+static FAutoConsoleVariableRef CVarDisplayClusterPostProcessOverrideEnable(
+	TEXT("nDisplay.render.postprocess.override.enable"),
+	GDisplayClusterPostProcessOverrideEnable,
+	TEXT("Enable postprocess overrides for nDisplay (0 to disable).\n"),
 	ECVF_RenderThreadSafe
 );
 
+// Override post-processing for InCamera viewports is allowed by default.
+int32 GDisplayClusterPostProcessOverrideInCameraVFX = 1;
+static FAutoConsoleVariableRef CVarDisplayClusterPostProcessOverrideInCameraVFX(
+	TEXT("nDisplay.render.postprocess.override.InCameraVFX"),
+	GDisplayClusterPostProcessOverrideInCameraVFX,
+	TEXT("Enable post-processing override for ICVFX Camera viewport (0 to disable).\n"),
+	ECVF_RenderThreadSafe
+);
+
+// By default, post-processing for Outers viewports is disabled.
+// Because of some issues with the depth of field effect.
+int32 GDisplayClusterPostProcessOverrideOutersVFX = 0;
+static FAutoConsoleVariableRef CVarDisplayClusterPostProcessOverrideOutersVFX(
+	TEXT("nDisplay.render.postprocess.override.OutersVFX"),
+	GDisplayClusterPostProcessOverrideOutersVFX,
+	TEXT("Enable postprocess override for ICVFX Outer viewports (0 to disable).\n"),
+	ECVF_RenderThreadSafe
+);
+
+/**
+* Auxiliary functions for post-processing.
+*/
+namespace UE::DisplayClusterViewport::CustomPostProcess
+{
+	/** Overrides the DoF post-processing parameters for the nDisplay viewport. */
+	static inline bool OverrideDepthOfFieldPostProcessSettings(IDisplayClusterViewport* InViewport, const uint32 InContextNum, FPostProcessSettings& InOutPostProcessSettings)
+	{
+		check(InViewport);
+
+		// This math only works for the CineCamera DoF because it provides a valid 'DepthOfFieldSensorFocalLength' value.
+		if (!(InOutPostProcessSettings.DepthOfFieldFstop > 0.0f
+			&& InOutPostProcessSettings.DepthOfFieldFocalDistance > 0.0f))
+		{
+			return false;
+		}
+
+		const TArray<FDisplayClusterViewport_Context>& InViewportContexts = InViewport->GetContexts();
+		if (!InViewportContexts.IsValidIndex(InContextNum))
+		{
+			return false;
+		}
+
+		const FDisplayClusterViewport_Context& ViewportContext = InViewportContexts[InContextNum];
+		if (!(ViewportContext.DepthOfField.SensorFocalLength > 0.f))
+		{
+			return false;
+		}
+
+		const FMatrix& ProjectionMatrix =
+			ViewportContext.ProjectionData.bUseOverscan ?
+			ViewportContext.OverscanProjectionMatrix :
+			ViewportContext.ProjectionMatrix;
+
+		// Ignore if the projection matrix is invalid.
+		if (ProjectionMatrix.M[0][0] == 0.0f || ProjectionMatrix.M[1][1] == 0.0f)
+		{
+			return false;
+		}
+
+		// M00 = 2n/(r-l)
+		// M11 = 2n/(t-b)
+		// => (r-l)/(t-b) = M11/M00 (= "SensorAspectRatio")
+		const double SensorAspectRatio = ProjectionMatrix.M[1][1] / ProjectionMatrix.M[0][0];
+		const double RenderingAspectRatio = double(ViewportContext.RenderTargetRect.Width()) / double(ViewportContext.RenderTargetRect.Height());
+		const double SensorToRenderAspectRatio = SensorAspectRatio / RenderingAspectRatio;
+
+		// Override sensor width such that DoF recovers our desired focal length
+		// 
+		// FocalLength = SensorWidth * M00 / 2
+		// => SensorWidth = 2 * FocalLength / M00
+		//
+		InOutPostProcessSettings.bOverride_DepthOfFieldSensorWidth = true;
+		InOutPostProcessSettings.DepthOfFieldSensorWidth = 2.0 * ViewportContext.DepthOfField.SensorFocalLength / ProjectionMatrix.M[0][0] / FMath::Pow(SensorToRenderAspectRatio, 2);
+
+		// Compensate with squeeze factor the effect of non-square pixels onto bokeh squeeze.
+		InOutPostProcessSettings.bOverride_DepthOfFieldSqueezeFactor = true;
+		InOutPostProcessSettings.DepthOfFieldSqueezeFactor = ViewportContext.DepthOfField.SqueezeFactor * SensorToRenderAspectRatio;
+
+		return true;
+	}
+};
+
 ///////////////////////////////////////////////////////////////////////////////////////
 // FDisplayClusterViewport_CustomPostProcessSettings
 ///////////////////////////////////////////////////////////////////////////////////////
@@ -115,12 +198,33 @@ bool FDisplayClusterViewport_CustomPostProcessSettings::ApplyCustomPostProcess(I
 
 bool FDisplayClusterViewport_CustomPostProcessSettings::ConfigurePostProcessSettingsForViewport(IDisplayClusterViewport* InViewport, const uint32 InContextNum, const ERenderPass InRenderPass, FPostProcessSettings& InOutPostProcessSettings) const
 {
-	if (!InViewport || !GDisplayClusterPostProcessConfigureForViewport)
+	using namespace UE::DisplayClusterViewport::CustomPostProcess;
+
+	if (!InViewport || !GDisplayClusterPostProcessOverrideEnable)
 	{
 		return false;
 	}
 
-	// Todo: Updates DoF and other PP settings for the current viewport.(JIRAs UE-219457,UE-219466)
+	// Ignore ICVFX cameras
+	if (!GDisplayClusterPostProcessOverrideInCameraVFX && EnumHasAnyFlags(InViewport->GetRenderSettingsICVFX().RuntimeFlags, EDisplayClusterViewportRuntimeICVFXFlags::InCamera))
+	{
+		return false;
+	}
+
+	// Ignore Outers for ICVFX
+	if (!GDisplayClusterPostProcessOverrideOutersVFX && EnumHasAnyFlags(InViewport->GetRenderSettingsICVFX().RuntimeFlags, EDisplayClusterViewportRuntimeICVFXFlags::Target))
+	{
+		return false;
+	}
+
+	// If there are any changes to the PP settings, return true.
+	bool bPostProcessSettingsHaveChanged = false;
+
+	if (InRenderPass == IDisplayClusterViewport_CustomPostProcessSettings::ERenderPass::Override)
+	{
+		// Updates the CineCamera DoF PP settings for the current viewport.
+		bPostProcessSettingsHaveChanged |= OverrideDepthOfFieldPostProcessSettings(InViewport, InContextNum, InOutPostProcessSettings);
+	}
 
 	return false;
 }
diff --git a/Engine/Plugins/Runtime/nDisplay/Source/DisplayCluster/Private/Render/Viewport/DisplayClusterViewport_ViewPoint.cpp b/Engine/Plugins/Runtime/nDisplay/Source/DisplayCluster/Private/Render/Viewport/DisplayClusterViewport_ViewPoint.cpp
--- a/Engine/Plugins/Runtime/nDisplay/Source/DisplayCluster/Private/Render/Viewport/DisplayClusterViewport_ViewPoint.cpp
+++ b/Engine/Plugins/Runtime/nDisplay/Source/DisplayCluster/Private/Render/Viewport/DisplayClusterViewport_ViewPoint.cpp
@@ -166,7 +166,7 @@ UDisplayClusterCameraComponent* FDisplayClusterViewport::GetViewPointCameraCompo
 	return CameraId.IsEmpty() ? nullptr : RootActor->GetDefaultCamera();
 }
 
-bool FDisplayClusterViewport::SetupViewPoint(FMinimalViewInfo& InOutViewInfo)
+bool FDisplayClusterViewport::SetupViewPoint(const uint32 InContextNum, FMinimalViewInfo& InOutViewInfo)
 {
 	if (UDisplayClusterCameraComponent* SceneCameraComponent = GetViewPointCameraComponent(EDisplayClusterRootActorType::Scene))
 	{
@@ -179,6 +179,30 @@ bool FDisplayClusterViewport::SetupViewPoint(FMinimalViewInfo& InOutViewInfo)
 			ProjectionPolicy->SetupProjectionViewPoint(this, Configuration->GetRootActorWorldDeltaSeconds(), InOutViewInfo, &CustomNearClippingPlane);
 		}
 
+		// Save additional data at this point:
+		if (Contexts.IsValidIndex(InContextNum))
+		{
+			FDisplayClusterViewport_Context& DestContext = Contexts[InContextNum];
+
+			// DoF FocalLength
+			if (InOutViewInfo.PostProcessSettings.DepthOfFieldFstop > 0.0f && InOutViewInfo.PostProcessSettings.DepthOfFieldFocalDistance > 0.0f)
+			{
+				// Convert FOV to focal length,
+				// 
+				// fov = 2 * atan(d/(2*f))
+				// where,
+				//   d = sensor dimension (APS-C 24.576 mm)
+				//   f = focal length
+				// 
+				// f = 0.5 * d * (1/tan(fov/2))
+				const FMatrix ProjectionMatrix = InOutViewInfo.CalculateProjectionMatrix();
+				DestContext.DepthOfField.SensorFocalLength = 0.5f * InOutViewInfo.PostProcessSettings.DepthOfFieldSensorWidth * ProjectionMatrix.M[0][0];
+
+				// Save actual squeeze factor value
+				DestContext.DepthOfField.SqueezeFactor = FMath::Clamp(InOutViewInfo.PostProcessSettings.DepthOfFieldSqueezeFactor, 1.0f, 2.0f);
+			}
+		}
+
 		return true;
 	}
 
diff --git a/Engine/Plugins/Runtime/nDisplay/Source/DisplayCluster/Private/Render/Viewport/Preview/DisplayClusterViewportPreview_Render.cpp b/Engine/Plugins/Runtime/nDisplay/Source/DisplayCluster/Private/Render/Viewport/Preview/DisplayClusterViewportPreview_Render.cpp
--- a/Engine/Plugins/Runtime/nDisplay/Source/DisplayCluster/Private/Render/Viewport/Preview/DisplayClusterViewportPreview_Render.cpp
+++ b/Engine/Plugins/Runtime/nDisplay/Source/DisplayCluster/Private/Render/Viewport/Preview/DisplayClusterViewportPreview_Render.cpp
@@ -159,7 +159,7 @@ bool FDisplayClusterViewportPreview::CalculateStereoViewOffset(FDisplayClusterVi
 
 	// Obtaining the internal viewpoint for a given viewport with stereo eye offset distance.
 	FMinimalViewInfo ViewInfo;
-	if (!InViewport.SetupViewPoint(ViewInfo))
+	if (!InViewport.SetupViewPoint(InContextNum, ViewInfo))
 	{
 		return false;
 	}
diff --git a/Engine/Plugins/Runtime/nDisplay/Source/DisplayCluster/Public/Render/Viewport/Containers/DisplayClusterViewport_Context.h b/Engine/Plugins/Runtime/nDisplay/Source/DisplayCluster/Public/Render/Viewport/Containers/DisplayClusterViewport_Context.h
--- a/Engine/Plugins/Runtime/nDisplay/Source/DisplayCluster/Public/Render/Viewport/Containers/DisplayClusterViewport_Context.h
+++ b/Engine/Plugins/Runtime/nDisplay/Source/DisplayCluster/Public/Render/Viewport/Containers/DisplayClusterViewport_Context.h
@@ -64,6 +64,17 @@ public:
 	// This values updated from function FDisplayClusterViewport::CalculateProjectionMatrix()
 	FCachedProjectionData ProjectionData;
 
+	/** Additional data for the Depth of Field (DoF). */
+	struct FDepthOfFieldSettings
+	{
+		// Focal length of the Depth of Field effect camera in mm.
+		float SensorFocalLength = 0.f;
+
+		// This is the squeeze factor for the DOF, which emulates the properties of anamorphic lenses.
+		float SqueezeFactor = 1.f;
+
+	} DepthOfField;
+
 	// World scale
 	float WorldToMeters = 100.f;
 
diff --git a/Engine/Plugins/Runtime/nDisplay/Source/DisplayCluster/Public/Render/Viewport/IDisplayClusterViewport.h b/Engine/Plugins/Runtime/nDisplay/Source/DisplayCluster/Public/Render/Viewport/IDisplayClusterViewport.h
--- a/Engine/Plugins/Runtime/nDisplay/Source/DisplayCluster/Public/Render/Viewport/IDisplayClusterViewport.h
+++ b/Engine/Plugins/Runtime/nDisplay/Source/DisplayCluster/Public/Render/Viewport/IDisplayClusterViewport.h
@@ -101,7 +101,7 @@ public:
 	 * 
 	 * @return - true if there is an internal viewpoint for the given viewport.
 	 */
-	virtual bool SetupViewPoint(struct FMinimalViewInfo& InOutViewInfo) = 0;
+	virtual bool SetupViewPoint(const uint32 InContextNum, struct FMinimalViewInfo& InOutViewInfo) = 0;
 
 	/** Return view point camera component for this viewport.
 	* @param InRootActorType - the root type
@@ -273,4 +273,17 @@ public:
 	{
 		return false;
 	}
+
+	/** Setup viewpoint for this viewport
+	 *
+	 * @param InOutViewInfo - [in\out] viewinfo
+	 *
+	 * @return - true if there is an internal viewpoint for the given viewport.
+	 */
+	UE_DEPRECATED(5.5, "This function has been deprecated. Please use 'SetupViewPoint()'.")
+	virtual bool SetupViewPoint(struct FMinimalViewInfo& InOutViewInfo)
+	{
+		return false;
+	}
+
 };
diff --git a/Engine/Plugins/Runtime/nDisplay/Source/DisplayClusterMoviePipeline/Private/DisplayClusterMoviePipelineViewportPass.cpp b/Engine/Plugins/Runtime/nDisplay/Source/DisplayClusterMoviePipeline/Private/DisplayClusterMoviePipelineViewportPass.cpp
--- a/Engine/Plugins/Runtime/nDisplay/Source/DisplayClusterMoviePipeline/Private/DisplayClusterMoviePipelineViewportPass.cpp
+++ b/Engine/Plugins/Runtime/nDisplay/Source/DisplayClusterMoviePipeline/Private/DisplayClusterMoviePipelineViewportPass.cpp
@@ -566,7 +566,7 @@ IDisplayClusterViewport* UDisplayClusterMoviePipelineViewportPassBase::GetAndCal
 		{
 			// Obtaining the internal viewpoint for a given viewport with stereo eye offset distance.
 			FMinimalViewInfo ViewInfo;
-			if (!DCViewport->SetupViewPoint(ViewInfo))
+			if (!DCViewport->SetupViewPoint(InContextNum, ViewInfo))
 			{
 				return nullptr;
 			}
diff --git a/Engine/Source/Runtime/Renderer/Private/PostProcess/DiaphragmDOFUtils.cpp b/Engine/Source/Runtime/Renderer/Private/PostProcess/DiaphragmDOFUtils.cpp
--- a/Engine/Source/Runtime/Renderer/Private/PostProcess/DiaphragmDOFUtils.cpp
+++ b/Engine/Source/Runtime/Renderer/Private/PostProcess/DiaphragmDOFUtils.cpp
@@ -107,7 +107,7 @@ void DiaphragmDOF::FPhysicalCocModel::Compile(const FViewInfo& View)
 
 		FocusDistance = View.FinalPostProcessSettings.DepthOfFieldFocalDistance;
 		FStops = View.FinalPostProcessSettings.DepthOfFieldFstop;
-		Squeeze = FMath::Clamp(View.FinalPostProcessSettings.DepthOfFieldSqueezeFactor, 1.0f, 2.0f);
+		Squeeze = FMath::Clamp(View.FinalPostProcessSettings.DepthOfFieldSqueezeFactor, 0.1f, 10.0f);
 
 		RenderingAspectRatio = float(View.UnscaledViewRect.Width()) / float(View.UnscaledViewRect.Height());
 		const float HorizontalHalfFOV = FMath::Atan(1.0f / View.ViewMatrices.GetProjectionMatrix().M[0][0]);
