Fix minor shader warnings in nDisplay

#rnx

diff --git a/Engine/Plugins/Runtime/nDisplay/Shaders/Private/MPCDIOverlayShaders.usf b/Engine/Plugins/Runtime/nDisplay/Shaders/Private/MPCDIOverlayShaders.usf
--- a/Engine/Plugins/Runtime/nDisplay/Shaders/Private/MPCDIOverlayShaders.usf
+++ b/Engine/Plugins/Runtime/nDisplay/Shaders/Private/MPCDIOverlayShaders.usf
@@ -150,7 +150,7 @@ float3 ComposeLightCardUnderColor(float4 InColor, OutputVS IN, float2 OverlayUV)
 		return ComposeOverlayColor(InColor, LightCardUnderColor);
 	}
 	
-	return InColor;
+	return InColor.xyz;
 }
 
 float3 ComposeLightCardOverColor(float4 InColor, OutputVS IN, float2 OverlayUV)
@@ -162,7 +162,7 @@ float3 ComposeLightCardOverColor(float4 InColor, OutputVS IN, float2 OverlayUV)
 		return ComposeOverlayColor(InColor, LightCardOverColor);
 	}
 	
-	return InColor;
+	return InColor.xyz;
 }
 
 ////////////////////////////////////////////////////////////////////////////////
