Display: Removed duplicate DoF PPS parameter modifiers.
Adapting the FOV of the nDisplay viewport to DoF is already done in FDisplayClusterViewport_CustomPostProcessSettings::ConfigurePostProcessSettingsForViewport().

diff --git a/Engine/Plugins/Runtime/nDisplay/Source/DisplayClusterConfiguration/Private/DisplayClusterConfigurationTypes_ICVFX.cpp b/Engine/Plugins/Runtime/nDisplay/Source/DisplayClusterConfiguration/Private/DisplayClusterConfigurationTypes_ICVFX.cpp
--- a/Engine/Plugins/Runtime/nDisplay/Source/DisplayClusterConfiguration/Private/DisplayClusterConfigurationTypes_ICVFX.cpp
+++ b/Engine/Plugins/Runtime/nDisplay/Source/DisplayClusterConfiguration/Private/DisplayClusterConfigurationTypes_ICVFX.cpp
@@ -755,50 +755,7 @@ void FDisplayClusterConfigurationICVFX_CameraDepthOfField::UpdateDynamicCompensa
 void FDisplayClusterConfigurationICVFX_CameraCustomFrustum::SetupViewInfo(const FDisplayClusterConfigurationICVFX_StageSettings& InStageSettings, const FDisplayClusterConfigurationICVFX_CameraSettings& InCameraSettings, FMinimalViewInfo& InOutViewInfo) const
 {
 	// Since Circle of confusion is directly proportional to aperature, with wider FOV focal length needs to be shortened by the same amount as FOV.
-	if (bEnable && InCameraSettings.ExternalCameraActor.IsValid())
-	{
-		// default - percents
-		const float ConvertToPercent = 0.01f;
-		const float MaxPercentOverscan = .5f;
-		const float MaxPixelOverscan = 5.f;
-
-		float LocLeft = Left;
-		float LocRight = Right;
-		float LocTop = Top;
-		float LocBottom = Bottom;
-
-		if (Mode == EDisplayClusterConfigurationViewportCustomFrustumMode::Pixels)
-		{
-			LocLeft = ClampCustomFrustum(Left * ConvertToPercent, MaxPixelOverscan);
-			LocRight = ClampCustomFrustum(Right * ConvertToPercent, MaxPixelOverscan);
-			LocTop = ClampCustomFrustum(Top * ConvertToPercent, MaxPixelOverscan);
-			LocBottom = ClampCustomFrustum(Bottom * ConvertToPercent, MaxPixelOverscan);
-
-			const float CameraBufferRatio = InCameraSettings.GetCameraBufferRatio(InStageSettings);
-			const FIntPoint FrameSize = InCameraSettings.GetCameraFrameSize(InStageSettings, *InCameraSettings.ExternalCameraActor->GetCineCameraComponent());
-			const float  FrameWidth = FrameSize.X * CameraBufferRatio;
-			const float FrameHeight = FrameSize.Y * CameraBufferRatio;
-
-			LocLeft = LocLeft / FrameWidth;
-			LocRight = LocRight / FrameWidth;
-			LocTop = LocTop / FrameHeight;
-			LocBottom = LocBottom / FrameHeight;
-		}
-
-		LocLeft = ClampCustomFrustum(LocLeft * ConvertToPercent, MaxPercentOverscan);
-		LocRight = ClampCustomFrustum(LocRight * ConvertToPercent, MaxPercentOverscan);
-		LocTop = ClampCustomFrustum(LocTop * ConvertToPercent, MaxPercentOverscan);
-		LocBottom = ClampCustomFrustum(LocBottom * ConvertToPercent, MaxPercentOverscan);
-
-		const float FOVMultiplier = GetCameraFieldOfViewMultiplier(InStageSettings);
-		const float ClampedFieldOfViewMultiplier = (FOVMultiplier > 0.f) ? FOVMultiplier : 1.f;
-		InOutViewInfo.PostProcessSettings.DepthOfFieldMinFstop /= ClampedFieldOfViewMultiplier;
-		InOutViewInfo.PostProcessSettings.DepthOfFieldFstop /= ClampedFieldOfViewMultiplier;
-
-		float Multiplier = (1. + (LocLeft + LocRight) / 2.) * (1. + (LocTop + LocBottom) / 2.);
-		InOutViewInfo.PostProcessSettings.DepthOfFieldMinFstop /= Multiplier;
-		InOutViewInfo.PostProcessSettings.DepthOfFieldFstop /= Multiplier;
-	}
+	// Adapting the FOV of the nDisplay viewport to DoF is already done in FDisplayClusterViewport_CustomPostProcessSettings::ConfigurePostProcessSettingsForViewport().
 }
 
 float FDisplayClusterConfigurationICVFX_CameraCustomFrustum::GetCameraFieldOfViewMultiplier(const FDisplayClusterConfigurationICVFX_StageSettings& InStageSettings) const
