nDisplay - Fix deleting the NDC from a level not updating the ICVFX panel.

diff --git a/Engine/Plugins/Runtime/nDisplay/Source/DisplayClusterOperator/Private/SDisplayClusterOperatorToolbar.cpp b/Engine/Plugins/Runtime/nDisplay/Source/DisplayClusterOperator/Private/SDisplayClusterOperatorToolbar.cpp
--- a/Engine/Plugins/Runtime/nDisplay/Source/DisplayClusterOperator/Private/SDisplayClusterOperatorToolbar.cpp
+++ b/Engine/Plugins/Runtime/nDisplay/Source/DisplayClusterOperator/Private/SDisplayClusterOperatorToolbar.cpp
@@ -12,6 +12,7 @@
 #include "LevelEditorActions.h"
 #include "EditorSupportDelegates.h"
 #include "Framework/MultiBox/MultiBoxBuilder.h"
+#include "Misc/TransactionObjectEvent.h"
 #include "Toolkits/AssetEditorToolkit.h"
 #include "Widgets/Layout/SBorder.h"
 #include "Widgets/Input/STextComboBox.h"
@@ -30,6 +31,11 @@ SDisplayClusterOperatorToolbar::~SDisplayClusterOperatorToolbar()
 		FLevelEditorModule& LevelEditor = FModuleManager::LoadModuleChecked<FLevelEditorModule>("LevelEditor");
 		LevelEditor.OnMapChanged().Remove(MapChangedHandle);
 	}
+
+	if (OnObjectTransactedHandle.IsValid())
+	{
+		FCoreUObjectDelegates::OnObjectTransacted.Remove(OnObjectTransactedHandle);
+	}
 	
 	FEditorSupportDelegates::PrepareToCleanseEditorObject.RemoveAll(this);
 
@@ -116,6 +122,8 @@ void SDisplayClusterOperatorToolbar::Construct(const FArguments& InArgs)
 	MapChangedHandle = LevelEditor.OnMapChanged().AddRaw(this, &SDisplayClusterOperatorToolbar::HandleMapChanged);
 
 	FEditorSupportDelegates::PrepareToCleanseEditorObject.AddSP(this, &SDisplayClusterOperatorToolbar::OnPrepareToCleanseEditorObject);
+
+	OnObjectTransactedHandle = FCoreUObjectDelegates::OnObjectTransacted.AddSP(this, &SDisplayClusterOperatorToolbar::OnObjectTransacted);
 }
 
 TSharedPtr<FString> SDisplayClusterOperatorToolbar::FillRootActorList(const FString& InitiallySelectedRootActor)
@@ -142,7 +150,7 @@ TSharedPtr<FString> SDisplayClusterOperatorToolbar::FillRootActorList(const FStr
 
 void SDisplayClusterOperatorToolbar::ClearSelectedRootActor()
 {
-	if (ADisplayClusterRootActor* ActiveRootActor = ViewModel->GetRootActor())
+	if (ADisplayClusterRootActor* ActiveRootActor = ViewModel->GetRootActor(/*bEvenIfPendingKill*/ true))
 	{
 		if (UBlueprint* Blueprint = UBlueprint::GetBlueprintFromClass(ActiveRootActor->GetClass()))
 		{
@@ -239,7 +247,7 @@ FText SDisplayClusterOperatorToolbar::GetRootActorComboBoxText() const
 
 void SDisplayClusterOperatorToolbar::OnLevelActorDeleted(AActor* Actor)
 {
-	if (Actor == ViewModel->GetRootActor())
+	if (Actor == ViewModel->GetRootActor(/*bEvenIfPendingKill*/ true))
 	{
 		if (Actor && Actor->GetClass()->HasAnyClassFlags(CLASS_NewerVersionExists))
 		{
@@ -275,4 +283,15 @@ void SDisplayClusterOperatorToolbar::OnPrepareToCleanseEditorObject(UObject* Obj
 	}
 }
 
+void SDisplayClusterOperatorToolbar::OnObjectTransacted(UObject* Object, const FTransactionObjectEvent& TransactionObjectEvent)
+{
+	if (TransactionObjectEvent.GetEventType() == ETransactionObjectEventType::UndoRedo
+		&& TransactionObjectEvent.HasPendingKillChange()
+		&& Object == ViewModel->GetRootActor(/*bEvenIfPendingKill*/ true))
+	{
+		// In this case the user has redone a deletion, and we need to update appropriately.
+		ClearSelectedRootActor();
+	}
+}
+
 #undef LOCTEXT_NAMESPACE
diff --git a/Engine/Plugins/Runtime/nDisplay/Source/DisplayClusterOperator/Private/SDisplayClusterOperatorToolbar.h b/Engine/Plugins/Runtime/nDisplay/Source/DisplayClusterOperator/Private/SDisplayClusterOperatorToolbar.h
--- a/Engine/Plugins/Runtime/nDisplay/Source/DisplayClusterOperator/Private/SDisplayClusterOperatorToolbar.h
+++ b/Engine/Plugins/Runtime/nDisplay/Source/DisplayClusterOperator/Private/SDisplayClusterOperatorToolbar.h
@@ -11,6 +11,7 @@ enum class EMapChangeType : uint8;
 class FUICommandList;
 class IDisplayClusterOperatorViewModel;
 class ADisplayClusterRootActor;
+class FTransactionObjectEvent;
 
 template<class T>
 class SComboBox;
@@ -61,8 +62,11 @@ private:
 	/** When the map changes */
 	void HandleMapChanged(UWorld* InWorld, EMapChangeType InMapChangeType);
 
-	/** Raised when an editor object is going to be cleased, usually used when a sublevel is about to be removed */
+	/** Raised when an editor object is going to be cleansed, usually used when a sublevel is about to be removed */
 	void OnPrepareToCleanseEditorObject(UObject* Object);
+	
+	/** Raised when any object is transacted */
+	void OnObjectTransacted(UObject* Object, const FTransactionObjectEvent& TransactionObjectEvent);
 
 private:
 	/** A reference the the operator panel's view model, which stores the operator panel state */
@@ -85,4 +89,7 @@ private:
 
 	/** The handle to OnMapChanged */
 	FDelegateHandle MapChangedHandle;
+
+	/** Delegate handle for when an object is transacted */
+	FDelegateHandle OnObjectTransactedHandle;
 };
\ No newline at end of file
