[nDisplay] Not placeable DCRA
- Set CLASS_NotPlaceable flag to ADisplayClusterRootActor to forbid adding it manually into the level
- Clean CLASS_NotPlaceable explicitly on the child blueprints so they are able to be placed

#rnx

diff --git a/Engine/Plugins/Runtime/nDisplay/Source/DisplayCluster/Public/DisplayClusterRootActor.h b/Engine/Plugins/Runtime/nDisplay/Source/DisplayCluster/Public/DisplayClusterRootActor.h
--- a/Engine/Plugins/Runtime/nDisplay/Source/DisplayCluster/Public/DisplayClusterRootActor.h
+++ b/Engine/Plugins/Runtime/nDisplay/Source/DisplayCluster/Public/DisplayClusterRootActor.h
@@ -48,7 +48,7 @@ struct FDisplayClusterRootActorPropertyOverrides;
 /**
  * VR root. This contains nDisplay VR hierarchy in the game.
  */
-UCLASS(HideCategories=(Replication, Collision, Input, Actor, HLOD, Cooking, Physics, Activation, AssetUserData, ActorTick, Advanced, DataLayers, Events), meta = (DisplayName = "nDisplay Root Actor"))
+UCLASS(notplaceable, HideCategories=(Replication, Collision, Input, Actor, HLOD, Cooking, Physics, Activation, AssetUserData, ActorTick, Advanced, DataLayers, Events), meta = (DisplayName = "nDisplay Root Actor"))
 class DISPLAYCLUSTER_API ADisplayClusterRootActor
 	: public AActor
 {
diff --git a/Engine/Plugins/Runtime/nDisplay/Source/DisplayClusterConfigurator/Private/DisplayClusterConfiguratorCompiler.cpp b/Engine/Plugins/Runtime/nDisplay/Source/DisplayClusterConfigurator/Private/DisplayClusterConfiguratorCompiler.cpp
--- a/Engine/Plugins/Runtime/nDisplay/Source/DisplayClusterConfigurator/Private/DisplayClusterConfiguratorCompiler.cpp
+++ b/Engine/Plugins/Runtime/nDisplay/Source/DisplayClusterConfigurator/Private/DisplayClusterConfiguratorCompiler.cpp
@@ -92,6 +92,18 @@ void FDisplayClusterConfiguratorKismetCompilerContext::PreCompile()
 	ValidateConfiguration();
 }
 
+void FDisplayClusterConfiguratorKismetCompilerContext::PostCompile()
+{
+	Super::PostCompile();
+
+	UDisplayClusterBlueprint* DCBlueprint = CastChecked<UDisplayClusterBlueprint>(Blueprint);
+	check(DCBlueprint->GeneratedClass);
+
+	// The parent class (DCRA) may have NotPlaceable flag set. This prevents the child blueprints from
+	// placing into a world. To avoid this limitation, we explicitly clear it because we know it's safe.
+	DCBlueprint->GeneratedClass->ClassFlags &= ~CLASS_NotPlaceable;
+}
+
 void FDisplayClusterConfiguratorKismetCompilerContext::SaveSubObjectsFromCleanAndSanitizeClass(
 	FSubobjectCollection& SubObjectsToSave, UBlueprintGeneratedClass* ClassToClean)
 {
diff --git a/Engine/Plugins/Runtime/nDisplay/Source/DisplayClusterConfigurator/Private/DisplayClusterConfiguratorCompiler.h b/Engine/Plugins/Runtime/nDisplay/Source/DisplayClusterConfigurator/Private/DisplayClusterConfiguratorCompiler.h
--- a/Engine/Plugins/Runtime/nDisplay/Source/DisplayClusterConfigurator/Private/DisplayClusterConfiguratorCompiler.h
+++ b/Engine/Plugins/Runtime/nDisplay/Source/DisplayClusterConfigurator/Private/DisplayClusterConfiguratorCompiler.h
@@ -29,6 +29,7 @@ protected:
 	virtual void SpawnNewClass(const FString& NewClassName) override;
 	virtual void OnNewClassSet(UBlueprintGeneratedClass* ClassToUse) override;
 	virtual void PreCompile() override;
+	virtual void PostCompile() override;
 	virtual void SaveSubObjectsFromCleanAndSanitizeClass(FSubobjectCollection& SubObjectsToSave, UBlueprintGeneratedClass* ClassToClean) override;
 	virtual void CopyTermDefaultsToDefaultObject(UObject* DefaultObject) override;
 	/** ~FKismetCompilerContext interface */
