[nDisplay] Some clean up to the media configuration GUI (properties)

#rnx

diff --git a/Engine/Plugins/Runtime/nDisplay/Source/DisplayClusterConfiguration/Public/DisplayClusterConfigurationStrings.h b/Engine/Plugins/Runtime/nDisplay/Source/DisplayClusterConfiguration/Public/DisplayClusterConfigurationStrings.h
--- a/Engine/Plugins/Runtime/nDisplay/Source/DisplayClusterConfiguration/Public/DisplayClusterConfigurationStrings.h
+++ b/Engine/Plugins/Runtime/nDisplay/Source/DisplayClusterConfiguration/Public/DisplayClusterConfigurationStrings.h
@@ -25,7 +25,8 @@ namespace DisplayClusterConfigurationStrings
 		static constexpr const TCHAR* ClusterConfigurationCategory = TEXT("NDisplay Cluster Configuration");
 		static constexpr const TCHAR* ClusterPostprocessCategory   = TEXT("Post Process");
 		static constexpr const TCHAR* ColorGradingCategory         = TEXT("Color Grading");
-		static constexpr const TCHAR* CameraColorGradingCategory   = TEXT("Inner Frustum Color Grading");
+		static constexpr const TCHAR* CameraColorGradingCategoryOrig   = TEXT("Inner Frustum Color Grading");
+		static constexpr const TCHAR* CameraColorGradingCategory   = TEXT("Color Grading Inner Frustum");
 		static constexpr const TCHAR* ChromaKeyCategory            = TEXT("Chromakey");
 		static constexpr const TCHAR* LightcardCategory            = TEXT("Light Cards");
 		static constexpr const TCHAR* OCIOCategory                 = TEXT("OCIO");
@@ -34,14 +35,19 @@ namespace DisplayClusterConfigurationStrings
 		static constexpr const TCHAR* OverrideCategory             = TEXT("Texture Replacement");
 		static constexpr const TCHAR* ViewportsCategory            = TEXT("Viewports");
 		
-		static constexpr const TCHAR* InCameraVFXCategory  = TEXT("In Camera VFX");
-		static constexpr const TCHAR* InnerFrustumCategory = TEXT("Inner Frustum");
-		static constexpr const TCHAR* ICVFXCameraCategory  = TEXT("ICVFX Camera");
+		static constexpr const TCHAR* InCameraVFXCategory          = TEXT("In Camera VFX");
+		static constexpr const TCHAR* InnerFrustumCategory         = TEXT("Inner Frustum");
+
+		static constexpr const TCHAR* ICVFXCameraCategoryOrig      = TEXT("ICVFX Camera");
+		static constexpr const TCHAR* ICVFXCameraCategory          = TEXT("Camera ICVFX");
 
 		static constexpr const TCHAR* ConfigurationCategory        = TEXT("Configuration");
 		static constexpr const TCHAR* PreviewCategory              = TEXT("Preview");
 		static constexpr const TCHAR* AdvancedCategory             = TEXT("Advanced");
 		static constexpr const TCHAR* TextureShareCategory         = TEXT("Texture Share");
+		static constexpr const TCHAR* NetworkCategory              = TEXT("Network");
+		static constexpr const TCHAR* RenderingCategory            = TEXT("Rendering");
+		static constexpr const TCHAR* StereoCategory               = TEXT("Stereo");
 
 		static constexpr const TCHAR* ViewPointStereoCategory              = TEXT("Stereo");
 		static constexpr const TCHAR* ViewPointCameraPostProcessCategory   = TEXT("Camera Post Process");
diff --git a/Engine/Plugins/Runtime/nDisplay/Source/DisplayClusterConfiguration/Public/DisplayClusterConfigurationTypes.h b/Engine/Plugins/Runtime/nDisplay/Source/DisplayClusterConfiguration/Public/DisplayClusterConfigurationTypes.h
--- a/Engine/Plugins/Runtime/nDisplay/Source/DisplayClusterConfiguration/Public/DisplayClusterConfigurationTypes.h
+++ b/Engine/Plugins/Runtime/nDisplay/Source/DisplayClusterConfiguration/Public/DisplayClusterConfigurationTypes.h
@@ -315,7 +315,7 @@ public:
 
 public:
 	/** IP address of this specific cluster Node */
-	UPROPERTY(EditAnywhere, BlueprintReadonly, Category = "Configuration", meta = (DisplayName = "Host IP Address"))
+	UPROPERTY(EditAnywhere, BlueprintReadonly, Category = "Network", meta = (DisplayName = "Host IP Address"))
 	FString Host;
 	
 	/** Enables or disables sound on nDisplay primary Node */
diff --git a/Engine/Plugins/Runtime/nDisplay/Source/DisplayClusterConfiguration/Public/DisplayClusterConfigurationTypes_Viewport.h b/Engine/Plugins/Runtime/nDisplay/Source/DisplayClusterConfiguration/Public/DisplayClusterConfigurationTypes_Viewport.h
--- a/Engine/Plugins/Runtime/nDisplay/Source/DisplayClusterConfiguration/Public/DisplayClusterConfigurationTypes_Viewport.h
+++ b/Engine/Plugins/Runtime/nDisplay/Source/DisplayClusterConfiguration/Public/DisplayClusterConfigurationTypes_Viewport.h
@@ -113,7 +113,7 @@ public:
 	FDisplayClusterConfigurationViewport_Overscan Overscan;
 
 	// Media settings
-	UPROPERTY(EditAnywhere, BlueprintReadWrite, Category = "Configuration", meta = (DisplayName = "Media"))
+	UPROPERTY(EditAnywhere, BlueprintReadWrite, Category = "Media", meta = (DisplayName = "Media", ShowOnlyInnerProperties))
 	FDisplayClusterConfigurationMediaViewport Media;
 
 	// Experimental: Support special frame builder mode - merge viewports to single viewfamily by group num
diff --git a/Engine/Plugins/Runtime/nDisplay/Source/DisplayClusterConfigurator/Private/DisplayClusterConfiguratorModule.cpp b/Engine/Plugins/Runtime/nDisplay/Source/DisplayClusterConfigurator/Private/DisplayClusterConfiguratorModule.cpp
--- a/Engine/Plugins/Runtime/nDisplay/Source/DisplayClusterConfigurator/Private/DisplayClusterConfiguratorModule.cpp
+++ b/Engine/Plugins/Runtime/nDisplay/Source/DisplayClusterConfigurator/Private/DisplayClusterConfiguratorModule.cpp
@@ -26,7 +26,9 @@
 #include "Views/Details/Cluster/DisplayClusterConfiguratorRectangleCustomization.h"
 #include "Views/Details/Components/DisplayClusterConfiguratorScreenComponentDetailsCustomization.h"
 #include "Views/Details/Components/DisplayClusterICVFXCameraComponentDetailsCustomization.h"
-#include "Views/Details/Media/DisplayClusterConfiguratorICVFXMediaCustomization.h"
+#include "Views/Details/Media/DCConfiguratorClusterNodeMediaCustomization.h"
+#include "Views/Details/Media/DCConfiguratorICVFXMediaCustomization.h"
+#include "Views/Details/Media/DCConfiguratorViewportMediaCustomization.h"
 #include "Views/Details/Media/DisplayClusterConfiguratorMediaFullFrameCustomization.h"
 #include "Views/Details/Media/DisplayClusterConfiguratorMediaTileCustomization.h"
 #include "Views/Details/Components/DisplayClusterCameraComponentDetailsCustomization.h"
@@ -230,7 +232,9 @@ void FDisplayClusterConfiguratorModule::RegisterCustomLayouts()
 	REGISTER_PROPERTY_LAYOUT(FDisplayClusterConfigurationClusterItemReferenceList, FDisplayClusterConfiguratorClusterReferenceListCustomization);
 	REGISTER_PROPERTY_LAYOUT(FDisplayClusterConfigurationViewport_RemapData, FDisplayClusterConfiguratorViewportRemapCustomization);
 	REGISTER_PROPERTY_LAYOUT(FDisplayClusterConfigurationRectangle, FDisplayClusterConfiguratorRectangleCustomization);
-	REGISTER_PROPERTY_LAYOUT(FDisplayClusterConfigurationMediaICVFX, FDisplayClusterConfiguratorICVFXMediaCustomization);
+	REGISTER_PROPERTY_LAYOUT(FDisplayClusterConfigurationMediaICVFX, FDCConfiguratorICVFXMediaCustomization);
+	REGISTER_PROPERTY_LAYOUT(FDisplayClusterConfigurationMediaNodeBackbuffer, FDCConfiguratorClusterNodeMediaCustomization);
+	REGISTER_PROPERTY_LAYOUT(FDisplayClusterConfigurationMediaViewport, FDCConfiguratorViewportMediaCustomization);
 	REGISTER_PROPERTY_LAYOUT(FDisplayClusterConfigurationMediaUniformTileInput,  FDisplayClusterConfiguratorMediaInputTileCustomization);
 	REGISTER_PROPERTY_LAYOUT(FDisplayClusterConfigurationMediaUniformTileOutput, FDisplayClusterConfiguratorMediaOutputTileCustomization);
 	REGISTER_PROPERTY_LAYOUT(FDisplayClusterConfigurationMediaInput,  FDisplayClusterConfiguratorMediaFullFrameInputCustomization);
@@ -302,11 +306,13 @@ void FDisplayClusterConfiguratorModule::RegisterSectionMappings()
 		const TSharedRef<FPropertySection> Section = PropertyModule.FindOrCreateSection(UDisplayClusterICVFXCameraComponent::StaticClass()->GetFName(),
 			DisplayClusterConfigurationStrings::categories::ICVFXCameraCategory, LOCTEXT("InnerFrustumCameraSectionLabel", "Camera"));
 		Section->AddCategory(DisplayClusterConfigurationStrings::categories::ICVFXCameraCategory);
+		Section->AddCategory(DisplayClusterConfigurationStrings::categories::ICVFXCameraCategoryOrig);
 	}
 	{
 		const TSharedRef<FPropertySection> Section = PropertyModule.FindOrCreateSection(UDisplayClusterICVFXCameraComponent::StaticClass()->GetFName(),
 		DisplayClusterConfigurationStrings::categories::CameraColorGradingCategory, LOCTEXT("InnerFrustumColorGradingLabel", "Color Grading"));
 		Section->AddCategory(DisplayClusterConfigurationStrings::categories::CameraColorGradingCategory);
+		Section->AddCategory(DisplayClusterConfigurationStrings::categories::CameraColorGradingCategoryOrig);
 	}
 	{
 		const TSharedRef<FPropertySection> Section = PropertyModule.FindOrCreateSection(UDisplayClusterICVFXCameraComponent::StaticClass()->GetFName(),
@@ -343,6 +349,45 @@ void FDisplayClusterConfiguratorModule::RegisterSectionMappings()
 		Section->AddCategory(DisplayClusterConfigurationStrings::categories::ViewPointInFrustumProjectionCategory);
 	}
 
+	// Cluster node
+	{
+		const TSharedRef<FPropertySection> Section = PropertyModule.FindOrCreateSection(UDisplayClusterConfigurationClusterNode::StaticClass()->GetFName(),
+			DisplayClusterConfigurationStrings::categories::NetworkCategory, LOCTEXT("ClusterNodeNetworkSectionLabel", "Network"));
+		Section->AddCategory(DisplayClusterConfigurationStrings::categories::NetworkCategory);
+	}
+	{
+		const TSharedRef<FPropertySection> Section = PropertyModule.FindOrCreateSection(UDisplayClusterConfigurationClusterNode::StaticClass()->GetFName(),
+			DisplayClusterConfigurationStrings::categories::ConfigurationCategory, LOCTEXT("ClusterNodeConfigurationSectionLabel", "Configuration"));
+		Section->AddCategory(DisplayClusterConfigurationStrings::categories::ConfigurationCategory);
+	}
+	{
+		const TSharedRef<FPropertySection> Section = PropertyModule.FindOrCreateSection(UDisplayClusterConfigurationClusterNode::StaticClass()->GetFName(),
+			DisplayClusterConfigurationStrings::categories::MediaCategory, LOCTEXT("ClusterNodeMediaSectionLabel", "Media"));
+		Section->AddCategory(DisplayClusterConfigurationStrings::categories::MediaCategory);
+	}
+
+	// Viewport
+	{
+		const TSharedRef<FPropertySection> Section = PropertyModule.FindOrCreateSection(UDisplayClusterConfigurationViewport::StaticClass()->GetFName(),
+			DisplayClusterConfigurationStrings::categories::ConfigurationCategory, LOCTEXT("ViewportConfigurationSectionLabel", "Configuration"));
+		Section->AddCategory(DisplayClusterConfigurationStrings::categories::ConfigurationCategory);
+	}
+	{
+		const TSharedRef<FPropertySection> Section = PropertyModule.FindOrCreateSection(UDisplayClusterConfigurationViewport::StaticClass()->GetFName(),
+			DisplayClusterConfigurationStrings::categories::MediaCategory, LOCTEXT("ViewportMediaSectionLabel", "Media"));
+		Section->AddCategory(DisplayClusterConfigurationStrings::categories::MediaCategory);
+	}
+	{
+		const TSharedRef<FPropertySection> Section = PropertyModule.FindOrCreateSection(UDisplayClusterConfigurationViewport::StaticClass()->GetFName(),
+			DisplayClusterConfigurationStrings::categories::PreviewCategory, LOCTEXT("ViewportPreviewSectionLabel", "Preview"));
+		Section->AddCategory(DisplayClusterConfigurationStrings::categories::PreviewCategory);
+	}
+	{
+		const TSharedRef<FPropertySection> Section = PropertyModule.FindOrCreateSection(UDisplayClusterConfigurationViewport::StaticClass()->GetFName(),
+			DisplayClusterConfigurationStrings::categories::RenderingCategory, LOCTEXT("ViewportRenderingSectionLabel", "Rendering"));
+		Section->AddCategory(DisplayClusterConfigurationStrings::categories::RenderingCategory);
+		Section->AddCategory(DisplayClusterConfigurationStrings::categories::StereoCategory);
+	}
 }
 
 void FDisplayClusterConfiguratorModule::UnregisterSectionMappings()
@@ -350,16 +395,40 @@ void FDisplayClusterConfiguratorModule::UnregisterSectionMappings()
 	if (FModuleManager::Get().IsModuleLoaded("PropertyEditor") && FSlateApplication::IsInitialized())
 	{
 		FPropertyEditorModule& PropertyModule = FModuleManager::GetModuleChecked<FPropertyEditorModule>("PropertyEditor");
+
+		// DCRA
 		PropertyModule.RemoveSection(ADisplayClusterRootActor::StaticClass()->GetFName(), DisplayClusterConfigurationStrings::categories::ViewportsCategory);
 		PropertyModule.RemoveSection(ADisplayClusterRootActor::StaticClass()->GetFName(), DisplayClusterConfigurationStrings::categories::InCameraVFXCategory);
 		PropertyModule.RemoveSection(ADisplayClusterRootActor::StaticClass()->GetFName(), DisplayClusterConfigurationStrings::categories::ColorGradingCategory);
 		PropertyModule.RemoveSection(ADisplayClusterRootActor::StaticClass()->GetFName(), DisplayClusterConfigurationStrings::categories::OCIOCategory);
 		PropertyModule.RemoveSection(ADisplayClusterRootActor::StaticClass()->GetFName(), DisplayClusterConfigurationStrings::categories::LightcardCategory);
+		PropertyModule.RemoveSection(ADisplayClusterRootActor::StaticClass()->GetFName(), DisplayClusterConfigurationStrings::categories::PreviewCategory);
 
+		// ICVFX camera component
 		PropertyModule.RemoveSection(UDisplayClusterICVFXCameraComponent::StaticClass()->GetFName(), DisplayClusterConfigurationStrings::categories::InnerFrustumCategory);
+		PropertyModule.RemoveSection(UDisplayClusterICVFXCameraComponent::StaticClass()->GetFName(), DisplayClusterConfigurationStrings::categories::ICVFXCameraCategory);
 		PropertyModule.RemoveSection(UDisplayClusterICVFXCameraComponent::StaticClass()->GetFName(), DisplayClusterConfigurationStrings::categories::CameraColorGradingCategory);
 		PropertyModule.RemoveSection(UDisplayClusterICVFXCameraComponent::StaticClass()->GetFName(), DisplayClusterConfigurationStrings::categories::OCIOCategory);
+		PropertyModule.RemoveSection(UDisplayClusterICVFXCameraComponent::StaticClass()->GetFName(), DisplayClusterConfigurationStrings::categories::MediaCategory);
 		PropertyModule.RemoveSection(UDisplayClusterICVFXCameraComponent::StaticClass()->GetFName(), DisplayClusterConfigurationStrings::categories::ChromaKeyCategory);
+
+		// ViewPoint component
+		PropertyModule.RemoveSection(UDisplayClusterCameraComponent::StaticClass()->GetFName(), DisplayClusterConfigurationStrings::categories::ViewPointStereoCategory);
+		PropertyModule.RemoveSection(UDisplayClusterCameraComponent::StaticClass()->GetFName(), DisplayClusterConfigurationStrings::categories::ViewPointCameraPostProcessCategory);
+
+		// InFrustum ViewPoint
+		PropertyModule.RemoveSection(UDisplayClusterInFrustumFitCameraComponent::StaticClass()->GetFName(), DisplayClusterConfigurationStrings::categories::ViewPointInFrustumProjectionCategory);
+
+		// Cluster node
+		PropertyModule.RemoveSection(UDisplayClusterConfigurationClusterNode::StaticClass()->GetFName(), DisplayClusterConfigurationStrings::categories::NetworkCategory);
+		PropertyModule.RemoveSection(UDisplayClusterConfigurationClusterNode::StaticClass()->GetFName(), DisplayClusterConfigurationStrings::categories::ConfigurationCategory);
+		PropertyModule.RemoveSection(UDisplayClusterConfigurationClusterNode::StaticClass()->GetFName(), DisplayClusterConfigurationStrings::categories::MediaCategory);
+
+		// Viewport
+		PropertyModule.RemoveSection(UDisplayClusterConfigurationViewport::StaticClass()->GetFName(), DisplayClusterConfigurationStrings::categories::ConfigurationCategory);
+		PropertyModule.RemoveSection(UDisplayClusterConfigurationViewport::StaticClass()->GetFName(), DisplayClusterConfigurationStrings::categories::MediaCategory);
+		PropertyModule.RemoveSection(UDisplayClusterConfigurationViewport::StaticClass()->GetFName(), DisplayClusterConfigurationStrings::categories::PreviewCategory);
+		PropertyModule.RemoveSection(UDisplayClusterConfigurationViewport::StaticClass()->GetFName(), DisplayClusterConfigurationStrings::categories::RenderingCategory);
 	}
 }
 
diff --git a/Engine/Plugins/Runtime/nDisplay/Source/DisplayClusterConfigurator/Private/Views/Details/Media/DCConfiguratorBaseMediaCustomization.cpp b/Engine/Plugins/Runtime/nDisplay/Source/DisplayClusterConfigurator/Private/Views/Details/Media/DCConfiguratorBaseMediaCustomization.cpp
new file mode 100644
--- /dev/null
+++ b/Engine/Plugins/Runtime/nDisplay/Source/DisplayClusterConfigurator/Private/Views/Details/Media/DCConfiguratorBaseMediaCustomization.cpp
@@ -0,0 +1,73 @@
+// Copyright Epic Games, Inc. All Rights Reserved.
+
+#include "Views/Details/Media/DCConfiguratorBaseMediaCustomization.h"
+#include "Views/Details/Media/DisplayClusterConfiguratorMediaUtils.h"
+
+#include "Widgets/Input/SButton.h"
+#include "Widgets/Layout/SBox.h"
+#include "Widgets/Text/STextBlock.h"
+
+#include "DetailWidgetRow.h"
+#include "DisplayClusterConfiguratorBlueprintEditor.h"
+#include "IDetailChildrenBuilder.h"
+#include "PropertyHandle.h"
+#include "ScopedTransaction.h"
+
+#define LOCTEXT_NAMESPACE "FDCConfiguratorBaseMediaCustomization"
+
+
+void FDCConfiguratorBaseMediaCustomization::AddResetButton(IDetailChildrenBuilder& InChildBuilder, const FText& ButtonText)
+{
+	InChildBuilder.AddCustomRow(FText::GetEmpty())
+		.WholeRowContent()
+		[
+			SNew(SBox)
+			.Padding(5.f)
+			[
+				SNew(SButton)
+				.HAlign(HAlign_Center)
+				.OnClicked(this, &FDCConfiguratorBaseMediaCustomization::OnResetButtonClicked)
+				[
+					SNew(STextBlock)
+					.Font(FAppStyle::GetFontStyle(TEXT("PropertyWindow.NormalFont")))
+					.Text(ButtonText)
+				]
+			]
+		];
+}
+
+FReply FDCConfiguratorBaseMediaCustomization::OnResetButtonClicked()
+{
+	if (EditingObject.IsValid())
+	{
+		FScopedTransaction Transaction(LOCTEXT("ResetMediaSettings", "Reset Media Settings"));
+		EditingObject->Modify();
+
+		// Notify tile customizers to re-initialize their media objects
+		FDisplayClusterConfiguratorMediaUtils::Get().OnMediaResetToDefaults().Broadcast(EditingObject.Get());
+
+		// Set owning package dirty
+		MarkDirty();
+	}
+
+	return FReply::Handled();
+}
+
+void FDCConfiguratorBaseMediaCustomization::MarkDirty()
+{
+	if (EditingObject.IsValid())
+	{
+		// Blueprint
+		if (EditingObject->IsInBlueprint())
+		{
+			ModifyBlueprint();
+		}
+		// Instance
+		else
+		{
+			EditingObject->MarkPackageDirty();
+		}
+	}
+}
+
+#undef LOCTEXT_NAMESPACE
diff --git a/Engine/Plugins/Runtime/nDisplay/Source/DisplayClusterConfigurator/Private/Views/Details/Media/DCConfiguratorBaseMediaCustomization.h b/Engine/Plugins/Runtime/nDisplay/Source/DisplayClusterConfigurator/Private/Views/Details/Media/DCConfiguratorBaseMediaCustomization.h
new file mode 100644
--- /dev/null
+++ b/Engine/Plugins/Runtime/nDisplay/Source/DisplayClusterConfigurator/Private/Views/Details/Media/DCConfiguratorBaseMediaCustomization.h
@@ -0,0 +1,32 @@
+// Copyright Epic Games, Inc. All Rights Reserved.
+
+#pragma once
+
+#include "Views/Details/DisplayClusterConfiguratorBaseTypeCustomization.h"
+
+#include "Input/Reply.h"
+
+
+/**
+ * Base customization class. Provides some common functionality for concrete implementations.
+ */
+class FDCConfiguratorBaseMediaCustomization
+	: public FDisplayClusterConfiguratorBaseTypeCustomization
+{
+public:
+
+	/** Abstract */
+	virtual ~FDCConfiguratorBaseMediaCustomization() = 0
+	{ }
+
+protected:
+
+	/** Builds reset button widget. */
+	void AddResetButton(IDetailChildrenBuilder& ChildBuilder, const FText& ButtonText);
+
+	/** Handles reset button clicks. */
+	FReply OnResetButtonClicked();
+
+	/** Marks package as dirty */
+	void MarkDirty();
+};
diff --git a/Engine/Plugins/Runtime/nDisplay/Source/DisplayClusterConfigurator/Private/Views/Details/Media/DCConfiguratorClusterNodeMediaCustomization.cpp b/Engine/Plugins/Runtime/nDisplay/Source/DisplayClusterConfigurator/Private/Views/Details/Media/DCConfiguratorClusterNodeMediaCustomization.cpp
new file mode 100644
--- /dev/null
+++ b/Engine/Plugins/Runtime/nDisplay/Source/DisplayClusterConfigurator/Private/Views/Details/Media/DCConfiguratorClusterNodeMediaCustomization.cpp
@@ -0,0 +1,21 @@
+// Copyright Epic Games, Inc. All Rights Reserved.
+
+#include "Views/Details/Media/DCConfiguratorClusterNodeMediaCustomization.h"
+
+#include "IDetailChildrenBuilder.h"
+#include "IPropertyUtilities.h"
+#include "PropertyHandle.h"
+
+#define LOCTEXT_NAMESPACE "FDCConfiguratorClusterNodeMediaCustomization"
+
+
+void FDCConfiguratorClusterNodeMediaCustomization::CustomizeChildren(TSharedRef<IPropertyHandle> InPropertyHandle, IDetailChildrenBuilder& InChildBuilder, IPropertyTypeCustomizationUtils& InCustomizationUtils)
+{
+	// Build all
+	Super::CustomizeChildren(InPropertyHandle, InChildBuilder, InCustomizationUtils);
+
+	// Create 'reset' button at the bottom
+	AddResetButton(InChildBuilder, LOCTEXT("ResetToDefaultButtonTitle", "Reset Media Output to Default"));
+}
+
+#undef LOCTEXT_NAMESPACE
diff --git a/Engine/Plugins/Runtime/nDisplay/Source/DisplayClusterConfigurator/Private/Views/Details/Media/DCConfiguratorClusterNodeMediaCustomization.h b/Engine/Plugins/Runtime/nDisplay/Source/DisplayClusterConfigurator/Private/Views/Details/Media/DCConfiguratorClusterNodeMediaCustomization.h
new file mode 100644
--- /dev/null
+++ b/Engine/Plugins/Runtime/nDisplay/Source/DisplayClusterConfigurator/Private/Views/Details/Media/DCConfiguratorClusterNodeMediaCustomization.h
@@ -0,0 +1,30 @@
+// Copyright Epic Games, Inc. All Rights Reserved.
+
+#pragma once
+
+#include "Views/Details/Media/DCConfiguratorBaseMediaCustomization.h"
+
+
+/**
+ * Details panel customization for 'FDisplayClusterConfigurationMediaNodeBackbuffer' struct.
+ */
+class FDCConfiguratorClusterNodeMediaCustomization
+	: public FDCConfiguratorBaseMediaCustomization
+{
+private:
+
+	using Super = FDCConfiguratorBaseMediaCustomization;
+
+public:
+
+	static TSharedRef<IPropertyTypeCustomization> MakeInstance()
+	{
+		return MakeShared<FDCConfiguratorClusterNodeMediaCustomization>();
+	}
+
+protected:
+
+	//~ Begin IPropertyTypeCustomization
+	virtual void CustomizeChildren(TSharedRef<IPropertyHandle> PropertyHandle, IDetailChildrenBuilder& ChildBuilder, IPropertyTypeCustomizationUtils& CustomizationUtils) override;
+	//~ End IPropertyTypeCustomization
+};
diff --git a/Engine/Plugins/Runtime/nDisplay/Source/DisplayClusterConfigurator/Private/Views/Details/Media/DCConfiguratorICVFXMediaCustomization.cpp b/Engine/Plugins/Runtime/nDisplay/Source/DisplayClusterConfigurator/Private/Views/Details/Media/DCConfiguratorICVFXMediaCustomization.cpp
new file mode 100644
--- /dev/null
+++ b/Engine/Plugins/Runtime/nDisplay/Source/DisplayClusterConfigurator/Private/Views/Details/Media/DCConfiguratorICVFXMediaCustomization.cpp
@@ -0,0 +1,235 @@
+// Copyright Epic Games, Inc. All Rights Reserved.
+
+#include "Views/Details/Media/DCConfiguratorICVFXMediaCustomization.h"
+#include "Views/Details/Media/DisplayClusterConfiguratorMediaUtils.h"
+#include "Views/Details/Media/SMediaTilesConfigurationDialog.h"
+
+#include "DisplayClusterConfigurationTypes_Media.h"
+#include "DisplayClusterConfiguratorBlueprintEditor.h"
+#include "DisplayClusterConfiguratorUtils.h"
+#include "DisplayClusterRootActor.h"
+
+#include "Blueprints/DisplayClusterBlueprint.h"
+#include "Components/DisplayClusterICVFXCameraComponent.h"
+
+#include "DetailWidgetRow.h"
+#include "IDetailChildrenBuilder.h"
+#include "IPropertyUtilities.h"
+#include "PropertyHandle.h"
+
+#include "Framework/Application/SlateApplication.h"
+#include "HAL/PlatformApplicationMisc.h"
+#include "Interfaces/IMainFrameModule.h"
+#include "Misc/MessageDialog.h"
+#include "Modules/ModuleManager.h"
+
+#include "Widgets/Input/SButton.h"
+#include "Widgets/Layout/SBox.h"
+#include "Widgets/Text/STextBlock.h"
+
+#include "MediaSource.h"
+#include "MediaOutput.h"
+
+#define LOCTEXT_NAMESPACE "FDCConfiguratorICVFXMediaCustomization"
+
+
+void FDCConfiguratorICVFXMediaCustomization::CustomizeChildren(TSharedRef<IPropertyHandle> InPropertyHandle, IDetailChildrenBuilder& InChildBuilder, IPropertyTypeCustomizationUtils& InCustomizationUtils)
+{
+	// SplitType property
+	TSharedPtr<IPropertyHandle> SplitTypeHandle = GET_CHILD_HANDLE(FDisplayClusterConfigurationMediaICVFX, SplitType);
+	check(SplitTypeHandle && SplitTypeHandle->IsValidHandle());
+
+	// Layout property
+	TSharedPtr<IPropertyHandle> TilesLayoutHandle = GET_CHILD_HANDLE(FDisplayClusterConfigurationMediaICVFX, TiledSplitLayout);
+	check(TilesLayoutHandle && TilesLayoutHandle->IsValidHandle());
+
+	// Separate groups specific for every split type available
+	TArray<TSharedPtr<IPropertyHandle>> FullFramePropertyHandles;
+	TArray<TSharedPtr<IPropertyHandle>> UniformTilePropertyHandles;
+
+	// FullFrame properties
+	FullFramePropertyHandles.Add(GET_CHILD_HANDLE(FDisplayClusterConfigurationMediaICVFX, MediaInputGroups));
+	FullFramePropertyHandles.Add(GET_CHILD_HANDLE(FDisplayClusterConfigurationMediaICVFX, MediaOutputGroups));
+
+	// UniformTile properties
+	UniformTilePropertyHandles.Add(GET_CHILD_HANDLE(FDisplayClusterConfigurationMediaICVFX, TiledSplitLayout));
+	UniformTilePropertyHandles.Add(GET_CHILD_HANDLE(FDisplayClusterConfigurationMediaICVFX, TileOverscan));
+	UniformTilePropertyHandles.Add(GET_CHILD_HANDLE(FDisplayClusterConfigurationMediaICVFX, ClusterNodesToRenderUnboundTiles));
+	UniformTilePropertyHandles.Add(GET_CHILD_HANDLE(FDisplayClusterConfigurationMediaICVFX, TiledMediaInputGroups));
+	UniformTilePropertyHandles.Add(GET_CHILD_HANDLE(FDisplayClusterConfigurationMediaICVFX, TiledMediaOutputGroups));
+
+	// Get current split type
+	EDisplayClusterConfigurationMediaSplitType SplitTypeValue = EDisplayClusterConfigurationMediaSplitType::UniformTiles;
+	SplitTypeHandle->GetValue((uint8&)SplitTypeValue);
+
+	TSharedPtr<IPropertyUtilities> PropertyUtils = InCustomizationUtils.GetPropertyUtilities();
+	check(PropertyUtils);
+
+	// Set details update request on frustum type change
+	SplitTypeHandle->SetOnPropertyValueChanged(
+		FSimpleDelegate::CreateLambda([PropertyUtils]()
+		{
+			if (PropertyUtils)
+			{
+				PropertyUtils->RequestForceRefresh();
+			}
+		}));
+
+	// Filter properties to hide
+	TArray<TArray<TSharedPtr<IPropertyHandle>>*> HiddenPropertyHandles;
+	switch (SplitTypeValue)
+	{
+	case EDisplayClusterConfigurationMediaSplitType::FullFrame:
+		HiddenPropertyHandles.Add(&UniformTilePropertyHandles);
+		break;
+
+	case EDisplayClusterConfigurationMediaSplitType::UniformTiles:
+		HiddenPropertyHandles.Add(&FullFramePropertyHandles);
+		break;
+
+	default:
+		checkNoEntry();
+	}
+
+	// Hide unnecessary properties depending on the frustum (split) type currently selected
+	for (const TArray<TSharedPtr<IPropertyHandle>>* HiddenGroup : HiddenPropertyHandles)
+	{
+		for (const TSharedPtr<IPropertyHandle>& PropertyHandle : *HiddenGroup)
+		{
+			PropertyHandle->MarkHiddenByCustomization();
+		}
+	}
+
+	// Finally, build the panel
+	if (ShouldShowChildren(InPropertyHandle))
+	{
+		const bool bUsingTiles = (SplitTypeValue == EDisplayClusterConfigurationMediaSplitType::UniformTiles);
+
+		uint32 NumChildren = 0;
+		InPropertyHandle->GetNumChildren(NumChildren);
+
+		// For each child property, build its own layout
+		for (uint32 ChildIndex = 0; ChildIndex < NumChildren; ++ChildIndex)
+		{
+			TSharedPtr<IPropertyHandle> ChildHandle = InPropertyHandle->GetChildHandle(ChildIndex);
+			if (ChildHandle && ChildHandle->IsValidHandle() && !ChildHandle->IsCustomized())
+			{
+				FText ChildTooltip = ApplySubstitutions(ChildHandle->GetToolTipText());
+				ChildHandle->SetToolTipText(ChildTooltip);
+
+				InChildBuilder.AddProperty(ChildHandle.ToSharedRef());
+
+				// Insert tile configuration button after the split type combobox
+				if (bUsingTiles && ChildHandle->IsSamePropertyNode(SplitTypeHandle))
+				{
+					AddConfigureTilesButton(InChildBuilder);
+				}
+			}
+		}
+
+		// Create 'reset' button at the bottom
+		AddResetButton(InChildBuilder, LOCTEXT("ResetToDefaultButtonTitle", "Reset Media Input and Output to Default"));
+	}
+}
+
+void FDCConfiguratorICVFXMediaCustomization::AddConfigureTilesButton(IDetailChildrenBuilder& InChildBuilder)
+{
+	InChildBuilder.AddCustomRow(FText::GetEmpty())
+		.WholeRowContent()
+		[
+			SNew(SBox)
+			.Padding(5.f)
+			[
+				SNew(SButton)
+				.HAlign(HAlign_Center)
+				.OnClicked(this, &FDCConfiguratorICVFXMediaCustomization::OnConfigureTilesButtonClicked)
+				[
+					SNew(STextBlock)
+					.Font(FAppStyle::GetFontStyle(TEXT("PropertyWindow.NormalFont")))
+					.Text(LOCTEXT("ConfigureTilesButtonTitle", "Configure Tiles"))
+				]
+			]
+		];
+}
+
+FReply FDCConfiguratorICVFXMediaCustomization::OnConfigureTilesButtonClicked()
+{
+	// We're in camera tiles customization so let's get the camera component
+	UDisplayClusterICVFXCameraComponent* ICVFXCamera = Cast<UDisplayClusterICVFXCameraComponent>(EditingObject.Get());
+	if (!ICVFXCamera)
+	{
+		return FReply::Handled();
+	}
+
+	// Get config data of the DCRA owning the camera being edited
+	const UDisplayClusterConfigurationData* ConfigData = GetConfig();
+	if (!ConfigData)
+	{
+		return FReply::Handled();
+	}
+
+	// Nothing to do if no cluster nodes available
+	if (ConfigData->Cluster->Nodes.Num() < 1)
+	{
+		FMessageDialog::Open(EAppMsgType::Ok, LOCTEXT("MessageNoClusterNodesAvailable", "There are no cluster nodes assigned. Can't configure media."));
+		return FReply::Handled();
+	}
+
+	// Camera's media settings
+	FDisplayClusterConfigurationMediaICVFX& MediaSettings = ICVFXCamera->CameraSettings.RenderSettings.Media;
+
+	// Prepare everything for the configuration dialog
+	FMediaTilesConfigurationDialogParameters Parameters;
+	Parameters.Owner        = ICVFXCamera;
+	Parameters.ConfigData   = ConfigData;
+	Parameters.SplitLayout  = &MediaSettings.TiledSplitLayout;
+	Parameters.InputGroups  = &MediaSettings.TiledMediaInputGroups;
+	Parameters.OutputGroups = &MediaSettings.TiledMediaOutputGroups;
+	Parameters.bAutoPreconfigureOutputMapping = true;
+
+	// Instantiate and show the config dialog
+	const TSharedRef<SMediaTilesConfigurationDialog> TilesConfigurationDialog = SNew(SMediaTilesConfigurationDialog, Parameters);
+	if (GEditor)
+	{
+		GEditor->EditorAddModalWindow(TilesConfigurationDialog);
+	}
+
+	// Process configuration results
+	if (TilesConfigurationDialog->WasConfigurationCompleted())
+	{
+		// Redraw property views
+		if (PropertyUtilities.IsValid())
+		{
+			PropertyUtilities.Pin()->ForceRefresh();
+		}
+
+		// Notify tile customizers to re-initialize all media we just generated
+		FDisplayClusterConfiguratorMediaUtils::Get().OnMediaResetToDefaults().Broadcast(EditingObject.Get());
+
+		// Set owning package dirty
+		MarkDirty();
+	}
+
+	return FReply::Handled();
+}
+
+UDisplayClusterConfigurationData* FDCConfiguratorICVFXMediaCustomization::GetConfig() const
+{
+	if (UDisplayClusterICVFXCameraComponent* ICVFXCameraComponent = Cast<UDisplayClusterICVFXCameraComponent>(EditingObject))
+	{
+		// For instances
+		if (ADisplayClusterRootActor* DCRA = Cast<ADisplayClusterRootActor>(ICVFXCameraComponent->GetOwner()))
+		{
+			return DCRA->GetConfigData();
+		}
+		// For DCRA configurator
+		else if (FDisplayClusterConfiguratorBlueprintEditor* BlueprintEditor = FDisplayClusterConfiguratorUtils::GetBlueprintEditorForObject(ICVFXCameraComponent))
+		{
+			return BlueprintEditor->GetConfig();
+		}
+	}
+
+	return nullptr;
+}
+
+#undef LOCTEXT_NAMESPACE
diff --git a/Engine/Plugins/Runtime/nDisplay/Source/DisplayClusterConfigurator/Private/Views/Details/Media/DCConfiguratorICVFXMediaCustomization.h b/Engine/Plugins/Runtime/nDisplay/Source/DisplayClusterConfigurator/Private/Views/Details/Media/DCConfiguratorICVFXMediaCustomization.h
new file mode 100644
--- /dev/null
+++ b/Engine/Plugins/Runtime/nDisplay/Source/DisplayClusterConfigurator/Private/Views/Details/Media/DCConfiguratorICVFXMediaCustomization.h
@@ -0,0 +1,49 @@
+// Copyright Epic Games, Inc. All Rights Reserved.
+
+#pragma once
+
+#include "Views/Details/Media/DCConfiguratorBaseMediaCustomization.h"
+
+#include "Input/Reply.h"
+
+#include "IPropertyTypeCustomization.h"
+
+class UDisplayClusterConfigurationData;
+
+
+/**
+ * Details panel customization for the FDisplayClusterConfigurationMediaICVFX struct.
+ */
+class FDCConfiguratorICVFXMediaCustomization
+	: public FDCConfiguratorBaseMediaCustomization
+{
+private:
+
+	using Super = FDCConfiguratorBaseMediaCustomization;
+
+public:
+
+	static TSharedRef<IPropertyTypeCustomization> MakeInstance()
+	{
+		return MakeShared<FDCConfiguratorICVFXMediaCustomization>();
+	}
+
+protected:
+
+	//~ Begin IPropertyTypeCustomization
+	virtual void CustomizeChildren(TSharedRef<IPropertyHandle> PropertyHandle, IDetailChildrenBuilder& ChildBuilder, IPropertyTypeCustomizationUtils& CustomizationUtils) override;
+	//~ End IPropertyTypeCustomization
+
+private:
+
+	/** Builds the button widget for tiles configuration. */
+	void AddConfigureTilesButton(IDetailChildrenBuilder& ChildBuilder);
+
+	/** Handles configure tiles button clicks. */
+	FReply OnConfigureTilesButtonClicked();
+
+private:
+
+	/** Returns configuration of a DCRA owning the camera being edited. */
+	UDisplayClusterConfigurationData* GetConfig() const;
+};
diff --git a/Engine/Plugins/Runtime/nDisplay/Source/DisplayClusterConfigurator/Private/Views/Details/Media/DCConfiguratorViewportMediaCustomization.cpp b/Engine/Plugins/Runtime/nDisplay/Source/DisplayClusterConfigurator/Private/Views/Details/Media/DCConfiguratorViewportMediaCustomization.cpp
new file mode 100644
--- /dev/null
+++ b/Engine/Plugins/Runtime/nDisplay/Source/DisplayClusterConfigurator/Private/Views/Details/Media/DCConfiguratorViewportMediaCustomization.cpp
@@ -0,0 +1,20 @@
+// Copyright Epic Games, Inc. All Rights Reserved.
+
+#include "Views/Details/Media/DCConfiguratorViewportMediaCustomization.h"
+
+#include "IDetailChildrenBuilder.h"
+#include "IPropertyUtilities.h"
+#include "PropertyHandle.h"
+
+#define LOCTEXT_NAMESPACE "FDCConfiguratorViewportMediaCustomization"
+
+void FDCConfiguratorViewportMediaCustomization::CustomizeChildren(TSharedRef<IPropertyHandle> InPropertyHandle, IDetailChildrenBuilder& InChildBuilder, IPropertyTypeCustomizationUtils& InCustomizationUtils)
+{
+	// Build all
+	Super::CustomizeChildren(InPropertyHandle, InChildBuilder, InCustomizationUtils);
+
+	// Create 'reset' button at the bottom
+	AddResetButton(InChildBuilder, LOCTEXT("ResetToDefaultButtonTitle", "Reset Media Input and Output to Default"));
+}
+
+#undef LOCTEXT_NAMESPACE
diff --git a/Engine/Plugins/Runtime/nDisplay/Source/DisplayClusterConfigurator/Private/Views/Details/Media/DCConfiguratorViewportMediaCustomization.h b/Engine/Plugins/Runtime/nDisplay/Source/DisplayClusterConfigurator/Private/Views/Details/Media/DCConfiguratorViewportMediaCustomization.h
new file mode 100644
--- /dev/null
+++ b/Engine/Plugins/Runtime/nDisplay/Source/DisplayClusterConfigurator/Private/Views/Details/Media/DCConfiguratorViewportMediaCustomization.h
@@ -0,0 +1,30 @@
+// Copyright Epic Games, Inc. All Rights Reserved.
+
+#pragma once
+
+#include "Views/Details/Media/DCConfiguratorBaseMediaCustomization.h"
+
+
+/**
+ * Details panel customization for 'FDisplayClusterConfigurationMediaViewport' struct.
+ */
+class FDCConfiguratorViewportMediaCustomization
+	: public FDCConfiguratorBaseMediaCustomization
+{
+private:
+
+	using Super = FDCConfiguratorBaseMediaCustomization;
+
+public:
+
+	static TSharedRef<IPropertyTypeCustomization> MakeInstance()
+	{
+		return MakeShared<FDCConfiguratorViewportMediaCustomization>();
+	}
+
+protected:
+
+	//~ Begin IPropertyTypeCustomization
+	virtual void CustomizeChildren(TSharedRef<IPropertyHandle> PropertyHandle, IDetailChildrenBuilder& ChildBuilder, IPropertyTypeCustomizationUtils& CustomizationUtils) override;
+	//~ End IPropertyTypeCustomization
+};
diff --git a/Engine/Plugins/Runtime/nDisplay/Source/DisplayClusterConfigurator/Private/Views/Details/Media/DisplayClusterConfiguratorICVFXMediaCustomization.cpp b/Engine/Plugins/Runtime/nDisplay/Source/DisplayClusterConfigurator/Private/Views/Details/Media/DisplayClusterConfiguratorICVFXMediaCustomization.cpp
deleted file mode 100644
diff --git a/Engine/Plugins/Runtime/nDisplay/Source/DisplayClusterConfigurator/Private/Views/Details/Media/DisplayClusterConfiguratorICVFXMediaCustomization.h b/Engine/Plugins/Runtime/nDisplay/Source/DisplayClusterConfigurator/Private/Views/Details/Media/DisplayClusterConfiguratorICVFXMediaCustomization.h
deleted file mode 100644
diff --git a/Engine/Plugins/Runtime/nDisplay/Source/DisplayClusterConfigurator/Private/Views/Details/Media/DisplayClusterConfiguratorMediaCustomizationBase.cpp b/Engine/Plugins/Runtime/nDisplay/Source/DisplayClusterConfigurator/Private/Views/Details/Media/DisplayClusterConfiguratorMediaCustomizationBase.cpp
--- a/Engine/Plugins/Runtime/nDisplay/Source/DisplayClusterConfigurator/Private/Views/Details/Media/DisplayClusterConfiguratorMediaCustomizationBase.cpp
+++ b/Engine/Plugins/Runtime/nDisplay/Source/DisplayClusterConfigurator/Private/Views/Details/Media/DisplayClusterConfiguratorMediaCustomizationBase.cpp
@@ -15,6 +15,10 @@
 #include "IDetailChildrenBuilder.h"
 #include "IPropertyUtilities.h"
 #include "PropertyHandle.h"
+#include "ScopedTransaction.h"
+
+
+#define LOCTEXT_NAMESPACE "FDisplayClusterConfiguratorMediaFullFrameCustomizationBase"
 
 
 FDisplayClusterConfiguratorMediaFullFrameCustomizationBase::FDisplayClusterConfiguratorMediaFullFrameCustomizationBase()
@@ -81,7 +85,12 @@ void FDisplayClusterConfiguratorMediaFullFrameCustomizationBase::ModifyMediaObje
 	{
 		if (Initializer && Initializer->IsMediaObjectSupported(NewMediaObject))
 		{
+			FScopedTransaction Transaction(LOCTEXT("ResetMediaSettings", "Reset Media Settings"));
+
+			NewMediaObject->Modify();
+
 			PerformMediaInitialization(Owner, NewMediaObject, Initializer);
+
 			break;
 		}
 	}
