[nDisplay] Fixed CIS issues

#rnx

diff --git a/Engine/Plugins/Runtime/nDisplay/Source/DisplayClusterConfigurator/Private/Views/Details/Media/DCConfiguratorBaseMediaCustomization.h b/Engine/Plugins/Runtime/nDisplay/Source/DisplayClusterConfigurator/Private/Views/Details/Media/DCConfiguratorBaseMediaCustomization.h
--- a/Engine/Plugins/Runtime/nDisplay/Source/DisplayClusterConfigurator/Private/Views/Details/Media/DCConfiguratorBaseMediaCustomization.h
+++ b/Engine/Plugins/Runtime/nDisplay/Source/DisplayClusterConfigurator/Private/Views/Details/Media/DCConfiguratorBaseMediaCustomization.h
@@ -16,8 +16,7 @@ class FDCConfiguratorBaseMediaCustomization
 public:
 
 	/** Abstract */
-	virtual ~FDCConfiguratorBaseMediaCustomization() = 0
-	{ }
+	virtual ~FDCConfiguratorBaseMediaCustomization() = default;
 
 protected:
 
