[nDisplay] Fixed CIS issue

- Added missing  #undef LOCTEXT_NAMESPACE

#rnx

diff --git a/Engine/Plugins/Runtime/nDisplay/Source/DisplayClusterConfigurator/Private/Views/Details/Media/DisplayClusterConfiguratorMediaCustomizationBase.cpp b/Engine/Plugins/Runtime/nDisplay/Source/DisplayClusterConfigurator/Private/Views/Details/Media/DisplayClusterConfiguratorMediaCustomizationBase.cpp
--- a/Engine/Plugins/Runtime/nDisplay/Source/DisplayClusterConfigurator/Private/Views/Details/Media/DisplayClusterConfiguratorMediaCustomizationBase.cpp
+++ b/Engine/Plugins/Runtime/nDisplay/Source/DisplayClusterConfigurator/Private/Views/Details/Media/DisplayClusterConfiguratorMediaCustomizationBase.cpp
@@ -338,3 +338,5 @@ void FDisplayClusterConfiguratorMediaTileCustomizationBase::OnTilePositionChange
 {
 	ModifyMediaObjectParameters();
 }
+
+#undef LOCTEXT_NAMESPACE
