[nDisplay] Fixed issue when media is not working in offscreen with 0x0 window size

- Fixed startup command line parameters generation for offscreen nodes
- No window geometry specification required when running offscreen
- Minor clean up to the cluster node configuration data (default values)

#rnx

diff --git a/Engine/Plugins/Editor/DisplayClusterLaunch/Source/DisplayClusterLaunchEditor/Private/DisplayClusterLaunchEditorModule.cpp b/Engine/Plugins/Editor/DisplayClusterLaunch/Source/DisplayClusterLaunchEditor/Private/DisplayClusterLaunchEditorModule.cpp
--- a/Engine/Plugins/Editor/DisplayClusterLaunch/Source/DisplayClusterLaunchEditor/Private/DisplayClusterLaunchEditorModule.cpp
+++ b/Engine/Plugins/Editor/DisplayClusterLaunch/Source/DisplayClusterLaunchEditor/Private/DisplayClusterLaunchEditorModule.cpp
@@ -520,6 +520,10 @@ void FDisplayClusterLaunchEditorModule::LaunchDisplayClusterProcess()
 				{
 					ConcatenatedCommandLineArguments += "-fullscreen ";
 				}
+				else if (NodePtr->bRenderHeadless)
+				{
+					ConcatenatedCommandLineArguments += "-RenderOffscreen ";
+				}
 				else
 				{
 					ConcatenatedCommandLineArguments +=
@@ -530,12 +534,6 @@ void FDisplayClusterLaunchEditorModule::LaunchDisplayClusterProcess()
 						);
 				}
 
-				// Headless
-				if (NodePtr->bRenderHeadless)
-				{
-					ConcatenatedCommandLineArguments += "-RenderOffscreen ";
-				}
-
 				// Graphics adapter
 				if (NodePtr->GraphicsAdapter >= 0)
 				{
diff --git a/Engine/Plugins/Runtime/nDisplay/Source/DisplayClusterConfiguration/Private/DisplayClusterConfigurationTypes.cpp b/Engine/Plugins/Runtime/nDisplay/Source/DisplayClusterConfiguration/Private/DisplayClusterConfigurationTypes.cpp
--- a/Engine/Plugins/Runtime/nDisplay/Source/DisplayClusterConfiguration/Private/DisplayClusterConfigurationTypes.cpp
+++ b/Engine/Plugins/Runtime/nDisplay/Source/DisplayClusterConfiguration/Private/DisplayClusterConfigurationTypes.cpp
@@ -350,15 +350,6 @@ void UDisplayClusterConfigurationClusterNode::Serialize(FArchive& Ar)
 	}
 }
 
-UDisplayClusterConfigurationClusterNode::UDisplayClusterConfigurationClusterNode()
-	: bIsSoundEnabled(false)
-#if WITH_EDITORONLY_DATA
-	, bIsVisible(true)
-	, bIsUnlocked(true)
-#endif
-{
-}
-
 UDisplayClusterConfigurationHostDisplayData::UDisplayClusterConfigurationHostDisplayData()
 	: bIsVisible(true)
 	, bIsUnlocked(true)
diff --git a/Engine/Plugins/Runtime/nDisplay/Source/DisplayClusterConfiguration/Public/DisplayClusterConfigurationTypes.h b/Engine/Plugins/Runtime/nDisplay/Source/DisplayClusterConfiguration/Public/DisplayClusterConfigurationTypes.h
--- a/Engine/Plugins/Runtime/nDisplay/Source/DisplayClusterConfiguration/Public/DisplayClusterConfigurationTypes.h
+++ b/Engine/Plugins/Runtime/nDisplay/Source/DisplayClusterConfiguration/Public/DisplayClusterConfigurationTypes.h
@@ -291,7 +291,6 @@ public:
 	FOnPostEditChangeChainProperty OnPostEditChangeChainProperty;
 
 public:
-	UDisplayClusterConfigurationClusterNode();
 
 	UFUNCTION(BlueprintPure, Category = "NDisplay|Configuration")
 	void GetViewportIds(TArray<FString>& OutViewportIds) const;
@@ -320,24 +319,24 @@ public:
 	
 	/** Enables or disables sound on nDisplay primary Node */
 	UPROPERTY(EditAnywhere, BlueprintReadWrite, Category = "Configuration", meta = (DisplayName = "Enable Sound"))
-	bool bIsSoundEnabled;
+	bool bIsSoundEnabled = false;
 
 	/** Enables application window native fullscreen support */
-	UPROPERTY(EditAnywhere, BlueprintReadonly, Category = "Configuration", meta = (DisplayName = "Fullscreen"))
-	bool bIsFullscreen;
+	UPROPERTY(EditAnywhere, BlueprintReadonly, Category = "Configuration", meta = (DisplayName = "Fullscreen", EditCondition = "!bRenderHeadless"))
+	bool bIsFullscreen = false;
+
+	/** Activates headless rendering for this cluster node */
+	UPROPERTY(EditAnywhere, BlueprintReadonly, Category = "Configuration", meta = (DisplayName = "Headless Rendering", EditCondition = "!bIsFullscreen"))
+	bool bRenderHeadless = false;
 
 	/** Defines the application window size in pixels */
-	UPROPERTY(EditAnywhere, BlueprintReadWrite, Category = "Configuration", meta = (DisplayName = "Window", DisplayMode = "Compound", FixedAspectRatioProperty = "bFixedAspectRatio"))
+	UPROPERTY(EditAnywhere, BlueprintReadWrite, Category = "Configuration", meta = (DisplayName = "Window", DisplayMode = "Compound", FixedAspectRatioProperty = "bFixedAspectRatio", EditCondition = "!bRenderHeadless && !bIsFullscreen"))
 	FDisplayClusterConfigurationRectangle WindowRect;
 
 	/** Output remapping settings for the selected cluster node */
 	UPROPERTY(EditAnywhere, BlueprintReadWrite, Category = "Configuration", meta = (DisplayName = "Output Remapping"))
 	FDisplayClusterConfigurationFramePostProcess_OutputRemap OutputRemap;
 
-	/** Activates headless rendering for this cluster node */
-	UPROPERTY(EditAnywhere, BlueprintReadonly, Category = "Configuration", meta = (DisplayName = "Headless Rendering"))
-	bool bRenderHeadless = false;
-
 	/** Hint for setting the r.GraphicsAdapter CVar when launching this cluster node. Note that this is distinct from the GPU Node Indices assigned to viewports. */
 	UPROPERTY(EditAnywhere, BlueprintReadonly, Category = "Configuration", meta = (DisplayName = "Graphics Adapter"))
 	int32 GraphicsAdapter = -1;
@@ -349,7 +348,7 @@ public:
 #if WITH_EDITORONLY_DATA
 	/** Locks the application window aspect ratio for easier resizing */
 	UPROPERTY(EditAnywhere, Category = "Configuration", meta = (HideProperty))
-	bool bFixedAspectRatio;
+	bool bFixedAspectRatio = false;
 #endif
 
 	UPROPERTY(VisibleInstanceOnly, BlueprintReadOnly, EditFixedSize, Instanced, Category = "Configuration", meta = (DisplayThumbnail = false))
@@ -364,10 +363,10 @@ public:
 
 #if WITH_EDITORONLY_DATA
 	UPROPERTY(EditDefaultsOnly, Category = "Configuration", meta = (HideProperty))
-	bool bIsVisible;
+	bool bIsVisible = true;
 
 	UPROPERTY(EditDefaultsOnly, Category = "Configuration", meta = (HideProperty))
-	bool bIsUnlocked;
+	bool bIsUnlocked = true;
 
 	/** Binds a background preview image for easier output mapping */
 	UPROPERTY(EditDefaultsOnly, Category = "Configuration")
diff --git a/Engine/Plugins/VirtualProduction/Switchboard/Source/Switchboard/switchboard/devices/ndisplay/plugin_ndisplay.py b/Engine/Plugins/VirtualProduction/Switchboard/Source/Switchboard/switchboard/devices/ndisplay/plugin_ndisplay.py
--- a/Engine/Plugins/VirtualProduction/Switchboard/Source/Switchboard/switchboard/devices/ndisplay/plugin_ndisplay.py
+++ b/Engine/Plugins/VirtualProduction/Switchboard/Source/Switchboard/switchboard/devices/ndisplay/plugin_ndisplay.py
@@ -1178,6 +1178,10 @@ class DevicenDisplay(DeviceUnreal):
             args.extend([
                 '-fullscreen',
             ])
+        elif headless:
+            args.extend([
+                '-RenderOffscreen',
+            ])
         else:
             args.extend([
                 '-windowed',
@@ -1188,11 +1192,6 @@ class DevicenDisplay(DeviceUnreal):
                 f'ResY={win_res[1]}',
             ])
 
-        if headless:
-            args.extend([
-                '-RenderOffscreen',
-            ])
-
         # MultiUser parameters
         if CONFIG.MUSERVER_AUTO_JOIN.get_value() and self.autojoin_mu_server.get_value():
             args.extend([
