nDisplay: Fix ConfigurePostProcessSettingsForViewport return when DoF is overridden.

#rnx

diff --git a/Engine/Plugins/Runtime/nDisplay/Source/DisplayCluster/Private/Render/Viewport/DisplayClusterViewport_CustomPostProcessSettings.cpp b/Engine/Plugins/Runtime/nDisplay/Source/DisplayCluster/Private/Render/Viewport/DisplayClusterViewport_CustomPostProcessSettings.cpp
--- a/Engine/Plugins/Runtime/nDisplay/Source/DisplayCluster/Private/Render/Viewport/DisplayClusterViewport_CustomPostProcessSettings.cpp
+++ b/Engine/Plugins/Runtime/nDisplay/Source/DisplayCluster/Private/Render/Viewport/DisplayClusterViewport_CustomPostProcessSettings.cpp
@@ -226,5 +226,5 @@ bool FDisplayClusterViewport_CustomPostProcessSettings::ConfigurePostProcessSett
 		bPostProcessSettingsHaveChanged |= OverrideDepthOfFieldPostProcessSettings(InViewport, InContextNum, InOutPostProcessSettings);
 	}
 
-	return false;
+	return bPostProcessSettingsHaveChanged;
 }
