nDisplay: The logic related with the "nDisplay.render.postprocess.override.*" CVar has been improved to support tiled viewports.

diff --git a/Engine/Plugins/Runtime/nDisplay/Source/DisplayCluster/Private/Render/Viewport/DisplayClusterViewport_CustomPostProcessSettings.cpp b/Engine/Plugins/Runtime/nDisplay/Source/DisplayCluster/Private/Render/Viewport/DisplayClusterViewport_CustomPostProcessSettings.cpp
--- a/Engine/Plugins/Runtime/nDisplay/Source/DisplayCluster/Private/Render/Viewport/DisplayClusterViewport_CustomPostProcessSettings.cpp
+++ b/Engine/Plugins/Runtime/nDisplay/Source/DisplayCluster/Private/Render/Viewport/DisplayClusterViewport_CustomPostProcessSettings.cpp
@@ -2,6 +2,7 @@
 
 #include "Render/Viewport/DisplayClusterViewport_CustomPostProcessSettings.h"
 #include "Render/Viewport/DisplayClusterViewport.h"
+#include "Render/Viewport/DisplayClusterViewportManager.h"
 
 #include "Render/Viewport/Configuration/DisplayClusterViewportConfigurationHelpers_Postprocess.h"
 #include "DisplayClusterConfigurationTypes_Postprocess.h"
@@ -205,14 +206,33 @@ bool FDisplayClusterViewport_CustomPostProcessSettings::ConfigurePostProcessSett
 		return false;
 	}
 
+	// These flags are used to define the purpose of the viewport.
+	EDisplayClusterViewportRuntimeICVFXFlags ICVFXRuntimeFlags = InViewport->GetRenderSettingsICVFX().RuntimeFlags;
+
+	// The tile viewport does not contain ICVFX flags, and they must be obtained from a reference to the original viewport.
+	const FDisplayClusterViewport_RenderSettings& RenderSettings = InViewport->GetRenderSettings();
+	if (RenderSettings.TileSettings.IsInternalViewport())
+	{
+		if (IDisplayClusterViewportManager* ViewportManager = InViewport->GetConfiguration().GetViewportManager())
+		{
+			// Find source viewport.
+			const FString& SourceViewportId = RenderSettings.TileSettings.GetSourceViewportId();
+			if (IDisplayClusterViewport* SourceViewport = ViewportManager->FindViewport(SourceViewportId))
+			{
+				// Use the ICVFX flags from the source viewport.
+				ICVFXRuntimeFlags = SourceViewport->GetRenderSettingsICVFX().RuntimeFlags;
+			}
+		}
+	}
+
 	// Ignore ICVFX cameras
-	if (!GDisplayClusterPostProcessOverrideInCameraVFX && EnumHasAnyFlags(InViewport->GetRenderSettingsICVFX().RuntimeFlags, EDisplayClusterViewportRuntimeICVFXFlags::InCamera))
+	if (!GDisplayClusterPostProcessOverrideInCameraVFX && EnumHasAnyFlags(ICVFXRuntimeFlags, EDisplayClusterViewportRuntimeICVFXFlags::InCamera))
 	{
 		return false;
 	}
 
 	// Ignore Outers for ICVFX
-	if (!GDisplayClusterPostProcessOverrideOutersVFX && EnumHasAnyFlags(InViewport->GetRenderSettingsICVFX().RuntimeFlags, EDisplayClusterViewportRuntimeICVFXFlags::Target))
+	if (!GDisplayClusterPostProcessOverrideOutersVFX && EnumHasAnyFlags(ICVFXRuntimeFlags, EDisplayClusterViewportRuntimeICVFXFlags::Target))
 	{
 		return false;
 	}
