nDisplay: Fix tile seams with certain tile overscan values.

The issue was that the tile overscan percentage may not land precisely on pixel boundaries, and the frustum was calculated using the precise overscan percentage, so there was a mistmatch of up to 1 pixel betwen them which would make the cropped tile view not match precisely.

#rnx

diff --git a/Engine/Plugins/Runtime/nDisplay/Source/DisplayCluster/Private/Render/Viewport/Containers/DisplayClusterViewport_CustomFrustumRuntimeSettings.cpp b/Engine/Plugins/Runtime/nDisplay/Source/DisplayCluster/Private/Render/Viewport/Containers/DisplayClusterViewport_CustomFrustumRuntimeSettings.cpp
--- a/Engine/Plugins/Runtime/nDisplay/Source/DisplayCluster/Private/Render/Viewport/Containers/DisplayClusterViewport_CustomFrustumRuntimeSettings.cpp
+++ b/Engine/Plugins/Runtime/nDisplay/Source/DisplayCluster/Private/Render/Viewport/Containers/DisplayClusterViewport_CustomFrustumRuntimeSettings.cpp
@@ -157,6 +157,12 @@ void FDisplayClusterViewport_CustomFrustumRuntimeSettings::UpdateCustomFrustumSe
 	InOutRuntimeSettings.CustomFrustumPixels.Top    = FMath::RoundToInt(Size.Y * InOutRuntimeSettings.CustomFrustumPercent.Top);
 	InOutRuntimeSettings.CustomFrustumPixels.Bottom = FMath::RoundToInt(Size.Y * InOutRuntimeSettings.CustomFrustumPercent.Bottom);
 
+	// Quantize the percentage to exactly fit the number of pixels.
+	InOutRuntimeSettings.CustomFrustumPercent.Left   = double(InOutRuntimeSettings.CustomFrustumPixels.Left)   / double(Size.X);
+	InOutRuntimeSettings.CustomFrustumPercent.Right  = double(InOutRuntimeSettings.CustomFrustumPixels.Right)  / double(Size.X);
+	InOutRuntimeSettings.CustomFrustumPercent.Top    = double(InOutRuntimeSettings.CustomFrustumPixels.Top)    / double(Size.Y);
+	InOutRuntimeSettings.CustomFrustumPercent.Bottom = double(InOutRuntimeSettings.CustomFrustumPixels.Bottom) / double(Size.Y);
+
 	// Update RTT size for CustomFrustum when we need to scale target resolution
 	if (InCustomFrustumSettings.bAdaptResolution)
 	{
diff --git a/Engine/Plugins/Runtime/nDisplay/Source/DisplayCluster/Private/Render/Viewport/Containers/DisplayClusterViewport_OverscanRuntimeSettings.cpp b/Engine/Plugins/Runtime/nDisplay/Source/DisplayCluster/Private/Render/Viewport/Containers/DisplayClusterViewport_OverscanRuntimeSettings.cpp
--- a/Engine/Plugins/Runtime/nDisplay/Source/DisplayCluster/Private/Render/Viewport/Containers/DisplayClusterViewport_OverscanRuntimeSettings.cpp
+++ b/Engine/Plugins/Runtime/nDisplay/Source/DisplayCluster/Private/Render/Viewport/Containers/DisplayClusterViewport_OverscanRuntimeSettings.cpp
@@ -115,10 +115,17 @@ void FDisplayClusterViewport_OverscanRuntimeSettings::UpdateOverscanSettings(
 	if (InOutOverscanRuntimeSettings.bIsEnabled)
 	{
 		// Calc pixels from percent
-		InOutOverscanRuntimeSettings.OverscanPixels.Left   = Size.X * InOutOverscanRuntimeSettings.OverscanPercent.Left;
-		InOutOverscanRuntimeSettings.OverscanPixels.Right  = Size.X * InOutOverscanRuntimeSettings.OverscanPercent.Right;
-		InOutOverscanRuntimeSettings.OverscanPixels.Top    = Size.Y * InOutOverscanRuntimeSettings.OverscanPercent.Top;
-		InOutOverscanRuntimeSettings.OverscanPixels.Bottom = Size.Y * InOutOverscanRuntimeSettings.OverscanPercent.Bottom;
+		InOutOverscanRuntimeSettings.OverscanPixels.Left   = FMath::RoundToInt(Size.X * InOutOverscanRuntimeSettings.OverscanPercent.Left);
+		InOutOverscanRuntimeSettings.OverscanPixels.Right  = FMath::RoundToInt(Size.X * InOutOverscanRuntimeSettings.OverscanPercent.Right);
+		InOutOverscanRuntimeSettings.OverscanPixels.Top    = FMath::RoundToInt(Size.Y * InOutOverscanRuntimeSettings.OverscanPercent.Top);
+		InOutOverscanRuntimeSettings.OverscanPixels.Bottom = FMath::RoundToInt(Size.Y * InOutOverscanRuntimeSettings.OverscanPercent.Bottom);
+
+		// Quantize the overscan percentage to exactly fit the number of pixels.
+		// This will avoid a mismatch between the overscanned frustum calculated in UpdateProjectionAngles and the pixel crop in GetFinalContextRect.
+		InOutOverscanRuntimeSettings.OverscanPercent.Left   = double(InOutOverscanRuntimeSettings.OverscanPixels.Left)   / double(Size.X);
+		InOutOverscanRuntimeSettings.OverscanPercent.Right  = double(InOutOverscanRuntimeSettings.OverscanPixels.Right)  / double(Size.X);
+		InOutOverscanRuntimeSettings.OverscanPercent.Top    = double(InOutOverscanRuntimeSettings.OverscanPixels.Top)    / double(Size.Y);
+		InOutOverscanRuntimeSettings.OverscanPercent.Bottom = double(InOutOverscanRuntimeSettings.OverscanPixels.Bottom) / double(Size.Y);
 
 		const FIntPoint OverscanSize = Size + InOutOverscanRuntimeSettings.OverscanPixels.Size();
 		const FIntPoint ValidOverscanSize = FDisplayClusterViewportHelpers::GetValidViewportRect(FIntRect(FIntPoint(0, 0), OverscanSize), InViewportId, TEXT("Overscan")).Size();
@@ -137,13 +144,24 @@ void FDisplayClusterViewport_OverscanRuntimeSettings::UpdateOverscanSettings(
 		}
 		else
 		{
-			double scaleX = double(Size.X) / OverscanSize.X;
-			double scaleY = double(Size.Y) / OverscanSize.Y;
-
-			InOutOverscanRuntimeSettings.OverscanPixels.Left   *= scaleX;
-			InOutOverscanRuntimeSettings.OverscanPixels.Right  *= scaleX;
-			InOutOverscanRuntimeSettings.OverscanPixels.Top    *= scaleY;
-			InOutOverscanRuntimeSettings.OverscanPixels.Bottom *= scaleY;
+			const double ScaleX = double(Size.X) / double(OverscanSize.X);
+			const double ScaleY = double(Size.Y) / double(OverscanSize.Y);
+
+			InOutOverscanRuntimeSettings.OverscanPixels.Left   *= ScaleX;
+			InOutOverscanRuntimeSettings.OverscanPixels.Right  *= ScaleX;
+			InOutOverscanRuntimeSettings.OverscanPixels.Top    *= ScaleY;
+			InOutOverscanRuntimeSettings.OverscanPixels.Bottom *= ScaleY;
+
+			const FIntPoint ScaledSize(
+				Size.X - (InOutOverscanRuntimeSettings.OverscanPixels.Left + InOutOverscanRuntimeSettings.OverscanPixels.Right),
+				Size.Y - (InOutOverscanRuntimeSettings.OverscanPixels.Top  + InOutOverscanRuntimeSettings.OverscanPixels.Bottom)
+			);
+
+			// Quantize the overscan percentage to exactly fit the number of pixels.
+			InOutOverscanRuntimeSettings.OverscanPercent.Left   = double(InOutOverscanRuntimeSettings.OverscanPixels.Left)   / double(ScaledSize.X);
+			InOutOverscanRuntimeSettings.OverscanPercent.Right  = double(InOutOverscanRuntimeSettings.OverscanPixels.Right)  / double(ScaledSize.X);
+			InOutOverscanRuntimeSettings.OverscanPercent.Top    = double(InOutOverscanRuntimeSettings.OverscanPixels.Top)    / double(ScaledSize.Y);
+			InOutOverscanRuntimeSettings.OverscanPercent.Bottom = double(InOutOverscanRuntimeSettings.OverscanPixels.Bottom) / double(ScaledSize.Y);
 		}
 	}
 }
