nDisplay: Added more pixel rounding for consistency when calculating overscan.

#rnx

diff --git a/Engine/Plugins/Runtime/nDisplay/Source/DisplayCluster/Private/Render/Viewport/Containers/DisplayClusterViewport_OverscanRuntimeSettings.cpp b/Engine/Plugins/Runtime/nDisplay/Source/DisplayCluster/Private/Render/Viewport/Containers/DisplayClusterViewport_OverscanRuntimeSettings.cpp
--- a/Engine/Plugins/Runtime/nDisplay/Source/DisplayCluster/Private/Render/Viewport/Containers/DisplayClusterViewport_OverscanRuntimeSettings.cpp
+++ b/Engine/Plugins/Runtime/nDisplay/Source/DisplayCluster/Private/Render/Viewport/Containers/DisplayClusterViewport_OverscanRuntimeSettings.cpp
@@ -147,10 +147,10 @@ void FDisplayClusterViewport_OverscanRuntimeSettings::UpdateOverscanSettings(
 			const double ScaleX = double(Size.X) / double(OverscanSize.X);
 			const double ScaleY = double(Size.Y) / double(OverscanSize.Y);
 
-			InOutOverscanRuntimeSettings.OverscanPixels.Left   *= ScaleX;
-			InOutOverscanRuntimeSettings.OverscanPixels.Right  *= ScaleX;
-			InOutOverscanRuntimeSettings.OverscanPixels.Top    *= ScaleY;
-			InOutOverscanRuntimeSettings.OverscanPixels.Bottom *= ScaleY;
+			InOutOverscanRuntimeSettings.OverscanPixels.Left   = FMath::RoundToInt(InOutOverscanRuntimeSettings.OverscanPixels.Left   * ScaleX);
+			InOutOverscanRuntimeSettings.OverscanPixels.Right  = FMath::RoundToInt(InOutOverscanRuntimeSettings.OverscanPixels.Right  * ScaleX);
+			InOutOverscanRuntimeSettings.OverscanPixels.Top    = FMath::RoundToInt(InOutOverscanRuntimeSettings.OverscanPixels.Top    * ScaleY);
+			InOutOverscanRuntimeSettings.OverscanPixels.Bottom = FMath::RoundToInt(InOutOverscanRuntimeSettings.OverscanPixels.Bottom * ScaleY);
 
 			const FIntPoint ScaledSize(
 				Size.X - (InOutOverscanRuntimeSettings.OverscanPixels.Left + InOutOverscanRuntimeSettings.OverscanPixels.Right),
