[Backout] - CL36460984
Original CL Desc
-----------------------------------------------------------------
[nDisplay] Fixed issue when media is not working in offscreen with 0x0 window size

- Fixed startup command line parameters generation for offscreen nodes
- No window geometry specification required when running offscreen
- Minor clean up to the cluster node configuration data (default values)

#rnx

diff --git a/Engine/Plugins/Editor/DisplayClusterLaunch/Source/DisplayClusterLaunchEditor/Private/DisplayClusterLaunchEditorModule.cpp b/Engine/Plugins/Editor/DisplayClusterLaunch/Source/DisplayClusterLaunchEditor/Private/DisplayClusterLaunchEditorModule.cpp
--- a/Engine/Plugins/Editor/DisplayClusterLaunch/Source/DisplayClusterLaunchEditor/Private/DisplayClusterLaunchEditorModule.cpp
+++ b/Engine/Plugins/Editor/DisplayClusterLaunch/Source/DisplayClusterLaunchEditor/Private/DisplayClusterLaunchEditorModule.cpp
@@ -520,10 +520,6 @@ void FDisplayClusterLaunchEditorModule::LaunchDisplayClusterProcess()
 				{
 					ConcatenatedCommandLineArguments += "-fullscreen ";
 				}
-				else if (NodePtr->bRenderHeadless)
-				{
-					ConcatenatedCommandLineArguments += "-RenderOffscreen ";
-				}
 				else
 				{
 					ConcatenatedCommandLineArguments +=
@@ -534,6 +530,12 @@ void FDisplayClusterLaunchEditorModule::LaunchDisplayClusterProcess()
 						);
 				}
 
+				// Headless
+				if (NodePtr->bRenderHeadless)
+				{
+					ConcatenatedCommandLineArguments += "-RenderOffscreen ";
+				}
+
 				// Graphics adapter
 				if (NodePtr->GraphicsAdapter >= 0)
 				{
diff --git a/Engine/Plugins/Runtime/nDisplay/Source/DisplayClusterConfiguration/Private/DisplayClusterConfigurationTypes.cpp b/Engine/Plugins/Runtime/nDisplay/Source/DisplayClusterConfiguration/Private/DisplayClusterConfigurationTypes.cpp
--- a/Engine/Plugins/Runtime/nDisplay/Source/DisplayClusterConfiguration/Private/DisplayClusterConfigurationTypes.cpp
+++ b/Engine/Plugins/Runtime/nDisplay/Source/DisplayClusterConfiguration/Private/DisplayClusterConfigurationTypes.cpp
@@ -350,6 +350,15 @@ void UDisplayClusterConfigurationClusterNode::Serialize(FArchive& Ar)
 	}
 }
 
+UDisplayClusterConfigurationClusterNode::UDisplayClusterConfigurationClusterNode()
+	: bIsSoundEnabled(false)
+#if WITH_EDITORONLY_DATA
+	, bIsVisible(true)
+	, bIsUnlocked(true)
+#endif
+{
+}
+
 UDisplayClusterConfigurationHostDisplayData::UDisplayClusterConfigurationHostDisplayData()
 	: bIsVisible(true)
 	, bIsUnlocked(true)
diff --git a/Engine/Plugins/Runtime/nDisplay/Source/DisplayClusterConfiguration/Public/DisplayClusterConfigurationTypes.h b/Engine/Plugins/Runtime/nDisplay/Source/DisplayClusterConfiguration/Public/DisplayClusterConfigurationTypes.h
--- a/Engine/Plugins/Runtime/nDisplay/Source/DisplayClusterConfiguration/Public/DisplayClusterConfigurationTypes.h
+++ b/Engine/Plugins/Runtime/nDisplay/Source/DisplayClusterConfiguration/Public/DisplayClusterConfigurationTypes.h
@@ -291,6 +291,7 @@ public:
 	FOnPostEditChangeChainProperty OnPostEditChangeChainProperty;
 
 public:
+	UDisplayClusterConfigurationClusterNode();
 
 	UFUNCTION(BlueprintPure, Category = "NDisplay|Configuration")
 	void GetViewportIds(TArray<FString>& OutViewportIds) const;
@@ -319,24 +320,24 @@ public:
 	
 	/** Enables or disables sound on nDisplay primary Node */
 	UPROPERTY(EditAnywhere, BlueprintReadWrite, Category = "Configuration", meta = (DisplayName = "Enable Sound"))
-	bool bIsSoundEnabled = false;
+	bool bIsSoundEnabled;
 
 	/** Enables application window native fullscreen support */
-	UPROPERTY(EditAnywhere, BlueprintReadonly, Category = "Configuration", meta = (DisplayName = "Fullscreen", EditCondition = "!bRenderHeadless"))
-	bool bIsFullscreen = false;
-
-	/** Activates headless rendering for this cluster node */
-	UPROPERTY(EditAnywhere, BlueprintReadonly, Category = "Configuration", meta = (DisplayName = "Headless Rendering", EditCondition = "!bIsFullscreen"))
-	bool bRenderHeadless = false;
+	UPROPERTY(EditAnywhere, BlueprintReadonly, Category = "Configuration", meta = (DisplayName = "Fullscreen"))
+	bool bIsFullscreen;
 
 	/** Defines the application window size in pixels */
-	UPROPERTY(EditAnywhere, BlueprintReadWrite, Category = "Configuration", meta = (DisplayName = "Window", DisplayMode = "Compound", FixedAspectRatioProperty = "bFixedAspectRatio", EditCondition = "!bRenderHeadless && !bIsFullscreen"))
+	UPROPERTY(EditAnywhere, BlueprintReadWrite, Category = "Configuration", meta = (DisplayName = "Window", DisplayMode = "Compound", FixedAspectRatioProperty = "bFixedAspectRatio"))
 	FDisplayClusterConfigurationRectangle WindowRect;
 
 	/** Output remapping settings for the selected cluster node */
 	UPROPERTY(EditAnywhere, BlueprintReadWrite, Category = "Configuration", meta = (DisplayName = "Output Remapping"))
 	FDisplayClusterConfigurationFramePostProcess_OutputRemap OutputRemap;
 
+	/** Activates headless rendering for this cluster node */
+	UPROPERTY(EditAnywhere, BlueprintReadonly, Category = "Configuration", meta = (DisplayName = "Headless Rendering"))
+	bool bRenderHeadless = false;
+
 	/** Hint for setting the r.GraphicsAdapter CVar when launching this cluster node. Note that this is distinct from the GPU Node Indices assigned to viewports. */
 	UPROPERTY(EditAnywhere, BlueprintReadonly, Category = "Configuration", meta = (DisplayName = "Graphics Adapter"))
 	int32 GraphicsAdapter = -1;
@@ -348,7 +349,7 @@ public:
 #if WITH_EDITORONLY_DATA
 	/** Locks the application window aspect ratio for easier resizing */
 	UPROPERTY(EditAnywhere, Category = "Configuration", meta = (HideProperty))
-	bool bFixedAspectRatio = false;
+	bool bFixedAspectRatio;
 #endif
 
 	UPROPERTY(VisibleInstanceOnly, BlueprintReadOnly, EditFixedSize, Instanced, Category = "Configuration", meta = (DisplayThumbnail = false))
@@ -363,10 +364,10 @@ public:
 
 #if WITH_EDITORONLY_DATA
 	UPROPERTY(EditDefaultsOnly, Category = "Configuration", meta = (HideProperty))
-	bool bIsVisible = true;
+	bool bIsVisible;
 
 	UPROPERTY(EditDefaultsOnly, Category = "Configuration", meta = (HideProperty))
-	bool bIsUnlocked = true;
+	bool bIsUnlocked;
 
 	/** Binds a background preview image for easier output mapping */
 	UPROPERTY(EditDefaultsOnly, Category = "Configuration")
diff --git a/Engine/Plugins/VirtualProduction/Switchboard/Source/Switchboard/switchboard/devices/ndisplay/plugin_ndisplay.py b/Engine/Plugins/VirtualProduction/Switchboard/Source/Switchboard/switchboard/devices/ndisplay/plugin_ndisplay.py
--- a/Engine/Plugins/VirtualProduction/Switchboard/Source/Switchboard/switchboard/devices/ndisplay/plugin_ndisplay.py
+++ b/Engine/Plugins/VirtualProduction/Switchboard/Source/Switchboard/switchboard/devices/ndisplay/plugin_ndisplay.py
@@ -1178,10 +1178,6 @@ class DevicenDisplay(DeviceUnreal):
             args.extend([
                 '-fullscreen',
             ])
-        elif headless:
-            args.extend([
-                '-RenderOffscreen',
-            ])
         else:
             args.extend([
                 '-windowed',
@@ -1192,6 +1188,11 @@ class DevicenDisplay(DeviceUnreal):
                 f'ResY={win_res[1]}',
             ])
 
+        if headless:
+            args.extend([
+                '-RenderOffscreen',
+            ])
+
         # MultiUser parameters
         if CONFIG.MUSERVER_AUTO_JOIN.get_value() and self.autojoin_mu_server.get_value():
             args.extend([
