nDisplay - Prevent being able to add a new viewport from the output mapping graph when no cluster node is present. The option to add a viewport is removed unless it is safe to do, fixing a crash and keeping the behavior the same as in the treeview.

diff --git a/Engine/Plugins/Runtime/nDisplay/Source/DisplayClusterConfigurator/Private/ClusterConfiguration/DisplayClusterConfiguratorClusterEditorUtils.cpp b/Engine/Plugins/Runtime/nDisplay/Source/DisplayClusterConfigurator/Private/ClusterConfiguration/DisplayClusterConfiguratorClusterEditorUtils.cpp
--- a/Engine/Plugins/Runtime/nDisplay/Source/DisplayClusterConfigurator/Private/ClusterConfiguration/DisplayClusterConfiguratorClusterEditorUtils.cpp
+++ b/Engine/Plugins/Runtime/nDisplay/Source/DisplayClusterConfigurator/Private/ClusterConfiguration/DisplayClusterConfiguratorClusterEditorUtils.cpp
@@ -323,6 +323,14 @@ UDisplayClusterConfigurationViewport* UE::DisplayClusterConfiguratorClusterEdito
 	return NewViewport;
 }
 
+bool UE::DisplayClusterConfiguratorClusterEditorUtils::CanAddNewViewportToToolkit(const TWeakPtr<FDisplayClusterConfiguratorBlueprintEditor>& InToolkit)
+{
+	check(InToolkit.IsValid());
+	const UDisplayClusterConfigurationData* EditorData = InToolkit.Pin()->GetEditorData();
+	check(EditorData);
+	return EditorData->Cluster && EditorData->Cluster->Nodes.Num() > 0;
+}
+
 void UE::DisplayClusterConfiguratorClusterEditorUtils::ShowNewClusterItemDialogWindow(TSharedRef<SDisplayClusterConfiguratorNewClusterItemDialog> DialogContent, TSharedPtr<SWidget> ParentElement, FText WindowTitle, FVector2D WindowSize)
 {
 	TSharedPtr<SWidget> ParentWindow = FSlateApplication::Get().FindBestParentWindowForDialogs(ParentElement);
diff --git a/Engine/Plugins/Runtime/nDisplay/Source/DisplayClusterConfigurator/Private/ClusterConfiguration/DisplayClusterConfiguratorClusterEditorUtils.h b/Engine/Plugins/Runtime/nDisplay/Source/DisplayClusterConfigurator/Private/ClusterConfiguration/DisplayClusterConfiguratorClusterEditorUtils.h
--- a/Engine/Plugins/Runtime/nDisplay/Source/DisplayClusterConfigurator/Private/ClusterConfiguration/DisplayClusterConfiguratorClusterEditorUtils.h
+++ b/Engine/Plugins/Runtime/nDisplay/Source/DisplayClusterConfigurator/Private/ClusterConfiguration/DisplayClusterConfiguratorClusterEditorUtils.h
@@ -47,6 +47,12 @@ namespace UE::DisplayClusterConfiguratorClusterEditorUtils
 		UDisplayClusterConfigurationClusterNode* ClusterNode, const FDisplayClusterConfigurationRectangle& PresetRect,
 		TSharedPtr<FScopedTransaction>& OutTransaction);
 
+	/**
+	 * Checks if a new viewport is allowed to be added by making sure there is a cluster node available.
+	 * @param InToolkit The toolkit hosting the viewport.
+	 */
+	bool CanAddNewViewportToToolkit(const TWeakPtr<FDisplayClusterConfiguratorBlueprintEditor>& InToolkit);
+	
 	/**
 	 * Displays a dialog box used to configure new cluster item properties before adding the new item to the cluster configuration.
 	 * @param DialogContent - The widget to display in the dialog box's main section.
diff --git a/Engine/Plugins/Runtime/nDisplay/Source/DisplayClusterConfigurator/Private/Views/OutputMapping/DisplayClusterConfiguratorGraphSchema.cpp b/Engine/Plugins/Runtime/nDisplay/Source/DisplayClusterConfigurator/Private/Views/OutputMapping/DisplayClusterConfiguratorGraphSchema.cpp
--- a/Engine/Plugins/Runtime/nDisplay/Source/DisplayClusterConfigurator/Private/Views/OutputMapping/DisplayClusterConfiguratorGraphSchema.cpp
+++ b/Engine/Plugins/Runtime/nDisplay/Source/DisplayClusterConfigurator/Private/Views/OutputMapping/DisplayClusterConfiguratorGraphSchema.cpp
@@ -70,6 +70,8 @@ void UDisplayClusterConfiguratorGraphSchema::GetGraphContextActions(FGraphContex
 	}
 
 	// Action for adding a new viewport
+	const UDisplayClusterConfiguratorGraph* ClusterGraph = CastChecked<UDisplayClusterConfiguratorGraph>(ContextMenuBuilder.CurrentGraph);
+	if (UE::DisplayClusterConfiguratorClusterEditorUtils::CanAddNewViewportToToolkit(ClusterGraph->GetToolkit()))
 	{
 		FText Desc = LOCTEXT("NewViewport_Desc", "New Viewport");
 		FText Tooltip = LOCTEXT("NewViewport_Tooltip", "Add a new viewport");
diff --git a/Engine/Plugins/Runtime/nDisplay/Source/DisplayClusterConfigurator/Private/Views/OutputMapping/SDisplayClusterConfiguratorGraphEditor.cpp b/Engine/Plugins/Runtime/nDisplay/Source/DisplayClusterConfigurator/Private/Views/OutputMapping/SDisplayClusterConfiguratorGraphEditor.cpp
--- a/Engine/Plugins/Runtime/nDisplay/Source/DisplayClusterConfigurator/Private/Views/OutputMapping/SDisplayClusterConfiguratorGraphEditor.cpp
+++ b/Engine/Plugins/Runtime/nDisplay/Source/DisplayClusterConfigurator/Private/Views/OutputMapping/SDisplayClusterConfiguratorGraphEditor.cpp
@@ -465,8 +465,9 @@ bool SDisplayClusterConfiguratorGraphEditor::CanAddNewViewport() const
 	const TSet<UObject*>& SelectedNodes = GetSelectedNodes();
 	if (SelectedNodes.Num() == 1)
 	{
-		UObject* Node = *SelectedNodes.begin();
-		return Node->IsA<UDisplayClusterConfiguratorWindowNode>() || Node->IsA<UDisplayClusterConfiguratorCanvasNode>();
+		const UObject* Node = *SelectedNodes.begin();
+		return UE::DisplayClusterConfiguratorClusterEditorUtils::CanAddNewViewportToToolkit(ToolkitPtr)
+			&& (Node->IsA<UDisplayClusterConfiguratorWindowNode>() || Node->IsA<UDisplayClusterConfiguratorCanvasNode>());
 	}
 
 	return false;
diff --git a/Engine/Plugins/Runtime/nDisplay/Source/DisplayClusterConfigurator/Private/Views/TreeViews/Cluster/DisplayClusterConfiguratorViewCluster.cpp b/Engine/Plugins/Runtime/nDisplay/Source/DisplayClusterConfigurator/Private/Views/TreeViews/Cluster/DisplayClusterConfiguratorViewCluster.cpp
--- a/Engine/Plugins/Runtime/nDisplay/Source/DisplayClusterConfigurator/Private/Views/TreeViews/Cluster/DisplayClusterConfiguratorViewCluster.cpp
+++ b/Engine/Plugins/Runtime/nDisplay/Source/DisplayClusterConfigurator/Private/Views/TreeViews/Cluster/DisplayClusterConfiguratorViewCluster.cpp
@@ -659,8 +659,7 @@ void FDisplayClusterConfiguratorViewCluster::AddNewViewport(FVector2D PresetSize
 bool FDisplayClusterConfiguratorViewCluster::CanAddNewViewport() const
 {
 	// Can only add a new viewport if there is at least one cluster node to add it to.
-	UDisplayClusterConfigurationCluster* Cluster = ToolkitPtr.Pin()->GetEditorData()->Cluster;
-	return Cluster->Nodes.Num() > 0;
+	return UE::DisplayClusterConfiguratorClusterEditorUtils::CanAddNewViewportToToolkit(ToolkitPtr);
 }
 
 #undef LOCTEXT_NAMESPACE
\ No newline at end of file
