nDisplay: Fixed screen percentage
- The screen percentage implementation has been moved to FDisplayClusterViewportManager::PostConfigureViewFamily().
- Removed the dynamic resolution feature from nDisplay.
- Some display names and comments in the UI have been renamed.

diff --git a/Engine/Plugins/Runtime/nDisplay/Source/DisplayCluster/Private/Blueprints/DisplayClusterBlueprint.cpp b/Engine/Plugins/Runtime/nDisplay/Source/DisplayCluster/Private/Blueprints/DisplayClusterBlueprint.cpp
--- a/Engine/Plugins/Runtime/nDisplay/Source/DisplayCluster/Private/Blueprints/DisplayClusterBlueprint.cpp
+++ b/Engine/Plugins/Runtime/nDisplay/Source/DisplayCluster/Private/Blueprints/DisplayClusterBlueprint.cpp
@@ -106,7 +106,7 @@ void UDisplayClusterBlueprint::UpdateSummaryProperty()
 
 	Lines.Add(FString::Printf(TEXT("Sync Policy: %s"), *ConfigData->Cluster->Sync.RenderSyncPolicy.Type));
 	Lines.Add(FString::Printf(TEXT("Follow Local Player Camera: %s"), ConfigData->bFollowLocalPlayerCamera ? TEXT("Yes") : TEXT("No")));
-	Lines.Add(FString::Printf(TEXT("Viewport Screen Percentage Multiplier: %.2f"), ConfigData->RenderFrameSettings.ClusterBufferRatioMult));
+	Lines.Add(FString::Printf(TEXT("Viewports Screen %% Multiplier: %.2f"), ConfigData->RenderFrameSettings.ClusterICVFXOuterViewportBufferRatioMult));
 
 	Lines.Add(TEXT(""));
 
diff --git a/Engine/Plugins/Runtime/nDisplay/Source/DisplayCluster/Private/Game/EngineClasses/Basics/DisplayClusterViewportClient.cpp b/Engine/Plugins/Runtime/nDisplay/Source/DisplayCluster/Private/Game/EngineClasses/Basics/DisplayClusterViewportClient.cpp
--- a/Engine/Plugins/Runtime/nDisplay/Source/DisplayCluster/Private/Game/EngineClasses/Basics/DisplayClusterViewportClient.cpp
+++ b/Engine/Plugins/Runtime/nDisplay/Source/DisplayCluster/Private/Game/EngineClasses/Basics/DisplayClusterViewportClient.cpp
@@ -33,8 +33,6 @@
 #include "Framework/Application/SlateApplication.h"
 #include "Slate/SceneViewport.h"
 
-#include "LegacyScreenPercentageDriver.h"
-#include "DynamicResolutionState.h"
 #include "EngineStats.h"
 
 #include "Render/Device/IDisplayClusterRenderDevice.h"
@@ -515,17 +513,9 @@ void UDisplayClusterViewportClient::Draw(FViewport* InViewport, FCanvas* SceneCa
 			FSceneViewFamilyContext& ViewFamily = *ViewFamilies.Last();
 			bool bIsFamilyVisible = false;
 
-			// Configure family
+			// Configure family for the nDisplay.
 			RenderFrameViewportManager->ConfigureViewFamily(DCRenderTarget, DCViewFamily, ViewFamily);
 
-#if WITH_EDITOR
-			if (GIsEditor)
-			{
-				// Force enable view family show flag for HighDPI derived's screen percentage.
-				ViewFamily.EngineShowFlags.ScreenPercentage = true;
-			}
-#endif
-
 			ViewFamily.ViewMode = EViewModeIndex(ViewModeIndex);
 			EngineShowFlagOverride(ESFIM_Game, ViewFamily.ViewMode, ViewFamily.EngineShowFlags, false);
 
@@ -730,103 +720,17 @@ void UDisplayClusterViewportClient::Draw(FViewport* InViewport, FCanvas* SceneCa
 			{
 				FinalizeViews(&ViewFamily, PlayerViewMap);
 
-				// Force screen percentage show flag to be turned off if not supported.
-				if (!ViewFamily.SupportsScreenPercentage())
+				// Collect rendering flags for nDisplay:
+				EDisplayClusterViewportRenderingFlags ViewportRenderingFlags = EDisplayClusterViewportRenderingFlags::None;
+				if (bStereoRendering)
 				{
-					ViewFamily.EngineShowFlags.ScreenPercentage = false;
-				}
-
-				// Set up secondary resolution fraction for the view family.
-				if (ViewFamily.SupportsScreenPercentage())
-				{
-					float CustomSecondaryScreenPercentage = IConsoleManager::Get().FindConsoleVariable(TEXT("r.SecondaryScreenPercentage.GameViewport"), false)->GetFloat();
-					if (CustomSecondaryScreenPercentage > 0.0)
-					{
-						// Override secondary resolution fraction with CVar.
-						ViewFamily.SecondaryViewFraction = FMath::Min(CustomSecondaryScreenPercentage / 100.0f, 1.0f);
-					}
-					else
-					{
-						// Automatically compute secondary resolution fraction from DPI.
-						ViewFamily.SecondaryViewFraction = GetDPIDerivedResolutionFraction();
-					}
-
-					check(ViewFamily.SecondaryViewFraction > 0.0f);
-				}
-
-				checkf(ViewFamily.GetScreenPercentageInterface() == nullptr,
-					TEXT("Some code has tried to set up an alien screen percentage driver, that could be wrong if not supported very well by the RHI."));
-
-				// Setup main view family with screen percentage interface by dynamic resolution if screen percentage is enabled.
-	#if WITH_DYNAMIC_RESOLUTION
-				if (ViewFamily.EngineShowFlags.ScreenPercentage)
-				{
-					FDynamicResolutionStateInfos DynamicResolutionStateInfos;
-					GEngine->GetDynamicResolutionCurrentStateInfos(DynamicResolutionStateInfos);
-
-					// Do not allow dynamic resolution to touch the view family if not supported to ensure there is no possibility to ruin
-					// game play experience on platforms that does not support it, but have it enabled by mistake.
-					if (DynamicResolutionStateInfos.Status == EDynamicResolutionStatus::Enabled)
-					{
-						GEngine->EmitDynamicResolutionEvent(EDynamicResolutionStateEvent::BeginDynamicResolutionRendering);
-						GEngine->GetDynamicResolutionState()->SetupMainViewFamily(ViewFamily);
-					}
-					else if (DynamicResolutionStateInfos.Status == EDynamicResolutionStatus::DebugForceEnabled)
-					{
-						GEngine->EmitDynamicResolutionEvent(EDynamicResolutionStateEvent::BeginDynamicResolutionRendering);
-						ViewFamily.SetScreenPercentageInterface(new FLegacyScreenPercentageDriver(
-							ViewFamily,
-							DynamicResolutionStateInfos.ResolutionFractionApproximations[GDynamicPrimaryResolutionFraction],
-							DynamicResolutionStateInfos.ResolutionFractionUpperBounds[GDynamicPrimaryResolutionFraction]));
-					}
-
-	#if CSV_PROFILER_STATS
-					if (DynamicResolutionStateInfos.ResolutionFractionApproximations[GDynamicPrimaryResolutionFraction] >= 0.0f)
-					{
-						CSV_CUSTOM_STAT_GLOBAL(DynamicResolutionPercentage, DynamicResolutionStateInfos.ResolutionFractionApproximations[GDynamicPrimaryResolutionFraction] * 100.0f, ECsvCustomStatOp::Set);
-					}
-	#endif
-				}
-	#endif
-
-				// If a screen percentage interface was not set by dynamic resolution, then create one matching legacy behavior.
-				if (ViewFamily.GetScreenPercentageInterface() == nullptr)
-				{
-					// In case of stereo, we set the same buffer ratio for both left and right views (taken from left)
-					float CustomBufferRatio = DCViewFamily.CustomBufferRatio;
-
-					float GlobalResolutionFraction = 1.0f;
-					float SecondaryScreenPercentage = ViewFamily.SecondaryViewFraction;
-
-					if (ViewFamily.EngineShowFlags.ScreenPercentage)
-					{
-						// Get global view fraction set by r.ScreenPercentage.
-						GlobalResolutionFraction = CustomBufferRatio;
-
-						// We need to split the screen percentage if below 0.5 because TAA upscaling only works well up to 2x.
-						if (GlobalResolutionFraction < 0.5f)
-						{
-							SecondaryScreenPercentage = 2.0f * GlobalResolutionFraction;
-							GlobalResolutionFraction = 0.5f;
-						}
-					}
-
-					ViewFamily.SetScreenPercentageInterface(new FLegacyScreenPercentageDriver(
-						ViewFamily, GlobalResolutionFraction));
-
-					ViewFamily.SecondaryViewFraction = SecondaryScreenPercentage;
-				}
-				else if (bStereoRendering)
-				{
-					// Change screen percentage method to raw output when doing dynamic resolution with VR if not using TAA upsample.
-					for (FSceneView* View : Views)
-					{
-						if (View->PrimaryScreenPercentageMethod == EPrimaryScreenPercentageMethod::SpatialUpscale)
-						{
-							View->PrimaryScreenPercentageMethod = EPrimaryScreenPercentageMethod::RawOutput;
-						}
-					}
+					EnumAddFlags(ViewportRenderingFlags, EDisplayClusterViewportRenderingFlags::StereoRendering);
 				}
+				
+				// Completing the of a ViewDamily configuration.
+				// The screen percentage is configurable in this function.
+				RenderFrameViewportManager->PostConfigureViewFamily(DCRenderTarget, DCViewFamily, ViewFamily, Views,
+					ViewportRenderingFlags, GetDPIScale());
 
 				ViewFamily.bIsHDR = GetWindow().IsValid() ? GetWindow().Get()->GetIsHDR() : false;
 
@@ -921,9 +825,6 @@ void UDisplayClusterViewportClient::Draw(FViewport* InViewport, FCanvas* SceneCa
 	// custom postprocess single frame flag must be removed at frame end on game thread
 	DCRenderDevice->FinalizeNewFrame();
 
-	// Beyond this point, only UI rendering independent from dynamic resolution.
-	GEngine->EmitDynamicResolutionEvent(EDynamicResolutionStateEvent::EndDynamicResolutionRendering);
-
 	// Update level streaming.
 	MyWorld->UpdateLevelStreaming();
 
diff --git a/Engine/Plugins/Runtime/nDisplay/Source/DisplayCluster/Private/Render/Viewport/DisplayClusterViewportManager.cpp b/Engine/Plugins/Runtime/nDisplay/Source/DisplayCluster/Private/Render/Viewport/DisplayClusterViewportManager.cpp
--- a/Engine/Plugins/Runtime/nDisplay/Source/DisplayCluster/Private/Render/Viewport/DisplayClusterViewportManager.cpp
+++ b/Engine/Plugins/Runtime/nDisplay/Source/DisplayCluster/Private/Render/Viewport/DisplayClusterViewportManager.cpp
@@ -34,6 +34,8 @@
 
 #include "SceneViewExtension.h"
 
+#include "LegacyScreenPercentageDriver.h"
+
 #include "DisplayClusterRootActor.h"
 
 #include "Misc/DisplayClusterLog.h"
@@ -775,7 +777,10 @@ bool FDisplayClusterViewportManager::ShouldRenderFinalColor() const
 	return !Configuration->GetRenderFrameSettings().IsPostProcessDisabled();
 }
 
-void FDisplayClusterViewportManager::ConfigureViewFamily(const FDisplayClusterRenderFrameTarget& InFrameTarget, const FDisplayClusterRenderFrameTargetViewFamily& InFrameViewFamily, FSceneViewFamilyContext& ViewFamily)
+void FDisplayClusterViewportManager::ConfigureViewFamily(
+	const FDisplayClusterRenderFrameTarget& InFrameTarget,
+	const FDisplayClusterRenderFrameTargetViewFamily& InFrameViewFamily,
+	FSceneViewFamilyContext& ViewFamily)
 {
 	// Note: EngineShowFlags should have already been configured in CreateViewFamilyConstructionValues.
 	ViewFamily.SceneCaptureCompositeMode = ESceneCaptureCompositeMode::SCCM_Overwrite;
@@ -795,6 +800,117 @@ void FDisplayClusterViewportManager::ConfigureViewFamily(const FDisplayClusterRe
 	}
 }
 
+void FDisplayClusterViewportManager::PostConfigureViewFamily(
+	const FDisplayClusterRenderFrameTarget& InFrameTarget,
+	const FDisplayClusterRenderFrameTargetViewFamily& InFrameViewFamily,
+	FSceneViewFamilyContext& ViewFamily,
+	const TArray<FSceneView*>& Views,
+	const EDisplayClusterViewportRenderingFlags InRenderingFlags,
+	const float InDPIScale
+)
+{
+#if WITH_EDITOR
+	if (GIsEditor)
+	{
+		// Force enable view family show flag for HighDPI derived's screen percentage.
+		ViewFamily.EngineShowFlags.ScreenPercentage = true;
+	}
+#endif
+
+	// Force screen percentage show flag to be turned off if not supported.
+	if (!ViewFamily.SupportsScreenPercentage())
+	{
+		ViewFamily.EngineShowFlags.ScreenPercentage = false;
+	}
+
+	const float DisplayClusterScreenPercentage = FMath::Max(0.01f, InFrameViewFamily.CustomBufferRatio);
+
+	// Get DPIDerivedResolutionFraction value:
+	float DPIDerivedResolutionFraction = 1.0f;
+	{
+#if WITH_EDITOR
+		// Automatically compute secondary resolution fraction from DPI.
+		// When in high res screenshot do not modify screen percentage based on dpi scale
+		if (GIsEditor && !GIsHighResScreenshot)
+		{
+			static auto CVarEditorViewportHighDPIPtr = IConsoleManager::Get().FindConsoleVariable(TEXT("r.Editor.Viewport.HighDPI"));
+
+			if (CVarEditorViewportHighDPIPtr && CVarEditorViewportHighDPIPtr->GetInt() == 0)
+			{
+				DPIDerivedResolutionFraction = FMath::Min(1.0f / InDPIScale, 1.0f);
+			}
+		}
+#endif
+	}
+
+	// Setup the screen percentage and upscaling method for the view family.
+	if (ViewFamily.EngineShowFlags.ScreenPercentage)
+	{
+		// Set up secondary resolution fraction for the view family.
+		{
+			static IConsoleVariable* CVarCustomSecondaryScreenPercentage =
+				IConsoleManager::Get().FindConsoleVariable(TEXT("r.SecondaryScreenPercentage.GameViewport"), false);
+			const float CustomSecondaryScreenPercentage = CVarCustomSecondaryScreenPercentage ?
+				CVarCustomSecondaryScreenPercentage->GetFloat() : -1.f;
+
+			if (CustomSecondaryScreenPercentage > 0.0)
+			{
+				// Override secondary resolution fraction with CVar.
+				ViewFamily.SecondaryViewFraction = FMath::Min(CustomSecondaryScreenPercentage / 100.0f, 1.0f);
+			}
+			else
+			{
+				ViewFamily.SecondaryViewFraction = DPIDerivedResolutionFraction;
+			}
+
+			check(ViewFamily.SecondaryViewFraction > 0.0f);
+		}
+
+		// nDisplay cannot use the dynamic resolution feature because it uses multiple viewports (view families) in a single frame.
+		const bool bUsesDynamicResolution = false;
+
+		// If a screen percentage interface was not set by dynamic resolution, then create one matching legacy behavior.
+		if (ViewFamily.GetScreenPercentageInterface() == nullptr)
+		{
+			float GlobalResolutionFraction = 1.0f;
+			float SecondaryScreenPercentage = ViewFamily.SecondaryViewFraction;
+
+			if (ViewFamily.EngineShowFlags.ScreenPercentage)
+			{
+				// Get global view fraction set by r.ScreenPercentage.
+				GlobalResolutionFraction = DisplayClusterScreenPercentage;
+
+				// We need to split the screen percentage if below 0.5 because TAA upscaling only works well up to 2x.
+				if (GlobalResolutionFraction < 0.5f)
+				{
+					SecondaryScreenPercentage = 2.0f * GlobalResolutionFraction;
+					GlobalResolutionFraction = 0.5f;
+				}
+			}
+
+			ViewFamily.SetScreenPercentageInterface(new FLegacyScreenPercentageDriver(
+				ViewFamily, GlobalResolutionFraction));
+
+			ViewFamily.SecondaryViewFraction = SecondaryScreenPercentage;
+		}
+
+		check(ViewFamily.GetScreenPercentageInterface() != nullptr);
+
+		// Change screen percentage method to raw output when doing dynamic resolution with VR if not using TAA upsample.
+		if (EnumHasAnyFlags(InRenderingFlags, EDisplayClusterViewportRenderingFlags::StereoRendering)
+			&& bUsesDynamicResolution)
+		{
+			for (FSceneView* View : Views)
+			{
+				if (View->PrimaryScreenPercentageMethod == EPrimaryScreenPercentageMethod::SpatialUpscale)
+				{
+					View->PrimaryScreenPercentageMethod = EPrimaryScreenPercentageMethod::RawOutput;
+				}
+			}
+		}
+	}
+}
+
 void FDisplayClusterViewportManager::RenderFrame(FViewport* InViewport)
 {
 	LightCardManager->RenderFrame();
diff --git a/Engine/Plugins/Runtime/nDisplay/Source/DisplayCluster/Private/Render/Viewport/DisplayClusterViewportManager.h b/Engine/Plugins/Runtime/nDisplay/Source/DisplayCluster/Private/Render/Viewport/DisplayClusterViewportManager.h
--- a/Engine/Plugins/Runtime/nDisplay/Source/DisplayCluster/Private/Render/Viewport/DisplayClusterViewportManager.h
+++ b/Engine/Plugins/Runtime/nDisplay/Source/DisplayCluster/Private/Render/Viewport/DisplayClusterViewportManager.h
@@ -90,7 +90,21 @@ public:
 		const bool bInAdditionalViewFamily
 	) const override;
 
-	virtual void ConfigureViewFamily(const FDisplayClusterRenderFrameTarget& InFrameTarget, const FDisplayClusterRenderFrameTargetViewFamily& InFrameViewFamily, FSceneViewFamilyContext& InOutViewFamily) override;
+	virtual void ConfigureViewFamily(
+		const FDisplayClusterRenderFrameTarget& InFrameTarget,
+		const FDisplayClusterRenderFrameTargetViewFamily& InFrameViewFamily,
+		FSceneViewFamilyContext& InOutViewFamily
+	) override;
+
+	virtual void PostConfigureViewFamily(
+		const FDisplayClusterRenderFrameTarget& InFrameTarget,
+		const FDisplayClusterRenderFrameTargetViewFamily& InFrameViewFamily,
+		FSceneViewFamilyContext& InOutViewFamily,
+		const TArray<FSceneView*>& InOutViews,
+		const EDisplayClusterViewportRenderingFlags InRenderingFlags = EDisplayClusterViewportRenderingFlags::None,
+		const float InDPIScale = 1.f
+	) override;
+
 	virtual void RenderFrame(FViewport* InViewport) override;
 
 private:
diff --git a/Engine/Plugins/Runtime/nDisplay/Source/DisplayCluster/Private/Render/Viewport/Preview/DisplayClusterViewportManagerPreview_Render.cpp b/Engine/Plugins/Runtime/nDisplay/Source/DisplayCluster/Private/Render/Viewport/Preview/DisplayClusterViewportManagerPreview_Render.cpp
--- a/Engine/Plugins/Runtime/nDisplay/Source/DisplayCluster/Private/Render/Viewport/Preview/DisplayClusterViewportManagerPreview_Render.cpp
+++ b/Engine/Plugins/Runtime/nDisplay/Source/DisplayCluster/Private/Render/Viewport/Preview/DisplayClusterViewportManagerPreview_Render.cpp
@@ -160,6 +160,7 @@ bool FDisplayClusterViewportManagerPreview::InitializeClusterNodePreview(const E
 							}
 						}
 
+						TArray<FSceneView*> Views;
 						for (const FDisplayClusterRenderFrameTargetView& ViewIt : ViewFamiliesIt.Views)
 						{
 							FDisplayClusterViewport* ViewportPtr = static_cast<FDisplayClusterViewport*>(ViewIt.Viewport.Get());
@@ -175,6 +176,8 @@ bool FDisplayClusterViewportManagerPreview::InitializeClusterNodePreview(const E
 								// Calculate the player's view information.
 								if (FSceneView* View = ViewportPtr->ViewportPreview->CalcSceneView(*ViewFamily, ViewIt.ContextNum))
 								{
+									Views.Add(View);
+
 									// Apply viewport context settings to view (crossGPU, visibility, etc)
 									ViewportPtr->SetupSceneView(ViewIt.ContextNum, PreviewScene->GetWorld(), *ViewFamily, *View);
 								}
@@ -188,10 +191,7 @@ bool FDisplayClusterViewportManagerPreview::InitializeClusterNodePreview(const E
 
 						if (!ViewFamily->Views.IsEmpty())
 						{
-							// Screen percentage is still not supported in scene capture.
-							ViewFamily->EngineShowFlags.ScreenPercentage = false;
-							ViewFamily->SetScreenPercentageInterface(new FLegacyScreenPercentageDriver(*ViewFamily, 1.0f));
-
+							ViewportManager->PostConfigureViewFamily(RenderTargetIt, ViewFamiliesIt, *ViewFamily, Views);
 							ViewportsViewFamily.Add(ViewFamily);
 						}
 					}
diff --git a/Engine/Plugins/Runtime/nDisplay/Source/DisplayCluster/Public/Render/Viewport/Containers/DisplayClusterViewport_Enums.h b/Engine/Plugins/Runtime/nDisplay/Source/DisplayCluster/Public/Render/Viewport/Containers/DisplayClusterViewport_Enums.h
--- a/Engine/Plugins/Runtime/nDisplay/Source/DisplayCluster/Public/Render/Viewport/Containers/DisplayClusterViewport_Enums.h
+++ b/Engine/Plugins/Runtime/nDisplay/Source/DisplayCluster/Public/Render/Viewport/Containers/DisplayClusterViewport_Enums.h
@@ -239,3 +239,16 @@ enum class EDisplayClusterViewportCameraPostProcessFlags : uint8
 	All = 0xFF
 };
 ENUM_CLASS_FLAGS(EDisplayClusterViewportCameraPostProcessFlags);
+
+/**
+* Rules for customizing nDisplay views and view families for the renderer.
+*/
+enum class EDisplayClusterViewportRenderingFlags : uint8
+{
+	// No flags
+	None = 0,
+
+	// Stereo rendering: Change screen percentage method to raw output when doing dynamic resolution with VR if not using TAA upsample.
+	StereoRendering = 1 << 0,
+};
+ENUM_CLASS_FLAGS(EDisplayClusterViewportRenderingFlags);
diff --git a/Engine/Plugins/Runtime/nDisplay/Source/DisplayCluster/Public/Render/Viewport/IDisplayClusterViewportManager.h b/Engine/Plugins/Runtime/nDisplay/Source/DisplayCluster/Public/Render/Viewport/IDisplayClusterViewportManager.h
--- a/Engine/Plugins/Runtime/nDisplay/Source/DisplayCluster/Public/Render/Viewport/IDisplayClusterViewportManager.h
+++ b/Engine/Plugins/Runtime/nDisplay/Source/DisplayCluster/Public/Render/Viewport/IDisplayClusterViewportManager.h
@@ -11,6 +11,7 @@ class IDisplayClusterWarpPolicy;
 class UWorld;
 class FViewport;
 class FSceneViewFamilyContext;
+class FSceneView;
 class ADisplayClusterRootActor;
 class UDisplayClusterConfigurationViewport;
 class IDisplayClusterViewportManagerProxy;
@@ -106,10 +107,36 @@ public:
 	* [Game thread func]
 	*
 	* @param InFrameTarget          - frame target
-	* @param OutRenderOutViewFamily - output family
+	* @param InFrameViewFamily      - the frame data for the view family
+	* @param OutRenderOutViewFamily - (in, out) The view family that will be configured.
 	*
 	*/
-	virtual void ConfigureViewFamily(const FDisplayClusterRenderFrameTarget& InFrameTarget, const FDisplayClusterRenderFrameTargetViewFamily& InFrameViewFamily, FSceneViewFamilyContext& InOutViewFamily) = 0;
+	virtual void ConfigureViewFamily(
+		const FDisplayClusterRenderFrameTarget& InFrameTarget,
+		const FDisplayClusterRenderFrameTargetViewFamily& InFrameViewFamily,
+		FSceneViewFamilyContext& InOutViewFamily) = 0;
+
+	/**
+	* Post-Initialize view family, using rules
+	* Should be called after all 'FSceneView' have been added.
+	* This function configures the screen percentage.
+	* [Game thread func]
+	*
+	* @param InFrameTarget            - frame target
+	* @param InFrameViewFamily        - the frame data for the view family
+	* @param InOutRenderOutViewFamily - (in, out) The view family that will be configured.
+	* @param InOutViews               - (in, out) The views in the family that will be configured.
+	* @param InRenderingFlags         - (in, opt) Rules for customizing nDisplay views and view families for the renderer.
+	* @param InDPIScale               - (in, opt ) DPI scale value
+	*
+	*/
+	virtual void PostConfigureViewFamily(
+		const FDisplayClusterRenderFrameTarget& InFrameTarget,
+		const FDisplayClusterRenderFrameTargetViewFamily& InFrameViewFamily,
+		FSceneViewFamilyContext& InOutViewFamily,
+		const TArray<FSceneView*>& InOutViews,
+		const EDisplayClusterViewportRenderingFlags InRenderingFlags = EDisplayClusterViewportRenderingFlags::None,
+		const float InDPIScale = 1.f) = 0;
 
 	/** Send to render thread. */
 	virtual void RenderFrame(FViewport* InViewport) = 0;
diff --git a/Engine/Plugins/Runtime/nDisplay/Source/DisplayClusterConfiguration/Public/DisplayClusterConfigurationTypes_Viewport.h b/Engine/Plugins/Runtime/nDisplay/Source/DisplayClusterConfiguration/Public/DisplayClusterConfigurationTypes_Viewport.h
--- a/Engine/Plugins/Runtime/nDisplay/Source/DisplayClusterConfiguration/Public/DisplayClusterConfigurationTypes_Viewport.h
+++ b/Engine/Plugins/Runtime/nDisplay/Source/DisplayClusterConfiguration/Public/DisplayClusterConfigurationTypes_Viewport.h
@@ -268,28 +268,30 @@ public:
 	UPROPERTY()
 	bool bShouldUseParentViewportRenderFamily = false;
 
-	// Multiply all viewports RTT size's for whole cluster by this value
-	UPROPERTY(BlueprintReadWrite, Category = NDisplay, meta = (DisplayName = "Global Viewport RTT Size Multiplier", ClampMin = "0.05", UIMin = "0.05", ClampMax = "10", UIMax = "10"))
+	// Multiplies the RTT size of all viewports within nDisplay by this value.
+	UPROPERTY(BlueprintReadWrite, Category = NDisplay, meta = (DisplayName = "Global RTT Size Multiplier", ClampMin = "0.05", UIMin = "0.05", ClampMax = "10", UIMax = "10"))
 	float ClusterRenderTargetRatioMult = 1.f;
 
-	// Multiply inner frustum RTT size's for whole cluster by this value
+	// Multiplies the RTT size of the ICVFX Inner Frustum viewports by this value.
 	UPROPERTY(BlueprintReadWrite, Category = NDisplay, meta = (DisplayName = "Inner Frustum RTT Size Multiplier", ClampMin = "0.05", UIMin = "0.05", ClampMax = "10", UIMax = "10"))
 	float ClusterICVFXInnerViewportRenderTargetRatioMult = 1.f;
 
-	// Multiply outer viewports RTT size's for whole cluster by this value
-	UPROPERTY(BlueprintReadWrite, Category = NDisplay, meta = (DisplayName = "Outer Viewport RTT Size Multiplier", ClampMin = "0.05", UIMin = "0.05", ClampMax = "10", UIMax = "10"))
+	// Multiplies the RTT size of the viewports by this value.
+	// (Excluding ICVFX internal viewports such as Inner frustum, LightCards, Chromakey, etc.)
+	UPROPERTY(BlueprintReadWrite, Category = NDisplay, meta = (DisplayName = "Viewports RTT Size Multiplier", ClampMin = "0.05", UIMin = "0.05", ClampMax = "10", UIMax = "10"))
 	float ClusterICVFXOuterViewportRenderTargetRatioMult = 1.f;
 
-	// Multiply all buffer ratios for whole cluster by this value
-	UPROPERTY(BlueprintReadWrite, Category = NDisplay, meta = (DisplayName = "Global Viewport Screen Percentage Multiplier", ClampMin = "0.05", UIMin = "0.05", ClampMax = "10", UIMax = "10"))
+	// Multiplies all screen percentages within nDisplay by this value.
+	UPROPERTY(BlueprintReadWrite, Category = NDisplay, meta = (DisplayName = "Global Screen Percentage Multiplier", ClampMin = "0.05", UIMin = "0.05", ClampMax = "10", UIMax = "10"))
 	float ClusterBufferRatioMult = 1.f;
 
-	// Multiply inner frustums buffer ratios for whole cluster by this value
+	// Multiplies the screen percentage for all ICVFX Inner Frustum viewports by this value.
 	UPROPERTY(BlueprintReadWrite, Category = NDisplay, meta = (DisplayName = "Inner Frustum Screen Percentage Multiplier", ClampMin = "0.05", UIMin = "0.05", ClampMax = "10", UIMax = "10"))
 	float ClusterICVFXInnerFrustumBufferRatioMult = 1.f;
 
-	// Multiply the screen percentage for all viewports in the cluster by this value.
-	UPROPERTY(EditAnywhere, BlueprintReadWrite, Category = NDisplay, meta = (DisplayName = "Viewport Screen Percentage Multiplier", ClampMin = "0.05", UIMin = "0.05", ClampMax = "10", UIMax = "1"))
+	// Multiplies the screen percentage for viewports by this value.
+	// (Excluding ICVFX internal viewports such as Inner Frustum, LightCards and Chromakey.)
+	UPROPERTY(EditAnywhere, BlueprintReadWrite, Category = NDisplay, meta = (DisplayName = "Viewports Screen Percentage Multiplier", ClampMin = "0.05", UIMin = "0.05", ClampMax = "10", UIMax = "1"))
 	float ClusterICVFXOuterViewportBufferRatioMult = 1.f;
 
 	// Allow warpblend render
