Change the interceptors handling of the sync event so that it occurs immediately after the pre tick for all nodes.  This will allow the finalization of the workspace to happen at the same time across all nodes.  Previously the arrival of the finalize sync message could arrive at any point in the frame and thus introducing a race condition where some clients would process the finalize.  Now the gate for processing the finalize event is coordinate at the beginning of the tick.

diff --git a/Engine/Plugins/Runtime/nDisplay/Source/DisplayClusterMessageInterceptor/Private/DisplayClusterMessageInterceptionModule.cpp b/Engine/Plugins/Runtime/nDisplay/Source/DisplayClusterMessageInterceptor/Private/DisplayClusterMessageInterceptionModule.cpp
--- a/Engine/Plugins/Runtime/nDisplay/Source/DisplayClusterMessageInterceptor/Private/DisplayClusterMessageInterceptionModule.cpp
+++ b/Engine/Plugins/Runtime/nDisplay/Source/DisplayClusterMessageInterceptor/Private/DisplayClusterMessageInterceptionModule.cpp
@@ -40,6 +40,42 @@ namespace DisplayClusterInterceptionModuleUtils
 	static const FString EventParameterSettings = TEXT("Settings");
 }
 
+class FSetBoolOnPreTick
+{
+public:
+	operator bool() const
+	{
+		return bValue;
+	}
+
+	FSetBoolOnPreTick& operator=(bool bInValue)
+	{
+		if (bInValue != bValue)
+		{
+			bPendingValue = bInValue;
+		}
+		return *this;
+	}
+
+	void ResetWithValue(bool bInValue)
+	{
+		bValue = bInValue;
+		bPendingValue = {};
+	}
+
+	void UpdateOnPreTick()
+	{
+		if (bPendingValue)
+		{
+			bValue = *bPendingValue;
+			bPendingValue = {};
+		}
+	}
+
+private:
+	bool bValue = false;
+	TOptional<bool> bPendingValue = {};
+};
 /**
  * Display Cluster Message Interceptor module
  * Intercept a specified set of message bus messages that are received across all the display nodes
@@ -230,7 +266,7 @@ private:
 
 	bool CanFinalizeWorkspaceSync() const
 	{
-		return bCanFinalizeWorkspace || bWasEverDisconnected;
+		return CanFinalizeWorkspace || bWasEverDisconnected;
 	}
 
 	bool CanProcessPendingPackages() const
@@ -264,7 +300,7 @@ private:
 			// activity sync event.
 			//
 			bWasEverDisconnected = true;
-			bCanFinalizeWorkspace = true;
+			CanFinalizeWorkspace = true;
 			bAllowPackages = true;
 
 			UE_LOG(LogDisplayClusterInterception, Display, TEXT("Disconnecting from session."));
@@ -293,7 +329,7 @@ private:
 
 	void OnMultiUserShutdown(TSharedRef<IConcertClientSession> InSession)
 	{
-		bCanFinalizeWorkspace = true;
+		CanFinalizeWorkspace = true;
 	}
 #endif
 
@@ -370,6 +406,7 @@ private:
 		{
 			Interceptor->SyncMessages();
 		}
+		CanFinalizeWorkspace.UpdateOnPreTick();
 	}
 
 	void StartInterception()
@@ -412,8 +449,8 @@ private:
 		NodesReady.Add(InEvent.Name);
 		if (NodesReady.Num() >= (int32)ClusterManager->GetNodesAmount())
 		{
-			UE_LOG(LogDisplayClusterInterception, Display, TEXT("Allowing multi-user workspace sync."));
-			bCanFinalizeWorkspace = true;
+			UE_LOG(LogDisplayClusterInterception, Display, TEXT("Allowing multi-user workspace sync on next pre-tick."));
+			CanFinalizeWorkspace = true;
 		}
 	}
 
@@ -439,13 +476,15 @@ private:
 	void ResetFinalizeSync()
 	{
 		UE_LOG(LogDisplayClusterInterception, Display, TEXT("Temporarily disabling multi-user workspace sync."));
-		bCanFinalizeWorkspace = false;
+		CanFinalizeWorkspace.ResetWithValue(false);
+
 		// Note, we intentially do not reset NodesReady here because of the timing of joining the MU session thereby causing the node
 		// to forever not receive updates.
 	}
 
 	bool bWasEverDisconnected = false;
-	bool bCanFinalizeWorkspace = true;
+	FSetBoolOnPreTick CanFinalizeWorkspace;
+
 	TSet<FString> NodesReady;
 
 	bool bAllowPackages = true;
