nDisplay: The user interface elements on the OutputMapping have been modified to improve usability.

- the "Snapping" popup window has been renamed to "Overlap & Snapping".
- The "Overlap & Bounds" section has been moved from the "Advanced menu" to this popup.
- A new section has been added to this popup with the text "Snapping (To Activate, hold SHIFT)" to clarify the user experience.

diff --git a/Engine/Plugins/Runtime/nDisplay/Source/DisplayClusterConfigurator/Private/Views/OutputMapping/Widgets/SDisplayClusterConfiguratorOutputMappingToolbar.cpp b/Engine/Plugins/Runtime/nDisplay/Source/DisplayClusterConfigurator/Private/Views/OutputMapping/Widgets/SDisplayClusterConfiguratorOutputMappingToolbar.cpp
--- a/Engine/Plugins/Runtime/nDisplay/Source/DisplayClusterConfigurator/Private/Views/OutputMapping/Widgets/SDisplayClusterConfiguratorOutputMappingToolbar.cpp
+++ b/Engine/Plugins/Runtime/nDisplay/Source/DisplayClusterConfigurator/Private/Views/OutputMapping/Widgets/SDisplayClusterConfiguratorOutputMappingToolbar.cpp
@@ -108,9 +108,9 @@ TSharedRef<SWidget> SDisplayClusterConfiguratorOutputMappingToolbar::MakeToolBar
 			.Cursor(EMouseCursor::Default)
 			.Style(ToolBarStyle)
 			.OnGetMenuContent(this, &SDisplayClusterConfiguratorOutputMappingToolbar::MakeSnappingMenu)
-			.ToolTipText(LOCTEXT("SnappingMenu_ToolTip", "Node Snapping"))
+			.ToolTipText(LOCTEXT("SnappingMenu_ToolTip", "Node Overlap and Snapping"))
 			.Icon(FSlateIcon(FDisplayClusterConfiguratorStyle::Get().GetStyleSetName(), "DisplayClusterConfigurator.OutputMapping.Snapping"))
-			.Label(LOCTEXT("AlignmentSettings_Label", "Snapping"))
+			.Label(LOCTEXT("AlignmentSettings_Label", "Overlap & Snapping"))
 			.ParentToolBar(SharedThis(this))
 		);
 	}
@@ -182,7 +182,14 @@ TSharedRef<SWidget> SDisplayClusterConfiguratorOutputMappingToolbar::MakeSnappin
 	const bool bShouldCloseWindowAfterMenuSelection = false;
 	FMenuBuilder MenuBuilder(bShouldCloseWindowAfterMenuSelection, CommandList);
 
-	MenuBuilder.BeginSection(TEXT("AdjacentEdgeSnapping"));
+	MenuBuilder.BeginSection(TEXT("OverlapBoundsSection"), LOCTEXT("OverlapBoundsSectionLabel", "Overlap & Bounds"));
+	{
+		MenuBuilder.AddMenuEntry(FDisplayClusterConfiguratorOutputMappingCommands::Get().ToggleClusterItemOverlap);
+		MenuBuilder.AddMenuEntry(FDisplayClusterConfiguratorOutputMappingCommands::Get().ToggleLockClusterNodesInHosts);
+	}
+	MenuBuilder.EndSection();
+
+	MenuBuilder.BeginSection(TEXT("SnappingSection"), LOCTEXT("SnappingSectionLabel", "Snapping (To Activate, hold SHIFT)"));
 	{
 		MenuBuilder.AddMenuEntry(FDisplayClusterConfiguratorOutputMappingCommands::Get().ToggleAdjacentEdgeSnapping);
 
@@ -212,15 +219,11 @@ TSharedRef<SWidget> SDisplayClusterConfiguratorOutputMappingToolbar::MakeSnappin
 			LOCTEXT("AlignmentSettings_AdjacentEdgePadding", "Adjacent Edge Padding")
 		);
 	}
-	MenuBuilder.EndSection();
 
-	MenuBuilder.BeginSection(TEXT("SameEdgeSnapping"));
 	{
 		MenuBuilder.AddMenuEntry(FDisplayClusterConfiguratorOutputMappingCommands::Get().ToggleSameEdgeSnapping);
 	}
-	MenuBuilder.EndSection();
 
-	MenuBuilder.BeginSection(TEXT("GeneralSnapping"));
 	{
 		MenuBuilder.AddWidget(
 			SNew(SBox)
@@ -312,13 +315,6 @@ TSharedRef<SWidget> SDisplayClusterConfiguratorOutputMappingToolbar::MakeAdvance
 	}
 	MenuBuilder.EndSection();
 
-	MenuBuilder.BeginSection(TEXT("OverlapBoundsSection"), LOCTEXT("OverlapBoundsSectionLabel", "Overlap & Bounds"));
-	{
-		MenuBuilder.AddMenuEntry(FDisplayClusterConfiguratorOutputMappingCommands::Get().ToggleClusterItemOverlap);
-		MenuBuilder.AddMenuEntry(FDisplayClusterConfiguratorOutputMappingCommands::Get().ToggleLockClusterNodesInHosts);
-	}
-	MenuBuilder.EndSection();
-
 	MenuBuilder.BeginSection(TEXT("Hosts"), LOCTEXT("HostsSectionLabel", "Hosts"));
 	{
 		MenuBuilder.AddSubMenu(
