nDisplay: Add a new resizing method that respects pixels when adapting In-Camera RTT resolution.

Added new CVar "nDisplay.icvfx.camera.AdaptResolution":
Adapt camera viewport resolution with 'Filmback + CropSettings + SqueezeFactor' CineCamera settings.  (Default = 1)
  0 - Disabled.
  1 - Preserve Pixel Area: Pixels = NewWidth * NewHeight = Width * Height.
  2 - Preserve Longest Dimension: Uses the maximum value of the camera frame size as the basis for the longest side of the sensor.

Custom Frustum (Overscan)
* Resolution->Overscan: Camera resolution is used as the base value when the units are pixels to convert to percentages.
* Overscan->Resolution: Overscan with the bAdaptResolution=OFF option does not change the RTT aspect ratio.
* CustomFrustum has restrictions on values, and this logic is implemented in FDisplayClusterViewport_CustomFrustumRuntimeSettings. Therefore, we cannot use the parameters of the FDisplayClusterConfigurationICVFX_CameraCustomFrustum directly. All related to that logic have been moved to the 'DisplayCluster' module.

ICVFX shader parameters
* ICVFX Camera Border: the aspect ratio of the camera texture should be obtained directly from the actual texture at the very end.
* ICVFX Camera Soft edges:  the original code directly uses FDisplayClusterViewport_CustomFrustumRuntimeSettings. This is incorrect, as previously stated. Reworked.

From FDisplayClusterConfigurationICVFX_ChromakeyRenderSettings, the CustomSize parameter was removed and a new parameter ChromakeySizeMult (???Chromakey resolution multiplier???) was added.

diff --git a/Engine/Plugins/Runtime/nDisplay/Source/DisplayCluster/Private/Game/EngineClasses/Scene/DisplayClusterICVFXCameraComponent.cpp b/Engine/Plugins/Runtime/nDisplay/Source/DisplayCluster/Private/Game/EngineClasses/Scene/DisplayClusterICVFXCameraComponent.cpp
--- a/Engine/Plugins/Runtime/nDisplay/Source/DisplayCluster/Private/Game/EngineClasses/Scene/DisplayClusterICVFXCameraComponent.cpp
+++ b/Engine/Plugins/Runtime/nDisplay/Source/DisplayCluster/Private/Game/EngineClasses/Scene/DisplayClusterICVFXCameraComponent.cpp
@@ -23,6 +23,83 @@
 #include "DisplayClusterEnums.h"
 #include "Version/DisplayClusterICVFXCameraCustomVersion.h"
 
+/** Enumerates the values of CVar "nDisplay.icvfx.camera.AdaptResolution". */
+enum class EICVFXCameraAdaptResolutionMethod : uint8
+{
+	// The size doesn't change at all.
+	Disabled = 0,
+
+	// Respect pixels: Pixels = NewWidth * NewHeight = Width * Height.
+	PreservePixelArea = 1,
+
+	// Constant Pixel Area : Uses the maximum value of the camera frame size as the basis for the longest side of the sensor.
+	PreserveLongestDimension = 2,
+
+	// Max value in this enum.
+	MAX = PreserveLongestDimension
+};
+
+/** Current method used to change the icvfx camera resolution. */
+int32 GDisplayClusterICVFXCameraAdaptResolution = (uint8)EICVFXCameraAdaptResolutionMethod::PreservePixelArea;
+static FAutoConsoleVariableRef CVarGDisplayClusterICVFXCameraAdaptResolution(
+	TEXT("nDisplay.icvfx.camera.AdaptResolution"),
+	GDisplayClusterICVFXCameraAdaptResolution,
+	TEXT("Adapt camera viewport resolution with 'Filmback + CropSettings + SqueezeFactor' CineCamera settings.  (Default = 1)\n")
+	TEXT("0 - Disabled.\n")
+	TEXT("1 - Preserve Pixel Area: Pixels = NewWidth * NewHeight = Width * Height.\n")
+	TEXT("2 - Preserve Longest Dimension : Uses the maximum value of the camera frame size as the basis for the longest side of the sensor.\n"),
+	ECVF_Default
+);
+
+namespace UE::DisplayClusterICVFXCameraComponent
+{
+	static inline FIntPoint AdaptResolutionToAspectRatio(const FIntPoint& InResolution, const float InDesiredAspectRatio)
+	{
+		// Decode the method of resizing.
+		const EICVFXCameraAdaptResolutionMethod ResizeMethod = (EICVFXCameraAdaptResolutionMethod)(FMath::Clamp(
+			GDisplayClusterICVFXCameraAdaptResolution, 0, (uint8)EICVFXCameraAdaptResolutionMethod::MAX));
+
+		// Implements resizing methods.
+		switch (ResizeMethod)
+		{
+		case EICVFXCameraAdaptResolutionMethod::PreservePixelArea:
+		{
+			// AR = (W/H) -> W = (AR*H)
+			// Pixels = W * H ->   Pixels = ((AR*H) * H) -> H^2 = (Pixels/AR)
+			// H = sqrt(Pixels/AR), W = (AR*H)
+			const int32 Pixels = InResolution.X * InResolution.Y;
+			const float Height = FMath::Sqrt(Pixels / InDesiredAspectRatio);
+			const float Width = InDesiredAspectRatio * Height;
+
+			// Get new camera size
+			return FIntPoint(FMath::RoundToInt(Width), FMath::RoundToInt(Height));
+		}
+
+		case EICVFXCameraAdaptResolutionMethod::PreserveLongestDimension:
+		{
+			// Use the max size of the RTT as a basis.
+			const float BasisDimension = InResolution.GetMax();
+
+			const float Width = (InDesiredAspectRatio >= 1.0)
+				? BasisDimension
+				: BasisDimension * InDesiredAspectRatio;
+
+			const float Height = (InDesiredAspectRatio >= 1.0)
+				? BasisDimension / InDesiredAspectRatio
+				: BasisDimension;
+
+			// Get new camera size
+			return FIntPoint(FMath::RoundToInt(Width), FMath::RoundToInt(Height));
+		}
+
+		default:
+			break;
+		}
+
+		return InResolution;
+	}
+};
+
 UDisplayClusterICVFXCameraComponent::UDisplayClusterICVFXCameraComponent(const FObjectInitializer& ObjectInitializer)
 {
 	PrimaryComponentTick.bCanEverTick = true;
@@ -185,6 +262,122 @@ void UDisplayClusterICVFXCameraComponent::ApplyICVFXCameraPostProcessesToViewpor
 	}
 }
 
+FDisplayClusterShaderParameters_ICVFX::FCameraSettings UDisplayClusterICVFXCameraComponent::GetICVFXCameraShaderParameters(const FDisplayClusterConfigurationICVFX_StageSettings& InStageSettings, const FDisplayClusterConfigurationICVFX_CameraSettings& InCameraSettings)
+{
+	FDisplayClusterShaderParameters_ICVFX::FCameraSettings OutCameraSettings;
+
+	const float CameraMult =
+		CameraSettings.GetCameraBufferRatio(InStageSettings)
+		* CameraSettings.CustomFrustum.GetCameraAdaptResolutionRatio(InStageSettings);
+
+	const FIntPoint CameraFrameSize = GetICVFXCameraFrameSize(InStageSettings, InCameraSettings);
+	const FIntPoint RealInnerFrustumResolution(CameraFrameSize.X * CameraMult, CameraFrameSize.Y * CameraMult);
+
+	FIntRect RealViewportRect(FIntPoint(0, 0), RealInnerFrustumResolution);
+	FDisplayClusterViewport_CustomFrustumSettings RealFrustumSettings;
+	FDisplayClusterViewport_CustomFrustumRuntimeSettings RealFrustumRuntimeSettings;
+
+	FDisplayClusterViewportConfigurationHelpers_ICVFX::UpdateCameraCustomFrustum(CameraSettings.CustomFrustum, RealFrustumSettings);
+	FDisplayClusterViewport_CustomFrustumRuntimeSettings::UpdateCustomFrustumSettings(GetName(), RealFrustumSettings, RealFrustumRuntimeSettings, RealViewportRect);
+
+	const FDisplayClusterViewport_CustomFrustumRuntimeSettings::FCustomFrustumPercent& Angles = RealFrustumRuntimeSettings.CustomFrustumPercent;
+
+	// Camera Border
+	{
+		if (InCameraSettings.Border.Enable)
+		{
+			const float RealThicknessScaleValue = 0.1f;
+
+			OutCameraSettings.InnerCameraBorderColor = InCameraSettings.Border.Color;
+			OutCameraSettings.InnerCameraBorderThickness = InCameraSettings.Border.Thickness;
+		}
+		else
+		{
+			// No border:
+			OutCameraSettings.InnerCameraBorderColor = FLinearColor::Black;
+			OutCameraSettings.InnerCameraBorderThickness = 0.0f;
+		}
+	}
+
+	// Camera Soft Edges
+	{
+		FVector4 SoftEdge(
+			// remap values from 0-1 GUI range into acceptable 0.0 - 0.25 shader range
+			FMath::GetMappedRangeValueClamped(FVector2D(0.0, 1.0f), FVector2D(0.0, 0.25), InCameraSettings.SoftEdge.Horizontal), // Left
+			FMath::GetMappedRangeValueClamped(FVector2D(0.0, 1.0f), FVector2D(0.0, 0.25), InCameraSettings.SoftEdge.Vertical), // Top
+
+			// ZW now used in other way
+			// Z for new parameter Feather
+			InCameraSettings.SoftEdge.Feather
+		);
+
+
+		SoftEdge.X /= (1 + Angles.Left + Angles.Right);
+		SoftEdge.Y /= (1 + Angles.Top + Angles.Bottom);
+
+		OutCameraSettings.SoftEdge = SoftEdge;
+	}
+
+	return OutCameraSettings;
+}
+
+FIntPoint UDisplayClusterICVFXCameraComponent::GetICVFXCameraFrameSize(const FDisplayClusterConfigurationICVFX_StageSettings& InStageSettings, const FDisplayClusterConfigurationICVFX_CameraSettings& InCameraSettings)
+{
+	using namespace UE::DisplayClusterICVFXCameraComponent;
+
+	const FIntPoint CameraFrameSize = InCameraSettings.RenderSettings.CustomFrameSize.bUseCustomSize
+		? FIntPoint(InCameraSettings.RenderSettings.CustomFrameSize.CustomWidth, InCameraSettings.RenderSettings.CustomFrameSize.CustomHeight)
+		: FIntPoint(InStageSettings.DefaultFrameSize.Width, InStageSettings.DefaultFrameSize.Height);
+
+	UCineCameraComponent* ActualCineCameraComponent = GetActualCineCameraComponent();
+	if (!ActualCineCameraComponent)
+	{
+		// Adaptation math requires an actual CineCamera component
+		return CameraFrameSize;
+	}
+
+	// User can disable this feature.
+	const bool bAdaptFrameSize = InCameraSettings.RenderSettings.CustomFrameSize.bUseCustomSize
+		? InCameraSettings.RenderSettings.CustomFrameSize.bAdaptSize
+		: InStageSettings.DefaultFrameSize.bAdaptSize;
+
+	// Get the size of the cinematic camera's cropped sensor:
+	const double CroppedSensorWidth  = FMath::Tan(FMath::DegreesToRadians(ActualCineCameraComponent->GetHorizontalFieldOfView()) / 2.f) * 2.f * ActualCineCameraComponent->CurrentFocalLength;
+	const double CroppedSensorHeight = FMath::Tan(FMath::DegreesToRadians(ActualCineCameraComponent->GetVerticalFieldOfView()) / 2.f) * 2.f * ActualCineCameraComponent->CurrentFocalLength;
+
+	if (!(CroppedSensorWidth > 0.f && CroppedSensorHeight > 0.f) || !bAdaptFrameSize)
+	{
+		// The CineCamera cropped sensor size has invalid values.
+		// or this feature is disabled.
+		return CameraFrameSize;
+	}
+
+	// Desired aspect ratio
+	double CroppedSensorAR = CroppedSensorWidth / CroppedSensorHeight;
+
+	// When no adopt resolution is used we must compensate for the change in aspect ratio caused by overscan.
+	if (InCameraSettings.CustomFrustum.bEnable && !InCameraSettings.CustomFrustum.bAdaptResolution)
+	{
+		// Overscan should only be used through this api:
+		FDisplayClusterViewport_CustomFrustumSettings CustomFrustumSettings;
+		FDisplayClusterViewport_CustomFrustumRuntimeSettings CustomFrustumRuntimeSettings;
+
+		FDisplayClusterViewportConfigurationHelpers_ICVFX::UpdateCameraCustomFrustum(InCameraSettings.CustomFrustum, CustomFrustumSettings);
+
+		FIntPoint DesiredSize = AdaptResolutionToAspectRatio(CameraFrameSize, CroppedSensorAR);
+		FIntRect ViewportRect(FIntPoint(0, 0), DesiredSize);
+		FDisplayClusterViewport_CustomFrustumRuntimeSettings::UpdateCustomFrustumSettings(GetName(), CustomFrustumSettings, CustomFrustumRuntimeSettings, ViewportRect);
+
+		// Overscan without the bAdaptResolution option does not change the RTT aspect ratio.
+		// In this case, the sensor AR must be modified to include CustomFrustumPercent values.
+		CroppedSensorAR *= CustomFrustumRuntimeSettings.CustomFrustumPercent.GetAspectRatioMult();
+	}
+
+	const FIntPoint AdaptedCameraFrameSize = AdaptResolutionToAspectRatio(CameraFrameSize, CroppedSensorAR);
+
+	return AdaptedCameraFrameSize;
+}
+
 void UDisplayClusterICVFXCameraComponent::UpdateOverscanEstimatedFrameSize()
 {
 	const ADisplayClusterRootActor* RootActor = Cast<ADisplayClusterRootActor>(GetOwner());
@@ -195,31 +388,24 @@ void UDisplayClusterICVFXCameraComponent::UpdateOverscanEstimatedFrameSize()
 
 	const FDisplayClusterConfigurationICVFX_StageSettings& StageSettings = RootActor->GetStageSettings();
 
-	UCineCameraComponent* ActualCineCameraComponent = GetActualCineCameraComponent();
-	check(ActualCineCameraComponent);
-
-	// additional multipliers from FDisplayClusterConfigurationRenderFrame are not used in following calculations
-	const float CameraBufferRatio = CameraSettings.GetCameraBufferRatio(StageSettings);
-	const FIntPoint CameraFrameSize = CameraSettings.GetCameraFrameSize(StageSettings, *ActualCineCameraComponent);
-	const FIntPoint InnerFrustumResolution( CameraFrameSize.X * CameraBufferRatio, CameraFrameSize.Y * CameraBufferRatio);
-	
 	{
 		// calculate estimations
-		FDisplayClusterConfigurationICVFX_CameraCustomFrustum EstimatedCustomFrustum = CameraSettings.CustomFrustum;
-		EstimatedCustomFrustum.bEnable = true;
-		EstimatedCustomFrustum.bAdaptResolution = true;
+		FDisplayClusterConfigurationICVFX_CameraSettings EstimatedCameraSettings = CameraSettings;
+		EstimatedCameraSettings.CustomFrustum.bEnable = true;
+		EstimatedCameraSettings.CustomFrustum.bAdaptResolution = true;
+
+		const float CameraMult =
+			EstimatedCameraSettings.GetCameraBufferRatio(StageSettings)
+			* EstimatedCameraSettings.CustomFrustum.GetCameraAdaptResolutionRatio(StageSettings);
 
-		const float EstimatedCameraAdaptResolutionRatio = EstimatedCustomFrustum.GetCameraAdaptResolutionRatio(StageSettings);
-		const FIntPoint EstimatedInnerFrustumResolution(
-			InnerFrustumResolution.X * EstimatedCameraAdaptResolutionRatio,
-			InnerFrustumResolution.Y * EstimatedCameraAdaptResolutionRatio
-	);
+		const FIntPoint CameraFrameSize = GetICVFXCameraFrameSize(StageSettings, EstimatedCameraSettings);
+		const FIntPoint EstimatedInnerFrustumResolution(CameraFrameSize.X * CameraMult, CameraFrameSize.Y * CameraMult);
 
 		FIntRect EstimatedViewportRect(FIntPoint(0, 0), EstimatedInnerFrustumResolution);
 		FDisplayClusterViewport_CustomFrustumSettings EstimatedFrustumSettings;
 		FDisplayClusterViewport_CustomFrustumRuntimeSettings EstimatedFrustumRuntimeSettings;
 
-		FDisplayClusterViewportConfigurationHelpers_ICVFX::UpdateCameraCustomFrustum(EstimatedCustomFrustum, EstimatedFrustumSettings);
+		FDisplayClusterViewportConfigurationHelpers_ICVFX::UpdateCameraCustomFrustum(EstimatedCameraSettings.CustomFrustum, EstimatedFrustumSettings);
 		FDisplayClusterViewport_CustomFrustumRuntimeSettings::UpdateCustomFrustumSettings(GetName(), EstimatedFrustumSettings, EstimatedFrustumRuntimeSettings, EstimatedViewportRect);
 
 		// Assign estimated calculated values
@@ -227,12 +413,12 @@ void UDisplayClusterICVFXCameraComponent::UpdateOverscanEstimatedFrameSize()
 	}
 
 	{
-		// calculate real
-		const float RealCameraAdaptResolutionRatio = CameraSettings.CustomFrustum.GetCameraAdaptResolutionRatio(StageSettings);
-		const FIntPoint RealInnerFrustumResolution(
-			InnerFrustumResolution.X * RealCameraAdaptResolutionRatio,
-			InnerFrustumResolution.Y * RealCameraAdaptResolutionRatio
-		);
+		const float CameraMult =
+			CameraSettings.GetCameraBufferRatio(StageSettings)
+			* CameraSettings.CustomFrustum.GetCameraAdaptResolutionRatio(StageSettings);
+
+		const FIntPoint CameraFrameSize = GetICVFXCameraFrameSize(StageSettings, CameraSettings);
+		const FIntPoint RealInnerFrustumResolution(CameraFrameSize.X * CameraMult, CameraFrameSize.Y * CameraMult);
 
 		FIntRect RealViewportRect(FIntPoint(0, 0), RealInnerFrustumResolution);
 		FDisplayClusterViewport_CustomFrustumSettings RealFrustumSettings;
diff --git a/Engine/Plugins/Runtime/nDisplay/Source/DisplayCluster/Private/Render/Viewport/Configuration/DisplayClusterViewportConfigurationHelpers_ICVFX.cpp b/Engine/Plugins/Runtime/nDisplay/Source/DisplayCluster/Private/Render/Viewport/Configuration/DisplayClusterViewportConfigurationHelpers_ICVFX.cpp
--- a/Engine/Plugins/Runtime/nDisplay/Source/DisplayCluster/Private/Render/Viewport/Configuration/DisplayClusterViewportConfigurationHelpers_ICVFX.cpp
+++ b/Engine/Plugins/Runtime/nDisplay/Source/DisplayCluster/Private/Render/Viewport/Configuration/DisplayClusterViewportConfigurationHelpers_ICVFX.cpp
@@ -470,37 +470,6 @@ bool FDisplayClusterViewportConfigurationHelpers_ICVFX::IsCameraUsed(const FDisp
 	return true;
 }
 
-FDisplayClusterShaderParameters_ICVFX::FCameraSettings FDisplayClusterViewportConfigurationHelpers_ICVFX::GetShaderParametersCameraSettings(const FDisplayClusterViewport& InCameraViewport, UDisplayClusterICVFXCameraComponent& InCameraComponent, const FDisplayClusterConfigurationICVFX_CameraSettings& InCameraSettings)
-{
-	FDisplayClusterShaderParameters_ICVFX::FCameraSettings Result;
-
-	ADisplayClusterRootActor* SceneRootActor = InCameraViewport.Configuration->GetRootActor(EDisplayClusterRootActorType::Scene);
-	ADisplayClusterRootActor* ConfigurationRootActor = InCameraViewport.Configuration->GetRootActor(EDisplayClusterRootActorType::Configuration);
-	const FDisplayClusterConfigurationICVFX_StageSettings* StageSettings = InCameraViewport.Configuration->GetStageSettings();
-
-	if (SceneRootActor && ConfigurationRootActor && StageSettings)
-	{
-		Result.Resource.ViewportId = InCameraViewport.GetId();
-
-		UCineCameraComponent* ActualCineCameraComponent = InCameraComponent.GetActualCineCameraComponent();
-		check(ActualCineCameraComponent);
-
-		// Get camera border settings
-		InCameraSettings.GetCameraBorder(*StageSettings, Result.InnerCameraBorderColor, Result.InnerCameraBorderThickness);
-		Result.InnerCameraFrameAspectRatio = InCameraSettings.GetCameraFrameAspectRatio(*StageSettings, *ActualCineCameraComponent);
-
-		// Soft edges
-		Result.SoftEdge = InCameraSettings.GetCameraSoftEdge(*StageSettings, *ActualCineCameraComponent);
-
-		// Rendering order for camera overlap
-		const FString InnerFrustumID = InCameraComponent.GetCameraUniqueId();
-		const int32 CameraRenderOrder = ConfigurationRootActor->GetInnerFrustumPriority(InnerFrustumID);
-		Result.RenderOrder = (CameraRenderOrder < 0) ? InCameraSettings.RenderSettings.RenderOrder : CameraRenderOrder;
-	}
-
-	return Result;
-}
-
 void FDisplayClusterViewportConfigurationHelpers_ICVFX::UpdateCameraViewportSettings(FDisplayClusterViewport& DstViewport, UDisplayClusterICVFXCameraComponent& InCameraComponent, const FDisplayClusterConfigurationICVFX_CameraSettings& InCameraSettings)
 {
 	const FDisplayClusterConfigurationICVFX_StageSettings* StageSettings = DstViewport.Configuration->GetStageSettings();
@@ -530,11 +499,8 @@ void FDisplayClusterViewportConfigurationHelpers_ICVFX::UpdateCameraViewportSett
 	// FDisplayClusterConfigurationICVFX_CameraSettings
 	InOutRenderSettings.CameraId.Empty();
 
-	UCineCameraComponent* ActualCineCameraComponent = InCameraComponent.GetActualCineCameraComponent();
-	check(ActualCineCameraComponent);
-
 	// UDisplayClusterConfigurationICVFX_CameraRenderSettings
-	const FIntPoint DesiredSize = InCameraSettings.GetCameraFrameSize(*StageSettings, *ActualCineCameraComponent);
+	const FIntPoint DesiredSize = InCameraComponent.GetICVFXCameraFrameSize(*StageSettings, InCameraSettings);
 
 	InOutRenderSettings.Rect = FDisplayClusterViewportHelpers::GetValidViewportRect(FIntRect(FIntPoint(0, 0), DesiredSize), DstViewport.GetId(), TEXT("Configuration Camera Frame Size"));
 
@@ -604,15 +570,9 @@ void FDisplayClusterViewportConfigurationHelpers_ICVFX::UpdateChromakeyViewportS
 
 		DstViewport.UpdateConfiguration_OverlayRenderSettings(InChromakeyRenderSettings->AdvancedRenderSettings);
 
-		// Support custom overlay size
-		if (InChromakeyRenderSettings->CustomSize.bUseCustomSize)
-		{
-			FIntPoint DesiredSize;
-			DesiredSize.X = InChromakeyRenderSettings->CustomSize.CustomWidth;
-			DesiredSize.Y = InChromakeyRenderSettings->CustomSize.CustomHeight;
-
-			InOutRenderSettings.Rect = FDisplayClusterViewportHelpers::GetValidViewportRect(FIntRect(FIntPoint(0, 0), DesiredSize), DstViewport.GetId(), TEXT("Configuration custom chromakey Frame Size"));
-		}
+		// Resize chromakey RTT
+		const FIntPoint ChromakeyTextureSize = InOutRenderSettings.Rect.Size() * InChromakeyRenderSettings->ChromakeySizeMult;
+		InOutRenderSettings.Rect = FDisplayClusterViewportHelpers::GetValidViewportRect(FIntRect(FIntPoint(0, 0), ChromakeyTextureSize), DstViewport.GetId(), TEXT("Configuration custom chromakey Frame Size"));
 
 		// Debug: override the texture of the target viewport from this chromakeyRTT
 		if (InChromakeyRenderSettings->bReplaceCameraViewport)
diff --git a/Engine/Plugins/Runtime/nDisplay/Source/DisplayCluster/Private/Render/Viewport/Configuration/DisplayClusterViewportConfigurationHelpers_ICVFX.h b/Engine/Plugins/Runtime/nDisplay/Source/DisplayCluster/Private/Render/Viewport/Configuration/DisplayClusterViewportConfigurationHelpers_ICVFX.h
--- a/Engine/Plugins/Runtime/nDisplay/Source/DisplayCluster/Private/Render/Viewport/Configuration/DisplayClusterViewportConfigurationHelpers_ICVFX.h
+++ b/Engine/Plugins/Runtime/nDisplay/Source/DisplayCluster/Private/Render/Viewport/Configuration/DisplayClusterViewportConfigurationHelpers_ICVFX.h
@@ -61,8 +61,6 @@ public:
 	static bool CreateProjectionPolicyCameraICVFX(FDisplayClusterViewportConfiguration& InConfiguration, UDisplayClusterICVFXCameraComponent& CameraComponent, const FDisplayClusterConfigurationICVFX_CameraSettings& InCameraSettings, TSharedPtr<IDisplayClusterProjectionPolicy, ESPMode::ThreadSafe>& OutProjPolicy);
 	static bool UpdateCameraProjectionSettingsICVFX(FDisplayClusterViewportConfiguration& InConfiguration, UDisplayClusterICVFXCameraComponent& CameraComponent, const FDisplayClusterConfigurationICVFX_CameraSettings& InCameraSettings, const TSharedPtr<IDisplayClusterProjectionPolicy, ESPMode::ThreadSafe>& InProjectionPolicy);
 
-	static FDisplayClusterShaderParameters_ICVFX::FCameraSettings GetShaderParametersCameraSettings(const FDisplayClusterViewport& InCameraViewport, UDisplayClusterICVFXCameraComponent& InCameraComponent, const FDisplayClusterConfigurationICVFX_CameraSettings& InCameraSettings);
-	
 	/** Configures the camera's chromakey render settings to match the specified configuration settings  */
 	static void UpdateCameraSettings_Chromakey(FDisplayClusterShaderParameters_ICVFX::FCameraSettings& InOutCameraSettings, const FDisplayClusterConfigurationICVFX_StageSettings& InStageSettings, const FDisplayClusterConfigurationICVFX_CameraSettings& InCameraSettings, bool bEnableChromakeyMarkers, const FString& InChromakeyViewportId);
 	static void UpdateCameraSettings_OverlapChromakey(FDisplayClusterShaderParameters_ICVFX::FCameraSettings& InOutCameraSettings, const FDisplayClusterConfigurationICVFX_StageSettings& InStageSettings, const FDisplayClusterConfigurationICVFX_CameraSettings& InCameraSettings, bool bEnableChromakeyMarkers);
diff --git a/Engine/Plugins/Runtime/nDisplay/Source/DisplayCluster/Private/Render/Viewport/Configuration/DisplayClusterViewportConfiguration_ICVFXCamera.cpp b/Engine/Plugins/Runtime/nDisplay/Source/DisplayCluster/Private/Render/Viewport/Configuration/DisplayClusterViewportConfiguration_ICVFXCamera.cpp
--- a/Engine/Plugins/Runtime/nDisplay/Source/DisplayCluster/Private/Render/Viewport/Configuration/DisplayClusterViewportConfiguration_ICVFXCamera.cpp
+++ b/Engine/Plugins/Runtime/nDisplay/Source/DisplayCluster/Private/Render/Viewport/Configuration/DisplayClusterViewportConfiguration_ICVFXCamera.cpp
@@ -76,17 +76,34 @@ void FDisplayClusterViewportConfiguration_ICVFXCamera::Update()
 	{
 		if (CameraViewport.IsValid())
 		{
-			FDisplayClusterShaderParameters_ICVFX::FCameraSettings ShaderParametersCameraSettings = FDisplayClusterViewportConfigurationHelpers_ICVFX::GetShaderParametersCameraSettings(*CameraViewport, CameraComponent, GetCameraSettings());
+			ADisplayClusterRootActor* SceneRootActor = CameraViewport->Configuration->GetRootActor(EDisplayClusterRootActorType::Scene);
+			ADisplayClusterRootActor* ConfigurationRootActor = CameraViewport->Configuration->GetRootActor(EDisplayClusterRootActorType::Configuration);
+			const FDisplayClusterConfigurationICVFX_StageSettings* StageSettings = CameraViewport->Configuration->GetStageSettings();
 
-			// Add this camera data to all visible targets:
-			for (TSharedPtr<FDisplayClusterViewport, ESPMode::ThreadSafe>& ViewportIt : VisibleTargets)
+			if (SceneRootActor && ConfigurationRootActor && StageSettings)
 			{
-				if (ViewportIt.IsValid())
+				const FDisplayClusterConfigurationICVFX_CameraSettings& CameraSettings = GetCameraSettings();
+
+				FDisplayClusterShaderParameters_ICVFX::FCameraSettings ShaderParametersCameraSettings =
+					CameraComponent.GetICVFXCameraShaderParameters(*StageSettings, CameraSettings);
+
+				ShaderParametersCameraSettings.Resource.ViewportId = CameraViewport->GetId();
+
+				// Rendering order for camera overlap
+				const FString InnerFrustumID = CameraComponent.GetCameraUniqueId();
+				const int32 CameraRenderOrder = ConfigurationRootActor->GetInnerFrustumPriority(InnerFrustumID);
+				ShaderParametersCameraSettings.RenderOrder = (CameraRenderOrder < 0) ? CameraSettings.RenderSettings.RenderOrder : CameraRenderOrder;
+
+				// Add this camera data to all visible targets:
+				for (TSharedPtr<FDisplayClusterViewport, ESPMode::ThreadSafe>& ViewportIt : VisibleTargets)
 				{
-					// Gain direct access to internal settings of the viewport:
-					FDisplayClusterViewport_RenderSettingsICVFX& InOutRenderSettingsICVFX = ViewportIt->GetRenderSettingsICVFXImpl();
+					if (ViewportIt.IsValid())
+					{
+						// Gain direct access to internal settings of the viewport:
+						FDisplayClusterViewport_RenderSettingsICVFX& InOutRenderSettingsICVFX = ViewportIt->GetRenderSettingsICVFXImpl();
 
-					InOutRenderSettingsICVFX.ICVFX.Cameras.Add(ShaderParametersCameraSettings);
+						InOutRenderSettingsICVFX.ICVFX.Cameras.Add(ShaderParametersCameraSettings);
+					}
 				}
 			}
 		}
diff --git a/Engine/Plugins/Runtime/nDisplay/Source/DisplayCluster/Private/Render/Viewport/Containers/DisplayClusterViewport_CustomFrustumRuntimeSettings.h b/Engine/Plugins/Runtime/nDisplay/Source/DisplayCluster/Private/Render/Viewport/Containers/DisplayClusterViewport_CustomFrustumRuntimeSettings.h
--- a/Engine/Plugins/Runtime/nDisplay/Source/DisplayCluster/Private/Render/Viewport/Containers/DisplayClusterViewport_CustomFrustumRuntimeSettings.h
+++ b/Engine/Plugins/Runtime/nDisplay/Source/DisplayCluster/Private/Render/Viewport/Containers/DisplayClusterViewport_CustomFrustumRuntimeSettings.h
@@ -27,13 +27,41 @@ struct FDisplayClusterViewport_CustomFrustumRuntimeSettings
 	* @param InOutTop                  - the value of the top projection plane that you want to change
 	* @parma InOutBottom               - the value of the bottom projection plane that you want to change
 	*/
-	static bool UpdateProjectionAngles(const FDisplayClusterViewport_CustomFrustumRuntimeSettings& InRuntimeSettings, const FIntPoint& InRenderTargetSize, double& InOutLeft, double& InOutRight, double& InOutTop, double& InOutBottom);
+	static bool UpdateProjectionAngles(
+		const FDisplayClusterViewport_CustomFrustumRuntimeSettings& InRuntimeSettings,
+		const FIntPoint& InRenderTargetSize,
+		double& InOutLeft,
+		double& InOutRight,
+		double& InOutTop,
+		double& InOutBottom);
 
 	/**
-	* Values in percent
+	* These are frustum modifiers, where 1 means 100% of the frustum size.
+	* 
+	* Positive values mean frustum expansion:
+	*   FrustumAngles.Left   -= FrustumAngles.Width * Left.
+	*   FrustumAngles.Right  += FrustumAngles.Width * Right.
+	*   FrustumAngles.Top    += FrustumAngles.Height * Top.
+	*   FrustumAngles.Bottom -= FrustumAngles.Height * Bottom.
 	*/
 	struct FCustomFrustumPercent
 	{
+		/** Returns the aspect ratio multiplier. */
+		inline double GetAspectRatioMult() const
+		{
+			// Positive values mean frustum expansion
+			const double WidthMult = 1 + Left + Right;
+			const double HeightMult = 1 + Top + Bottom;
+
+			if (WidthMult <= 0.f || HeightMult <= 0.f)
+			{
+				// If extreme values are used, 1 is returned, which means that the aspect ratio remains unchanged.
+				return 1;
+			}
+
+			return WidthMult / HeightMult;
+		}
+
 		double Left = 0;
 		double Right = 0;
 		double Top = 0;
diff --git a/Engine/Plugins/Runtime/nDisplay/Source/DisplayCluster/Private/Render/Viewport/Containers/DisplayClusterViewport_OverscanRuntimeSettings.cpp b/Engine/Plugins/Runtime/nDisplay/Source/DisplayCluster/Private/Render/Viewport/Containers/DisplayClusterViewport_OverscanRuntimeSettings.cpp
--- a/Engine/Plugins/Runtime/nDisplay/Source/DisplayCluster/Private/Render/Viewport/Containers/DisplayClusterViewport_OverscanRuntimeSettings.cpp
+++ b/Engine/Plugins/Runtime/nDisplay/Source/DisplayCluster/Private/Render/Viewport/Containers/DisplayClusterViewport_OverscanRuntimeSettings.cpp
@@ -33,7 +33,7 @@ namespace UE::DisplayCluster::Viewport::OverscanHelpers
 		// The idea behind the overscan is to add extra space on the sides of the RTT.
 		// Note: this only applies to regular viewports (Outers, etc.).
 		// 
-		// The inner frustum viewport has its own implementation for the overscan feature called “CustomFrustum”.
+		// The inner frustum viewport has its own implementation for the overscan feature called "CustomFrustum".
 		// (see FDisplayClusterViewport_CustomFrustumRuntimeSettings)
 
 		return FMath::Clamp(InValue, 0, MaxCustomFrustumValue);
diff --git a/Engine/Plugins/Runtime/nDisplay/Source/DisplayCluster/Private/Render/Viewport/LightCard/DisplayClusterViewportLightCardEnums.h b/Engine/Plugins/Runtime/nDisplay/Source/DisplayCluster/Private/Render/Viewport/LightCard/DisplayClusterViewportLightCardEnums.h
--- a/Engine/Plugins/Runtime/nDisplay/Source/DisplayCluster/Private/Render/Viewport/LightCard/DisplayClusterViewportLightCardEnums.h
+++ b/Engine/Plugins/Runtime/nDisplay/Source/DisplayCluster/Private/Render/Viewport/LightCard/DisplayClusterViewportLightCardEnums.h
@@ -10,10 +10,10 @@ enum class EDisplayClusterUVLightCardRenderMode : uint8
 	// Use light card render mode from the StageSettings
 	Default,
 
-	// The lightcard will always be displayed only “over the In-Camera”.
+	// The lightcard will always be displayed only "Over the In-Camera".
 	AlwaysOver,
 
-	// The lightcard will always be displayed only “under the In-Camera”.
+	// The lightcard will always be displayed only "Under the In-Camera".
 	AlwaysUnder,
 
 	// Lighting maps are not used.
diff --git a/Engine/Plugins/Runtime/nDisplay/Source/DisplayCluster/Private/Render/Viewport/LightCard/DisplayClusterViewportLightCardManager.cpp b/Engine/Plugins/Runtime/nDisplay/Source/DisplayCluster/Private/Render/Viewport/LightCard/DisplayClusterViewportLightCardManager.cpp
--- a/Engine/Plugins/Runtime/nDisplay/Source/DisplayCluster/Private/Render/Viewport/LightCard/DisplayClusterViewportLightCardManager.cpp
+++ b/Engine/Plugins/Runtime/nDisplay/Source/DisplayCluster/Private/Render/Viewport/LightCard/DisplayClusterViewportLightCardManager.cpp
@@ -103,12 +103,12 @@ EDisplayClusterUVLightCardRenderMode FDisplayClusterViewportLightCardManager::Ge
 			const EDisplayClusterViewportICVFXFlags LightcardRenderModeFlags = ICVFXFlags & EDisplayClusterViewportICVFXFlags::LightcardRenderModeMask;
 			if (LightcardRenderModeFlags == EDisplayClusterViewportICVFXFlags::LightcardAlwaysUnder)
 			{
-				// The lightcard will always be displayed only “under the In-Camera” for this cluster node.
+				// The lightcard will always be displayed only "Under the In-Camera" for this cluster node.
 				return EDisplayClusterUVLightCardRenderMode::AlwaysUnder;
 			}
 			else if (LightcardRenderModeFlags == EDisplayClusterViewportICVFXFlags::LightcardAlwaysOver)
 			{
-				// The lightcard will always be displayed only “over the In-Camera” for this cluster node.
+				// The lightcard will always be displayed only "Over the In-Camera" for this cluster node.
 				return EDisplayClusterUVLightCardRenderMode::AlwaysOver;
 			}
 		}
diff --git a/Engine/Plugins/Runtime/nDisplay/Source/DisplayCluster/Public/Components/DisplayClusterICVFXCameraComponent.h b/Engine/Plugins/Runtime/nDisplay/Source/DisplayCluster/Public/Components/DisplayClusterICVFXCameraComponent.h
--- a/Engine/Plugins/Runtime/nDisplay/Source/DisplayCluster/Public/Components/DisplayClusterICVFXCameraComponent.h
+++ b/Engine/Plugins/Runtime/nDisplay/Source/DisplayCluster/Public/Components/DisplayClusterICVFXCameraComponent.h
@@ -11,6 +11,7 @@
 #include "DisplayClusterEditorPropertyReference.h"
 #include "Render/Viewport/Containers/DisplayClusterViewport_Enums.h"
 #include "Render/Viewport/Containers/DisplayClusterViewport_CameraMotionBlur.h"
+#include "ShaderParameters/DisplayClusterShaderParameters_ICVFX.h"
 
 #include "DisplayClusterICVFXCameraComponent.generated.h"
 
@@ -91,6 +92,12 @@ public:
 	UFUNCTION(BlueprintCallable, Category = "ICVFX Camera")
 	void SetDepthOfFieldParameters(const FDisplayClusterConfigurationICVFX_CameraDepthOfField& NewDepthOfFieldParams);
 
+	/** Get ICVFX camera frame resolution. */
+	FIntPoint GetICVFXCameraFrameSize(const FDisplayClusterConfigurationICVFX_StageSettings& InStageSettings, const FDisplayClusterConfigurationICVFX_CameraSettings& InCameraSettings);
+
+	/** Get ICVFX camera shader parameters. */
+	FDisplayClusterShaderParameters_ICVFX::FCameraSettings GetICVFXCameraShaderParameters(const FDisplayClusterConfigurationICVFX_StageSettings& InStageSettings, const FDisplayClusterConfigurationICVFX_CameraSettings& InCameraSettings);
+
 private:
 	void UpdateOverscanEstimatedFrameSize();
 
diff --git a/Engine/Plugins/Runtime/nDisplay/Source/DisplayClusterConfiguration/Private/DisplayClusterConfigurationTypes_ICVFX.cpp b/Engine/Plugins/Runtime/nDisplay/Source/DisplayClusterConfiguration/Private/DisplayClusterConfigurationTypes_ICVFX.cpp
--- a/Engine/Plugins/Runtime/nDisplay/Source/DisplayClusterConfiguration/Private/DisplayClusterConfigurationTypes_ICVFX.cpp
+++ b/Engine/Plugins/Runtime/nDisplay/Source/DisplayClusterConfiguration/Private/DisplayClusterConfigurationTypes_ICVFX.cpp
@@ -8,30 +8,7 @@
 #include "CineCameraComponent.h"
 #include "CineCameraActor.h"
 
-namespace UE::DisplayClusterConfiguration::ICVFX
-{
-	static float ClampPercent(float InValue)
-	{
-		static const float MaxCustomFrustumValue = 5.f;
-
-		return FMath::Clamp(InValue, -MaxCustomFrustumValue, MaxCustomFrustumValue);
-	}
-
-	static float ClampCustomFrustum(float InValue, float InMax)
-	{
-		return FMath::Clamp(InValue, -InMax, InMax);
-	}
-};
-
-using namespace UE::DisplayClusterConfiguration::ICVFX;
 
-int32 GDisplayClusterICVFXCameraAdoptResolution = 1;
-static FAutoConsoleVariableRef CVarGDisplayClusterICVFXCameraAdoptResolution(
-	TEXT("nDisplay.icvfx.camera.AdoptResolution"),
-	GDisplayClusterICVFXCameraAdoptResolution,
-	TEXT("Adopt camera viewport resolution with 'Filmback + CropSettings + SqueezeFactor' CineCamera settings.  (Default = 1)"),
-	ECVF_Default
-);
 
 ///////////////////////////////////////////////////////////////////////////////////////
 // FDisplayClusterConfigurationICVFX_ChromakeyMarkers
@@ -284,24 +261,6 @@ float FDisplayClusterConfigurationICVFX_CameraSettings::GetCameraBufferRatio(con
 	return BufferRatio;
 }
 
-bool FDisplayClusterConfigurationICVFX_CameraSettings::GetCameraBorder(const FDisplayClusterConfigurationICVFX_StageSettings& InStageSettings, FLinearColor& OutBorderColor, float& OutBorderThickness) const
-{
-	if (!Border.Enable)
-	{
-		OutBorderColor = FLinearColor::Black;
-		OutBorderThickness = 0.0f;
-
-		return false;
-	}
-
-	const float RealThicknessScaleValue = 0.1f;
-
-	OutBorderColor = Border.Color;
-	OutBorderThickness = Border.Thickness * RealThicknessScaleValue;
-
-	return true;
-}
-
 void FDisplayClusterConfigurationICVFX_CameraSettings::SetupViewInfo(const FDisplayClusterConfigurationICVFX_StageSettings& InStageSettings, FMinimalViewInfo& InOutViewInfo)
 {
 	RenderSettings.SetupViewInfo(InStageSettings, InOutViewInfo);
@@ -309,112 +268,6 @@ void FDisplayClusterConfigurationICVFX_CameraSettings::SetupViewInfo(const FDisp
 	CameraMotionBlur.SetupViewInfo(InStageSettings, InOutViewInfo);
 }
 
-FIntPoint FDisplayClusterConfigurationICVFX_CameraSettings::GetCameraFrameSize(const FDisplayClusterConfigurationICVFX_StageSettings& InStageSettings, const UCineCameraComponent& InCineCameraComponent) const
-{
-	const FIntPoint CameraFrameSize = RenderSettings.CustomFrameSize.bUseCustomSize
-		? FIntPoint(RenderSettings.CustomFrameSize.CustomWidth, RenderSettings.CustomFrameSize.CustomHeight)
-		: FIntPoint(InStageSettings.DefaultFrameSize.Width, InStageSettings.DefaultFrameSize.Height);
-
-	if (GDisplayClusterICVFXCameraAdoptResolution)
-	{
-		// Get the size of the cinematic camera's cropped sensor:
-		const double CropedSensorWidth  = FMath::Tan(FMath::DegreesToRadians(InCineCameraComponent.GetHorizontalFieldOfView()) / 2.f) * 2.f * InCineCameraComponent.CurrentFocalLength;
-		const double CropedSensorHeight = FMath::Tan(FMath::DegreesToRadians(InCineCameraComponent.GetVerticalFieldOfView()) / 2.f) * 2.f * InCineCameraComponent.CurrentFocalLength;
-
-		// Get the ratio of the cinematic camera's cropped sensor size to the base sensor size.
-		const double CroppedSensorWidthRatio  = CropedSensorWidth / InCineCameraComponent.Filmback.SensorWidth;
-		const double CroppedSensorHeightRatio = CropedSensorHeight / InCineCameraComponent.Filmback.SensorHeight;
-
-		// Adapt camera resolution to the filmback sensor aspect ratio
-		const double SensorAspectRatio = (InCineCameraComponent.Filmback.SensorHeight > 0.f && InCineCameraComponent.Filmback.SensorWidth > 0.f)
-			? InCineCameraComponent.Filmback.SensorWidth / InCineCameraComponent.Filmback.SensorHeight
-			: -1;
-
-		if (SensorAspectRatio < 0)
-		{
-			// The CineCamera sensor size has invalid values.
-			return CameraFrameSize;
-		}
-
-		// Use the max size as a base.
-		const double MaxSize = CameraFrameSize.GetMax();
-
-		const double CameraFrameWidth  = (SensorAspectRatio >= 1.0)
-			? MaxSize
-			: MaxSize * SensorAspectRatio;
-		const double CameraFrameHeight = (SensorAspectRatio >= 1.0)
-			? MaxSize / SensorAspectRatio
-			: MaxSize;
-
-		// Get cropped camera size
-		const FIntPoint CroppedCameraFrameSize(
-			FMath::RoundToInt(CameraFrameWidth  * CroppedSensorWidthRatio),
-			FMath::RoundToInt(CameraFrameHeight * CroppedSensorHeightRatio)
-		);
-
-		return CroppedCameraFrameSize;
-	}
-
-	return CameraFrameSize;
-}
-
-float FDisplayClusterConfigurationICVFX_CameraSettings::GetCameraFrameAspectRatio(const FDisplayClusterConfigurationICVFX_StageSettings& InStageSettings, const UCineCameraComponent& InCineCameraComponent) const
-{
-	FIntPoint FrameSize = GetCameraFrameSize(InStageSettings, InCineCameraComponent);
-
-	return (FrameSize.Y > 0 && FrameSize.X > 0) ? (float)FrameSize.X / float(FrameSize.Y) : 0;
-}
-
-FVector4 FDisplayClusterConfigurationICVFX_CameraSettings::GetCameraSoftEdge(const FDisplayClusterConfigurationICVFX_StageSettings& InStageSettings, const UCineCameraComponent& InCineCameraComponent) const
-{
-	FVector4 ResultSoftEdge(ForceInitToZero);
-
-	const float FieldOfViewMultiplier = CustomFrustum.GetCameraFieldOfViewMultiplier(InStageSettings);
-
-	// softedge adjustments	
-	const float Overscan = (FieldOfViewMultiplier > 0) ? FieldOfViewMultiplier : 1;
-
-	// remap values from 0-1 GUI range into acceptable 0.0 - 0.25 shader range
-	ResultSoftEdge.X = FMath::GetMappedRangeValueClamped(FVector2D(0.0, 1.0f), FVector2D(0.0, 0.25), SoftEdge.Horizontal) / Overscan; // Left
-	ResultSoftEdge.Y = FMath::GetMappedRangeValueClamped(FVector2D(0.0, 1.0f), FVector2D(0.0, 0.25), SoftEdge.Vertical) / Overscan; // Top
-
-	// ZW now used in other way
-	// Z for new parameter Feather
-	ResultSoftEdge.Z = SoftEdge.Feather;
-
-	// Custom frustum made changes to soft edges
-	if (CustomFrustum.bEnable)
-	{
-		// default - percents
-		const float ConvertToPercent = 0.01f;
-
-		float Left = ClampPercent(CustomFrustum.Left * ConvertToPercent);
-		float Right = ClampPercent(CustomFrustum.Right * ConvertToPercent);
-		float Top = ClampPercent(CustomFrustum.Top * ConvertToPercent);
-		float Bottom = ClampPercent(CustomFrustum.Bottom * ConvertToPercent);
-
-		if (CustomFrustum.Mode == EDisplayClusterConfigurationViewportCustomFrustumMode::Pixels)
-		{
-			const float CameraBufferRatio = GetCameraBufferRatio(InStageSettings);
-			const FIntPoint FrameSize = GetCameraFrameSize(InStageSettings, InCineCameraComponent);
-
-			const float  FrameWidth = FrameSize.X * CameraBufferRatio;
-			const float FrameHeight = FrameSize.Y * CameraBufferRatio;
-
-			Left = ClampPercent(CustomFrustum.Left / FrameWidth);
-			Right = ClampPercent(CustomFrustum.Right / FrameWidth);
-			Top = ClampPercent(CustomFrustum.Top / FrameHeight);
-			Bottom = ClampPercent(CustomFrustum.Bottom / FrameHeight);
-		}
-
-		// recalculate soft edge related offsets based on frustum
-		ResultSoftEdge.X /= (1 + Left + Right);
-		ResultSoftEdge.Y /= (1 + Top + Bottom);
-	}
-
-	return ResultSoftEdge;
-}
-
 ///////////////////////////////////////////////////////////////////////////////////////
 // FDisplayClusterConfigurationICVFX_StageSettings
 ///////////////////////////////////////////////////////////////////////////////////////
@@ -689,7 +542,7 @@ EDisplayClusterShaderParametersICVFX_LightCardRenderMode FDisplayClusterConfigur
 		break;
 	};
 
-	// By default, lightcards are rendered in “Over” mode.
+	// By default, lightcards are rendered in "Over" mode.
 	return EDisplayClusterShaderParametersICVFX_LightCardRenderMode::Over;
 }
 
diff --git a/Engine/Plugins/Runtime/nDisplay/Source/DisplayClusterConfiguration/Public/DisplayClusterConfigurationTypes_ICVFX.h b/Engine/Plugins/Runtime/nDisplay/Source/DisplayClusterConfiguration/Public/DisplayClusterConfigurationTypes_ICVFX.h
--- a/Engine/Plugins/Runtime/nDisplay/Source/DisplayClusterConfiguration/Public/DisplayClusterConfigurationTypes_ICVFX.h
+++ b/Engine/Plugins/Runtime/nDisplay/Source/DisplayClusterConfiguration/Public/DisplayClusterConfigurationTypes_ICVFX.h
@@ -135,38 +135,60 @@ public:
 	TArray<TSoftObjectPtr<AActor>> AutoAddedActors;
 };
 
+/** Custom resolution settings for ICVFX. */
 USTRUCT(Blueprintable)
 struct DISPLAYCLUSTERCONFIGURATION_API FDisplayClusterConfigurationICVFX_CustomSize
 {
 	GENERATED_BODY()
 
 public:
-	// Use custom size
-	UPROPERTY(EditAnywhere, BlueprintReadWrite, Category = NDisplay)
+	/** Enabling this option will allow these custom settings to be used instead of the default settings. */
+	UPROPERTY(EditAnywhere, BlueprintReadWrite, Category = NDisplay, meta = (DisplayName = "Use Custom Resolution"))
 	bool bUseCustomSize = false;
 
-	// Used when enabled "bUseCustomSize"
-	UPROPERTY(EditAnywhere, BlueprintReadWrite, Category = NDisplay, meta = (ClampMin = "32", UIMin = "32"))
+	/** Custom Width, in pixels. */
+	UPROPERTY(EditAnywhere, BlueprintReadWrite, Category = NDisplay, meta = (DisplayName = "Custom Width", ClampMin = "32", UIMin = "32", EditCondition = "bUseCustomSize"))
 	int CustomWidth = 2560;
 
-	// Used when enabled "bUseCustomSize"
-	UPROPERTY(EditAnywhere, BlueprintReadWrite, Category = NDisplay, meta = (ClampMin = "32", UIMin = "32"))
+	/** Custom Height, in pixels. */
+	UPROPERTY(EditAnywhere, BlueprintReadWrite, Category = NDisplay, meta = (DisplayName = "Custom Height", ClampMin = "32", UIMin = "32", EditCondition = "bUseCustomSize"))
 	int CustomHeight = 1440;
+
+	/** Automatically adapts resolution to be proportional
+	* to the filmback aspect ratio for the ICVFX Camera
+	* while continuing to render the same total amount of
+	* pixels as specified by the size. Enabling this can help
+	* avoid visual artifacts without affecting the
+	* performance budget.
+	*/
+	UPROPERTY(EditAnywhere, BlueprintReadWrite, Category = NDisplay, meta = (DisplayName = "Adapt Resolution to Match Filmback Aspect Ratio", EditCondition = "bUseCustomSize"))
+	bool bAdaptSize = true;
 };
 
+/** Default resolution settings for ICVFX. */
 USTRUCT(Blueprintable)
 struct DISPLAYCLUSTERCONFIGURATION_API FDisplayClusterConfigurationICVFX_Size
 {
 	GENERATED_BODY()
 
 public:
-	// Viewport width in pixels
-	UPROPERTY(EditAnywhere, BlueprintReadWrite, Category = NDisplay, meta = (ClampMin = "32", UIMin = "32"))
+	/** The default width of In-Cameras, in pixels. */
+	UPROPERTY(EditAnywhere, BlueprintReadWrite, Category = NDisplay, meta = (DisplayName = "Default Width", ClampMin = "32", UIMin = "32"))
 	int Width = 2560;
 
-	// Viewport height  in pixels
-	UPROPERTY(EditAnywhere, BlueprintReadWrite, Category = NDisplay, meta = (ClampMin = "32", UIMin = "32"))
+	/** The default height of In-Cameras, in pixels. */
+	UPROPERTY(EditAnywhere, BlueprintReadWrite, Category = NDisplay, meta = (DisplayName = "Default Height", ClampMin = "32", UIMin = "32"))
 	int Height = 1440;
+
+	/** Automatically adapts the default frame resolution to be proportional
+	* to the filmback aspect ratio for the ICVFX Camera
+	* while continuing to render the same total amount of
+	* pixels as specified by the size. Enabling this can help
+	* avoid visual artifacts without affecting the
+	* performance budget.
+	*/
+	UPROPERTY(EditAnywhere, BlueprintReadWrite, Category = NDisplay, meta = (DisplayName = "Adapt Resolution to Match Filmback Aspect Ratio"))
+	bool bAdaptSize = true;
 };
 
 USTRUCT(Blueprintable)
@@ -183,13 +205,23 @@ public:
 	UPROPERTY(EditAnywhere, BlueprintReadWrite, Category = NDisplay, meta = (ClampMin = "0.01", UIMin = "0.01", ClampMax = "1.0", UIMax = "1.0"))
 	float RenderTargetRatio = 1;
 
-	// Performance, Multi-GPU: Asign GPU for viewport rendering. The Value '-1' used to default gpu mapping (EYE_LEFT and EYE_RIGHT GPU)
-	UPROPERTY(EditAnywhere, BlueprintReadWrite, Category = NDisplay)
-	int GPUIndex = -1;
+	/**
+	* Specifies the GPU index for the ICVFX viewport (LC or CK).
+	* A value of '-1' means using the same GPU index as defined in the base viewport:
+	* - The In-Camera viewport is used as the base viewport for the Chromakey (CK) viewport.
+	* - An outer viewport is used as the base viewport for the Light Card (LC) viewport.
+	* Used to improve rendering performance by spreading the load across multiple GPUs.
+	*/
+	UPROPERTY(EditAnywhere, BlueprintReadWrite, Category = NDisplay, meta = (DisplayName = "GPU Index", ClampMin = "-1", UIMin = "-1", ClampMax = "8", UIMax = "8"))
+	int GPUIndex = INDEX_NONE;
 
-	// Performance, Multi-GPU: Customize GPU for stereo mode second view (EYE_RIGHT GPU)
-	UPROPERTY(EditAnywhere, BlueprintReadWrite, Category = NDisplay)
-	int StereoGPUIndex = -1;
+	/**
+	* Specifies the GPU index for the ICVFX viewport (LC or CK) in stereo rendering for the second eye.
+	* A value of '-1' means to use the value from the GPU Index parameter. (the same value is used for both eyes).
+	* Used to improve rendering performance by spreading the load across multiple GPUs.
+	*/
+	UPROPERTY(EditAnywhere, BlueprintReadWrite, Category = NDisplay, meta = (DisplayName = "Stereo GPU Index", ClampMin = "-1", UIMin = "-1", ClampMax = "8", UIMax = "8"))
+	int StereoGPUIndex = INDEX_NONE;
 
 	// Performance: force monoscopic render, resolved to stereo viewport
 	UPROPERTY(EditAnywhere, BlueprintReadWrite, Category = NDisplay)
@@ -231,9 +263,14 @@ public:
 	UPROPERTY(BlueprintReadWrite,Category = NDisplay)
 	bool bReplaceCameraViewport = false;
 
-	// Performance: Use custom size (low-res) for chromakey RTT frame. Default size same as camera frame
-	UPROPERTY(BlueprintReadWrite, Category = NDisplay)
-	FDisplayClusterConfigurationICVFX_CustomSize CustomSize;
+	/** Custom resolution of the chromakey RTT. If it is not used, the Default Frame Resolution value is used by default. */
+	UE_DEPRECATED(5.5, "Use the 'ChromakeySizeMult' instead")
+	UPROPERTY(meta = (DeprecatedProperty, DeprecationMessage = "Use the 'Chromakey Resolution Multiplier' instead"))
+	FDisplayClusterConfigurationICVFX_CustomSize CustomSize_DEPRECATED;
+
+	/** Resolution multiplier for RTT chromakey. The default is the ICVFX camera frame resolution. */
+	UPROPERTY(BlueprintReadWrite, Category = NDisplay, meta = (DisplayName = "Chromakey Resolution Multiplier", ClampMin = "0.05", UIMin = "0.05", ClampMax = "1.0", UIMax = "1.0"))
+	float ChromakeySizeMult = 1.f;
 
 	/** Content specified here will be overridden to use the chromakey color specified and include chromakey markers if enabled. */
 	UPROPERTY(EditAnywhere, BlueprintReadWrite, Category = NDisplay, meta = (DisplayName = "Custom Chromakey Content"))
@@ -465,13 +502,21 @@ public:
 	UPROPERTY(EditAnywhere, BlueprintReadWrite, Category = NDisplay, meta = (ClampMin = "0.01", UIMin = "0.01", ClampMax = "1.0", UIMax = "1.0"))
 	float RenderTargetRatio = 1;
 
-	// Performance, Multi-GPU: Asign GPU for viewport rendering. The Value '-1' used to default gpu mapping (EYE_LEFT and EYE_RIGHT GPU)
-	UPROPERTY(EditAnywhere, BlueprintReadWrite, Category = NDisplay)
-	int GPUIndex = -1;
+	/**
+	* Specifies the GPU index for the ICVFX camera viewport.
+	* Value '-1' means do not use multi-GPU
+	* Used to improve rendering performance by spreading the load across multiple GPUs.
+	*/
+	UPROPERTY(EditAnywhere, BlueprintReadWrite, Category = NDisplay, meta = (DisplayName = "GPU Index", ClampMin = "-1", UIMin = "-1", ClampMax = "8", UIMax = "8"))
+	int GPUIndex = INDEX_NONE;
 
-	// Performance, Multi-GPU: Customize GPU for stereo mode second view (EYE_RIGHT GPU)
-	UPROPERTY(EditAnywhere, BlueprintReadWrite, Category = NDisplay)
-	int StereoGPUIndex = -1;
+	/**
+	* Specifies the GPU index for the ICVFX camera viewport in stereo rendering for the second eye.
+	* A value of '-1' means to use the value from the GPU Index parameter. (the same value is used for both eyes).
+	* Used to improve rendering performance by spreading the load across multiple GPUs.
+	*/
+	UPROPERTY(EditAnywhere, BlueprintReadWrite, Category = NDisplay, meta = (DisplayName = "Stereo GPU Index", ClampMin = "-1", UIMin = "-1", ClampMax = "8", UIMax = "8"))
+	int StereoGPUIndex = INDEX_NONE;
 
 	// Performance: force monoscopic render, resolved to stereo viewport
 	UPROPERTY(EditAnywhere, BlueprintReadWrite, Category = NDisplay)
@@ -492,8 +537,8 @@ public:
 	FDisplayClusterConfigurationICVFX_CameraRenderSettings();
 
 public:
-	// Define custom inner camera viewport size
-	UPROPERTY(EditAnywhere, BlueprintReadWrite, Category = NDisplay)
+	/** Custom resolution of the ICVFX Camera. If it is not used, the Default Frame Resolution value is used by default. */
+	UPROPERTY(EditAnywhere, BlueprintReadWrite, Category = NDisplay, meta = (DisplayName = "Custom Frame Resolution"))
 	FDisplayClusterConfigurationICVFX_CustomSize CustomFrameSize;
 
 	// Camera render order, bigger value is over
@@ -665,12 +710,15 @@ struct DISPLAYCLUSTERCONFIGURATION_API FDisplayClusterConfigurationICVFX_CameraC
 {
 	GENERATED_BODY()
 
+	/** Expected ICVFX camera resolution when both 'Adapt Resolution' and 'Enable Inner Frustum Overscan' are enabled. */
 	UPROPERTY(VisibleAnywhere, BlueprintReadOnly, Category = NDisplay)
 	FIntPoint EstimatedOverscanResolution = { 2560, 1440 };
 	
+	/** Real ICVFX camera resolution for current settings. */
 	UPROPERTY(VisibleAnywhere, BlueprintReadOnly, Category = NDisplay)
 	FIntPoint InnerFrustumResolution = { 2560, 1440 };
 
+	/** This value shows the ratio of "Overscan Estimated Resolution" pixels to "Frustum Internal Resolution" pixels. */
 	UPROPERTY(BlueprintReadOnly, Category = NDisplay, meta = (HideInDetailPanel))
 	float OverscanPixelsIncrease = 0.f;
 
@@ -752,19 +800,23 @@ public:
 	bool IsChromakeyViewportSettingsEqual(const FString& InClusterNodeId1, const FString& InClusterNodeId2) const;
 
 	/** Return calculated soft edges values. */
-	FVector4 GetCameraSoftEdge(const FDisplayClusterConfigurationICVFX_StageSettings& InStageSettings, const UCineCameraComponent& InCineCameraComponent) const;
+	UE_DEPRECATED(5.5, "Use the UDisplayClusterICVFXCameraComponent::GetICVFXCameraShaderParameters() instead")
+	FVector4 GetCameraSoftEdge(const FDisplayClusterConfigurationICVFX_StageSettings& InStageSettings, const UCineCameraComponent& InCineCameraComponent) const { return FVector4(); }
 
 	/** Get camera buffer ratio. */
 	float GetCameraBufferRatio(const FDisplayClusterConfigurationICVFX_StageSettings& InStageSettings) const;
 
-	/** Get camera frame size. */
-	FIntPoint GetCameraFrameSize(const FDisplayClusterConfigurationICVFX_StageSettings& InStageSettings, const UCineCameraComponent& InCineCameraComponent) const;
+	/** Get camera frame resolution. */
+	UE_DEPRECATED(5.5, "Use the UDisplayClusterICVFXCameraComponent::GetICVFXCameraFrameSize() instead")
+	FIntPoint GetCameraFrameSize(const FDisplayClusterConfigurationICVFX_StageSettings& InStageSettings, const UCineCameraComponent& InCineCameraComponent) const { return FIntPoint(0, 0); }
 
 	/** Get camera frame aspect ratio. */
-	float GetCameraFrameAspectRatio(const FDisplayClusterConfigurationICVFX_StageSettings& InStageSettings, const UCineCameraComponent& InCineCameraComponent) const;
+	UE_DEPRECATED(5.5, "Use the UDisplayClusterICVFXCameraComponent::GetICVFXCameraFrameSize() instead")
+	float GetCameraFrameAspectRatio(const FDisplayClusterConfigurationICVFX_StageSettings& InStageSettings, const UCineCameraComponent& InCineCameraComponent) const { return 0.f; }
 
 	/** Get camera border settings. */
-	bool GetCameraBorder(const FDisplayClusterConfigurationICVFX_StageSettings& InStageSettings, FLinearColor& OutBorderColor, float& OutBorderThickness) const;
+	UE_DEPRECATED(5.5, "Use the UDisplayClusterICVFXCameraComponent::GetICVFXCameraShaderParameters() instead")
+	bool GetCameraBorder(const FDisplayClusterConfigurationICVFX_StageSettings& InStageSettings, FLinearColor& OutBorderColor, float& OutBorderThickness) const { return false; }
 
 	/** Sets up view info for each relevant setting such as render, custom frustrum and motion blur settings. */
 	void SetupViewInfo(const FDisplayClusterConfigurationICVFX_StageSettings& InStageSettings, FMinimalViewInfo& InOutViewInfo);
@@ -869,8 +921,8 @@ public:
 	UPROPERTY(EditAnywhere, BlueprintReadWrite, Category = "In Camera VFX", meta = (DisplayName = "Enable Inner Frustum Chromakey Overlap"))
 	bool bEnableInnerFrustumChromakeyOverlap = false;
 
-	/** Default incameras RTT texture size. */
-	UPROPERTY(EditAnywhere, BlueprintReadWrite, Category = Configuration)
+	/** Default In-Cameras texture resolution. */
+	UPROPERTY(EditAnywhere, BlueprintReadWrite, Category = Configuration, meta = (DisplayName = "Default Frame Resolution"))
 	FDisplayClusterConfigurationICVFX_Size DefaultFrameSize;
 
 	UPROPERTY(EditAnywhere, BlueprintReadWrite, Category = "Light Cards", meta = (ExpandProperties))
diff --git a/Engine/Plugins/Runtime/nDisplay/Source/DisplayClusterConfiguration/Public/DisplayClusterConfigurationTypes_Viewport.h b/Engine/Plugins/Runtime/nDisplay/Source/DisplayClusterConfiguration/Public/DisplayClusterConfigurationTypes_Viewport.h
--- a/Engine/Plugins/Runtime/nDisplay/Source/DisplayClusterConfiguration/Public/DisplayClusterConfigurationTypes_Viewport.h
+++ b/Engine/Plugins/Runtime/nDisplay/Source/DisplayClusterConfiguration/Public/DisplayClusterConfigurationTypes_Viewport.h
@@ -77,9 +77,13 @@ public:
 	UPROPERTY(EditAnywhere, BlueprintReadWrite, Category = "Rendering", meta = (DisplayName = "Enable Cross-GPU Transfer"))
 	bool bEnableCrossGPUTransfer = true;
 
-	/** Specify which GPU should render the second Stereo eye */
-	UPROPERTY(EditAnywhere, Category = "Stereo", meta = (DisplayName = "Stereo GPU Index"))
-	int StereoGPUIndex = -1;
+	/**
+	* Specifies the GPU index for the nDisplay viewport in stereo rendering for the second eye.
+	* A value of '-1' means to use the value from the GPU Index parameter. (the same value is used for both eyes).
+	* Used to improve rendering performance by spreading the load across multiple GPUs.
+	*/
+	UPROPERTY(EditAnywhere, Category = "Stereo", meta = (DisplayName = "Stereo GPU Index", ClampMin = "-1", UIMin = "-1", ClampMax = "8", UIMax = "8"))
+	int StereoGPUIndex = INDEX_NONE;
 
 	/** Enables and sets Stereo mode */
 	UPROPERTY(EditAnywhere, BlueprintReadWrite, Category = "Stereo")
@@ -216,8 +220,12 @@ public:
 	UPROPERTY(EditAnywhere, BlueprintReadWrite, Category = "Rendering")
 	int OverlapOrder = 0;
 
-	/** Specify which GPU should render this Viewport. "-1" is default. */
-	UPROPERTY(EditAnywhere, BlueprintReadWrite, Category = "Rendering", meta = (DisplayName = "GPU Index"))
+	/**
+	* Specifies the GPU index for the nDisplay viewport.
+	* Value '-1' means do not use multi-GPU
+	* Used to improve rendering performance by spreading the load across multiple GPUs.
+	*/
+	UPROPERTY(EditAnywhere, BlueprintReadWrite, Category = "Rendering", meta = (DisplayName = "GPU Index", ClampMin = "-1", UIMin = "-1", ClampMax = "8", UIMax = "8"))
 	int GPUIndex = -1;
 
 	// Configure render for this viewport
diff --git a/Engine/Plugins/Runtime/nDisplay/Source/DisplayClusterShaders/Private/Shaders/DisplayClusterShadersWarpblend_ICVFX.cpp b/Engine/Plugins/Runtime/nDisplay/Source/DisplayClusterShaders/Private/Shaders/DisplayClusterShadersWarpblend_ICVFX.cpp
--- a/Engine/Plugins/Runtime/nDisplay/Source/DisplayClusterShaders/Private/Shaders/DisplayClusterShadersWarpblend_ICVFX.cpp
+++ b/Engine/Plugins/Runtime/nDisplay/Source/DisplayClusterShaders/Private/Shaders/DisplayClusterShadersWarpblend_ICVFX.cpp
@@ -833,11 +833,20 @@ public:
 		RenderPassData.PSParameters.InnerCameraSampler = TStaticSamplerState<SF_Trilinear, AM_Clamp, AM_Clamp, AM_Clamp>::GetRHI();
 
 		RenderPassData.PSParameters.InnerCameraProjectionMatrix = FMatrix44f(InnerCameraProjectionMatrix);
+
+		float InnerCameraFrameAspectRatio = 1.f;
+		if (Camera.Resource.Texture)
+		{
+			// Calc inner camera texture aspect ratio:
+			const FIntPoint InnerCameraTextureSize = Camera.Resource.Texture->GetSizeXY();
+			InnerCameraFrameAspectRatio = float(InnerCameraTextureSize.X) / InnerCameraTextureSize.Y;
+		}
+
 		RenderPassData.PSParameters.InnerCameraSoftEdge = FVector4f(Camera.SoftEdge);
 
 		RenderPassData.PSParameters.InnerCameraBorderColor = Camera.InnerCameraBorderColor;
 		RenderPassData.PSParameters.InnerCameraBorderThickness = Camera.InnerCameraBorderThickness;
-		RenderPassData.PSParameters.InnerCameraFrameAspectRatio = Camera.InnerCameraFrameAspectRatio;
+		RenderPassData.PSParameters.InnerCameraFrameAspectRatio = InnerCameraFrameAspectRatio;
 
 		return true;
 	}
diff --git a/Engine/Plugins/Runtime/nDisplay/Source/DisplayClusterShaders/Public/ShaderParameters/DisplayClusterShaderParameters_ICVFX.h b/Engine/Plugins/Runtime/nDisplay/Source/DisplayClusterShaders/Public/ShaderParameters/DisplayClusterShaderParameters_ICVFX.h
--- a/Engine/Plugins/Runtime/nDisplay/Source/DisplayClusterShaders/Public/ShaderParameters/DisplayClusterShaderParameters_ICVFX.h
+++ b/Engine/Plugins/Runtime/nDisplay/Source/DisplayClusterShaders/Public/ShaderParameters/DisplayClusterShaderParameters_ICVFX.h
@@ -217,7 +217,6 @@ public:
 
 			InnerCameraBorderColor      = InCameraSettings.InnerCameraBorderColor;
 			InnerCameraBorderThickness  = InCameraSettings.InnerCameraBorderThickness;
-			InnerCameraFrameAspectRatio = InCameraSettings.InnerCameraFrameAspectRatio;
 
 			ViewProjection = InCameraSettings.ViewProjection;
 
@@ -246,7 +245,6 @@ public:
 
 		FLinearColor InnerCameraBorderColor = FLinearColor::Black;
 		float InnerCameraBorderThickness = 0.1f;
-		float InnerCameraFrameAspectRatio = 1.0f;
 
 		// Camera view projection data
 		FDisplayClusterShaderParametersICVFX_CameraViewProjection ViewProjection;
