[nDisplay] Removing the workaround for ReturnUntouchedSceneColorForPostProcessing

- The original problem has been fixed (UE-219768) so we can now remove the workaround, and simply call ReturnUntouchedSceneColorForPostProcessing

#rnx

diff --git a/Engine/Plugins/Runtime/nDisplay/Source/DisplayCluster/Private/Render/Viewport/DisplayClusterViewportManagerViewExtension.cpp b/Engine/Plugins/Runtime/nDisplay/Source/DisplayCluster/Private/Render/Viewport/DisplayClusterViewportManagerViewExtension.cpp
--- a/Engine/Plugins/Runtime/nDisplay/Source/DisplayCluster/Private/Render/Viewport/DisplayClusterViewportManagerViewExtension.cpp
+++ b/Engine/Plugins/Runtime/nDisplay/Source/DisplayCluster/Private/Render/Viewport/DisplayClusterViewportManagerViewExtension.cpp
@@ -181,27 +181,7 @@ FScreenPassTexture FDisplayClusterViewportManagerViewExtension::PostProcessPassA
 		}
 	}
 
-// UE-219768
-// This is a temporary workaround that prevents crash in ReturnUntouchedSceneColorForPostProcessing.
-// Once the jira is fixed, we can remove this and leave original code from the #else-block
-#if true
-	const FScreenPassTexture& SceneColor = FScreenPassTexture::CopyFromSlice(GraphBuilder, Inputs.GetInput(EPostProcessMaterialInput::SceneColor));
-	check(SceneColor.IsValid());
-
-	FScreenPassRenderTarget Output = Inputs.OverrideOutput;
-
-	// If the override output is provided, it means that this is the last pass in post processing.
-	if (!Output.IsValid())
-	{
-		Output = FScreenPassRenderTarget::CreateFromInput(GraphBuilder, SceneColor, View.GetOverwriteLoadAction(), TEXT("FinalSceneColor"));
-	}
-
-	AddDrawTexturePass(GraphBuilder, View, SceneColor, Output);
-
-	return Output;
-#else
 	return Inputs.ReturnUntouchedSceneColorForPostProcessing(GraphBuilder);
-#endif
 }
 
 bool FDisplayClusterViewportManagerViewExtension::IsActiveThisFrame_Internal(const FSceneViewExtensionContext& Context) const
