nDisplay: Fixed a crash where the nDisplay data cache item was not being released properly.

diff --git a/Engine/Plugins/Runtime/nDisplay/Source/DisplayCluster/Private/Render/Containers/DisplayClusterRender_Texture.cpp b/Engine/Plugins/Runtime/nDisplay/Source/DisplayCluster/Private/Render/Containers/DisplayClusterRender_Texture.cpp
--- a/Engine/Plugins/Runtime/nDisplay/Source/DisplayCluster/Private/Render/Containers/DisplayClusterRender_Texture.cpp
+++ b/Engine/Plugins/Runtime/nDisplay/Source/DisplayCluster/Private/Render/Containers/DisplayClusterRender_Texture.cpp
@@ -2,6 +2,7 @@
 
 #include "Render/Containers/DisplayClusterRender_Texture.h"
 #include "Misc/App.h"
+#include "Misc/DisplayClusterLog.h"
 #include "RenderingThread.h"
 
 int32 GDisplayClusterRender_TextureCacheEnable = 1;
@@ -29,7 +30,19 @@ FDisplayClusterRender_Texture::FDisplayClusterRender_Texture(const FString& InUn
 { }
 
 FDisplayClusterRender_Texture::~FDisplayClusterRender_Texture()
-{ }
+{
+	if (TextureResourcePtr.IsValid())
+	{
+		UE_LOG(LogDisplayClusterRender, Error, TEXT("FDisplayClusterRender_Texture::SetTextureResource(nullptr) function must be called before the destructor."));
+
+		// Forced release of a resource
+		TextureResourcePtr->ReleaseRenderResource();
+		TextureResourcePtr.Reset();
+
+		TextureResourceProxyPtr->ReleaseRenderResource();
+		TextureResourceProxyPtr.Reset();
+	}
+}
 
 const TSharedPtr<FDisplayClusterRender_TextureResource, ESPMode::ThreadSafe>& FDisplayClusterRender_Texture::GetTextureResource() const
 {
diff --git a/Engine/Plugins/Runtime/nDisplay/Source/DisplayCluster/Public/Misc/DisplayClusterDataCache.h b/Engine/Plugins/Runtime/nDisplay/Source/DisplayCluster/Public/Misc/DisplayClusterDataCache.h
--- a/Engine/Plugins/Runtime/nDisplay/Source/DisplayCluster/Public/Misc/DisplayClusterDataCache.h
+++ b/Engine/Plugins/Runtime/nDisplay/Source/DisplayCluster/Public/Misc/DisplayClusterDataCache.h
@@ -111,9 +111,6 @@ protected:
 				// Remove after timout
 				if (CachedObjects[Index].FramesInCachedState > TimeOutInFrames)
 				{
-					// Deleting an unused item after time has elapsed
-					CachedObjects[Index].DataRef.Reset();
-
 					// Delete the current element, and reuse the current array index,
 					// because the next element will move to the index of the current
 					CachedObjects.RemoveAt(Index--);
@@ -156,12 +153,20 @@ private:
 
 		~FCachedObject()
 		{
-			// Calling the release function for referenced data before the destructor
+			Release();
+		}
+
+		/** Release referenced data. */
+		void Release()
+		{
 			if (DataRef.IsValid())
 			{
+				// Calls the special release function before the destructor
 				DataRef->ReleaseDataCacheItem();
 				DataRef.Reset();
 			}
+
+			FramesInCachedState = 0;
 		}
 
 		inline bool IsNameEqual(const FString& InName) const
