nDisplay: Fixed issue where root component transform properties were not added to the Transform details panel in the ICVFX drawer in the ICVFX window.

diff --git a/Engine/Plugins/Runtime/nDisplay/Source/DisplayClusterDetails/Private/DataModelGenerators/DisplayClusterDetailsGenerator_RootActor.cpp b/Engine/Plugins/Runtime/nDisplay/Source/DisplayClusterDetails/Private/DataModelGenerators/DisplayClusterDetailsGenerator_RootActor.cpp
--- a/Engine/Plugins/Runtime/nDisplay/Source/DisplayClusterDetails/Private/DataModelGenerators/DisplayClusterDetailsGenerator_RootActor.cpp
+++ b/Engine/Plugins/Runtime/nDisplay/Source/DisplayClusterDetails/Private/DataModelGenerators/DisplayClusterDetailsGenerator_RootActor.cpp
@@ -13,6 +13,7 @@
 #include "DisplayClusterConfigurationTypes_ICVFX.h"
 #include "DisplayClusterConfigurationTypes_Postprocess.h"
 
+#include "Algo/Transform.h"
 #include "DetailCategoryBuilder.h"
 #include "DetailLayoutBuilder.h"
 #include "DetailWidgetRow.h"
@@ -137,6 +138,26 @@ public:
 			Category.AddProperty(PropertyHandle);
 		};
 
+		// Add root component transform properties to layout builder so that the details panel can find them when constructing TransformCommon
+		const TArray<TWeakObjectPtr<UObject>>& SelectedObjects = DetailBuilder.GetSelectedObjects();
+		{
+			TArray<UObject*> RootComponents;
+			Algo::Transform(SelectedObjects, RootComponents, [](TWeakObjectPtr<UObject> Obj)
+			{
+				USceneComponent* RootComponent = nullptr;
+				if (ADisplayClusterRootActor* RootActor = Cast<ADisplayClusterRootActor>(Obj.Get()))
+				{
+					RootComponent = RootActor->GetRootComponent();
+				}
+
+				return RootComponent;
+			});
+
+			DetailBuilder.AddObjectPropertyData(RootComponents, USceneComponent::GetRelativeLocationPropertyName());
+			DetailBuilder.AddObjectPropertyData(RootComponents, USceneComponent::GetRelativeRotationPropertyName());
+			DetailBuilder.AddObjectPropertyData(RootComponents, USceneComponent::GetRelativeScale3DPropertyName());
+		}
+		
 		IDetailCategoryBuilder& ViewportsCategoryBuilder = DetailBuilder.EditCategory(TEXT("CustomViewportsCategory"), LOCTEXT("CustomViewportsCategoryLabel", "Viewports"));
 		AddProperty(ViewportsCategoryBuilder, TEXT("ViewportScreenPercentageMultiplierRef"));
 		AddProperty(ViewportsCategoryBuilder, TEXT("FreezeRenderOuterViewportsRef"));
