nDisplay: Fix media output of inner frustum not adjusting to changed resolution.

Even though nDisplay was setting bAutoRestartOnSourceSizeChange, because nDisplay specifies a region when capturing, MediaCapture is automatically changing the crop method to Custom, and in that case bAutoRestartOnSourceSizeChange only triggers when the DesiredSize is bigger than the captured texture (this logic might need revision). In the case of nDisplay overscan, the captured texture size is typically increased from the default, so this was not happening.

To resolve the issue, the media stop and start is now handled directly by nDisplay, by detecting differences between the desired media capture size and the texture that it is capturing.

Also made MediaCapture::StopCapture thread safe by reinvoking it in the game thread, since there are several render thread functions that can potentially call it.

#rnx

diff --git a/Engine/Plugins/Media/MediaIOFramework/Source/MediaIOCore/Private/MediaCapture.cpp b/Engine/Plugins/Media/MediaIOFramework/Source/MediaIOCore/Private/MediaCapture.cpp
--- a/Engine/Plugins/Media/MediaIOFramework/Source/MediaIOCore/Private/MediaCapture.cpp
+++ b/Engine/Plugins/Media/MediaIOFramework/Source/MediaIOCore/Private/MediaCapture.cpp
@@ -512,7 +512,20 @@ bool UMediaCapture::UpdateTextureRenderTarget2D(UTextureRenderTarget2D* InRender
 
 void UMediaCapture::StopCapture(bool bAllowPendingFrameToBeProcess)
 {
-	check(IsInGameThread());
+	if (!IsInGameThread())
+	{
+		TWeakObjectPtr<UMediaCapture> Self = this;
+		AsyncTask(ENamedThreads::GameThread, [Self, bAllowPendingFrameToBeProcess]()
+			{
+				UMediaCapture* MediaCapture = Self.Get();
+				if (UObjectInitialized() && MediaCapture)
+				{
+					MediaCapture->StopCapture(bAllowPendingFrameToBeProcess);
+				}
+			});
+
+		return;
+	}
 
 	if (GetState() != EMediaCaptureState::StopRequested && GetState() != EMediaCaptureState::Capturing)
 	{
diff --git a/Engine/Plugins/Runtime/nDisplay/Source/DisplayClusterMedia/Private/Capture/DisplayClusterMediaCaptureBase.cpp b/Engine/Plugins/Runtime/nDisplay/Source/DisplayClusterMedia/Private/Capture/DisplayClusterMediaCaptureBase.cpp
--- a/Engine/Plugins/Runtime/nDisplay/Source/DisplayClusterMedia/Private/Capture/DisplayClusterMediaCaptureBase.cpp
+++ b/Engine/Plugins/Runtime/nDisplay/Source/DisplayClusterMedia/Private/Capture/DisplayClusterMediaCaptureBase.cpp
@@ -155,7 +155,26 @@ void FDisplayClusterMediaCaptureBase::OnPostClusterTick()
 {
 	if (MediaCapture)
 	{
-		const EMediaCaptureState MediaCaptureState = MediaCapture->GetState();
+		EMediaCaptureState MediaCaptureState = MediaCapture->GetState();
+
+		// If we're capturing but the desired capture resolution does not match the texture being captured,
+		// restart the capture with the updated size.
+
+		if (MediaCaptureState == EMediaCaptureState::Capturing)
+		{
+			const FIntPoint LastSrcRegionIntPoint = LastSrcRegionSize.load().ToIntPoint();
+			const FIntPoint DesiredSize = MediaCapture->GetDesiredSize();
+
+			if (DesiredSize != LastSrcRegionIntPoint)
+			{
+				UE_LOG(LogDisplayClusterMedia, Log, TEXT("Stopping MediaCapture '%s' because its DesiredSize (%d, %d) doesn't match the captured texture size (%d, %d)"), 
+					*GetMediaId(), DesiredSize.X, DesiredSize.Y, LastSrcRegionIntPoint.X, LastSrcRegionIntPoint.Y);
+
+				MediaCapture->StopCapture(false /* bAllowPendingFrameToBeProcess */);
+				MediaCaptureState = MediaCapture->GetState(); // Re-sample state to restart the media capture right away
+			}
+		}
+
 		const bool bMediaCaptureNeedsRestart = (MediaCaptureState == EMediaCaptureState::Error) || (MediaCaptureState == EMediaCaptureState::Stopped);
 
 		if (!bWasCaptureStarted || bMediaCaptureNeedsRestart)
@@ -194,7 +213,7 @@ bool FDisplayClusterMediaCaptureBase::StartMediaCapture()
 
 	FMediaCaptureOptions MediaCaptureOptions;
 	MediaCaptureOptions.NumberOfFramesToCapture = -1;
-	MediaCaptureOptions.bAutoRestartOnSourceSizeChange = true;
+	MediaCaptureOptions.bAutoRestartOnSourceSizeChange = false; // true won't work due to MediaCapture auto-changing crop mode to custom when capture region is specified.
 	MediaCaptureOptions.bSkipFrameWhenRunningExpensiveTasks = false;
 	MediaCaptureOptions.OverrunAction = EMediaCaptureOverrunAction::Flush;
 
