nDisplay: Fixed a bug where Per-Viewport OutputRemap caused the RTT of the entire frame to be cleared for each viewport.
- The RTT cleanup has been removed from the FDisplayClusterShadersPostprocess_OutputRemap::RenderPostprocess_OutputRemap(). Because this function is used for both types of PP: per-viewport and full-frame.
* CVar 'nDisplay.render.output_remap.clear_rtt' was removed.

- The RTT cleanup has been added to FDisplayClusterViewportPostProcessManager::ImplPerformPostProcessFrameAfterWarpBlend_RenderThread(). The idea is to only clear RTT for full-frame post processes, not for per-viewport PP.
* Added new CVar 'nDisplay.render.postprocess.FrameAfterWarpBlend.AlwaysClearRTT':  Enables RTT clear before each full-frame post processing. (0 to disable).

- Code cleanup

diff --git a/Engine/Plugins/Runtime/nDisplay/Source/DisplayCluster/Private/Render/Viewport/Postprocess/DisplayClusterViewportPostProcessManager.cpp b/Engine/Plugins/Runtime/nDisplay/Source/DisplayCluster/Private/Render/Viewport/Postprocess/DisplayClusterViewportPostProcessManager.cpp
--- a/Engine/Plugins/Runtime/nDisplay/Source/DisplayCluster/Private/Render/Viewport/Postprocess/DisplayClusterViewportPostProcessManager.cpp
+++ b/Engine/Plugins/Runtime/nDisplay/Source/DisplayCluster/Private/Render/Viewport/Postprocess/DisplayClusterViewportPostProcessManager.cpp
@@ -45,13 +45,24 @@ static TAutoConsoleVariable<int32> CVarPostprocessViewAfterWarpBlend(
 );
 
 // Enable/disable PP round 5
-static TAutoConsoleVariable<int32> CVarPostprocessFrameAfterWarpBlend(
+int32 GDisplayClusterPostprocessFrameAfterWarpBlendEnable = 1;
+static FAutoConsoleVariableRef CVarDisplayClusterPostprocessFrameAfterWarpBlendEnable(
 	TEXT("nDisplay.render.postprocess.FrameAfterWarpBlend"),
-	1,
+	GDisplayClusterPostprocessFrameAfterWarpBlendEnable,
 	TEXT("Enable PP per eye frame after warp&blend (0 = disabled)\n"),
 	ECVF_RenderThreadSafe
 );
 
+// Enables RTT cleanup before each full-frame post processing.
+int32 GDisplayClusterPostProcessFrameAfterWarpBlendAlwaysClearRTT = 1;
+static FAutoConsoleVariableRef CVarDisplayClusterPostProcessFrameAfterWarpBlendAlwaysClearRTT(
+	TEXT("nDisplay.render.postprocess.FrameAfterWarpBlend.AlwaysClearRTT"),
+	GDisplayClusterPostProcessFrameAfterWarpBlendAlwaysClearRTT,
+	TEXT("Enables RTT cleanup before each full-frame post processing. (0 to disable).\n"),
+	ECVF_RenderThreadSafe
+);
+
+
 //----------------------------------------------------------------------------------------------------------
 // FDisplayClusterViewportPostProcessManager
 //----------------------------------------------------------------------------------------------------------
@@ -269,9 +280,10 @@ bool FDisplayClusterViewportPostProcessManager::IsPostProcessViewAfterWarpBlendR
 
 bool FDisplayClusterViewportPostProcessManager::IsPostProcessFrameAfterWarpBlendRequired(const TSharedPtr<IDisplayClusterPostProcess, ESPMode::ThreadSafe>& PostprocessInstance) const
 {
-	if (PostprocessInstance.IsValid() && PostprocessInstance->IsPostProcessFrameAfterWarpBlendRequired())
+	if (GDisplayClusterPostprocessFrameAfterWarpBlendEnable
+		&& PostprocessInstance.IsValid() && PostprocessInstance->IsPostProcessFrameAfterWarpBlendRequired())
 	{
-		return (CVarPostprocessFrameAfterWarpBlend.GetValueOnAnyThread() != 0);
+		return true;
 	}
 
 	return false;
@@ -531,32 +543,52 @@ void FDisplayClusterViewportPostProcessManager::ImplPerformPostProcessViewAfterW
 //////////////////////////////////////////////////////////////////////////////////////////////
 void FDisplayClusterViewportPostProcessManager::ImplPerformPostProcessFrameAfterWarpBlend_RenderThread(FRHICommandListImmediate& RHICmdList, const FDisplayClusterViewportManagerProxy* InViewportManagerProxy) const
 {
-	const bool bEnabled = (CVarPostprocessFrameAfterWarpBlend.GetValueOnRenderThread() != 0);
-	UE_LOG(LogDisplayClusterRender, VeryVerbose, TEXT("Postprocess VIEW after WarpBlend: %d"), bEnabled ? 1 : 0);
-
-	if (bEnabled != 0 && InViewportManagerProxy)
+	if (GDisplayClusterPostprocessFrameAfterWarpBlendEnable && InViewportManagerProxy)
 	{
 		TArray<FRHITexture*> FrameResources;
-		TArray<FRHITexture*> AdditionalFrameResources;
+		TArray<FRHITexture*> TempFrameResources;
 		TArray<FIntPoint> TargetOffset;
-		if (InViewportManagerProxy->GetFrameTargets_RenderThread(FrameResources, TargetOffset, &AdditionalFrameResources))
+		if (InViewportManagerProxy->GetFrameTargets_RenderThread(FrameResources, TargetOffset, &TempFrameResources))
 		{
-			for (const TSharedPtr<IDisplayClusterPostProcess, ESPMode::ThreadSafe >& It : PostprocessProxy)
+			// This is expected to be the way the postprocess is implemented:
+			// 1. The 'FrameResources' is used as input and output data.
+			// 2. The 'AdditionalFrameResources' is a temporary RTT that can be provided by nDisplay to PP if requested by the PP party.
+			// 3. The 'AdditionalFrameResources' RTT will be cleared before each PP (if the GDisplayClusterPostProcessFrameAfterWarpBlendAlwaysClearRTT is set)
+
+			for (const TSharedPtr<IDisplayClusterPostProcess, ESPMode::ThreadSafe >& PostProcessProxyIt : PostprocessProxy)
 			{
-				if (IsPostProcessFrameAfterWarpBlendRequired(It))
+				if (IsPostProcessFrameAfterWarpBlendRequired(PostProcessProxyIt))
 				{
-					UE_LOG(LogDisplayClusterRender, VeryVerbose, TEXT("Postprocess FRAME after WarpBlend"));
+					TArray<FRHITexture*>* TempFrameResourcesPtr = (!TempFrameResources.IsEmpty() && PostProcessProxyIt->ShouldUseAdditionalFrameTargetableResource()) ? &TempFrameResources : nullptr;
 
-					TArray<FRHITexture*>* AdditionalResources = (AdditionalFrameResources.Num() > 0 && It->ShouldUseAdditionalFrameTargetableResource())? &AdditionalFrameResources: nullptr;
-					It->PerformPostProcessFrameAfterWarpBlend_RenderThread(RHICmdList, &FrameResources, AdditionalResources);
+					UE_LOG(LogDisplayClusterRender, VeryVerbose, TEXT("Perform PostProcessFrameAfterWarpBlend '%s'"), *PostProcessProxyIt->GetId());
+
+					if (GDisplayClusterPostProcessFrameAfterWarpBlendAlwaysClearRTT && TempFrameResourcesPtr)
+					{
+						// Always perform RTT cleanup before each full-frame post-processing.
+						for (FRHITexture* FrameResourceIt : *TempFrameResourcesPtr)
+						{
+							FDisplayClusterViewportProxy::FillTextureWithColor_RenderThread(RHICmdList, FrameResourceIt, FLinearColor::Black);
+						}
+					}
+
+					PostProcessProxyIt->PerformPostProcessFrameAfterWarpBlend_RenderThread(RHICmdList, &FrameResources, TempFrameResourcesPtr);
 				}
 			}
 
-			// Apply OutputRemap after all postprocess
-			if (OutputRemap.IsValid() && OutputRemap->IsEnabled())
+			// Apply OutputRemap after all postprocess. The TempFrameResources is required.
+			if (OutputRemap.IsValid() && OutputRemap->IsEnabled() && !TempFrameResources.IsEmpty())
 			{
-				TArray<FRHITexture*>* AdditionalResources = (AdditionalFrameResources.Num() > 0) ? &AdditionalFrameResources : nullptr;
-				OutputRemap->PerformPostProcessFrame_RenderThread(RHICmdList, &FrameResources, AdditionalResources);
+				if (GDisplayClusterPostProcessFrameAfterWarpBlendAlwaysClearRTT)
+				{
+					// Always perform RTT cleanup before each full-frame post-processing.
+					for (FRHITexture* FrameResourceIt : TempFrameResources)
+					{
+						FDisplayClusterViewportProxy::FillTextureWithColor_RenderThread(RHICmdList, FrameResourceIt, FLinearColor::Black);
+					}
+				}
+
+				OutputRemap->PerformPostProcessFrame_RenderThread(RHICmdList, &FrameResources, &TempFrameResources);
 			}
 		}
 	}
diff --git a/Engine/Plugins/Runtime/nDisplay/Source/DisplayClusterShaders/Private/Shaders/DisplayClusterShadersPostprocess_OutputRemap.cpp b/Engine/Plugins/Runtime/nDisplay/Source/DisplayClusterShaders/Private/Shaders/DisplayClusterShadersPostprocess_OutputRemap.cpp
--- a/Engine/Plugins/Runtime/nDisplay/Source/DisplayClusterShaders/Private/Shaders/DisplayClusterShadersPostprocess_OutputRemap.cpp
+++ b/Engine/Plugins/Runtime/nDisplay/Source/DisplayClusterShaders/Private/Shaders/DisplayClusterShadersPostprocess_OutputRemap.cpp
@@ -49,17 +49,6 @@ static TAutoConsoleVariable<int32> CVarOutputRemapShaderType(
 	,ECVF_RenderThreadSafe
 );
 
-// Enable/Disable ClearTexture for OutputRemap RTT
-static TAutoConsoleVariable<int32> CVarClearOutputRemapClearRTT(
-	TEXT("nDisplay.render.output_remap.clear_rtt"),
-	1,
-	TEXT("Enables OutputRemap RTT clearing before postprocessing.\n")
-	TEXT("0 : disabled\n")
-	TEXT("1 : enabled\n")
-	,
-	ECVF_RenderThreadSafe
-);
-
 BEGIN_SHADER_PARAMETER_STRUCT(FOutputRemapVertexShaderParameters, )
 END_SHADER_PARAMETER_STRUCT()
 
@@ -148,8 +137,6 @@ bool FDisplayClusterShadersPostprocess_OutputRemap::RenderPostprocess_OutputRema
 
 	bool bResult = false;
 
-	const bool bClearTextureEnabled = CVarClearOutputRemapClearRTT.GetValueOnRenderThread() != 0;
-
 	RHI_BREADCRUMB_EVENT_STAT(RHICmdList, nDisplay_PostProcess_OutputRemap, "nDisplay_PostProcess_OutputRemap");
 	SCOPED_GPU_STAT(RHICmdList, nDisplay_PostProcess_OutputRemap);
 
@@ -161,12 +148,6 @@ bool FDisplayClusterShadersPostprocess_OutputRemap::RenderPostprocess_OutputRema
 		const FIntPoint TargetSizeXY = InRenderTargetableDestTexture->GetSizeXY();
 		RHICmdList.SetViewport(0, 0, 0.0f, TargetSizeXY.X, TargetSizeXY.Y, 1.0f);
 
-		if (bClearTextureEnabled)
-		{
-			// Clean the texture of the output frame before remapping the whole frame
-			DrawClearQuad(RHICmdList, FLinearColor::Black);
-		}
-
 		const ERHIFeatureLevel::Type RenderFeatureLevel = GMaxRHIFeatureLevel;
 		const auto GlobalShaderMap = GetGlobalShaderMap(RenderFeatureLevel);
 
