nDisplay: Fixed a bug where the log was spammed when the viewport size exceeded the maximum texture size.

Note: It's a temporary solution. It does not work properly for multiple DCRAs with the same viewport names.

diff --git a/Engine/Plugins/Runtime/nDisplay/Source/DisplayCluster/Private/Game/EngineClasses/Scene/DisplayClusterICVFXCameraComponent.cpp b/Engine/Plugins/Runtime/nDisplay/Source/DisplayCluster/Private/Game/EngineClasses/Scene/DisplayClusterICVFXCameraComponent.cpp
--- a/Engine/Plugins/Runtime/nDisplay/Source/DisplayCluster/Private/Game/EngineClasses/Scene/DisplayClusterICVFXCameraComponent.cpp
+++ b/Engine/Plugins/Runtime/nDisplay/Source/DisplayCluster/Private/Game/EngineClasses/Scene/DisplayClusterICVFXCameraComponent.cpp
@@ -273,12 +273,16 @@ FDisplayClusterShaderParameters_ICVFX::FCameraSettings UDisplayClusterICVFXCamer
 	const FIntPoint CameraFrameSize = GetICVFXCameraFrameSize(InStageSettings, InCameraSettings);
 	const FIntPoint RealInnerFrustumResolution(CameraFrameSize.X * CameraMult, CameraFrameSize.Y * CameraMult);
 
+	// Creates unique name "DCRA.Component"
+	const FString UniqueComponentName = FString::Printf(TEXT("%s.%s"), *GetOwner()->GetName(), *GetName());
+
 	FIntRect RealViewportRect(FIntPoint(0, 0), RealInnerFrustumResolution);
 	FDisplayClusterViewport_CustomFrustumSettings RealFrustumSettings;
 	FDisplayClusterViewport_CustomFrustumRuntimeSettings RealFrustumRuntimeSettings;
 
 	FDisplayClusterViewportConfigurationHelpers_ICVFX::UpdateCameraCustomFrustum(CameraSettings.CustomFrustum, RealFrustumSettings);
-	FDisplayClusterViewport_CustomFrustumRuntimeSettings::UpdateCustomFrustumSettings(GetName(), RealFrustumSettings, RealFrustumRuntimeSettings, RealViewportRect);
+	FDisplayClusterViewport_CustomFrustumRuntimeSettings::UpdateCustomFrustumSettings(
+		UniqueComponentName, RealFrustumSettings, RealFrustumRuntimeSettings, RealViewportRect, TEXT("ShaderParameters CustomFrustum"));
 
 	const FDisplayClusterViewport_CustomFrustumRuntimeSettings::FCustomFrustumPercent& Angles = RealFrustumRuntimeSettings.CustomFrustumPercent;
 
@@ -358,6 +362,9 @@ FIntPoint UDisplayClusterICVFXCameraComponent::GetICVFXCameraFrameSize(const FDi
 	// When no adopt resolution is used we must compensate for the change in aspect ratio caused by overscan.
 	if (InCameraSettings.CustomFrustum.bEnable && !InCameraSettings.CustomFrustum.bAdaptResolution)
 	{
+		// Creates unique name "DCRA.Component"
+		const FString UniqueComponentName = FString::Printf(TEXT("%s.%s"), *GetOwner()->GetName(), *GetName());
+
 		// Overscan should only be used through this api:
 		FDisplayClusterViewport_CustomFrustumSettings CustomFrustumSettings;
 		FDisplayClusterViewport_CustomFrustumRuntimeSettings CustomFrustumRuntimeSettings;
@@ -366,7 +373,8 @@ FIntPoint UDisplayClusterICVFXCameraComponent::GetICVFXCameraFrameSize(const FDi
 
 		FIntPoint DesiredSize = AdaptResolutionToAspectRatio(CameraFrameSize, CroppedSensorAR);
 		FIntRect ViewportRect(FIntPoint(0, 0), DesiredSize);
-		FDisplayClusterViewport_CustomFrustumRuntimeSettings::UpdateCustomFrustumSettings(GetName(), CustomFrustumSettings, CustomFrustumRuntimeSettings, ViewportRect);
+		FDisplayClusterViewport_CustomFrustumRuntimeSettings::UpdateCustomFrustumSettings(
+			UniqueComponentName, CustomFrustumSettings, CustomFrustumRuntimeSettings, ViewportRect, TEXT("CameraFrame Size CustomFrustum"));
 
 		// Overscan without the bAdaptResolution option does not change the RTT aspect ratio.
 		// In this case, the sensor AR must be modified to include CustomFrustumPercent values.
@@ -386,8 +394,10 @@ void UDisplayClusterICVFXCameraComponent::UpdateOverscanEstimatedFrameSize()
 		return;
 	}
 
-	const FDisplayClusterConfigurationICVFX_StageSettings& StageSettings = RootActor->GetStageSettings();
+	// Creates unique name "DCRA.Component"
+	const FString UniqueComponentName = FString::Printf(TEXT("%s.%s"), *GetOwner()->GetName(), *GetName());
 
+	const FDisplayClusterConfigurationICVFX_StageSettings& StageSettings = RootActor->GetStageSettings();
 	{
 		// calculate estimations
 		FDisplayClusterConfigurationICVFX_CameraSettings EstimatedCameraSettings = CameraSettings;
@@ -406,7 +416,8 @@ void UDisplayClusterICVFXCameraComponent::UpdateOverscanEstimatedFrameSize()
 		FDisplayClusterViewport_CustomFrustumRuntimeSettings EstimatedFrustumRuntimeSettings;
 
 		FDisplayClusterViewportConfigurationHelpers_ICVFX::UpdateCameraCustomFrustum(EstimatedCameraSettings.CustomFrustum, EstimatedFrustumSettings);
-		FDisplayClusterViewport_CustomFrustumRuntimeSettings::UpdateCustomFrustumSettings(GetName(), EstimatedFrustumSettings, EstimatedFrustumRuntimeSettings, EstimatedViewportRect);
+		FDisplayClusterViewport_CustomFrustumRuntimeSettings::UpdateCustomFrustumSettings(
+			UniqueComponentName, EstimatedFrustumSettings, EstimatedFrustumRuntimeSettings, EstimatedViewportRect, TEXT("Estimated CustomFrustum"));
 
 		// Assign estimated calculated values
 		CameraSettings.CustomFrustum.EstimatedOverscanResolution = EstimatedViewportRect.Size();
@@ -425,7 +436,8 @@ void UDisplayClusterICVFXCameraComponent::UpdateOverscanEstimatedFrameSize()
 		FDisplayClusterViewport_CustomFrustumRuntimeSettings RealFrustumRuntimeSettings;
 
 		FDisplayClusterViewportConfigurationHelpers_ICVFX::UpdateCameraCustomFrustum(CameraSettings.CustomFrustum, RealFrustumSettings);
-		FDisplayClusterViewport_CustomFrustumRuntimeSettings::UpdateCustomFrustumSettings(GetName(), RealFrustumSettings, RealFrustumRuntimeSettings, RealViewportRect);
+		FDisplayClusterViewport_CustomFrustumRuntimeSettings::UpdateCustomFrustumSettings(
+			UniqueComponentName, RealFrustumSettings, RealFrustumRuntimeSettings, RealViewportRect, TEXT("Real CustomFrustum"));
 
 		// Assign real calculated values
 		CameraSettings.CustomFrustum.InnerFrustumResolution = RealViewportRect.Size();
diff --git a/Engine/Plugins/Runtime/nDisplay/Source/DisplayCluster/Private/Render/Viewport/Configuration/DisplayClusterViewportConfiguration.cpp b/Engine/Plugins/Runtime/nDisplay/Source/DisplayCluster/Private/Render/Viewport/Configuration/DisplayClusterViewportConfiguration.cpp
--- a/Engine/Plugins/Runtime/nDisplay/Source/DisplayCluster/Private/Render/Viewport/Configuration/DisplayClusterViewportConfiguration.cpp
+++ b/Engine/Plugins/Runtime/nDisplay/Source/DisplayCluster/Private/Render/Viewport/Configuration/DisplayClusterViewportConfiguration.cpp
@@ -174,6 +174,16 @@ ADisplayClusterRootActor* FDisplayClusterViewportConfiguration::GetRootActor(con
 	return nullptr;
 }
 
+FString FDisplayClusterViewportConfiguration::GetRootActorName() const
+{
+	if (ADisplayClusterRootActor* RootActor = GetRootActor(EDisplayClusterRootActorType::Scene))
+	{
+		return RootActor->GetName();
+	}
+
+	return FString();
+}
+
 void FDisplayClusterViewportConfiguration::OnHandleStartScene()
 {
 	bCurrentSceneActive = true;
diff --git a/Engine/Plugins/Runtime/nDisplay/Source/DisplayCluster/Private/Render/Viewport/Configuration/DisplayClusterViewportConfiguration.h b/Engine/Plugins/Runtime/nDisplay/Source/DisplayCluster/Private/Render/Viewport/Configuration/DisplayClusterViewportConfiguration.h
--- a/Engine/Plugins/Runtime/nDisplay/Source/DisplayCluster/Private/Render/Viewport/Configuration/DisplayClusterViewportConfiguration.h
+++ b/Engine/Plugins/Runtime/nDisplay/Source/DisplayCluster/Private/Render/Viewport/Configuration/DisplayClusterViewportConfiguration.h
@@ -139,6 +139,10 @@ public:
 	void OnHandleStartScene();
 	void OnHandleEndScene();
 
+
+	/** Returns the DCRA name or an empty string. */
+	FString GetRootActorName() const;
+
 private:
 	/** Update configuration implementation.
 	* 
diff --git a/Engine/Plugins/Runtime/nDisplay/Source/DisplayCluster/Private/Render/Viewport/Containers/DisplayClusterViewport_CustomFrustumRuntimeSettings.cpp b/Engine/Plugins/Runtime/nDisplay/Source/DisplayCluster/Private/Render/Viewport/Containers/DisplayClusterViewport_CustomFrustumRuntimeSettings.cpp
--- a/Engine/Plugins/Runtime/nDisplay/Source/DisplayCluster/Private/Render/Viewport/Containers/DisplayClusterViewport_CustomFrustumRuntimeSettings.cpp
+++ b/Engine/Plugins/Runtime/nDisplay/Source/DisplayCluster/Private/Render/Viewport/Containers/DisplayClusterViewport_CustomFrustumRuntimeSettings.cpp
@@ -106,7 +106,8 @@ void FDisplayClusterViewport_CustomFrustumRuntimeSettings::UpdateCustomFrustumSe
 	const FString& InViewportId,
 	const FDisplayClusterViewport_CustomFrustumSettings& InCustomFrustumSettings,
 	FDisplayClusterViewport_CustomFrustumRuntimeSettings& InOutRuntimeSettings,
-	FIntRect& InOutRenderTargetRect)
+	FIntRect& InOutRenderTargetRect,
+	const TCHAR* InCustomResourceName)
 {
 	using namespace UE::DisplayCluster::Viewport;
 
@@ -166,9 +167,15 @@ void FDisplayClusterViewport_CustomFrustumRuntimeSettings::UpdateCustomFrustumSe
 	// Update RTT size for CustomFrustum when we need to scale target resolution
 	if (InCustomFrustumSettings.bAdaptResolution)
 	{
-		const FIntPoint CustomFrustumSize = Size + InOutRuntimeSettings.CustomFrustumPixels.Size();
-		const FIntPoint ValidCustomFrustumSize = FDisplayClusterViewportHelpers::GetValidViewportRect(FIntRect(FIntPoint(0, 0), CustomFrustumSize), InViewportId, TEXT("CustomFrustum")).Size();
-
-		InOutRenderTargetRect.Max = ValidCustomFrustumSize;
+		const FIntRect NewCustomFrustumRect(
+			FIntPoint(0, 0),
+			Size + InOutRuntimeSettings.CustomFrustumPixels.Size()
+		);
+
+		// Return new custom frustum rect
+		InOutRenderTargetRect = FDisplayClusterViewportHelpers::GetValidViewportRect(
+			NewCustomFrustumRect,
+			InViewportId,
+			InCustomResourceName ? InCustomResourceName : TEXT("CustomFrustum"));
 	}
 }
diff --git a/Engine/Plugins/Runtime/nDisplay/Source/DisplayCluster/Private/Render/Viewport/Containers/DisplayClusterViewport_CustomFrustumRuntimeSettings.h b/Engine/Plugins/Runtime/nDisplay/Source/DisplayCluster/Private/Render/Viewport/Containers/DisplayClusterViewport_CustomFrustumRuntimeSettings.h
--- a/Engine/Plugins/Runtime/nDisplay/Source/DisplayCluster/Private/Render/Viewport/Containers/DisplayClusterViewport_CustomFrustumRuntimeSettings.h
+++ b/Engine/Plugins/Runtime/nDisplay/Source/DisplayCluster/Private/Render/Viewport/Containers/DisplayClusterViewport_CustomFrustumRuntimeSettings.h
@@ -12,11 +12,19 @@ struct FDisplayClusterViewport_CustomFrustumRuntimeSettings
 {
 	/** Update custom frustum settings
 	*
-	* @param InViewport            - owner viewport
-	* @param InOutRuntimeSettings  - the CustomFrustum runtime settings.
-	* @param InOutRenderTargetRect - Viewport rect, changeable during CustomFrustum
+	* @param InViewportId            - owner viewport name
+	* @param InCustomFrustumSettings - CustomFrustum settings
+	* @param InOutRuntimeSettings    - CustomFrustum runtime settings.
+	* @param InOutRenderTargetRect   - Viewport rect, changeable during CustomFrustum
+	* @param InCustomResourceName    - (opt) The unique resource name used for log messages.
+	*                                        Use a new custom resource name for each new location from which this function is called.
 	*/
-	static void UpdateCustomFrustumSettings(const FString& InViewportId, const FDisplayClusterViewport_CustomFrustumSettings& InCustomFrustumSettings, FDisplayClusterViewport_CustomFrustumRuntimeSettings& InOutRuntimeSettings, FIntRect& InOutRenderTargetRect);
+	static void UpdateCustomFrustumSettings(
+		const FString& InViewportId,
+		const FDisplayClusterViewport_CustomFrustumSettings& InCustomFrustumSettings,
+		FDisplayClusterViewport_CustomFrustumRuntimeSettings& InOutRuntimeSettings,
+		FIntRect& InOutRenderTargetRect,
+		const TCHAR* InCustomResourceName = nullptr);
 	
 	/** Update projection angles by custom frustum
 	*
diff --git a/Engine/Plugins/Runtime/nDisplay/Source/DisplayCluster/Private/Render/Viewport/DisplayClusterViewport.cpp b/Engine/Plugins/Runtime/nDisplay/Source/DisplayCluster/Private/Render/Viewport/DisplayClusterViewport.cpp
--- a/Engine/Plugins/Runtime/nDisplay/Source/DisplayCluster/Private/Render/Viewport/DisplayClusterViewport.cpp
+++ b/Engine/Plugins/Runtime/nDisplay/Source/DisplayCluster/Private/Render/Viewport/DisplayClusterViewport.cpp
@@ -610,7 +610,9 @@ bool FDisplayClusterViewport::UpdateFrameContexts(const uint32 InStereoViewIndex
 	// Support custom frustum rendering feature
 	if (!RenderSettings.bDisableCustomFrustumFeature)
 	{
-		FDisplayClusterViewport_CustomFrustumRuntimeSettings::UpdateCustomFrustumSettings(GetId(), RenderSettings.CustomFrustumSettings, CustomFrustumRuntimeSettings, RenderTargetRect);
+		// Creates unique name "DCRA.Viewport"
+		const FString UniqueViewportName = FString::Printf(TEXT("%s.%s"), *Configuration->GetRootActorName(), *GetId());
+		FDisplayClusterViewport_CustomFrustumRuntimeSettings::UpdateCustomFrustumSettings(UniqueViewportName, RenderSettings.CustomFrustumSettings, CustomFrustumRuntimeSettings, RenderTargetRect);
 	}
 
 	FIntPoint ContextSize = RenderTargetRect.Size();
diff --git a/Engine/Plugins/Runtime/nDisplay/Source/DisplayCluster/Private/Render/Viewport/DisplayClusterViewportHelpers.cpp b/Engine/Plugins/Runtime/nDisplay/Source/DisplayCluster/Private/Render/Viewport/DisplayClusterViewportHelpers.cpp
--- a/Engine/Plugins/Runtime/nDisplay/Source/DisplayCluster/Private/Render/Viewport/DisplayClusterViewportHelpers.cpp
+++ b/Engine/Plugins/Runtime/nDisplay/Source/DisplayCluster/Private/Render/Viewport/DisplayClusterViewportHelpers.cpp
@@ -130,15 +130,12 @@ FIntRect FDisplayClusterViewportHelpers::GetValidViewportRect(const FIntRect& In
 
 	FIntRect OutRect(InRect.Min, InRect.Min + FIntPoint(Width, Height));
 
-	float RectScale = 1;
 
 	// Make sure the rect doesn't exceed the maximum resolution, and preserve its aspect ratio if it needs to be clamped
-	int32 RectMaxSize = OutRect.Max.GetMax();
-	if (RectMaxSize > MaxTextureSize)
-	{
-		RectScale = float(MaxTextureSize) / RectMaxSize;
-		UE_LOG(LogDisplayClusterViewport, Error, TEXT("The viewport '%s' rect '%s' size %dx%d clamped: max texture dimensions is %d"), *InViewportId, (InResourceName == nullptr) ? TEXT("none") : InResourceName, InRect.Max.X, InRect.Max.Y, MaxTextureSize);
-	}
+	const int32 RectMaxSize = OutRect.Max.GetMax();
+	const float RectScale = (RectMaxSize > MaxTextureSize)
+		? float(MaxTextureSize) / RectMaxSize
+		: 1.f;
 
 	OutRect.Min.X = FMath::Min(OutRect.Min.X, MaxTextureSize);
 	OutRect.Min.Y = FMath::Min(OutRect.Min.Y, MaxTextureSize);
@@ -148,6 +145,42 @@ FIntRect FDisplayClusterViewportHelpers::GetValidViewportRect(const FIntRect& In
 	OutRect.Max.X = FMath::Clamp(ScaledRectMax.X, OutRect.Min.X, MaxTextureSize);
 	OutRect.Max.Y = FMath::Clamp(ScaledRectMax.Y, OutRect.Min.Y, MaxTextureSize);
 
+	// It's a temporary solution.
+	// It does not work properly for multiple DCRAs with the same viewport names.
+	static TMap<FString, bool> CachedLogMessages;
+	const FString UniqueLogMessageId = FString::Printf(TEXT("%s.%s"), *InViewportId, (InResourceName == nullptr) ? TEXT("none") : InResourceName);
+	bool* const LogMessageStatePtr = CachedLogMessages.Find(UniqueLogMessageId);
+
+	if (InRect != OutRect)
+	{
+		if (LogMessageStatePtr == nullptr || *LogMessageStatePtr == false)
+		{
+			UE_LOG(LogDisplayClusterViewport, Error,
+				TEXT("The '%s' %s has been clamped %dx%d->%dx%d (max texture dimensions is %d)"),
+				*InViewportId,
+				(InResourceName == nullptr) ? TEXT("none") : InResourceName,
+				InRect.Width(), InRect.Height(),
+				OutRect.Width(), OutRect.Height(),
+				MaxTextureSize);
+
+			// Raise the flag to show the log message once
+			if (LogMessageStatePtr)
+			{
+				*LogMessageStatePtr = true;
+			}
+			else
+			{
+				CachedLogMessages.Emplace(UniqueLogMessageId, true);
+			}
+		}
+	}
+	else if (LogMessageStatePtr)
+	{
+		// Reset flag
+		*LogMessageStatePtr = false;
+	}
+
+
 	return OutRect;
 }
 
