nDisplay: Fixed issue where stage geometry mesh had physics enabled, causing it to interact with physics objects.

diff --git a/Engine/Plugins/Runtime/nDisplay/Source/DisplayCluster/Private/Game/EngineClasses/Scene/DisplayClusterStageIsosphereComponent.cpp b/Engine/Plugins/Runtime/nDisplay/Source/DisplayCluster/Private/Game/EngineClasses/Scene/DisplayClusterStageIsosphereComponent.cpp
--- a/Engine/Plugins/Runtime/nDisplay/Source/DisplayCluster/Private/Game/EngineClasses/Scene/DisplayClusterStageIsosphereComponent.cpp
+++ b/Engine/Plugins/Runtime/nDisplay/Source/DisplayCluster/Private/Game/EngineClasses/Scene/DisplayClusterStageIsosphereComponent.cpp
@@ -14,6 +14,8 @@ UDisplayClusterStageIsosphereComponent::UDisplayClusterStageIsosphereComponent(c
 
 	IsosphereMesh = IsosphereMeshFinder.Object;
 
+	BodyInstance.SetCollisionEnabled(ECollisionEnabled::QueryOnly);
+	
 	SetHiddenInGame(true);
 	SetVisibility(false);
 }
