[nDisplay] Fixed incorrect colors with late OCIO

- There were multiple issues when the senders have no OCIO configured. Starting from rgb10 (instead of float16) and finishing with incorrect colors regardless of OCIO settings on the receiving side.

- The problem was caused by a logical error that prevented SVF to be prepared for OCIO properly (no tonecurve, linear sRGB)

#rnx

diff --git a/Engine/Plugins/Runtime/nDisplay/Source/DisplayCluster/Private/Render/Viewport/DisplayClusterViewport.cpp b/Engine/Plugins/Runtime/nDisplay/Source/DisplayCluster/Private/Render/Viewport/DisplayClusterViewport.cpp
--- a/Engine/Plugins/Runtime/nDisplay/Source/DisplayCluster/Private/Render/Viewport/DisplayClusterViewport.cpp
+++ b/Engine/Plugins/Runtime/nDisplay/Source/DisplayCluster/Private/Render/Viewport/DisplayClusterViewport.cpp
@@ -321,11 +321,14 @@ void FDisplayClusterViewport::SetupSceneView(uint32 ContextNum, class UWorld* Wo
 		return;
 	}
 
-	if (OpenColorIO.IsValid())
+	// Always modify rendering parameters if valid OCIO transformation is configured
+	if (OpenColorIO.IsValid() && OpenColorIO->GetConversionSettings().IsValid())
 	{
 		OpenColorIO->SetupSceneView(InOutViewFamily, InOutView);
 	}
-	// When capturing with late OCIO, we also need to modify some OCIO specific parameters
+	// When capturing with late OCIO enabled, we still need to modify the OCIO related
+	// rendering parameters even though OCIO is not set. The receivers might have valid OCIO
+	// transformations configured therefore should get a proper input texture.
 	else if (RenderSettings.HasAnyMediaStates(EDisplayClusterViewportMediaState::CaptureLateOCIO))
 	{
 		FOpenColorIORendering::PrepareView(InOutViewFamily, InOutView);
diff --git a/Engine/Plugins/Runtime/nDisplay/Source/DisplayClusterMedia/Private/Input/DisplayClusterMediaInputViewport.cpp b/Engine/Plugins/Runtime/nDisplay/Source/DisplayClusterMedia/Private/Input/DisplayClusterMediaInputViewport.cpp
--- a/Engine/Plugins/Runtime/nDisplay/Source/DisplayClusterMedia/Private/Input/DisplayClusterMediaInputViewport.cpp
+++ b/Engine/Plugins/Runtime/nDisplay/Source/DisplayClusterMedia/Private/Input/DisplayClusterMediaInputViewport.cpp
@@ -72,7 +72,7 @@ void FDisplayClusterMediaInputViewport::OnPreSubmitViewFamilies(TArray<FSceneVie
 			FOpenColorIOColorConversionSettings OCIOConversionSettings;
 
 			const bool bSettingsAvailable = Viewport->GetOCIOConversionSettings(OCIOConversionSettings);
-			if (bSettingsAvailable)
+			if (bSettingsAvailable && OCIOConversionSettings.IsValid())
 			{
 				const UWorld* const CurrentWorld = ViewportMgr->GetConfiguration().GetCurrentWorld();
 
@@ -82,14 +82,16 @@ void FDisplayClusterMediaInputViewport::OnPreSubmitViewFamilies(TArray<FSceneVie
 
 				// Get OCIO render pass resources
 				FOpenColorIORenderPassResources OCIOPassResources = FOpenColorIORendering::GetRenderPassResources(OCIOConversionSettings, FeatureLevel);
-
-				// And push it to the rendering thread
-				ENQUEUE_RENDER_COMMAND(DCMediaInputUpdateOCIOResources)(
-					[this, InOCIOPassResources = MoveTemp(OCIOPassResources)](FRHICommandListImmediate& RHICmdList)
-					{
-						OCIOPassResources_RT = InOCIOPassResources;
-					}
-				);
+				if (OCIOPassResources.IsValid())
+				{
+					// And push it to the rendering thread
+					ENQUEUE_RENDER_COMMAND(DCMediaInputUpdateOCIOResources)(
+						[this, InOCIOPassResources = MoveTemp(OCIOPassResources)](FRHICommandListImmediate& RHICmdList)
+						{
+							OCIOPassResources_RT = InOCIOPassResources;
+						}
+					);
+				}
 			}
 		}
 	}
