[nDisplay] Late OCIO has no effect if OCIO is disabled completely

- Previously, if OCIO is disabled and "Receivers Apply OCIO" is on, the sender used to multicast a no tone curve linear sRGB texture to the receivers. That looked confusing because the colors are different compared to the outer views. There was an illusion that the same OCIO conversion is applied on all the nodes.

- With this fix, when OCIO is diabled, the "Receivers Apply OCIO" is ignored, and the inner view remains untouched (same colors as non-OCIO outers have)

#rnx

diff --git a/Engine/Plugins/Runtime/nDisplay/Source/DisplayClusterMedia/Private/DisplayClusterMediaModule.cpp b/Engine/Plugins/Runtime/nDisplay/Source/DisplayClusterMedia/Private/DisplayClusterMediaModule.cpp
--- a/Engine/Plugins/Runtime/nDisplay/Source/DisplayClusterMedia/Private/DisplayClusterMediaModule.cpp
+++ b/Engine/Plugins/Runtime/nDisplay/Source/DisplayClusterMedia/Private/DisplayClusterMediaModule.cpp
@@ -303,6 +303,10 @@ void FDisplayClusterMediaModule::InitializeICVFXCameraFullFrameInput(const UDisp
 			const FString ICVFXCameraName = ExplicitCameraName.IsEmpty() ? ICVFXCameraComponent->GetName() : ExplicitCameraName;
 			const FString ICVFXViewportId = DisplayClusterMediaHelpers::GenerateICVFXViewportName(ClusterNodeId, ICVFXCameraName);
 
+			// Determine whether it's necessary to activate the late OCIO on media side
+			const bool bEnabledOCIO = ICVFXCameraComponent->CameraSettings.CameraOCIO.AllNodesOCIOConfiguration.bIsEnabled;
+			const bool bUseLateOCIO = MediaSettings.bLateOCIOPass && bEnabledOCIO;
+
 			if (UMediaSource* MediaSource = MediaSettings.GetMediaSource(ClusterNodeId))
 			{
 				const FString MediaInputId = DisplayClusterMediaHelpers::MediaId::GenerateMediaId(
@@ -316,7 +320,7 @@ void FDisplayClusterMediaModule::InitializeICVFXCameraFullFrameInput(const UDisp
 					MediaInputId, ClusterNodeId,
 					ICVFXViewportId,
 					MediaSource,
-					MediaSettings.bLateOCIOPass);
+					bUseLateOCIO);
 
 				InputViewports.Emplace(MediaInputId, MoveTemp(NewICVFXInput));
 			}
@@ -337,6 +341,10 @@ void FDisplayClusterMediaModule::InitializeICVFXCameraFullFrameOutput(const UDis
 			const FString ICVFXCameraName = ExplicitCameraName.IsEmpty() ? ICVFXCameraComponent->GetName() : ExplicitCameraName;
 			const FString ICVFXViewportId = DisplayClusterMediaHelpers::GenerateICVFXViewportName(ClusterNodeId, ICVFXCameraName);
 
+			// Determine whether it's necessary to activate the late OCIO on media side
+			const bool bEnabledOCIO = ICVFXCameraComponent->CameraSettings.CameraOCIO.AllNodesOCIOConfiguration.bIsEnabled;
+			const bool bUseLateOCIO = MediaSettings.bLateOCIOPass && bEnabledOCIO;
+
 			// Media capture (full frame)
 			const TArray<FDisplayClusterConfigurationMediaOutputGroup> MediaOutputItems = MediaSettings.GetMediaOutputGroups(ClusterNodeId);
 			uint8 CaptureIdx = 0;
@@ -356,7 +364,7 @@ void FDisplayClusterMediaModule::InitializeICVFXCameraFullFrameOutput(const UDis
 						ICVFXCameraName, ICVFXViewportId,
 						MediaOutputItem.MediaOutput,
 						MediaOutputItem.OutputSyncPolicy,
-						MediaSettings.bLateOCIOPass);
+						bUseLateOCIO);
 
 					CaptureViewports.Emplace(MediaCaptureId, MoveTemp(NewICVFXCapture));
 				}
@@ -377,15 +385,19 @@ void FDisplayClusterMediaModule::InitializeICVFXCameraUniformTilesInput(const UD
 
 		if (MediaSettings.bEnable && MediaSettings.SplitType == EDisplayClusterConfigurationMediaSplitType::UniformTiles)
 		{
-			const FString ICVFXCameraName = ExplicitCameraName.IsEmpty() ? ICVFXCameraComponent->GetName() : ExplicitCameraName;
-			const FString ICVFXViewportId = DisplayClusterMediaHelpers::GenerateICVFXViewportName(ClusterNodeId, ICVFXCameraName);
-
 			// Find corresponsing media group
 			TArray<FDisplayClusterConfigurationMediaUniformTileInput> MediaInputTiles;
 			const bool bMediaInputGroupFound = MediaSettings.GetMediaInputTiles(ClusterNodeId, MediaInputTiles);
 
 			if (bMediaInputGroupFound)
 			{
+				const FString ICVFXCameraName = ExplicitCameraName.IsEmpty() ? ICVFXCameraComponent->GetName() : ExplicitCameraName;
+				const FString ICVFXViewportId = DisplayClusterMediaHelpers::GenerateICVFXViewportName(ClusterNodeId, ICVFXCameraName);
+
+				// Determine whether it's necessary to activate the late OCIO on media side
+				const bool bEnabledOCIO = ICVFXCameraComponent->CameraSettings.CameraOCIO.AllNodesOCIOConfiguration.bIsEnabled;
+				const bool bUseLateOCIO = MediaSettings.bLateOCIOPass && bEnabledOCIO;
+
 				uint8 Index = 0;
 				for (const FDisplayClusterConfigurationMediaUniformTileInput& MediaInputTile : MediaInputTiles)
 				{
@@ -406,7 +418,7 @@ void FDisplayClusterMediaModule::InitializeICVFXCameraUniformTilesInput(const UD
 							MediaInputId, ClusterNodeId,
 							ICVFXViewportTileId,
 							MediaInputTile.MediaSource,
-							MediaSettings.bLateOCIOPass);
+							bUseLateOCIO);
 
 						InputViewports.Emplace(MediaInputId, MoveTemp(NewICVFXTileInput));
 					}
@@ -428,15 +440,19 @@ void FDisplayClusterMediaModule::InitializeICVFXCameraUniformTilesOutput(const U
 
 		if (MediaSettings.bEnable && MediaSettings.SplitType == EDisplayClusterConfigurationMediaSplitType::UniformTiles)
 		{
-			const FString ICVFXCameraName = ExplicitCameraName.IsEmpty() ? ICVFXCameraComponent->GetName() : ExplicitCameraName;
-			const FString ICVFXViewportId = DisplayClusterMediaHelpers::GenerateICVFXViewportName(ClusterNodeId, ICVFXCameraName);
-
 			// Find corresponsing media group
 			TArray<FDisplayClusterConfigurationMediaUniformTileOutput> MediaOutputTiles;
 			const bool bMediaOutputGroupFound = MediaSettings.GetMediaOutputTiles(ClusterNodeId, MediaOutputTiles);
 
 			if (bMediaOutputGroupFound)
 			{
+				const FString ICVFXCameraName = ExplicitCameraName.IsEmpty() ? ICVFXCameraComponent->GetName() : ExplicitCameraName;
+				const FString ICVFXViewportId = DisplayClusterMediaHelpers::GenerateICVFXViewportName(ClusterNodeId, ICVFXCameraName);
+
+				// Determine whether it's necessary to activate the late OCIO on media side
+				const bool bEnabledOCIO = ICVFXCameraComponent->CameraSettings.CameraOCIO.AllNodesOCIOConfiguration.bIsEnabled;
+				const bool bUseLateOCIO = MediaSettings.bLateOCIOPass && bEnabledOCIO;
+
 				uint8 Index = 0;
 				for (const FDisplayClusterConfigurationMediaUniformTileOutput& MediaOutputTile : MediaOutputTiles)
 				{
@@ -458,7 +474,7 @@ void FDisplayClusterMediaModule::InitializeICVFXCameraUniformTilesOutput(const U
 							ICVFXViewportTileId,
 							MediaOutputTile.MediaOutput,
 							nullptr, // no sync for tiles as it's basically for in-cluster use
-							MediaSettings.bLateOCIOPass);
+							bUseLateOCIO);
 
 						CaptureViewports.Emplace(MediaOutputId, MoveTemp(NewICVFXTileOutput));
 					}
