nDisplay: Fixed RHI crash - added check and fix for FRHICopyTextureInfo parameters

- Added new static function FDisplayClusterViewportHelpers::GetValidResourceRectsForResolve()
  * adjust the regions to the texture size
  * checks that the size of the area to be copied is not equal to zero.
  * compensates for the negative value of Rect.Min

- Copy to backbuffer has been modified.
- ImplResolveResource() has been modified.

diff --git a/Engine/Plugins/Runtime/nDisplay/Source/DisplayCluster/Private/Render/Viewport/DisplayClusterViewportHelpers.cpp b/Engine/Plugins/Runtime/nDisplay/Source/DisplayCluster/Private/Render/Viewport/DisplayClusterViewportHelpers.cpp
--- a/Engine/Plugins/Runtime/nDisplay/Source/DisplayCluster/Private/Render/Viewport/DisplayClusterViewportHelpers.cpp
+++ b/Engine/Plugins/Runtime/nDisplay/Source/DisplayCluster/Private/Render/Viewport/DisplayClusterViewportHelpers.cpp
@@ -260,3 +260,59 @@ EPixelFormat FDisplayClusterViewportHelpers::GetDefaultPixelFormat()
 {
 	return EPixelFormat::PF_FloatRGBA;
 }
+
+bool FDisplayClusterViewportHelpers::GetValidResourceRectsForResolve(
+	const FRHITexture* InSourceTexture,
+	const FRHITexture* InDestTexture,
+	FIntRect& InOutSourceRect,
+	FIntRect& InOutDestRect)
+{
+	if (!InSourceTexture || !InDestTexture)
+	{
+		return false;
+	}
+
+	FIntRect SrcRect(InOutSourceRect);
+	FIntRect DestRect(InOutDestRect);
+
+	const FIntPoint& InputResourceSize = InSourceTexture->GetDesc().Extent;
+	const FIntPoint& OutputResourceSize = InDestTexture->GetDesc().Extent;
+
+	// If SrcRect or DestRect exceeds the texture size, RHI will crash. Let's adjust it to the texture size.
+	{
+		SrcRect.Min.X = FMath::Clamp(SrcRect.Min.X, 0, InputResourceSize.X);
+		SrcRect.Min.Y = FMath::Clamp(SrcRect.Min.Y, 0, InputResourceSize.Y);
+		SrcRect.Max.X = FMath::Clamp(SrcRect.Max.X, 0, InputResourceSize.X);
+		SrcRect.Max.Y = FMath::Clamp(SrcRect.Max.Y, 0, InputResourceSize.Y);
+
+		DestRect.Min.X = FMath::Clamp(DestRect.Min.X, 0, OutputResourceSize.X);
+		DestRect.Min.Y = FMath::Clamp(DestRect.Min.Y, 0, OutputResourceSize.Y);
+		DestRect.Max.X = FMath::Clamp(DestRect.Max.X, 0, OutputResourceSize.X);
+		DestRect.Max.Y = FMath::Clamp(DestRect.Max.Y, 0, OutputResourceSize.Y);
+	}
+
+	// if InputRect.Min<0, also adjust the DestRect.Min
+	DestRect.Min += FIntPoint(
+		FMath::Max(0, -InOutSourceRect.Min.X),
+		FMath::Max(0, -InOutSourceRect.Min.Y));
+
+	// if OutputRect.Min<0, also adjust the SrcRect.Min
+	SrcRect.Min += FIntPoint(
+		FMath::Max(0, -InOutDestRect.Min.X),
+		FMath::Max(0, -InOutDestRect.Min.Y));
+
+	// SrcRect.Min and DestRect.Min always > 0
+
+	// Check the SrcRect and DestRect
+	if (SrcRect.Size().GetMin() <= 0 || DestRect.Size().GetMin() <= 0)
+	{
+		// The SrcRect or DestRect is invalid.
+		return false;
+	}
+
+	// Can be resolved
+	InOutSourceRect = SrcRect;
+	InOutDestRect = DestRect;
+
+	return true;
+}
diff --git a/Engine/Plugins/Runtime/nDisplay/Source/DisplayCluster/Private/Render/Viewport/DisplayClusterViewportHelpers.h b/Engine/Plugins/Runtime/nDisplay/Source/DisplayCluster/Private/Render/Viewport/DisplayClusterViewportHelpers.h
--- a/Engine/Plugins/Runtime/nDisplay/Source/DisplayCluster/Private/Render/Viewport/DisplayClusterViewportHelpers.h
+++ b/Engine/Plugins/Runtime/nDisplay/Source/DisplayCluster/Private/Render/Viewport/DisplayClusterViewportHelpers.h
@@ -5,6 +5,7 @@
 #include "PixelFormat.h"
 
 class FDisplayClusterViewportResource;
+class FRHITexture;
 struct FDisplayClusterRenderFrameSettings;
 
 /**
@@ -83,4 +84,22 @@ public:
 	 * Getting the default pixel format
 	 */
 	static EPixelFormat GetDefaultPixelFormat();
+
+	/**
+	* Check if resources with the specified regions can be resolved.
+	* If any rect exceeds the texture size, RHI will crash.
+	* This function adjusts the rects to the size of the textures.
+	* 
+	* @param InSourceTexture - (in) src RHI texture
+	* @param InDestTexture   - (in) dest RHI texture
+	* @param InOutSourceRect - (in, out) src rect
+	* @param InOutDestRect   - (in, out) dest rect
+	* 
+	* @return false, if it isn't possible
+	*/
+	static bool GetValidResourceRectsForResolve(
+		const FRHITexture* InSourceTexture,
+		const FRHITexture* InDestTexture,
+		FIntRect& InOutSourceRect,
+		FIntRect& InOutDestRect);
 };
diff --git a/Engine/Plugins/Runtime/nDisplay/Source/DisplayCluster/Private/Render/Viewport/DisplayClusterViewportManagerProxy.cpp b/Engine/Plugins/Runtime/nDisplay/Source/DisplayCluster/Private/Render/Viewport/DisplayClusterViewportManagerProxy.cpp
--- a/Engine/Plugins/Runtime/nDisplay/Source/DisplayCluster/Private/Render/Viewport/DisplayClusterViewportManagerProxy.cpp
+++ b/Engine/Plugins/Runtime/nDisplay/Source/DisplayCluster/Private/Render/Viewport/DisplayClusterViewportManagerProxy.cpp
@@ -3,6 +3,7 @@
 #include "Render/Viewport/DisplayClusterViewportManagerProxy.h"
 
 #include "Render/Viewport/DisplayClusterViewportManager.h"
+#include "Render/Viewport/DisplayClusterViewportHelpers.h"
 #include "Render/Viewport/Configuration/DisplayClusterViewportConfiguration.h"
 #include "Render/Viewport/Configuration/DisplayClusterViewportConfigurationProxy.h"
 
@@ -568,35 +569,32 @@ bool FDisplayClusterViewportManagerProxy::ResolveFrameTargetToBackBuffer_RenderT
 
 	TArray<FRHITexture*>   FrameResources;
 	TArray<FIntPoint>        TargetOffsets;
-	if (GetFrameTargets_RenderThread(FrameResources, TargetOffsets))
+	if (GetFrameTargets_RenderThread(FrameResources, TargetOffsets) && FrameResources.IsValidIndex(InContextNum))
 	{
 		// Use internal frame textures as source
-		int32 ContextNum = InContextNum;
-
-		FRHITexture* FrameTexture = FrameResources[ContextNum];
-		FIntPoint DstOffset = TargetOffsets[ContextNum];
-
-		if (FrameTexture)
+		if (FRHITexture* FrameTexture = FrameResources[InContextNum])
 		{
-			const FIntPoint SrcSize = FrameTexture->GetSizeXY();
-			const FIntPoint DstSize = DestTexture->GetSizeXY();;
+			FIntRect SrcRect(FIntPoint::ZeroValue, FrameTexture->GetDesc().Extent);
 
-			FIntRect DstRect(DstOffset, DstOffset + SrcSize);
+			const FIntPoint& DestOffset = TargetOffsets[InContextNum];
+			FIntRect DestRect(DestOffset, DestOffset + FrameTexture->GetDesc().Extent);
 
-			// Fit to backbuffer size
-			DstRect.Max.X = FMath::Min(DstSize.X, DstRect.Max.X);
-			DstRect.Max.Y = FMath::Min(DstSize.Y, DstRect.Max.Y);
+			// Check if resources with the specified regions can be resolved.
+			if (!FDisplayClusterViewportHelpers::GetValidResourceRectsForResolve(FrameTexture, DestTexture, SrcRect, DestRect))
+			{
+				// The SrcRect or DestRect is invalid.
+				return false;
+			}
 
 			FRHICopyTextureInfo CopyInfo;
 
 			CopyInfo.SourceSliceIndex = 0;
 			CopyInfo.DestSliceIndex = DestArrayIndex;
 
-			CopyInfo.Size.X = DstRect.Width();
-			CopyInfo.Size.Y = DstRect.Height();
+			CopyInfo.SourcePosition = FIntVector(SrcRect.Min.X, SrcRect.Min.Y, 0);
+			CopyInfo.DestPosition = FIntVector(DestRect.Min.X, DestRect.Min.Y, 0);
 
-			CopyInfo.DestPosition.X = DstRect.Min.X;
-			CopyInfo.DestPosition.Y = DstRect.Min.Y;
+			CopyInfo.Size = FIntVector(DestRect.Width(), DestRect.Height(), 0);
 
 			TransitionAndCopyTexture(RHICmdList, FrameTexture, DestTexture, CopyInfo);
 
diff --git a/Engine/Plugins/Runtime/nDisplay/Source/DisplayCluster/Private/Render/Viewport/DisplayClusterViewportProxy_Resources.cpp b/Engine/Plugins/Runtime/nDisplay/Source/DisplayCluster/Private/Render/Viewport/DisplayClusterViewportProxy_Resources.cpp
--- a/Engine/Plugins/Runtime/nDisplay/Source/DisplayCluster/Private/Render/Viewport/DisplayClusterViewportProxy_Resources.cpp
+++ b/Engine/Plugins/Runtime/nDisplay/Source/DisplayCluster/Private/Render/Viewport/DisplayClusterViewportProxy_Resources.cpp
@@ -5,6 +5,7 @@
 #include "Render/Viewport/DisplayClusterViewport.h"
 #include "Render/Viewport/DisplayClusterViewportManager.h"
 #include "Render/Viewport/DisplayClusterViewportManagerViewExtension.h"
+#include "Render/Viewport/DisplayClusterViewportHelpers.h"
 
 #include "Render/Viewport/Containers/DisplayClusterViewport_PostRenderSettings.h"
 #include "Render/Viewport/Containers/DisplayClusterViewportProxyData.h"
@@ -145,25 +146,33 @@ namespace UE::DisplayCluster::ViewportProxy
 		check(InputResource);
 		check(OutputResource);
 
+		// Check if resources with the specified regions can be resolved.
+		FIntRect SrcRect(InputRect);
+		FIntRect DestRect(OutputRect);
+		if(!FDisplayClusterViewportHelpers::GetValidResourceRectsForResolve(InputResource, OutputResource, SrcRect, DestRect))
+		{
+			// The SrcRect or DestRect is invalid.
+			return;
+		}
+
 		if (bOutputIsPreviewResource)
 		{
 			// The preview texture should use only RGB colors and ignore the alpha channel. The alpha channel may or may not be inverted in third-party libraries.
-			ResampleCopyTextureImpl_RenderThread<FScreenPS>(RHICmdList, InputResource, OutputResource, InputRect, OutputRect, EDisplayClusterTextureCopyMode::RGB);
+			ResampleCopyTextureImpl_RenderThread<FScreenPS>(RHICmdList, InputResource, OutputResource, SrcRect, DestRect, EDisplayClusterTextureCopyMode::RGB);
 		}
-		else if (InputRect.Size() == OutputRect.Size() && InputResource->GetFormat() == OutputResource->GetFormat())
+		else if (SrcRect.Size() == DestRect.Size() && InputResource->GetFormat() == OutputResource->GetFormat())
 		{
 			FRHICopyTextureInfo CopyInfo;
-			CopyInfo.Size = FIntVector(InputRect.Width(), InputRect.Height(), 0);
-			CopyInfo.SourcePosition.X = InputRect.Min.X;
-			CopyInfo.SourcePosition.Y = InputRect.Min.Y;
-			CopyInfo.DestPosition.X = OutputRect.Min.X;
-			CopyInfo.DestPosition.Y = OutputRect.Min.Y;
+			CopyInfo.SourcePosition = FIntVector(SrcRect.Min.X, SrcRect.Min.Y, 0);
+			CopyInfo.DestPosition = FIntVector(DestRect.Min.X, DestRect.Min.Y, 0);
+
+			CopyInfo.Size = FIntVector(DestRect.Width(), DestRect.Height(), 0);
 
 			TransitionAndCopyTexture(RHICmdList, InputResource, OutputResource, CopyInfo);
 		}
 		else
 		{
-			ResampleCopyTextureImpl_RenderThread<FScreenPS>(RHICmdList, InputResource, OutputResource, InputRect, OutputRect);
+			ResampleCopyTextureImpl_RenderThread<FScreenPS>(RHICmdList, InputResource, OutputResource, SrcRect, DestRect);
 		}
 	}
 
