SharedMemoryMedia: Added experimental option to use the new Fence::Wait function, which is more efficient that polling the fence.

The wait is however called from outside the render thread. This seems to be safe to do on DX12 and Vulkan RHIs but not others, so platform a platform agnostic function was provided but that can have platform specific implementations. Currently only dx12 was provided, with a default implementation that is the original polling method.

For the above reason and the fact that tests show that it introduces more latency with some  variance it is marked as experimental and is not the default.

To enable:
SharedMemoryMedia.FenceWaitingMethod 1

#rnx

diff --git a/Engine/Plugins/Runtime/nDisplay/Source/SharedMemoryMedia/Private/Platforms/Windows/SharedMemoryMediaPlatformWindowsD3D12.cpp b/Engine/Plugins/Runtime/nDisplay/Source/SharedMemoryMedia/Private/Platforms/Windows/SharedMemoryMediaPlatformWindowsD3D12.cpp
--- a/Engine/Plugins/Runtime/nDisplay/Source/SharedMemoryMedia/Private/Platforms/Windows/SharedMemoryMediaPlatformWindowsD3D12.cpp
+++ b/Engine/Plugins/Runtime/nDisplay/Source/SharedMemoryMedia/Private/Platforms/Windows/SharedMemoryMediaPlatformWindowsD3D12.cpp
@@ -294,6 +294,24 @@ FTextureRHIRef FSharedMemoryMediaPlatformWindowsD3D12::OpenSharedTextureByGuid(c
 	);
 }
 
+void FSharedMemoryMediaPlatformWindowsD3D12::GPUFenceWait_AnyThread(FGPUFenceRHIRef Fence, const FRHIGPUMask GpuMask)
+{
+	if (CVarSharedMemoryMediaFenceWaitingMethod.GetValueOnAnyThread() == static_cast<int32>(ESharedMemoryMediaFenceWaitingMethod::Poll))
+	{
+		return FSharedMemoryMediaPlatform::GPUFenceWait_AnyThread(Fence, GpuMask);
+	}
+
+	if (!Fence)
+	{
+		return;
+	}
+
+	// We only pass a null reference because we have knowledge that FD3D12GPUFence::Wait does not use the RHI Cmd list passed,
+	// and instead only calls graph event functions. If that changes, this will need to change as well.
+	// Note: FRHICommandListImmediate::Get() checks if we are in the rendering thread.
+
+	Fence->Wait(*static_cast<FRHICommandListImmediate*>(nullptr), GpuMask);
+}
 
 const FString FSharedMemoryMediaPlatformWindowsD3D12::GetD3D12ComErrorDescription(HRESULT Hresult)
 {
diff --git a/Engine/Plugins/Runtime/nDisplay/Source/SharedMemoryMedia/Private/Platforms/Windows/SharedMemoryMediaPlatformWindowsD3D12.h b/Engine/Plugins/Runtime/nDisplay/Source/SharedMemoryMedia/Private/Platforms/Windows/SharedMemoryMediaPlatformWindowsD3D12.h
--- a/Engine/Plugins/Runtime/nDisplay/Source/SharedMemoryMedia/Private/Platforms/Windows/SharedMemoryMediaPlatformWindowsD3D12.h
+++ b/Engine/Plugins/Runtime/nDisplay/Source/SharedMemoryMedia/Private/Platforms/Windows/SharedMemoryMediaPlatformWindowsD3D12.h
@@ -39,6 +39,7 @@ public:
 	virtual FTextureRHIRef CreateSharedTexture(EPixelFormat Format, bool bSrgb, int32 Width, int32 Height, const FGuid& Guid, uint32 BufferIdx, bool bCrossGpu) override;
 	virtual void ReleaseSharedTexture(uint32 BufferIdx) override;
 	virtual FTextureRHIRef OpenSharedTextureByGuid(const FGuid& Guid, FSharedMemoryMediaTextureDescription& OutTextureDescription) override;
+	virtual void GPUFenceWait_AnyThread(FGPUFenceRHIRef Fence, const FRHIGPUMask GpuMask) override;
 	//~ End FSharedMemoryMediaPlatform interface
 
 	FSharedMemoryMediaPlatformWindowsD3D12();
diff --git a/Engine/Plugins/Runtime/nDisplay/Source/SharedMemoryMedia/Private/SharedMemoryMediaCapture.cpp b/Engine/Plugins/Runtime/nDisplay/Source/SharedMemoryMedia/Private/SharedMemoryMediaCapture.cpp
--- a/Engine/Plugins/Runtime/nDisplay/Source/SharedMemoryMedia/Private/SharedMemoryMediaCapture.cpp
+++ b/Engine/Plugins/Runtime/nDisplay/Source/SharedMemoryMedia/Private/SharedMemoryMediaCapture.cpp
@@ -379,10 +379,13 @@ void USharedMemoryMediaCapture::AddCopyToSharedGpuTexturePass(FRDGBuilder& Graph
 			// Write GPU fence
 			RHICmdList.WriteGPUFence(TextureReadyFences[SharedTextureIdx]);
 
+			// Get the GPU Mask, which will be needed when waiting for the fence.
+			const FRHIGPUMask GPUMask = RHICmdList.GetGPUMask();
+
 			// Spawn a thread that via shared ram will notify receiver that data is ready
 			// It will also verify that the data has been consumed (with a timeout).
 			RunningTasksCount++;
-			UE::Tasks::Launch(UE_SOURCE_LOCATION, [FrameNumber = GFrameCounterRenderThread, SharedTextureIdx, this]()
+			UE::Tasks::Launch(UE_SOURCE_LOCATION, [FrameNumber = GFrameCounterRenderThread, SharedTextureIdx, GPUMask, this]()
 			{
 				// Decrement RunningTasksCount when the task exits
 				ON_SCOPE_EXIT
@@ -397,10 +400,9 @@ void USharedMemoryMediaCapture::AddCopyToSharedGpuTexturePass(FRDGBuilder& Graph
 				{
 					TRACE_CPUPROFILER_EVENT_SCOPE(WaitForGpuTextureReadyFence);
 
-					while (TextureReadyFences[SharedTextureIdx] && !TextureReadyFences[SharedTextureIdx]->Poll())
-					{
-						FPlatformProcess::SleepNoStats(0);
-					}
+					check(PlatformData.IsValid());
+
+					PlatformData->GPUFenceWait_AnyThread(TextureReadyFences[SharedTextureIdx], GPUMask);
 				}
 
 				// Update shared memory metadata to indicate to the receiver that there is new data
diff --git a/Engine/Plugins/Runtime/nDisplay/Source/SharedMemoryMedia/Private/SharedMemoryMediaPlatform.cpp b/Engine/Plugins/Runtime/nDisplay/Source/SharedMemoryMedia/Private/SharedMemoryMediaPlatform.cpp
--- a/Engine/Plugins/Runtime/nDisplay/Source/SharedMemoryMedia/Private/SharedMemoryMediaPlatform.cpp
+++ b/Engine/Plugins/Runtime/nDisplay/Source/SharedMemoryMedia/Private/SharedMemoryMediaPlatform.cpp
@@ -2,6 +2,17 @@
 
 #include "SharedMemoryMediaPlatform.h"
 
+
+// Declare the CVar to control GPU fence waiting method
+TAutoConsoleVariable<int32> CVarSharedMemoryMediaFenceWaitingMethod(
+	TEXT("SharedMemoryMedia.FenceWaitingMethod"),
+	static_cast<int32>(ESharedMemoryMediaFenceWaitingMethod::Poll),
+	TEXT("Choose how GPU fences are waited on in shared memory media.\n")
+	TEXT("0: Poll Fence for completion\n")
+	TEXT("1: Use Fence->Wait. It is experimental, adds a variable amount of latency, but is more CPU efficient."),
+	ECVF_Default
+);
+
 FString FSharedMemoryMediaPlatformFactory::GetRhiTypeString(const ERHIInterfaceType RhiType)
 {
 	static TMap<ERHIInterfaceType, FString> Rhis;
@@ -46,3 +57,10 @@ FSharedMemoryMediaPlatformFactory* FSharedMemoryMediaPlatformFactory::Get()
 }
 
 
+void FSharedMemoryMediaPlatform::GPUFenceWait_AnyThread(FGPUFenceRHIRef Fence, const FRHIGPUMask GpuMask)
+{
+	while (Fence && !Fence->Poll())
+	{
+		FPlatformProcess::SleepNoStats(0);
+	}
+}
\ No newline at end of file
diff --git a/Engine/Plugins/Runtime/nDisplay/Source/SharedMemoryMedia/Private/SharedMemoryMediaPlatform.h b/Engine/Plugins/Runtime/nDisplay/Source/SharedMemoryMedia/Private/SharedMemoryMediaPlatform.h
--- a/Engine/Plugins/Runtime/nDisplay/Source/SharedMemoryMedia/Private/SharedMemoryMediaPlatform.h
+++ b/Engine/Plugins/Runtime/nDisplay/Source/SharedMemoryMedia/Private/SharedMemoryMediaPlatform.h
@@ -7,6 +7,17 @@
 #include "SharedMemoryMediaTypes.h"
 
 
+// Declaring this CVar as an extern so it can be accessed from different platform files.
+extern TAutoConsoleVariable<int32> CVarSharedMemoryMediaFenceWaitingMethod;
+
+/** Helper enum to interpret values from CVarSharedMemoryMediaFenceWaitingMethod more conveniently. */
+enum class ESharedMemoryMediaFenceWaitingMethod
+{
+	Poll = 0,
+	Wait = 1
+};
+
+
 /** This is the interface to the platform specific functions needed by shared memory media. 
  *  It is mostly about the platform specific implementation of cross gpu textures.
  */
@@ -29,6 +40,9 @@ public:
 
 	/** Release any platform specific resources related to the indexed texture */
 	virtual void ReleaseSharedTexture(uint32 BufferIdx) {};
+
+	/** Waits for the completion of a GPU fence. Can be called from any thread, unlike FRHIGPUFence::Wait in general */
+	virtual void GPUFenceWait_AnyThread(FGPUFenceRHIRef Fence, const FRHIGPUMask GpuMask);
 };
 
 /** Factory of registered FSharedMemoryMediaPlatform RHI implementations. Use RegisterPlatformForRhi to register them. */
diff --git a/Engine/Plugins/Runtime/nDisplay/Source/SharedMemoryMedia/Private/SharedMemoryMediaPlayer.cpp b/Engine/Plugins/Runtime/nDisplay/Source/SharedMemoryMedia/Private/SharedMemoryMediaPlayer.cpp
--- a/Engine/Plugins/Runtime/nDisplay/Source/SharedMemoryMedia/Private/SharedMemoryMediaPlayer.cpp
+++ b/Engine/Plugins/Runtime/nDisplay/Source/SharedMemoryMedia/Private/SharedMemoryMediaPlayer.cpp
@@ -851,6 +851,9 @@ void FSharedMemoryMediaPlayer::JustInTimeSampleRender(FRHICommandListImmediate&
 		RHICmdList.WriteGPUFence(FrameAckFences[SharedMemoryIdx]);
 	}
 
+	// Get the GPU Mask, which will be needed when waiting for the fence.
+	const FRHIGPUMask GPUMask = RHICmdList.GetGPUMask();
+
 	// spawn async task that will wait on FrameAckFence, ack the frame to the Media Capture, and flag that it is clear for re-use
 
 	RunningTasksCount++;
@@ -860,7 +863,8 @@ void FSharedMemoryMediaPlayer::JustInTimeSampleRender(FRHICommandListImmediate&
 			SharedMemoryData, 
 			ExpectedFrameNumber, 
 			SharedMemoryIdx, 
-			ReceiverIdx = ReceiverIndex[SharedMemoryIdx], 
+			ReceiverIdx = ReceiverIndex[SharedMemoryIdx],
+			GPUMask,
 			this
 	    ]()
 		{
@@ -872,18 +876,9 @@ void FSharedMemoryMediaPlayer::JustInTimeSampleRender(FRHICommandListImmediate&
 				RunningTasksCount--;
 			};
 
-			do
-			{
-				const bool bFenceReached = !FrameAckFences[SharedMemoryIdx] || FrameAckFences[SharedMemoryIdx]->Poll();
-
-				if (bFenceReached)
-				{
-					break;
-				}
-
-				FPlatformProcess::SleepNoStats(SpinWaitTimeSeconds);
+			check(PlatformData.IsValid());
 
-			} while (true);
+			PlatformData->GPUFenceWait_AnyThread(FrameAckFences[SharedMemoryIdx], GPUMask);
 
 			// We clear the fence and signal that it can be re-used.
 			FrameAckFences[SharedMemoryIdx]->Clear();
