SharedMemoryMedia: Support capturing source texture subrects.

It now takes into account the source and destination rects in CopyInfo passed to OnCustomCapture_RenderingThread. While SMM shares the texture without copying it to the given OutputTexture, it will follow the resizing rules that resulted in the size of OutputTexture. If this sizing changes, SMM will error out and log the resolution mismatch. Also added sanity checks to ensure that the source and dest copy rects fit inside the source and dest texture sizes.

The need for this arose while testing tiled backbuffer media output of nDisplay nodes.

#rn SharedMemoryMedia now supports capturing subregions.

diff --git a/Engine/Plugins/Runtime/nDisplay/Source/SharedMemoryMedia/Private/SharedMemoryMediaCapture.cpp b/Engine/Plugins/Runtime/nDisplay/Source/SharedMemoryMedia/Private/SharedMemoryMediaCapture.cpp
--- a/Engine/Plugins/Runtime/nDisplay/Source/SharedMemoryMedia/Private/SharedMemoryMediaCapture.cpp
+++ b/Engine/Plugins/Runtime/nDisplay/Source/SharedMemoryMedia/Private/SharedMemoryMediaCapture.cpp
@@ -24,6 +24,24 @@ BEGIN_SHADER_PARAMETER_STRUCT(FCopyToSharedGpuTexturePass, )
 END_SHADER_PARAMETER_STRUCT()
 
 
+namespace UE::SMMCapture
+{
+	/**
+	 * Checks if a given Rect is completely contained within a size of a texture.
+	 *
+	 * @param Size The size of the texture
+	 * @param Rect The rectangle to check
+	 * @return True iff the rect is completely contained within the bounds of the texture size.
+	 */
+	static bool IsRectContainedInTextureSize(const FIntRect& Rect, const FIntVector& Size)
+	{
+		return Rect.Min.X >= 0         && Rect.Min.Y >= 0          // top-left should be 0+
+			&& Rect.Max.X <= Size.X    && Rect.Max.Y <= Size.Y     // bottom-right should be within the size bounds
+			&& Rect.Min.X < Rect.Max.X && Rect.Min.Y < Rect.Max.Y; // Min should be less than Max.
+	}
+}
+
+
 bool USharedMemoryMediaCapture::InitializeCapture()
 {
 	// Validate Media Output type
@@ -227,12 +245,15 @@ void USharedMemoryMediaCapture::OnCustomCapture_RenderingThread(
 		if (!SharedCrossGpuTextures[Idx].IsValid())
 		{
 			const FGuid Guid = FGuid::NewGuid();
+
+			// Shared texture size follows output texture which has already gone through MediaCapture's resizing logic.
+			const FIntVector SharedSize = OutputTexture->Desc.GetSize();
 						
 			SharedCrossGpuTextures[Idx] = PlatformData->CreateSharedTexture(
 				InSourceTexture->Desc.Format,
 				EnumHasAnyFlags(InSourceTexture->Desc.Flags, TexCreate_SRGB),
-				CopyInfo.Size.X, 
-				CopyInfo.Size.Y, 
+				SharedSize.X,
+				SharedSize.Y,
 				Guid, 
 				Idx,
 				SharedMemoryMediaOutput->bCrossGpu
@@ -254,9 +275,58 @@ void USharedMemoryMediaCapture::OnCustomCapture_RenderingThread(
 				Idx, *SharedCrossGpuTextureGuids[Idx].ToString(), *SharedMemoryMediaOutput->UniqueName
 			);
 		}
+
+		// Verify that the shared memory media size is still the same as the OutputTexture
+
+		const FIntVector SharedSize = SharedCrossGpuTextures[Idx]->GetDesc().GetSize();
+		const FIntVector OutputSize = OutputTexture->Desc.GetSize();
+
+		if (SharedSize != OutputSize)
+		{
+			UE_LOG(LogSharedMemoryMedia, Warning,
+				TEXT("Cross GPU texture is out of date for Unique Name '%s'. It's size was (%d,%d) but OutputTexture was (%d,%d)"),
+				*SharedMemoryMediaOutput->UniqueName, SharedSize.X, SharedSize.Y, OutputSize.X, OutputSize.Y
+			);
+
+			SetState(EMediaCaptureState::Error);
+			return;
+		}
+	}
+
+	// Sanity check the copy regions are contained inside the input and output texture sizes.
+	{
+		const FIntRect Rect = CopyInfo.GetSourceRect();
+		const FIntVector Size = InSourceTexture->Desc.GetSize();
+
+		if (!UE::SMMCapture::IsRectContainedInTextureSize(Rect, Size))
+		{
+			UE_LOG(LogSharedMemoryMedia, Error, 
+				TEXT("Invalid source CopyInfo passed SharedMemoryMedia with Unique Name '%s'. (%d,%d,%d,%d) not contained in (%d,%d)"),
+				*SharedMemoryMediaOutput->UniqueName, Rect.Min.X, Rect.Min.Y, Rect.Max.X, Rect.Max.Y, Size.X, Size.Y
+			);
+
+			SetState(EMediaCaptureState::Error);
+			return;
+		}
+	}
+	{
+		const FIntRect Rect = CopyInfo.GetDestRect();
+		const FIntVector Size = OutputTexture->Desc.GetSize();
+
+		if (!UE::SMMCapture::IsRectContainedInTextureSize(Rect, Size))
+		{
+			UE_LOG(LogSharedMemoryMedia, Error, 
+				TEXT("Invalid destination CopyInfo passed SharedMemoryMedia with Unique Name '%s'. [(%d,%d),(%d,%d)] not contained in (%d,%d)"),
+				*SharedMemoryMediaOutput->UniqueName, Rect.Min.X, Rect.Min.Y, Rect.Max.X, Rect.Max.Y, Size.X, Size.Y
+			);
+
+			SetState(EMediaCaptureState::Error);
+			return;
+		}
 	}
 
 	FRDGTextureRef SourceTexture = InSourceTexture;
+	FRHICopyTextureInfo CurrentCopyInfo = CopyInfo; // This may be modified if there are intermediate steps (i.e. invert alpha pass)
 
 	// When enabled, add pass to invert alpha
 
@@ -270,7 +340,10 @@ void USharedMemoryMediaCapture::OnCustomCapture_RenderingThread(
 		}
 
 		const FRDGTextureDesc Desc = FRDGTextureDesc::Create2D(
-			FIntPoint(CopyInfo.Size.X, CopyInfo.Size.Y),
+			FIntPoint(
+				OutputTexture->Desc.GetSize().X,
+				OutputTexture->Desc.GetSize().Y
+			),
 			SourceTexture->Desc.Format,
 			FClearValueBinding::Black,
 			InvertedAlphaTextureFlags
@@ -279,26 +352,26 @@ void USharedMemoryMediaCapture::OnCustomCapture_RenderingThread(
 		FRDGTextureRef InvertedAlphaTexture = GraphBuilder.CreateTexture(Desc, TEXT("SharedMemoryMediaInvertedAlphaTexture"));
 		check(InvertedAlphaTexture);
 
-		AddInvertAlphaConversionPass(GraphBuilder, SourceTexture, InvertedAlphaTexture);
+		AddInvertAlphaConversionPass(GraphBuilder, SourceTexture, InvertedAlphaTexture, CurrentCopyInfo);
 
 		SourceTexture = InvertedAlphaTexture;
+
+		// Adjust the copy info's Source Position since SourceTexture is now InvertedAlphaTexture which is like OutputTexture.
+		CurrentCopyInfo.SourcePosition = CurrentCopyInfo.DestPosition;
 	}
 
 	// Add the copy texture pass
-	AddCopyToSharedGpuTexturePass(GraphBuilder, SourceTexture, GFrameCounterRenderThread % NUMBUFFERS);
+	AddCopyToSharedGpuTexturePass(GraphBuilder, SourceTexture, GFrameCounterRenderThread % NUMBUFFERS, CurrentCopyInfo);
 }
 
-void USharedMemoryMediaCapture::AddInvertAlphaConversionPass(FRDGBuilder& GraphBuilder, const FRDGTextureRef& SourceTexture, FRDGTextureRef DestTexture)
+void USharedMemoryMediaCapture::AddInvertAlphaConversionPass(FRDGBuilder& GraphBuilder, const FRDGTextureRef& SourceTexture, FRDGTextureRef DestTexture, const FRHICopyTextureInfo& CopyInfo)
 {
-	// Rectangle area to use from source
-	const FIntRect ViewRect(0, 0, SourceTexture->Desc.GetSize().X, SourceTexture->Desc.GetSize().Y);
-
 	FGlobalShaderMap* GlobalShaderMap = GetGlobalShaderMap(GMaxRHIFeatureLevel);
 	TShaderMapRef<FScreenPassVS> VertexShader(GlobalShaderMap);
 
 	// Configure source/output viewport to get the right UV scaling from source texture to output texture
-	FScreenPassTextureViewport InputViewport(SourceTexture, ViewRect);
-	FScreenPassTextureViewport OutputViewport(DestTexture);
+	FScreenPassTextureViewport InputViewport(SourceTexture, CopyInfo.GetSourceRect());
+	FScreenPassTextureViewport OutputViewport(DestTexture, CopyInfo.GetDestRect());
 
 	// In cases where texture is converted from a format that doesn't have A channel, we want to force set it to 1.
 	EMediaCaptureConversionOperation MediaConversionOperation = EMediaCaptureConversionOperation::INVERT_ALPHA;
@@ -317,7 +390,7 @@ void USharedMemoryMediaCapture::AddInvertAlphaConversionPass(FRDGBuilder& GraphB
 	FSceneViewInitOptions ViewInitOptions;
 
 	ViewInitOptions.ViewFamily = &ViewFamily;
-	ViewInitOptions.SetViewRectangle(ViewRect);
+	ViewInitOptions.SetViewRectangle(CopyInfo.GetSourceRect());
 	ViewInitOptions.ViewOrigin = FVector::ZeroVector;
 	ViewInitOptions.ViewRotationMatrix = FMatrix::Identity;
 	ViewInitOptions.ProjectionMatrix = FMatrix::Identity;
@@ -337,7 +410,7 @@ void USharedMemoryMediaCapture::AddInvertAlphaConversionPass(FRDGBuilder& GraphB
 }
 
 
-void USharedMemoryMediaCapture::AddCopyToSharedGpuTexturePass(FRDGBuilder& GraphBuilder, FRDGTextureRef InSourceTexture, uint32 SharedTextureIdx)
+void USharedMemoryMediaCapture::AddCopyToSharedGpuTexturePass(FRDGBuilder& GraphBuilder, FRDGTextureRef InSourceTexture, uint32 SharedTextureIdx, const FRHICopyTextureInfo& CopyInfo)
 {
 	using namespace UE::SharedMemoryMedia;
 
@@ -352,7 +425,7 @@ void USharedMemoryMediaCapture::AddCopyToSharedGpuTexturePass(FRDGBuilder& Graph
 		RDG_EVENT_NAME("Copy%sToSharedGpuTexture", InSourceTexture->Name),
 		PassParameters,
 		ERDGPassFlags::Copy,
-		[InSourceTexture, SharedTextureIdx, this](FRHICommandList& RHICmdList)
+		[InSourceTexture, SharedTextureIdx, CopyInfo, this](FRHICommandList& RHICmdList)
 		{
 			// bTextureReadyFenceBusy will also signal that the resource is safe to reuse
 			if (bTextureReadyFenceBusy[SharedTextureIdx])
@@ -373,7 +446,6 @@ void USharedMemoryMediaCapture::AddCopyToSharedGpuTexturePass(FRDGBuilder& Graph
 			bTextureReadyFenceBusy[SharedTextureIdx] = true;
 
 			// Do the copy
-			FRHICopyTextureInfo CopyInfo;
 			RHICmdList.CopyTexture(InSourceTexture->GetRHI(), SharedCrossGpuTextures[SharedTextureIdx], CopyInfo);
 
 			// Write GPU fence
diff --git a/Engine/Plugins/Runtime/nDisplay/Source/SharedMemoryMedia/Public/SharedMemoryMediaCapture.h b/Engine/Plugins/Runtime/nDisplay/Source/SharedMemoryMedia/Public/SharedMemoryMediaCapture.h
--- a/Engine/Plugins/Runtime/nDisplay/Source/SharedMemoryMedia/Public/SharedMemoryMediaCapture.h
+++ b/Engine/Plugins/Runtime/nDisplay/Source/SharedMemoryMedia/Public/SharedMemoryMediaCapture.h
@@ -64,11 +64,23 @@ protected:
 
 	//~ End UMediaCapture interface
 
+	UE_DEPRECATED(5.5, "Use the version of AddCopyToSharedGpuTexturePass with FRHICopyTextureInfo.")
+	void AddCopyToSharedGpuTexturePass(FRDGBuilder& GraphBuilder, FRDGTextureRef InSourceTexture, uint32 SharedTextureIdx)
+	{
+		AddCopyToSharedGpuTexturePass(GraphBuilder, InSourceTexture, SharedTextureIdx, FRHICopyTextureInfo());
+	}
+
+	UE_DEPRECATED(5.5, "Use the version of AddInvertAlphaConversionPass with FRHICopyTextureInfo.")
+	void AddInvertAlphaConversionPass(FRDGBuilder& GraphBuilder, const FRDGTextureRef& SourceTexture, FRDGTextureRef DestTexture)
+	{
+		AddInvertAlphaConversionPass(GraphBuilder, SourceTexture, DestTexture, FRHICopyTextureInfo());
+	}
+
 	/** Adds to the graph builder the copy from the input texture to the shared gpu texture that corresponds to the current frame number */
-	void AddCopyToSharedGpuTexturePass(FRDGBuilder& GraphBuilder, FRDGTextureRef InSourceTexture, uint32 SharedTextureIdx);
+	void AddCopyToSharedGpuTexturePass(FRDGBuilder& GraphBuilder, FRDGTextureRef InSourceTexture, uint32 SharedTextureIdx, const FRHICopyTextureInfo& CopyInfo);
 
 	/** Adds a pass to invert the alpha of the texture */
-	void AddInvertAlphaConversionPass(FRDGBuilder& GraphBuilder, const FRDGTextureRef& SourceTexture, FRDGTextureRef DestTexture);
+	void AddInvertAlphaConversionPass(FRDGBuilder& GraphBuilder, const FRDGTextureRef& SourceTexture, FRDGTextureRef DestTexture, const FRHICopyTextureInfo& CopyInfo);
 
 private:
 
