nDisplay:  Improved the "Per-View Lumen" feature for rendering in DisplayCluster.

- Added new options to CVar "DC.LumenPerView" (Per-View Lumen scene): "Disabled", "Use only in Cluster", "Use in both Cluster and Preview", "Use only in Preview", "Use only for Inner Frustum in Cluster"

diff --git a/Engine/Plugins/Runtime/nDisplay/Source/DisplayCluster/Private/Game/EngineClasses/Basics/DisplayClusterViewportClient.cpp b/Engine/Plugins/Runtime/nDisplay/Source/DisplayCluster/Private/Game/EngineClasses/Basics/DisplayClusterViewportClient.cpp
--- a/Engine/Plugins/Runtime/nDisplay/Source/DisplayCluster/Private/Game/EngineClasses/Basics/DisplayClusterViewportClient.cpp
+++ b/Engine/Plugins/Runtime/nDisplay/Source/DisplayCluster/Private/Game/EngineClasses/Basics/DisplayClusterViewportClient.cpp
@@ -87,14 +87,6 @@ static FAutoConsoleVariableRef CVarDisplayClusterSortViews(
 	ECVF_RenderThreadSafe
 );
 
-int32 GDisplayClusterLumenPerView = 1;
-static FAutoConsoleVariableRef CVarDisplayClusterLumenPerView(
-	TEXT("DC.LumenPerView"),
-	GDisplayClusterLumenPerView,
-	TEXT("Separate Lumen scene cache allocated for each View.  Reduces artifacts where views affect one another, at a cost in GPU memory."),
-	ECVF_RenderThreadSafe
-);
-
 int32 GDisplayClusterDebugDraw = 1;
 static FAutoConsoleVariableRef CVarDisplayClusterDebugDraw(
 	TEXT("DC.DebugDraw"),
@@ -582,7 +574,7 @@ void UDisplayClusterViewportClient::Draw(FViewport* InViewport, FCanvas* SceneCa
 					Views.Add(View);
 
 					// Apply viewport context settings to view (crossGPU, visibility, etc)
-					DCView.Viewport->SetupSceneView(DCView.ContextNum, World, ViewFamily, *View);
+					DCView.Viewport->SetupSceneView(DCView.ContextNum, MyWorld, ViewFamily, *View);
 
 					// We don't allow instanced stereo currently
 					View->bIsInstancedStereoEnabled = false;
@@ -626,15 +618,7 @@ void UDisplayClusterViewportClient::Draw(FViewport* InViewport, FCanvas* SceneCa
 
 					View->CameraConstrainedViewRect = View->UnscaledViewRect;
 
-					// Enable per-view Lumen scene
-					if (GDisplayClusterLumenPerView)
-					{
-						View->State->AddLumenSceneData(MyWorld->Scene);
-					}
-					else
-					{
-						View->State->RemoveLumenSceneData(MyWorld->Scene);
-					}
+					
 
 					{
 						// Save the location of the view.
diff --git a/Engine/Plugins/Runtime/nDisplay/Source/DisplayCluster/Private/Render/Viewport/DisplayClusterViewport.cpp b/Engine/Plugins/Runtime/nDisplay/Source/DisplayCluster/Private/Render/Viewport/DisplayClusterViewport.cpp
--- a/Engine/Plugins/Runtime/nDisplay/Source/DisplayCluster/Private/Render/Viewport/DisplayClusterViewport.cpp
+++ b/Engine/Plugins/Runtime/nDisplay/Source/DisplayCluster/Private/Render/Viewport/DisplayClusterViewport.cpp
@@ -31,6 +31,7 @@
 #include "EngineUtils.h"
 #include "OpenColorIORendering.h"
 #include "PrimitiveDrawingUtils.h"
+#include "SceneManagement.h"
 #include "SceneView.h"
 #include "UnrealClient.h"
 
@@ -373,6 +374,15 @@ void FDisplayClusterViewport::SetupSceneView(uint32 ContextNum, class UWorld* Wo
 	{
 		InDisplayDeviceComponent->SetupSceneView(*ViewportPreview, ContextNum, InOutViewFamily, InOutView);
 	}
+
+	if (ShouldUseLumenPerView())
+	{
+		InOutView.State->AddLumenSceneData(World->Scene);
+	}
+	else
+	{
+		InOutView.State->RemoveLumenSceneData(World->Scene);
+	}
 }
 
 float FDisplayClusterViewport::GetClusterRenderTargetRatioMult(const FDisplayClusterRenderFrameSettings& InFrameSettings) const
diff --git a/Engine/Plugins/Runtime/nDisplay/Source/DisplayCluster/Private/Render/Viewport/DisplayClusterViewport.h b/Engine/Plugins/Runtime/nDisplay/Source/DisplayCluster/Private/Render/Viewport/DisplayClusterViewport.h
--- a/Engine/Plugins/Runtime/nDisplay/Source/DisplayCluster/Private/Render/Viewport/DisplayClusterViewport.h
+++ b/Engine/Plugins/Runtime/nDisplay/Source/DisplayCluster/Private/Render/Viewport/DisplayClusterViewport.h
@@ -333,6 +333,9 @@ public:
 	/** Returns true if this viewport is to be used as a tile source. */
 	bool CanSplitIntoTiles() const;
 
+	/** Returns true if per-view Lumen scene is enabled. */
+	bool ShouldUseLumenPerView() const;
+
 	/** Returns true if this viewport should use the RTT. */
 	bool ShouldUseRenderTargetResource() const;
 
@@ -386,7 +389,7 @@ public:
 	void SetViewportConfigurationData(const UDisplayClusterConfigurationViewport* InConfigurationData);
 
 	/** Support view states for preview. */
-	FSceneViewStateInterface* GetViewState(uint32 ViewIndex);
+	class FSceneViewStateInterface* GetViewState(uint32 ViewIndex);
 
 	/** Cleanup view states. */
 	void CleanupViewState();
diff --git a/Engine/Plugins/Runtime/nDisplay/Source/DisplayCluster/Private/Render/Viewport/DisplayClusterViewport_Rules.cpp b/Engine/Plugins/Runtime/nDisplay/Source/DisplayCluster/Private/Render/Viewport/DisplayClusterViewport_Rules.cpp
--- a/Engine/Plugins/Runtime/nDisplay/Source/DisplayCluster/Private/Render/Viewport/DisplayClusterViewport_Rules.cpp
+++ b/Engine/Plugins/Runtime/nDisplay/Source/DisplayCluster/Private/Render/Viewport/DisplayClusterViewport_Rules.cpp
@@ -44,6 +44,40 @@
 #include "Misc/DisplayClusterLog.h"
 #include "Misc/Parse.h"
 
+/** Auxiliary enumeration for CVar “DC.LumenPerView”. */
+enum class EDisplayClusterLumenPerView : uint8
+{
+	// Disabled
+	Disabled = 0,
+
+	// Use only in Cluster.
+	UseOnlyInCluster,
+
+	// Use in both Cluster and Preview.
+	UseEverywhere,
+
+	// Use only in Preview.
+	UseOnlyInPreview,
+
+	// Use only for Inner Frustum in Cluster.
+	UseOnlyForICVXFCameraInCluster,
+
+	MAX = UseOnlyForICVXFCameraInCluster
+};
+
+int32 GDisplayClusterLumenPerView = (uint8)EDisplayClusterLumenPerView::UseOnlyInCluster;
+static FAutoConsoleVariableRef CVarDisplayClusterLumenPerView(
+	TEXT("DC.LumenPerView"),
+	GDisplayClusterLumenPerView,
+	TEXT("Separate Lumen scene cache allocated for each View. (Default = 1)\n")
+	TEXT("Reduces artifacts where views affect one another, at a cost in GPU memory.\n")
+	TEXT(" 0 - Disabled.\n")
+	TEXT(" 1 - Use only in Cluster.\n")
+	TEXT(" 2 - Use in both Cluster and Preview.\n")
+	TEXT(" 3 - Use only in Preview.\n")
+	TEXT(" 4 - Use only for Inner Frustum in Cluster.\n"),
+	ECVF_RenderThreadSafe
+);
 
 ///////////////////////////////////////////////////////////////////////////////////////
 // FDisplayClusterViewport
@@ -444,3 +478,75 @@ uint8 FDisplayClusterViewport::GetPriority() const
 
 	return OutOrder;
 }
+
+bool FDisplayClusterViewport::ShouldUseLumenPerView() const
+{
+	switch (RenderSettings.CaptureMode)
+	{
+	case EDisplayClusterViewportCaptureMode::Chromakey:
+	case EDisplayClusterViewportCaptureMode::Lightcard:
+		// This type of viewport does not use Lumen
+		return false;
+
+	default:
+		break;
+	}
+
+	const EDisplayClusterLumenPerView LumenPerViewMode = (EDisplayClusterLumenPerView)FMath::Clamp(
+		GDisplayClusterLumenPerView, 0, (uint8)EDisplayClusterLumenPerView::MAX);
+
+	switch (LumenPerViewMode)
+	{
+	// Use only in Cluster.
+	case EDisplayClusterLumenPerView::UseOnlyInCluster:
+		return !Configuration->IsPreviewRendering();
+
+	// Use only in Preview.
+	case EDisplayClusterLumenPerView::UseOnlyInPreview:
+		return Configuration->IsPreviewRendering();
+
+	// Use only for Inner Frustum in Cluster.
+	case EDisplayClusterLumenPerView::UseOnlyForICVXFCameraInCluster:
+		if (Configuration->IsPreviewRendering())
+		{
+			// Not a cluster
+			return false;
+		}
+
+		// Full-frame In-Camera viewport
+		if (EnumHasAnyFlags(RenderSettingsICVFX.RuntimeFlags, EDisplayClusterViewportRuntimeICVFXFlags::InCamera))
+		{
+			return true;
+		}
+
+		// Tile In-Camera viewport
+		if (RenderSettings.TileSettings.GetType() == EDisplayClusterViewportTileType::Tile)
+		{
+			if (FDisplayClusterViewportManager* ViewportManager = Configuration->GetViewportManagerImpl())
+			{
+				TSharedPtr<FDisplayClusterViewport, ESPMode::ThreadSafe> SourceViewport =
+					ViewportManager->ImplFindViewport(RenderSettings.TileSettings.GetSourceViewportId());
+
+				if (SourceViewport.IsValid())
+				{
+					if (EnumHasAnyFlags(SourceViewport->RenderSettingsICVFX.RuntimeFlags, EDisplayClusterViewportRuntimeICVFXFlags::InCamera))
+					{
+						return true;
+					}
+				}
+			}
+		}
+		
+		// Other viewports
+		return false;
+
+	// Use in both Cluster and Preview.
+	case EDisplayClusterLumenPerView::UseEverywhere:
+		return true;
+
+	default:
+		break;
+	}
+
+	return false;
+}
