nDisplay: Fixed editor crash on opening nDisplay Config asset.

diff --git a/Engine/Plugins/Runtime/nDisplay/Source/DisplayCluster/Private/Render/Viewport/DisplayClusterViewport.cpp b/Engine/Plugins/Runtime/nDisplay/Source/DisplayCluster/Private/Render/Viewport/DisplayClusterViewport.cpp
--- a/Engine/Plugins/Runtime/nDisplay/Source/DisplayCluster/Private/Render/Viewport/DisplayClusterViewport.cpp
+++ b/Engine/Plugins/Runtime/nDisplay/Source/DisplayCluster/Private/Render/Viewport/DisplayClusterViewport.cpp
@@ -375,13 +375,16 @@ void FDisplayClusterViewport::SetupSceneView(uint32 ContextNum, class UWorld* Wo
 		InDisplayDeviceComponent->SetupSceneView(*ViewportPreview, ContextNum, InOutViewFamily, InOutView);
 	}
 
-	if (ShouldUseLumenPerView())
+	if (InOutView.State && World && World->Scene)
 	{
-		InOutView.State->AddLumenSceneData(World->Scene);
-	}
-	else
-	{
-		InOutView.State->RemoveLumenSceneData(World->Scene);
+		if (ShouldUseLumenPerView())
+		{
+			InOutView.State->AddLumenSceneData(World->Scene);
+		}
+		else
+		{
+			InOutView.State->RemoveLumenSceneData(World->Scene);
+		}
 	}
 }
 
