nDisplay: Fix ICVFX panel mesh flicker when 2+ viewports control the same mesh.

The behavior is now for the first viewport to control a given mesh takes priority and the others do not attempt to set its material as well.

#rnx

diff --git a/Engine/Plugins/Runtime/nDisplay/Source/DisplayCluster/Private/Render/Viewport/Preview/DisplayClusterViewportManagerPreview.cpp b/Engine/Plugins/Runtime/nDisplay/Source/DisplayCluster/Private/Render/Viewport/Preview/DisplayClusterViewportManagerPreview.cpp
--- a/Engine/Plugins/Runtime/nDisplay/Source/DisplayCluster/Private/Render/Viewport/Preview/DisplayClusterViewportManagerPreview.cpp
+++ b/Engine/Plugins/Runtime/nDisplay/Source/DisplayCluster/Private/Render/Viewport/Preview/DisplayClusterViewportManagerPreview.cpp
@@ -200,11 +200,15 @@ void FDisplayClusterViewportManagerPreview::OnPreviewRenderTick()
 
 void FDisplayClusterViewportManagerPreview::Update()
 {
+	// This set is passed to the preview update of each viewport so that they can avoid fighting to control
+	// the same mesh component by knowing if a previous viewport is already in control of it.
+	TSet<const UMeshComponent*> MeshComponentsVisited;
+
 	for (const TSharedPtr<FDisplayClusterViewportPreview, ESPMode::ThreadSafe>& ViewportPreviewIt : GetEntireClusterPreviewViewportsImpl())
 	{
 		if (ViewportPreviewIt.IsValid())
 		{
-			ViewportPreviewIt->Update();
+			ViewportPreviewIt->Update(MeshComponentsVisited);
 		}
 	}
 }
diff --git a/Engine/Plugins/Runtime/nDisplay/Source/DisplayCluster/Private/Render/Viewport/Preview/DisplayClusterViewportPreview.cpp b/Engine/Plugins/Runtime/nDisplay/Source/DisplayCluster/Private/Render/Viewport/Preview/DisplayClusterViewportPreview.cpp
--- a/Engine/Plugins/Runtime/nDisplay/Source/DisplayCluster/Private/Render/Viewport/Preview/DisplayClusterViewportPreview.cpp
+++ b/Engine/Plugins/Runtime/nDisplay/Source/DisplayCluster/Private/Render/Viewport/Preview/DisplayClusterViewportPreview.cpp
@@ -42,7 +42,7 @@ void FDisplayClusterViewportPreview::Release()
 	PreviewEditableMesh.Release(InViewport);
 	}
 
-void FDisplayClusterViewportPreview::Update()
+void FDisplayClusterViewportPreview::Update(TSet<const UMeshComponent*>& MeshComponentsVisited)
 {
 	RuntimeFlags = EDisplayClusterViewportPreviewFlags::None;
 
@@ -60,38 +60,80 @@ void FDisplayClusterViewportPreview::Update()
 		EnumAddFlags(RuntimeFlags, EDisplayClusterViewportPreviewFlags::HasValidPreviewRTT);
 	}
 
-	FDisplayClusterViewport* InViewport = GetViewportImpl();
-	UDisplayClusterCameraComponent* ViewPointComponent = InViewport ? InViewport->GetViewPointCameraComponent(EDisplayClusterRootActorType::Configuration) : nullptr;
-	UDisplayClusterDisplayDeviceBaseComponent* InDisplayDeviceComponent = InViewport ? InViewport->GetDisplayDeviceComponent(Configuration->GetPreviewSettings().DisplayDeviceRootActorType) : nullptr;
-
-	// Update preview meshes only if DisplayDevice is used
-	PreviewMesh.Update(InViewport, InDisplayDeviceComponent, ViewPointComponent);
-	PreviewEditableMesh.Update(InViewport, InDisplayDeviceComponent, ViewPointComponent);
-
-	// Update Runtime Flags:
-	if (PreviewMesh.HasAnyFlag(EDisplayClusterViewportPreviewMeshFlags::HasDeletedMaterialInstance | EDisplayClusterViewportPreviewMeshFlags::HasChangedMaterialInstance))
-	{
-		EnumAddFlags(RuntimeFlags, EDisplayClusterViewportPreviewFlags::HasChangedPreviewMeshMaterialInstance);
-	}
-	if (PreviewEditableMesh.HasAnyFlag(EDisplayClusterViewportPreviewMeshFlags::HasDeletedMaterialInstance | EDisplayClusterViewportPreviewMeshFlags::HasChangedMaterialInstance))
+	// Don't update the meshes if any of them are being controlled by other viewports
+	if (MeshComponentsVisited.Contains(PreviewMesh.GetMeshComponent()) || MeshComponentsVisited.Contains(PreviewEditableMesh.GetMeshComponent()))
 	{
-		EnumAddFlags(RuntimeFlags, EDisplayClusterViewportPreviewFlags::HasChangedPreviewEditableMeshMaterialInstance);
+		return;
 	}
 
-	// Update the preview mesh and materials in the DisplayDevice component
-	if (InDisplayDeviceComponent)
-	{
-		// Update mesh and material instances
-		InDisplayDeviceComponent->OnUpdateDisplayDeviceMeshAndMaterialInstance(*this, EDisplayClusterDisplayDeviceMeshType::PreviewMesh, PreviewMesh.GetCurrentMaterialType(), PreviewMesh.GetMeshComponent(), PreviewMesh.GetMaterialInstance());
-		InDisplayDeviceComponent->OnUpdateDisplayDeviceMeshAndMaterialInstance(*this, EDisplayClusterDisplayDeviceMeshType::PreviewEditableMesh, PreviewMesh.GetCurrentMaterialType(), PreviewEditableMesh.GetMeshComponent(), PreviewEditableMesh.GetMaterialInstance());
-	}
+	FDisplayClusterViewport* Viewport = GetViewportImpl();	
+	UDisplayClusterCameraComponent* ViewPointComponent = nullptr;
+	UDisplayClusterDisplayDeviceBaseComponent* DisplayDeviceComponent = nullptr;
 
-	// Update the preview mesh and materials in the ViewPointComponent component
-	if (ViewPointComponent)
+	if (Viewport)
 	{
-		ViewPointComponent->OnUpdateDisplayDeviceMeshAndMaterialInstance(*this, EDisplayClusterDisplayDeviceMeshType::PreviewMesh, PreviewMesh.GetCurrentMaterialType(), PreviewMesh.GetMeshComponent(), PreviewMesh.GetMaterialInstance());
-		ViewPointComponent->OnUpdateDisplayDeviceMeshAndMaterialInstance(*this, EDisplayClusterDisplayDeviceMeshType::PreviewEditableMesh, PreviewMesh.GetCurrentMaterialType(), PreviewEditableMesh.GetMeshComponent(), PreviewEditableMesh.GetMaterialInstance());
+		ViewPointComponent = Viewport->GetViewPointCameraComponent(EDisplayClusterRootActorType::Configuration);
+		DisplayDeviceComponent = Viewport->GetDisplayDeviceComponent(Configuration->GetPreviewSettings().DisplayDeviceRootActorType);
 	}
+	
+	// Utility function to update both PreviewMesh and PreviewEditableMesh
+	auto UpdatePreviewMeshAndMaterial = [&](
+		FDisplayClusterViewportPreviewMesh& PreviewMeshObj,
+		EDisplayClusterViewportPreviewFlags ChangedMaterialFlag,
+		EDisplayClusterDisplayDeviceMeshType MeshType)
+		{
+			PreviewMeshObj.Update(Viewport, DisplayDeviceComponent, ViewPointComponent);
+
+			// Update Runtime Flags
+			if (PreviewMeshObj.HasAnyFlag(EDisplayClusterViewportPreviewMeshFlags::HasDeletedMaterialInstance | EDisplayClusterViewportPreviewMeshFlags::HasChangedMaterialInstance))
+			{
+				EnumAddFlags(RuntimeFlags, ChangedMaterialFlag);
+			}
+
+			// Record that these meshes have been updated so that other viewports don't try to.
+			if (const UMeshComponent* MeshComponent = PreviewMeshObj.GetMeshComponent())
+			{
+				MeshComponentsVisited.Add(MeshComponent);
+			}
+
+			// Update the preview mesh and materials in the DisplayDevice component
+			if (DisplayDeviceComponent)
+			{
+				DisplayDeviceComponent->OnUpdateDisplayDeviceMeshAndMaterialInstance(
+					*this,
+					MeshType,
+					PreviewMeshObj.GetCurrentMaterialType(),
+					PreviewMeshObj.GetMeshComponent(),
+					PreviewMeshObj.GetMaterialInstance()
+				);
+			}
+
+			// Update the preview mesh and materials in the ViewPointComponent component
+			if (ViewPointComponent)
+			{
+				ViewPointComponent->OnUpdateDisplayDeviceMeshAndMaterialInstance(
+					*this,
+					MeshType,
+					PreviewMeshObj.GetCurrentMaterialType(),
+					PreviewMeshObj.GetMeshComponent(),
+					PreviewMeshObj.GetMaterialInstance()
+				);
+			}
+		};
+
+	// Update PreviewMesh and PreviewEditableMesh
+
+	UpdatePreviewMeshAndMaterial(
+		PreviewMesh,                                                                // PreviewMeshObj
+		EDisplayClusterViewportPreviewFlags::HasChangedPreviewMeshMaterialInstance, // ChangedMaterialFlag
+		EDisplayClusterDisplayDeviceMeshType::PreviewMesh                           // MeshType
+	);
+
+	UpdatePreviewMeshAndMaterial(
+		PreviewEditableMesh,                                                                // PreviewMeshObj
+		EDisplayClusterViewportPreviewFlags::HasChangedPreviewEditableMeshMaterialInstance, // ChangedMaterialFlag
+		EDisplayClusterDisplayDeviceMeshType::PreviewEditableMesh                           // MeshType
+	);
 }
 
 IDisplayClusterViewport* FDisplayClusterViewportPreview::GetViewport() const
diff --git a/Engine/Plugins/Runtime/nDisplay/Source/DisplayCluster/Private/Render/Viewport/Preview/DisplayClusterViewportPreview.h b/Engine/Plugins/Runtime/nDisplay/Source/DisplayCluster/Private/Render/Viewport/Preview/DisplayClusterViewportPreview.h
--- a/Engine/Plugins/Runtime/nDisplay/Source/DisplayCluster/Private/Render/Viewport/Preview/DisplayClusterViewportPreview.h
+++ b/Engine/Plugins/Runtime/nDisplay/Source/DisplayCluster/Private/Render/Viewport/Preview/DisplayClusterViewportPreview.h
@@ -76,7 +76,7 @@ public:
 
 public:
 	/** Update preview resources for this viewport. */
-	void Update();
+	void Update(TSet<const UMeshComponent*>& MeshComponentsVisited);
 
 	/** Release preview resources and materials for this viewport. */
 	void Release();
