nDisplay: Allow ICVFX camera priority inversion per viewport.

This is useful for dual camera ghost frame workflow such that only the overlap flickers.

#rn Inner frustum overlap order can be reversed per outer viewport. Use case is LED processors that can time-multiplex (such as Ghost Frame) for dual camera setups such that the flicker is limited to the inner overlap area.

diff --git a/Engine/Plugins/Runtime/nDisplay/Source/DisplayCluster/Private/Render/Viewport/Configuration/DisplayClusterViewportConfiguration_ICVFX.cpp b/Engine/Plugins/Runtime/nDisplay/Source/DisplayCluster/Private/Render/Viewport/Configuration/DisplayClusterViewportConfiguration_ICVFX.cpp
--- a/Engine/Plugins/Runtime/nDisplay/Source/DisplayCluster/Private/Render/Viewport/Configuration/DisplayClusterViewportConfiguration_ICVFX.cpp
+++ b/Engine/Plugins/Runtime/nDisplay/Source/DisplayCluster/Private/Render/Viewport/Configuration/DisplayClusterViewportConfiguration_ICVFX.cpp
@@ -175,8 +175,13 @@ void FDisplayClusterViewportConfiguration_ICVFX::Update()
 			// Gain direct access to internal settings of the viewport:
 			FDisplayClusterViewport_RenderSettingsICVFX& InOutRenderSettingsICVFX = TargetIt->GetRenderSettingsICVFXImpl();
 
+			// Reverse camera order when the option is set in the viewport.
+			const UDisplayClusterConfigurationViewport* ViewportConfig = TargetIt->GetViewportConfigurationData();
+			ensureMsgf(ViewportConfig, TEXT("Viewport '%s' did not have a config set"), *TargetIt->ViewportId);
+			const bool bReverseCameraPriority = ViewportConfig ? ViewportConfig->ICVFX.bReverseCameraPriority : false;
+
 			// Sort cameras by render order for each target
-			InOutRenderSettingsICVFX.ICVFX.SortCamerasRenderOrder();
+			InOutRenderSettingsICVFX.ICVFX.SortCamerasRenderOrder(bReverseCameraPriority);
 
 			// Setup incamera overlap mode
 			InOutRenderSettingsICVFX.ICVFX.CameraOverlappingRenderMode = StageSettings->GetCameraOverlappingRenderMode();
diff --git a/Engine/Plugins/Runtime/nDisplay/Source/DisplayClusterConfiguration/Public/DisplayClusterConfigurationTypes_Viewport.h b/Engine/Plugins/Runtime/nDisplay/Source/DisplayClusterConfiguration/Public/DisplayClusterConfigurationTypes_Viewport.h
--- a/Engine/Plugins/Runtime/nDisplay/Source/DisplayClusterConfiguration/Public/DisplayClusterConfigurationTypes_Viewport.h
+++ b/Engine/Plugins/Runtime/nDisplay/Source/DisplayClusterConfiguration/Public/DisplayClusterConfigurationTypes_Viewport.h
@@ -61,6 +61,10 @@ public:
 	/** Use unique lightcard mode for this viewport */
 	UPROPERTY(EditAnywhere, BlueprintReadWrite, Category = "In Camera VFX")
 	EDisplayClusterConfigurationICVFX_OverrideLightcardRenderMode LightcardRenderMode = EDisplayClusterConfigurationICVFX_OverrideLightcardRenderMode::Default;
+
+	/** The order in which the ICVFX cameras are composited over is reversed. Useful for time-multiplexed displays. */
+	UPROPERTY(EditAnywhere, BlueprintReadWrite, Category = "In Camera VFX")
+	bool bReverseCameraPriority = false;
 };
 
 USTRUCT(Blueprintable)
diff --git a/Engine/Plugins/Runtime/nDisplay/Source/DisplayClusterShaders/Public/ShaderParameters/DisplayClusterShaderParameters_ICVFX.h b/Engine/Plugins/Runtime/nDisplay/Source/DisplayClusterShaders/Public/ShaderParameters/DisplayClusterShaderParameters_ICVFX.h
--- a/Engine/Plugins/Runtime/nDisplay/Source/DisplayClusterShaders/Public/ShaderParameters/DisplayClusterShaderParameters_ICVFX.h
+++ b/Engine/Plugins/Runtime/nDisplay/Source/DisplayClusterShaders/Public/ShaderParameters/DisplayClusterShaderParameters_ICVFX.h
@@ -136,15 +136,27 @@ public:
 		LightCardGamma = InParameters.LightCardGamma;
 	}
 
-	inline void SortCamerasRenderOrder()
+	inline void SortCamerasRenderOrder(const bool bReverseOrder = false)
 	{
-		Cameras.Sort([](const FCameraSettings& It1, const FCameraSettings& It2)
+		Cameras.Sort([bReverseOrder](const FCameraSettings& It1, const FCameraSettings& It2)
 		{
+			bool bComparisonResult;
+
 			if (It1.RenderOrder == It2.RenderOrder)
 			{
-				return It1.Resource.ViewportId.Compare(It2.Resource.ViewportId, ESearchCase::IgnoreCase) < 0;
+				bComparisonResult = It1.Resource.ViewportId.Compare(It2.Resource.ViewportId, ESearchCase::IgnoreCase) < 0;
+			}
+			else
+			{
+				bComparisonResult = It1.RenderOrder < It2.RenderOrder;
 			}
-			return It1.RenderOrder < It2.RenderOrder;
+
+			if (bReverseOrder)
+			{
+				bComparisonResult = !bComparisonResult;
+			}
+
+			return bComparisonResult;
 		});
 	}
 
