nDisplay: added missed calls before destructor.

Before calling ~FDisplayClusterViewportManager(), we must pass events about it to other subsystems.
Because all subsystems connected to this class work via WeakPtr and will not work when we are in a destructor.

For example,:
The EndScene() event is widely used by many subsystems such as 'PostProcess', 'ProjectionPolicy' etc.
TextureSharePostProcess has logic based on the SceneEnd() event to release objects of another external subsystem (TextureShare module).
Without the HandleEndScene() event, these external objects will not be released, but will be reused in the next scene, which breaks the logic.

diff --git a/Engine/Plugins/Runtime/nDisplay/Source/DisplayCluster/Private/Game/EngineClasses/Scene/DisplayClusterRootActor.cpp b/Engine/Plugins/Runtime/nDisplay/Source/DisplayCluster/Private/Game/EngineClasses/Scene/DisplayClusterRootActor.cpp
--- a/Engine/Plugins/Runtime/nDisplay/Source/DisplayCluster/Private/Game/EngineClasses/Scene/DisplayClusterRootActor.cpp
+++ b/Engine/Plugins/Runtime/nDisplay/Source/DisplayCluster/Private/Game/EngineClasses/Scene/DisplayClusterRootActor.cpp
@@ -302,6 +302,9 @@ void ADisplayClusterRootActor::RemoveViewportManager()
 {
 	if (ViewportManagerPtr.IsValid())
 	{
+		// This function must be called before destructor.
+		ViewportManagerPtr->Release();
+
 		// Immediately release the viewport manager with resources
 		ViewportManagerPtr.Reset();
 	}
diff --git a/Engine/Plugins/Runtime/nDisplay/Source/DisplayCluster/Private/Render/Viewport/DisplayClusterViewportManager.cpp b/Engine/Plugins/Runtime/nDisplay/Source/DisplayCluster/Private/Render/Viewport/DisplayClusterViewportManager.cpp
--- a/Engine/Plugins/Runtime/nDisplay/Source/DisplayCluster/Private/Render/Viewport/DisplayClusterViewportManager.cpp
+++ b/Engine/Plugins/Runtime/nDisplay/Source/DisplayCluster/Private/Render/Viewport/DisplayClusterViewportManager.cpp
@@ -233,8 +233,30 @@ void FDisplayClusterViewportManager::Initialize()
 	ViewportManagerPreview->RegisterPreviewRendering();
 }
 
+void FDisplayClusterViewportManager::Release()
+{
+	check(IsInGameThread());
+
+	// Calls OnSceneEnd() for all subsystems. They will release their own resources.
+	HandleEndScene();
+
+	// Frees GPU memory as fast as possible.
+	ReleaseTextures();
+
+	// This is a flag to check in the destructor that Release() has been called before.
+	bReleaseFunctionCalledBeforeDestructor = true;
+}
+
 FDisplayClusterViewportManager::~FDisplayClusterViewportManager()
 {
+	if (!bReleaseFunctionCalledBeforeDestructor)
+	{
+		// This is the wrong way to call the destructor.
+		// Lets show a message in the log
+		UE_LOG(LogDisplayClusterViewport, Error,
+			TEXT("~FDisplayClusterViewportManager() is called at an inappropriate time. The Release() function must be called before this."));
+	}
+
 	UnregisterCallbacks();
 
 	// Remove this viewport manager from the preview rendering pipeline.
diff --git a/Engine/Plugins/Runtime/nDisplay/Source/DisplayCluster/Private/Render/Viewport/DisplayClusterViewportManager.h b/Engine/Plugins/Runtime/nDisplay/Source/DisplayCluster/Private/Render/Viewport/DisplayClusterViewportManager.h
--- a/Engine/Plugins/Runtime/nDisplay/Source/DisplayCluster/Private/Render/Viewport/DisplayClusterViewportManager.h
+++ b/Engine/Plugins/Runtime/nDisplay/Source/DisplayCluster/Private/Render/Viewport/DisplayClusterViewportManager.h
@@ -37,6 +37,12 @@ public:
 
 	void Initialize();
 
+	/** 
+	* Must be called before the destructor.
+	* Because all subsystems connected to this class via WeakPtr will not work when we are in destructor.
+	*/
+	void Release();
+
 	/** Release all texture resources used for viewports.
 	* After that, the BeginNewFrame() function should be called for a new resource allocation.*/
 	void ReleaseTextures();
@@ -251,6 +257,9 @@ private:
 	// This VE is used to display frame information such as timecode and frame number.
 	TSharedPtr<class FDisplayClusterViewportFrameStatsViewExtension, ESPMode::ThreadSafe> FrameStatsViewExtension;
 
+	// This is a flag to check in the destructor that Release() has been called before.
+	bool bReleaseFunctionCalledBeforeDestructor = false;
+
 #if WITH_EDITOR
 	/** The handle for FCoreUObjectDelegates::GetPreGarbageCollectDelegate */
 	FDelegateHandle PreGarbageCollectHandle;
