[nDisplay] Late OCIO improvements

- Added ability to manipulate late OCIO live
- Added ability to use PQ for transfer
- Restructured nDisplay media adapters

[RN] Added ability to use 10 bits perceptual quantizer (PQ) for the late OCIO media streams

diff --git a/Engine/Plugins/Runtime/nDisplay/Shaders/Private/MediaShaders.usf b/Engine/Plugins/Runtime/nDisplay/Shaders/Private/MediaShaders.usf
new file mode 100644
--- /dev/null
+++ b/Engine/Plugins/Runtime/nDisplay/Shaders/Private/MediaShaders.usf
@@ -0,0 +1,28 @@
+// Copyright Epic Games, Inc. All Rights Reserved.
+
+#include "/Engine/Private/Common.ush"
+#include "/Engine/Private/GammaCorrectionCommon.ush"
+
+Texture2D InputTexture;
+SamplerState InputSampler;
+
+
+// Pixel shader for Linear-PQ conversion
+float4 LinearToPQ_PS(float2 uv : TEXCOORD) : SV_Target
+{
+	float4 colorLinear = InputTexture.Sample(InputSampler, uv);
+	float3 colorPQ = LinearToST2084(colorLinear.rgb);
+	colorPQ = saturate(colorPQ);
+
+	return float4(colorPQ, colorLinear.a);
+}
+
+
+// Pixel shader for PQ-Linear conversion
+float4 PQToLinear_PS(float2 uv : TEXCOORD) : SV_Target
+{
+	float4 colorPQ = InputTexture.Sample(InputSampler, uv);
+	float3 colorLinear = ST2084ToLinear(colorPQ.rgb);
+
+	return float4(colorLinear, colorPQ.a);
+}
diff --git a/Engine/Plugins/Runtime/nDisplay/Source/DisplayCluster/Private/Render/Viewport/DisplayClusterViewportManagerViewExtension.cpp b/Engine/Plugins/Runtime/nDisplay/Source/DisplayCluster/Private/Render/Viewport/DisplayClusterViewportManagerViewExtension.cpp
--- a/Engine/Plugins/Runtime/nDisplay/Source/DisplayCluster/Private/Render/Viewport/DisplayClusterViewportManagerViewExtension.cpp
+++ b/Engine/Plugins/Runtime/nDisplay/Source/DisplayCluster/Private/Render/Viewport/DisplayClusterViewportManagerViewExtension.cpp
@@ -171,13 +171,7 @@ FScreenPassTexture FDisplayClusterViewportManagerViewExtension::PostProcessPassA
 		uint32 ContextNum = 0;
 		if (FDisplayClusterViewportProxy* ViewportProxyPtr = ViewportManagerProxy->ImplFindViewportProxy_RenderThread(View.StereoViewIndex, &ContextNum))
 		{
-			// Broadcast PassTonemap event
-			IDisplayCluster::Get().GetCallbacks().OnDisplayClusterPostTonemapPass_RenderThread().Broadcast(GraphBuilder, ViewportProxyPtr, View, Inputs, ContextNum);
-
-			if (ViewportProxyPtr->GetOpenColorIOMode() == EDisplayClusterViewportOpenColorIOMode::PostProcess)
-			{
-				return ViewportProxyPtr->OnPostProcessPassAfterTonemap_RenderThread(GraphBuilder, View, Inputs, ContextNum);
-			}
+			return ViewportProxyPtr->OnPostProcessPassAfterTonemap_RenderThread(GraphBuilder, View, Inputs, ContextNum);
 		}
 	}
 
diff --git a/Engine/Plugins/Runtime/nDisplay/Source/DisplayCluster/Private/Render/Viewport/DisplayClusterViewportProxy.cpp b/Engine/Plugins/Runtime/nDisplay/Source/DisplayCluster/Private/Render/Viewport/DisplayClusterViewportProxy.cpp
--- a/Engine/Plugins/Runtime/nDisplay/Source/DisplayCluster/Private/Render/Viewport/DisplayClusterViewportProxy.cpp
+++ b/Engine/Plugins/Runtime/nDisplay/Source/DisplayCluster/Private/Render/Viewport/DisplayClusterViewportProxy.cpp
@@ -415,6 +415,9 @@ FScreenPassTexture FDisplayClusterViewportProxy::OnPostProcessPassAfterFXAA_Rend
 
 FScreenPassTexture FDisplayClusterViewportProxy::OnPostProcessPassAfterTonemap_RenderThread(FRDGBuilder& GraphBuilder, const FSceneView& View, const FPostProcessMaterialInputs& Inputs, const uint32 ContextNum)
 {
+	// Broadcast PassTonemap event
+	IDisplayCluster::Get().GetCallbacks().OnDisplayClusterPostTonemapPass_RenderThread().Broadcast(GraphBuilder, this, View, Inputs, ContextNum);
+
 	// Perform OCIO rendering after the tonemapper
 	if (GetOpenColorIOMode() == EDisplayClusterViewportOpenColorIOMode::PostProcess)
 	{
diff --git a/Engine/Plugins/Runtime/nDisplay/Source/DisplayClusterConfiguration/Public/DisplayClusterConfigurationTypes_Media.h b/Engine/Plugins/Runtime/nDisplay/Source/DisplayClusterConfiguration/Public/DisplayClusterConfigurationTypes_Media.h
--- a/Engine/Plugins/Runtime/nDisplay/Source/DisplayClusterConfiguration/Public/DisplayClusterConfigurationTypes_Media.h
+++ b/Engine/Plugins/Runtime/nDisplay/Source/DisplayClusterConfiguration/Public/DisplayClusterConfigurationTypes_Media.h
@@ -302,10 +302,14 @@ public:
 	UPROPERTY(EditAnywhere, BlueprintReadWrite, Category = "Media", meta = (DisplayName = "Media Output Groups"))
 	TArray<FDisplayClusterConfigurationMediaTiledOutputGroup> TiledMediaOutputGroups;
 
-	/** Allows the receviers to apply their own OCIO transformations (per node OCIO override). Requires the media to support FloatRGBA.*/
+	/** Allows the receviers to apply their own OCIO transformations (per node OCIO override). Requires the media to support PF_FloatRGBA pixel format if PQ transfer is not used. */
 	UPROPERTY(EditAnywhere, BlueprintReadWrite, Category = "Media", meta = (DisplayName = "Receivers Apply OCIO"))
 	bool bLateOCIOPass = false;
 
+	/** Use perceptual quantizer (PQ) transfer function to reduce bandwith. Requires PF_A2B10G10R10 pixel format configured on media side. */
+	UPROPERTY(EditAnywhere, BlueprintReadWrite, Category = "Media", meta = (DisplayName = "Use PQ-Encoding for Transfer", EditCondition = "bLateOCIOPass"))
+	bool bTransferPQ = false;
+
 public:
 	/** Returns true if these media settings can be used for tiling. */
 	bool ShouldMediaICVFXSplitIntoTiles() const;
diff --git a/Engine/Plugins/Runtime/nDisplay/Source/DisplayClusterMedia/DisplayClusterMedia.Build.cs b/Engine/Plugins/Runtime/nDisplay/Source/DisplayClusterMedia/DisplayClusterMedia.Build.cs
--- a/Engine/Plugins/Runtime/nDisplay/Source/DisplayClusterMedia/DisplayClusterMedia.Build.cs
+++ b/Engine/Plugins/Runtime/nDisplay/Source/DisplayClusterMedia/DisplayClusterMedia.Build.cs
@@ -19,6 +19,7 @@ public class DisplayClusterMedia : ModuleRules
 				"Core",
 				"CoreUObject",
 				"DisplayCluster",
+				"DisplayClusterShaders",
 				"Engine",
 				"OpenColorIO",
 				"Renderer",
diff --git a/Engine/Plugins/Runtime/nDisplay/Source/DisplayClusterMedia/Private/Capture/DisplayClusterMediaCaptureBase.cpp b/Engine/Plugins/Runtime/nDisplay/Source/DisplayClusterMedia/Private/Capture/DisplayClusterMediaCaptureBase.cpp
--- a/Engine/Plugins/Runtime/nDisplay/Source/DisplayClusterMedia/Private/Capture/DisplayClusterMediaCaptureBase.cpp
+++ b/Engine/Plugins/Runtime/nDisplay/Source/DisplayClusterMedia/Private/Capture/DisplayClusterMediaCaptureBase.cpp
@@ -1,14 +1,16 @@
 // Copyright Epic Games, Inc. All Rights Reserved.
 
 #include "Capture/DisplayClusterMediaCaptureBase.h"
+
 #include "DisplayClusterMediaHelpers.h"
 #include "DisplayClusterMediaLog.h"
-#include "IDisplayCluster.h"
-#include "IDisplayClusterCallbacks.h"
-
 #include "DisplayClusterConfigurationTypes_Media.h"
 #include "DisplayClusterConfigurationTypes_MediaSync.h"
 
+#include "IDisplayCluster.h"
+#include "IDisplayClusterCallbacks.h"
+#include "IDisplayClusterShaders.h"
+
 #include "MediaCapture.h"
 #include "MediaOutput.h"
 
@@ -17,6 +19,8 @@
 #include "RHICommandList.h"
 #include "RHIUtilities.h"
 
+#include "ShaderParameters/DisplayClusterShaderParameters_Media.h"
+
 #include "UObject/UObjectGlobals.h"
 #include "UObject/Package.h"
 
@@ -25,10 +29,9 @@ FDisplayClusterMediaCaptureBase::FDisplayClusterMediaCaptureBase(
 	const FString& InMediaId,
 	const FString& InClusterNodeId,
 	UMediaOutput* InMediaOutput,
-	UDisplayClusterMediaOutputSynchronizationPolicy* InSyncPolicy,
-	bool bInLateOCIO
+	UDisplayClusterMediaOutputSynchronizationPolicy* InSyncPolicy
 )
-	: FDisplayClusterMediaBase(InMediaId, InClusterNodeId, bInLateOCIO)
+	: FDisplayClusterMediaBase(InMediaId, InClusterNodeId)
 	, SyncPolicy(InSyncPolicy)
 {
 	checkSlow(InMediaOutput);
@@ -126,31 +129,99 @@ void FDisplayClusterMediaCaptureBase::StopCapture()
 
 void FDisplayClusterMediaCaptureBase::ExportMediaData_RenderThread(FRDGBuilder& GraphBuilder, const FMediaOutputTextureInfo& TextureInfo)
 {
-	// Check if source texture is valid
-	if (!TextureInfo.Texture)
+	// Check if request data is valid
+	if (!IsValidRequestData(TextureInfo))
 	{
-		UE_LOG(LogDisplayClusterMedia, Warning, TEXT("MediaCapture '%s': invalid source texture on RT frame %lu"), *GetMediaId(), GFrameCounterRenderThread);
+		UE_LOG(LogDisplayClusterMedia, Warning, TEXT("MediaCapture '%s': no capture performed on RT frame %lu"), *GetMediaId(), GFrameCounterRenderThread);
 		return;
 	}
 
 	MediaCapture->SetValidSourceGPUMask(GraphBuilder.RHICmdList.GetGPUMask());
 
-	const FIntPoint SrcTextureSize = TextureInfo.Texture->Desc.Extent;
-	const FIntPoint SrcRegionSize  = TextureInfo.Region.Size();
+	{
+		const FIntPoint& SrcTextureSize = TextureInfo.Texture->Desc.Extent;
+		const FIntPoint  SrcRegionSize = TextureInfo.Region.Size();
 
-	LastSrcRegionSize = FIntSize(SrcRegionSize);
+		LastSrcRegionSize = FIntSize(SrcRegionSize);
+
+		UE_LOG(LogDisplayClusterMedia, VeryVerbose, TEXT("MediaCapture '%s': Requested texture export [size=%dx%d, rect=%dx%d, format=%u] on RT frame '%lu'..."),
+			*GetMediaId(), SrcTextureSize.X, SrcTextureSize.Y, SrcRegionSize.X, SrcRegionSize.Y, (uint32)TextureInfo.Texture->Desc.Format, GFrameCounterRenderThread);
+	}
 
-	UE_LOG(LogDisplayClusterMedia, VeryVerbose, TEXT("MediaCapture '%s': exporting texture [size=%dx%d, rect=%dx%d] on RT frame '%lu'..."),
-		*GetMediaId(), SrcTextureSize.X, SrcTextureSize.Y, SrcRegionSize.X, SrcRegionSize.Y, GFrameCounterRenderThread);
+	bool bCaptureSucceeded = false;
+
+	// Is PQ-encoding required?
+	constexpr bool bConsideringLateOCIOEnabled = true;
+	const bool bLateOCIOWithPQTransfer = IsTransferPQ(bConsideringLateOCIOEnabled);
+
+	// When PQ encoding is required, we have to add a separate PQ-encoding pass
+	if (bLateOCIOWithPQTransfer)
+	{
+		// Allocate intermediate PQ texture of PF_A2B10G10R10 pixel format
+		FRDGTextureDesc TexturePQDesc = FRDGTextureDesc::Create2D(TextureInfo.Region.Size(), PF_A2B10G10R10, FClearValueBinding::Black, TexCreate_ShaderResource | TexCreate_RenderTargetable);
+		FRDGTextureRef TexturePQ = GraphBuilder.CreateTexture(TexturePQDesc, TEXT("DC.MediaTexturePQ"));
+
+		// Add PQ-encoding pass
+		FDisplayClusterShaderParameters_MediaPQ Parameters{ TextureInfo.Texture, TextureInfo.Region, TexturePQ, { {0, 0}, TexturePQ->Desc.Extent } };
+		IDisplayClusterShaders::Get().AddLinearToPQPass(GraphBuilder, Parameters);
+
+		UE_LOG(LogDisplayClusterMedia, VeryVerbose, TEXT("MediaCapture '%s': PQ exporting TexSize[%s], TexRect[%s], TexFormat[%u] on RT frame '%lu'..."),
+			*GetMediaId(),
+			*Parameters.OutputTexture->Desc.Extent.ToString(),
+			*Parameters.OutputRect.ToString(),
+			(uint32)Parameters.OutputTexture->Desc.Format,
+			GFrameCounterRenderThread);
+
+		// Pass the texture to the capture device
+		bCaptureSucceeded = MediaCapture->TryCaptureImmediate_RenderThread(GraphBuilder, Parameters.OutputTexture, Parameters.OutputRect);
+	}
+	else
+	{
+		UE_LOG(LogDisplayClusterMedia, VeryVerbose, TEXT("MediaCapture '%s': Direct exporting TexSize[%s], TexRect[%s], TexFormat[%u] on RT frame '%lu'..."),
+			*GetMediaId(),
+			*TextureInfo.Texture->Desc.Extent.ToString(),
+			*TextureInfo.Region.ToString(),
+			(uint32)TextureInfo.Texture->Desc.Format,
+			GFrameCounterRenderThread);
+
+		// Direct capture
+		bCaptureSucceeded = MediaCapture->TryCaptureImmediate_RenderThread(GraphBuilder, TextureInfo.Texture, TextureInfo.Region);
+	}
 
-	// Capture
-	bool bCaptureSucceeded = MediaCapture->TryCaptureImmediate_RenderThread(GraphBuilder, TextureInfo.Texture, TextureInfo.Region);
 	if(!bCaptureSucceeded)
 	{
-		UE_LOG(LogDisplayClusterMedia, VeryVerbose, TEXT("MediaCapture '%s': failed to capture resource"), *GetMediaId());
+		UE_LOG(LogDisplayClusterMedia, Warning, TEXT("MediaCapture '%s': failed to capture resource"), *GetMediaId());
 	}
 }
 
+bool FDisplayClusterMediaCaptureBase::IsValidRequestData(const FMediaOutputTextureInfo& TextureInfo) const
+{
+	// Check if source texture is valid
+	if (!TextureInfo.Texture)
+	{
+		UE_LOG(LogDisplayClusterMedia, Warning, TEXT("MediaCapture '%s': invalid source texture on RT frame %lu"), *GetMediaId(), GFrameCounterRenderThread);
+		return false;
+	}
+
+	// Check if region matches the texture
+	const FIntPoint RegionSize = TextureInfo.Region.Size();
+	const bool bCorrectRegion =
+		TextureInfo.Region.Min.X >= 0 &&
+		TextureInfo.Region.Min.Y >= 0 &&
+		RegionSize.X > 0 &&
+		RegionSize.Y > 0 &&
+		RegionSize.X <= TextureInfo.Texture->Desc.Extent.X &&
+		RegionSize.Y <= TextureInfo.Texture->Desc.Extent.Y;
+
+	if (!bCorrectRegion)
+	{
+		UE_LOG(LogDisplayClusterMedia, Warning, TEXT("MediaCapture '%s': invalid source region on RT frame %lu"), *GetMediaId(), GFrameCounterRenderThread);
+		return false;
+	}
+
+	return true;
+}
+
 void FDisplayClusterMediaCaptureBase::OnPostClusterTick()
 {
 	if (MediaCapture)
diff --git a/Engine/Plugins/Runtime/nDisplay/Source/DisplayClusterMedia/Private/Capture/DisplayClusterMediaCaptureBase.h b/Engine/Plugins/Runtime/nDisplay/Source/DisplayClusterMedia/Private/Capture/DisplayClusterMediaCaptureBase.h
--- a/Engine/Plugins/Runtime/nDisplay/Source/DisplayClusterMedia/Private/Capture/DisplayClusterMediaCaptureBase.h
+++ b/Engine/Plugins/Runtime/nDisplay/Source/DisplayClusterMedia/Private/Capture/DisplayClusterMediaCaptureBase.h
@@ -2,8 +2,7 @@
 
 #pragma once
 
-#include "CoreMinimal.h"
-#include "DisplayClusterMediaBase.h"
+#include "Common/DisplayClusterMediaBase.h"
 #include "UObject/GCObject.h"
 
 #include "RenderGraphFwd.h"
@@ -16,6 +15,7 @@ class UMediaOutput;
 class UDisplayClusterMediaOutputSynchronizationPolicy;
 class IDisplayClusterMediaOutputSynchronizationPolicyHandler;
 
+
 /**
  * Base media capture adapter class
  */
@@ -24,17 +24,18 @@ class FDisplayClusterMediaCaptureBase
 	, public FGCObject
 {
 public:
+
 	FDisplayClusterMediaCaptureBase(
 		const FString& MediaId,
 		const FString& ClusterNodeId,
 		UMediaOutput* MediaOutput,
-		UDisplayClusterMediaOutputSynchronizationPolicy* SyncPolicy = nullptr,
-		bool bInLateOCIO = false
+		UDisplayClusterMediaOutputSynchronizationPolicy* SyncPolicy
 	);
 
 	virtual ~FDisplayClusterMediaCaptureBase();
 
 public:
+
 	//~ Begin FGCObject interface
 	virtual void AddReferencedObjects(FReferenceCollector& Collector) override;
 	virtual FString GetReferencerName() const override
@@ -85,7 +86,12 @@ protected:
 
 private:
 
-	// Trivial version of FIntPoint so that it can be std::atomic
+	/** Validate if capture request data is valid */
+	bool IsValidRequestData(const FMediaOutputTextureInfo& TextureInfo) const;
+
+private:
+
+	/** Trivial version of FIntPoint so that it can be std::atomic */
 	struct FIntSize
 	{
 		int32 X = 0;
@@ -101,19 +107,20 @@ private:
 	};
 
 private:
+
 	//~ Begin GC by AddReferencedObjects
 	TObjectPtr<UMediaOutput>  MediaOutput;
 	TObjectPtr<UMediaCapture> MediaCapture;
 	TObjectPtr<UDisplayClusterMediaOutputSynchronizationPolicy> SyncPolicy;
 	//~ End GC by AddReferencedObjects
 
-	// Used to restart media capture in the case it falls in error
+	/** Used to restart media capture in the case it falls in error */
 	bool bWasCaptureStarted = false;
 
-	// Used to control the rate at which we try to restart the capture
+	/** Used to control the rate at which we try to restart the capture */
 	double LastRestartTimestamp = 0;
 
-	// Last region size of the texture being exported. Used to restart the capture when in error.
+	/** Last region size of the texture being exported.Used to restart the capture when in error. */
 	std::atomic<FIntSize> LastSrcRegionSize { FIntSize(0,0) };
 
 	/** Sync policy handler to deal with synchronization logic */
diff --git a/Engine/Plugins/Runtime/nDisplay/Source/DisplayClusterMedia/Private/Capture/DisplayClusterMediaCaptureCamera.cpp b/Engine/Plugins/Runtime/nDisplay/Source/DisplayClusterMedia/Private/Capture/DisplayClusterMediaCaptureCamera.cpp
deleted file mode 100644
diff --git a/Engine/Plugins/Runtime/nDisplay/Source/DisplayClusterMedia/Private/Capture/DisplayClusterMediaCaptureCamera.h b/Engine/Plugins/Runtime/nDisplay/Source/DisplayClusterMedia/Private/Capture/DisplayClusterMediaCaptureCamera.h
deleted file mode 100644
diff --git a/Engine/Plugins/Runtime/nDisplay/Source/DisplayClusterMedia/Private/Capture/DisplayClusterMediaCaptureCameraBase.cpp b/Engine/Plugins/Runtime/nDisplay/Source/DisplayClusterMedia/Private/Capture/DisplayClusterMediaCaptureCameraBase.cpp
new file mode 100644
--- /dev/null
+++ b/Engine/Plugins/Runtime/nDisplay/Source/DisplayClusterMedia/Private/Capture/DisplayClusterMediaCaptureCameraBase.cpp
@@ -0,0 +1,59 @@
+// Copyright Epic Games, Inc. All Rights Reserved.
+
+#include "Capture/DisplayClusterMediaCaptureCameraBase.h"
+
+#include "Components/DisplayClusterICVFXCameraComponent.h"
+#include "DisplayClusterConfigurationTypes.h"
+#include "DisplayClusterRootActor.h"
+#include "IDisplayCluster.h"
+#include "Game/IDisplayClusterGameManager.h"
+
+#include "DisplayClusterMediaLog.h"
+
+
+FDisplayClusterMediaCaptureCameraBase::FDisplayClusterMediaCaptureCameraBase(
+	const FString& InMediaId,
+	const FString& InClusterNodeId,
+	const FString& InCameraId,
+	const FString& InCameraViewportId,
+	UMediaOutput* InMediaOutput,
+	UDisplayClusterMediaOutputSynchronizationPolicy* InSyncPolicy
+)
+	: FDisplayClusterMediaCaptureViewportBase(InMediaId, InClusterNodeId, InCameraViewportId, InMediaOutput, InSyncPolicy)
+	, FDisplayClusterMediaCameraCommon(InCameraId)
+{
+}
+
+bool FDisplayClusterMediaCaptureCameraBase::GetCaptureSizeFromConfig(FIntPoint& OutSize) const
+{
+	const UDisplayClusterICVFXCameraComponent* const ICVFXCameraComponent = GetCameraComponent();
+	if (!ICVFXCameraComponent)
+	{
+		return false;
+	}
+
+	const FDisplayClusterConfigurationICVFX_CameraSettings& CameraSettings = ICVFXCameraComponent->GetCameraSettingsICVFX();
+
+	if (CameraSettings.CustomFrustum.bEnable)
+	{
+		OutSize = CameraSettings.CustomFrustum.EstimatedOverscanResolution;
+	}
+	else
+	{
+		OutSize = CameraSettings.CustomFrustum.InnerFrustumResolution;
+	}
+
+	return true;
+}
+
+
+void FDisplayClusterMediaCaptureCameraBase::UpdateLateOCIOState(IDisplayClusterViewport* Viewport)
+{
+	FLateOCIOData NewLateOCIOConfiguration;
+
+	// Get current late OCIO parameters from the camera component
+	FDisplayClusterMediaCameraCommon::GetLateOCIOParameters(NewLateOCIOConfiguration.bLateOCIO, NewLateOCIOConfiguration.bTransferPQ);
+
+	// And store/update for current frame
+	SetLateOCIO(NewLateOCIOConfiguration);
+}
diff --git a/Engine/Plugins/Runtime/nDisplay/Source/DisplayClusterMedia/Private/Capture/DisplayClusterMediaCaptureCameraBase.h b/Engine/Plugins/Runtime/nDisplay/Source/DisplayClusterMedia/Private/Capture/DisplayClusterMediaCaptureCameraBase.h
new file mode 100644
--- /dev/null
+++ b/Engine/Plugins/Runtime/nDisplay/Source/DisplayClusterMedia/Private/Capture/DisplayClusterMediaCaptureCameraBase.h
@@ -0,0 +1,37 @@
+// Copyright Epic Games, Inc. All Rights Reserved.
+
+#pragma once
+
+#include "Common/DisplayClusterMediaCameraCommon.h"
+#include "Capture/DisplayClusterMediaCaptureViewportBase.h"
+
+
+/**
+ * Base camera capture adapter
+ */
+class FDisplayClusterMediaCaptureCameraBase
+	: public FDisplayClusterMediaCaptureViewportBase
+	, public FDisplayClusterMediaCameraCommon
+{
+public:
+
+	FDisplayClusterMediaCaptureCameraBase(
+		const FString& MediaId,
+		const FString& ClusterNodeId,
+		const FString& CameraId,
+		const FString& CameraViewportId,
+		UMediaOutput* MediaOutput,
+		UDisplayClusterMediaOutputSynchronizationPolicy* SyncPolicy
+	);
+
+public:
+
+	//~ Begin FDisplayClusterMediaCaptureViewport
+	virtual bool GetCaptureSizeFromConfig(FIntPoint& OutSize) const override;
+	//~ End FDisplayClusterMediaCaptureViewport
+
+protected:
+
+	/** Updates late OCIO state */
+	virtual void UpdateLateOCIOState(IDisplayClusterViewport* Viewport) override;
+};
diff --git a/Engine/Plugins/Runtime/nDisplay/Source/DisplayClusterMedia/Private/Capture/DisplayClusterMediaCaptureCameraFull.cpp b/Engine/Plugins/Runtime/nDisplay/Source/DisplayClusterMedia/Private/Capture/DisplayClusterMediaCaptureCameraFull.cpp
new file mode 100644
--- /dev/null
+++ b/Engine/Plugins/Runtime/nDisplay/Source/DisplayClusterMedia/Private/Capture/DisplayClusterMediaCaptureCameraFull.cpp
@@ -0,0 +1,23 @@
+// Copyright Epic Games, Inc. All Rights Reserved.
+
+#include "Capture/DisplayClusterMediaCaptureCameraFull.h"
+
+#include "DisplayClusterMediaHelpers.h"
+
+
+FDisplayClusterMediaCaptureCameraFull::FDisplayClusterMediaCaptureCameraFull(
+	const FString& InMediaId,
+	const FString& InClusterNodeId,
+	const FString& InCameraId,
+	UMediaOutput* InMediaOutput,
+	UDisplayClusterMediaOutputSynchronizationPolicy* InSyncPolicy
+)
+	: FDisplayClusterMediaCaptureCameraBase(
+		InMediaId,
+		InClusterNodeId,
+		InCameraId,
+		DisplayClusterMediaHelpers::GenerateICVFXViewportName(InClusterNodeId, InCameraId),
+		InMediaOutput,
+		InSyncPolicy)
+{
+}
diff --git a/Engine/Plugins/Runtime/nDisplay/Source/DisplayClusterMedia/Private/Capture/DisplayClusterMediaCaptureCameraFull.h b/Engine/Plugins/Runtime/nDisplay/Source/DisplayClusterMedia/Private/Capture/DisplayClusterMediaCaptureCameraFull.h
new file mode 100644
--- /dev/null
+++ b/Engine/Plugins/Runtime/nDisplay/Source/DisplayClusterMedia/Private/Capture/DisplayClusterMediaCaptureCameraFull.h
@@ -0,0 +1,23 @@
+// Copyright Epic Games, Inc. All Rights Reserved.
+
+#pragma once
+
+#include "Capture/DisplayClusterMediaCaptureCameraBase.h"
+
+
+/**
+ * Camera capture adapter (full frame)
+ */
+class FDisplayClusterMediaCaptureCameraFull
+	: public FDisplayClusterMediaCaptureCameraBase
+{
+public:
+
+	FDisplayClusterMediaCaptureCameraFull(
+		const FString& MediaId,
+		const FString& ClusterNodeId,
+		const FString& CameraId,
+		UMediaOutput* MediaOutput,
+		UDisplayClusterMediaOutputSynchronizationPolicy* SyncPolicy
+	);
+};
diff --git a/Engine/Plugins/Runtime/nDisplay/Source/DisplayClusterMedia/Private/Capture/DisplayClusterMediaCaptureCameraTile.cpp b/Engine/Plugins/Runtime/nDisplay/Source/DisplayClusterMedia/Private/Capture/DisplayClusterMediaCaptureCameraTile.cpp
new file mode 100644
--- /dev/null
+++ b/Engine/Plugins/Runtime/nDisplay/Source/DisplayClusterMedia/Private/Capture/DisplayClusterMediaCaptureCameraTile.cpp
@@ -0,0 +1,34 @@
+// Copyright Epic Games, Inc. All Rights Reserved.
+
+#include "Capture/DisplayClusterMediaCaptureCameraTile.h"
+
+#include "DisplayClusterMediaHelpers.h"
+
+
+FDisplayClusterMediaCaptureCameraTile::FDisplayClusterMediaCaptureCameraTile(
+	const FString& InMediaId,
+	const FString& InClusterNodeId,
+	const FString& InCameraId,
+	const FIntPoint& InTilePosition,
+	UMediaOutput* InMediaOutput,
+	UDisplayClusterMediaOutputSynchronizationPolicy* InSyncPolicy
+)
+	: FDisplayClusterMediaCaptureCameraBase(
+		InMediaId,
+		InClusterNodeId,
+		InCameraId,
+		DisplayClusterMediaHelpers::GenerateICVFXTileViewportName(InClusterNodeId, InCameraId, InTilePosition),
+		InMediaOutput,
+		InSyncPolicy)
+{
+}
+
+
+bool FDisplayClusterMediaCaptureCameraTile::GetCaptureSizeFromConfig(FIntPoint& OutSize) const
+{
+	// The same workaround as the viewport tile has
+
+	OutSize = { 64, 64 };
+
+	return true;
+}
diff --git a/Engine/Plugins/Runtime/nDisplay/Source/DisplayClusterMedia/Private/Capture/DisplayClusterMediaCaptureCameraTile.h b/Engine/Plugins/Runtime/nDisplay/Source/DisplayClusterMedia/Private/Capture/DisplayClusterMediaCaptureCameraTile.h
new file mode 100644
--- /dev/null
+++ b/Engine/Plugins/Runtime/nDisplay/Source/DisplayClusterMedia/Private/Capture/DisplayClusterMediaCaptureCameraTile.h
@@ -0,0 +1,30 @@
+// Copyright Epic Games, Inc. All Rights Reserved.
+
+#pragma once
+
+#include "Capture/DisplayClusterMediaCaptureCameraBase.h"
+
+
+/**
+ * Camera capture adapter (tile)
+ */
+class FDisplayClusterMediaCaptureCameraTile
+	: public FDisplayClusterMediaCaptureCameraBase
+{
+public:
+
+	FDisplayClusterMediaCaptureCameraTile(
+		const FString& MediaId,
+		const FString& ClusterNodeId,
+		const FString& CameraId,
+		const FIntPoint& TilePosition,
+		UMediaOutput* MediaOutput,
+		UDisplayClusterMediaOutputSynchronizationPolicy* SyncPolicy
+	);
+
+protected:
+
+	//~ Begin FDisplayClusterMediaCaptureViewport
+	virtual bool GetCaptureSizeFromConfig(FIntPoint& OutSize) const override;
+	//~ End FDisplayClusterMediaCaptureViewport
+};
diff --git a/Engine/Plugins/Runtime/nDisplay/Source/DisplayClusterMedia/Private/Capture/DisplayClusterMediaCaptureNode.h b/Engine/Plugins/Runtime/nDisplay/Source/DisplayClusterMedia/Private/Capture/DisplayClusterMediaCaptureNode.h
--- a/Engine/Plugins/Runtime/nDisplay/Source/DisplayClusterMedia/Private/Capture/DisplayClusterMediaCaptureNode.h
+++ b/Engine/Plugins/Runtime/nDisplay/Source/DisplayClusterMedia/Private/Capture/DisplayClusterMediaCaptureNode.h
@@ -22,7 +22,7 @@ public:
 		const FString& MediaId,
 		const FString& ClusterNodeId,
 		UMediaOutput* MediaOutput,
-		UDisplayClusterMediaOutputSynchronizationPolicy* SyncPolicy = nullptr
+		UDisplayClusterMediaOutputSynchronizationPolicy* SyncPolicy
 	);
 
 public:
diff --git a/Engine/Plugins/Runtime/nDisplay/Source/DisplayClusterMedia/Private/Capture/DisplayClusterMediaCaptureTile.cpp b/Engine/Plugins/Runtime/nDisplay/Source/DisplayClusterMedia/Private/Capture/DisplayClusterMediaCaptureTile.cpp
deleted file mode 100644
diff --git a/Engine/Plugins/Runtime/nDisplay/Source/DisplayClusterMedia/Private/Capture/DisplayClusterMediaCaptureTile.h b/Engine/Plugins/Runtime/nDisplay/Source/DisplayClusterMedia/Private/Capture/DisplayClusterMediaCaptureTile.h
deleted file mode 100644
diff --git a/Engine/Plugins/Runtime/nDisplay/Source/DisplayClusterMedia/Private/Capture/DisplayClusterMediaCaptureViewport.cpp b/Engine/Plugins/Runtime/nDisplay/Source/DisplayClusterMedia/Private/Capture/DisplayClusterMediaCaptureViewport.cpp
deleted file mode 100644
diff --git a/Engine/Plugins/Runtime/nDisplay/Source/DisplayClusterMedia/Private/Capture/DisplayClusterMediaCaptureViewport.h b/Engine/Plugins/Runtime/nDisplay/Source/DisplayClusterMedia/Private/Capture/DisplayClusterMediaCaptureViewport.h
deleted file mode 100644
diff --git a/Engine/Plugins/Runtime/nDisplay/Source/DisplayClusterMedia/Private/Capture/DisplayClusterMediaCaptureViewportBase.cpp b/Engine/Plugins/Runtime/nDisplay/Source/DisplayClusterMedia/Private/Capture/DisplayClusterMediaCaptureViewportBase.cpp
new file mode 100644
--- /dev/null
+++ b/Engine/Plugins/Runtime/nDisplay/Source/DisplayClusterMedia/Private/Capture/DisplayClusterMediaCaptureViewportBase.cpp
@@ -0,0 +1,236 @@
+// Copyright Epic Games, Inc. All Rights Reserved.
+
+#include "Capture/DisplayClusterMediaCaptureViewportBase.h"
+
+#include "IDisplayCluster.h"
+#include "IDisplayClusterCallbacks.h"
+
+#include "DisplayClusterConfigurationTypes_Viewport.h"
+#include "DisplayClusterMediaCVars.h"
+#include "DisplayClusterMediaLog.h"
+#include "DisplayClusterRootActor.h"
+
+#include "Config/IDisplayClusterConfigManager.h"
+#include "Game/IDisplayClusterGameManager.h"
+
+#include "PostProcess/PostProcessMaterialInputs.h"
+
+#include "Render/IDisplayClusterRenderManager.h"
+#include "Render/Viewport/IDisplayClusterViewport.h"
+#include "Render/Viewport/IDisplayClusterViewportManager.h"
+#include "Render/Viewport/IDisplayClusterViewportManagerProxy.h"
+#include "Render/Viewport/Containers/DisplayClusterViewport_Context.h"
+#include "Render/Viewport/Containers/DisplayClusterViewport_RenderSettings.h"
+
+#include "RHICommandList.h"
+#include "RHIResources.h"
+
+
+FDisplayClusterMediaCaptureViewportBase::FDisplayClusterMediaCaptureViewportBase(
+	const FString& InMediaId,
+	const FString& InClusterNodeId,
+	const FString& InViewportId,
+	UMediaOutput* InMediaOutput,
+	UDisplayClusterMediaOutputSynchronizationPolicy* SyncPolicy
+)
+	: FDisplayClusterMediaCaptureBase(InMediaId, InClusterNodeId, InMediaOutput, SyncPolicy)
+	, ViewportId(InViewportId)
+{
+}
+
+
+bool FDisplayClusterMediaCaptureViewportBase::StartCapture()
+{
+	// Subscribe for events
+	IDisplayCluster::Get().GetCallbacks().OnDisplayClusterUpdateViewportMediaState().AddRaw(this, &FDisplayClusterMediaCaptureViewportBase::OnUpdateViewportMediaState);
+
+	// PostTonemap callback for the late OCIO path
+	IDisplayCluster::Get().GetCallbacks().OnDisplayClusterPostTonemapPass_RenderThread().AddRaw(this, &FDisplayClusterMediaCaptureViewportBase::OnPostTonemapPass_RenderThread);
+
+	// PostRenderViewFamily for a regular capture path
+	IDisplayCluster::Get().GetCallbacks().OnDisplayClusterPostRenderViewFamily_RenderThread().AddRaw(this, &FDisplayClusterMediaCaptureViewportBase::OnPostRenderViewFamily_RenderThread);
+
+	// Start capture
+	const bool bStarted = FDisplayClusterMediaCaptureBase::StartCapture();
+
+	return bStarted;
+}
+
+void FDisplayClusterMediaCaptureViewportBase::StopCapture()
+{
+	// Unsubscribe from external events/callbacks
+	IDisplayCluster::Get().GetCallbacks().OnDisplayClusterPostRenderViewFamily_RenderThread().RemoveAll(this);
+	IDisplayCluster::Get().GetCallbacks().OnDisplayClusterPostTonemapPass_RenderThread().RemoveAll(this);
+	IDisplayCluster::Get().GetCallbacks().OnDisplayClusterUpdateViewportMediaState().RemoveAll(this);
+
+	// Stop capturing
+	FDisplayClusterMediaCaptureBase::StopCapture();
+}
+
+void FDisplayClusterMediaCaptureViewportBase::OnUpdateViewportMediaState(IDisplayClusterViewport* InViewport, EDisplayClusterViewportMediaState& InOutMediaState)
+{
+	// Set capture flag for the matching viewport
+	if (InViewport && InViewport->GetId().Equals(GetViewportId(), ESearchCase::IgnoreCase))
+	{
+		// Raise flags that this viewport will be captured by media.
+		InOutMediaState |= EDisplayClusterViewportMediaState::Capture;
+
+		// Update late OCIO state on current frame
+		UpdateLateOCIOState(InViewport);
+
+		// Late OCIO flag
+		if (IsLateOCIO())
+		{
+			InOutMediaState |= EDisplayClusterViewportMediaState::CaptureLateOCIO;
+		}
+	}
+}
+
+FIntPoint FDisplayClusterMediaCaptureViewportBase::GetCaptureSize() const
+{
+	return GetViewportSize();
+}
+
+FIntPoint FDisplayClusterMediaCaptureViewportBase::GetViewportSize() const
+{
+	FIntPoint CaptureSize{ FIntPoint::ZeroValue };
+
+	if (GetCaptureSizeFromGameProxy(CaptureSize))
+	{
+		UE_LOG(LogDisplayClusterMedia, Verbose, TEXT("'%s' acquired capture size from game proxy [%d, %d]"), *GetMediaId(), CaptureSize.X, CaptureSize.Y);
+	}
+	else if (GetCaptureSizeFromConfig(CaptureSize))
+	{
+		UE_LOG(LogDisplayClusterMedia, Verbose, TEXT("'%s' acquired capture size from config [%d, %d]"), *GetMediaId(), CaptureSize.X, CaptureSize.Y);
+	}
+	else
+	{
+		UE_LOG(LogDisplayClusterMedia, Verbose, TEXT("'%s' couldn't acquire capture size"), *GetMediaId());
+	}
+
+	return CaptureSize;
+}
+
+bool FDisplayClusterMediaCaptureViewportBase::GetCaptureSizeFromConfig(FIntPoint& OutSize) const
+{
+	if (const ADisplayClusterRootActor* const ActiveRootActor = IDisplayCluster::Get().GetGameMgr()->GetRootActor())
+	{
+		if (const UDisplayClusterConfigurationData* const ConfigData = ActiveRootActor->GetConfigData())
+		{
+			const FString& NodeId = GetClusterNodeId();
+			if (const UDisplayClusterConfigurationViewport* const ViewportCfg = ConfigData->GetViewport(NodeId, ViewportId))
+			{
+				const FIntRect ViewportRect = ViewportCfg->Region.ToRect();
+				OutSize = FIntPoint(ViewportRect.Width(), ViewportRect.Height());
+
+				return true;
+			}
+		}
+	}
+
+	return false;
+}
+
+bool FDisplayClusterMediaCaptureViewportBase::GetCaptureSizeFromGameProxy(FIntPoint& OutSize) const
+{
+	// We need to get actual texture size for the viewport
+	if (const IDisplayClusterRenderManager* const RenderMgr = IDisplayCluster::Get().GetRenderMgr())
+	{
+		if (const IDisplayClusterViewportManager* const ViewportMgr = RenderMgr->GetViewportManager())
+		{
+			if (const IDisplayClusterViewport* const Viewport = ViewportMgr->FindViewport(ViewportId))
+			{
+				const TArray<FDisplayClusterViewport_Context>& Contexts = Viewport->GetContexts();
+				if (Contexts.Num() > 0)
+				{
+					OutSize = Contexts[0].RenderTargetRect.Size();
+					return true;
+				}
+			}
+		}
+	}
+
+	return false;
+}
+
+void FDisplayClusterMediaCaptureViewportBase::UpdateLateOCIOState(IDisplayClusterViewport* Viewport)
+{
+	// So far, regular viewports dont't support late OCIO feature.
+}
+
+void FDisplayClusterMediaCaptureViewportBase::HandleLateOCIOChanged(const FLateOCIOData& NewLateOCIOConfiguration)
+{
+	// Restart capture
+	StopCapture();
+	StartCapture();
+}
+
+void FDisplayClusterMediaCaptureViewportBase::OnPostTonemapPass_RenderThread(FRDGBuilder& GraphBuilder, const IDisplayClusterViewportProxy* ViewportProxy, const FSceneView& View, const FPostProcessMaterialInputs& Inputs, const uint32 ContextNum)
+{
+	checkSlow(ViewportProxy);
+
+	// If late OCIO is not used currently, we have to ignore this PostTonemap callback.
+	if (!IsLateOCIO())
+	{
+		return;
+	}
+
+	// Media subsystem does not support stereo, therefore we process context 0 only
+	if (ContextNum != 0)
+	{
+		return;
+	}
+
+	// Check if proxy object is valid
+	if (!ViewportProxy)
+	{
+		return;
+	}
+
+	// Make sure this is our viewport
+	const bool bMatchingViewport = ViewportProxy->GetId().Equals(GetViewportId(), ESearchCase::IgnoreCase);
+	if (!bMatchingViewport)
+	{
+		return;
+	}
+
+	// Get current SceneColor texture
+	const FScreenPassTexture& SceneColor = FScreenPassTexture::CopyFromSlice(GraphBuilder, Inputs.GetInput(EPostProcessMaterialInput::SceneColor));
+
+	// Pass it to the media capture pipeline
+	if (SceneColor.IsValid())
+	{
+		FMediaOutputTextureInfo TextureInfo{ SceneColor.Texture, SceneColor.ViewRect };
+		ExportMediaData_RenderThread(GraphBuilder, TextureInfo);
+	}
+}
+
+void FDisplayClusterMediaCaptureViewportBase::OnPostRenderViewFamily_RenderThread(FRDGBuilder& GraphBuilder, const FSceneViewFamily& ViewFamily, const IDisplayClusterViewportProxy* ViewportProxy)
+{
+	checkSlow(ViewportProxy);
+
+	// If late OCIO is used currently, we have to ignore this PostRenderViewFamily callback.
+	if (IsLateOCIO())
+	{
+		return;
+	}
+
+	// Otherwise, find our viewport and export its texture
+	if (ViewportProxy && ViewportProxy->GetId().Equals(GetViewportId(), ESearchCase::IgnoreCase))
+	{
+		TArray<FRHITexture*> Textures;
+		TArray<FIntRect>     Regions;
+
+		// Get RHI texture and pass it to the media capture pipeline
+		if (ViewportProxy->GetResourcesWithRects_RenderThread(EDisplayClusterViewportResourceType::InternalRenderTargetResource, Textures, Regions))
+		{
+			if (Textures.Num() > 0 && Regions.Num() > 0 && Textures[0])
+			{
+				FRDGTextureRef SrcTextureRef = RegisterExternalTexture(GraphBuilder, Textures[0], TEXT("DCMediaOutViewportTex"));
+
+				FMediaOutputTextureInfo TextureInfo{ SrcTextureRef, Regions[0] };
+				ExportMediaData_RenderThread(GraphBuilder, TextureInfo);
+			}
+		}
+	}
+}
diff --git a/Engine/Plugins/Runtime/nDisplay/Source/DisplayClusterMedia/Private/Capture/DisplayClusterMediaCaptureViewportBase.h b/Engine/Plugins/Runtime/nDisplay/Source/DisplayClusterMedia/Private/Capture/DisplayClusterMediaCaptureViewportBase.h
new file mode 100644
--- /dev/null
+++ b/Engine/Plugins/Runtime/nDisplay/Source/DisplayClusterMedia/Private/Capture/DisplayClusterMediaCaptureViewportBase.h
@@ -0,0 +1,84 @@
+// Copyright Epic Games, Inc. All Rights Reserved.
+
+#pragma once
+
+#include "CoreMinimal.h"
+#include "Capture/DisplayClusterMediaCaptureBase.h"
+#include "Render/Viewport/Containers/DisplayClusterViewport_Enums.h"
+
+class FRDGBuilder;
+class FSceneView;
+class FSceneViewFamily;
+class IDisplayClusterViewport;
+class IDisplayClusterViewportProxy;
+struct FPostProcessMaterialInputs;
+
+
+/**
+ * Base viewport capture adapter
+ */
+class FDisplayClusterMediaCaptureViewportBase
+	: public FDisplayClusterMediaCaptureBase
+{
+public:
+
+	FDisplayClusterMediaCaptureViewportBase(
+		const FString& MediaId,
+		const FString& ClusterNodeId,
+		const FString& ViewportId,
+		UMediaOutput* MediaOutput,
+		UDisplayClusterMediaOutputSynchronizationPolicy* SyncPolicy
+	);
+
+public:
+
+	/** Start capturing */
+	virtual bool StartCapture() override;
+
+	/** Stop capturing */
+	virtual void StopCapture() override;
+
+	/** Returns viewport ID that is configured for capture */
+	const FString& GetViewportId() const
+	{
+		return ViewportId;
+	}
+
+	/** Returns texture size of a viewport assigned to capture (main thread) */
+	virtual FIntPoint GetCaptureSize() const override;
+
+	/** Provides default texture size from config */
+	virtual bool GetCaptureSizeFromConfig(FIntPoint& OutSize) const;
+
+	/** Provides texture size from a game proxy (if available) */
+	bool GetCaptureSizeFromGameProxy(FIntPoint& OutSize) const;
+
+protected:
+
+	/** Updates late OCIO state */
+	virtual void UpdateLateOCIOState(IDisplayClusterViewport* Viewport);
+
+	/** Process late OCIO state updates */
+	virtual void HandleLateOCIOChanged(const FLateOCIOData& NewLateOCIOConfiguration) override;
+
+private:
+
+	/** UpdateViewportMediaState callback to configure media state for a viewoprt */
+	void OnUpdateViewportMediaState(IDisplayClusterViewport* InViewport, EDisplayClusterViewportMediaState& InOutMediaState);
+
+	/** PostRenderViewFamily callback handler where data is captured (no late OCIO) */
+	void OnPostRenderViewFamily_RenderThread(FRDGBuilder& GraphBuilder, const FSceneViewFamily& ViewFamily, const IDisplayClusterViewportProxy* ViewportProxy);
+
+	/** PostTonemapPass callback handler (late OCIO) */
+	void OnPostTonemapPass_RenderThread(FRDGBuilder& GraphBuilder, const IDisplayClusterViewportProxy* ViewportProxy, const FSceneView& View, const FPostProcessMaterialInputs& Inputs, const uint32 ContextNum);
+
+private:
+
+	/** Returns size of the viewport bound to this media */
+	virtual FIntPoint GetViewportSize() const;
+
+private:
+
+	/** Viewport ID assigned to capture */
+	const FString ViewportId;
+};
diff --git a/Engine/Plugins/Runtime/nDisplay/Source/DisplayClusterMedia/Private/Capture/DisplayClusterMediaCaptureViewportFull.cpp b/Engine/Plugins/Runtime/nDisplay/Source/DisplayClusterMedia/Private/Capture/DisplayClusterMediaCaptureViewportFull.cpp
new file mode 100644
--- /dev/null
+++ b/Engine/Plugins/Runtime/nDisplay/Source/DisplayClusterMedia/Private/Capture/DisplayClusterMediaCaptureViewportFull.cpp
@@ -0,0 +1,15 @@
+// Copyright Epic Games, Inc. All Rights Reserved.
+
+#include "Capture/DisplayClusterMediaCaptureViewportFull.h"
+
+
+FDisplayClusterMediaCaptureViewportFull::FDisplayClusterMediaCaptureViewportFull(
+	const FString& InMediaId,
+	const FString& InClusterNodeId,
+	const FString& InViewportId,
+	UMediaOutput* InMediaOutput,
+	UDisplayClusterMediaOutputSynchronizationPolicy* SyncPolicy
+)
+	: FDisplayClusterMediaCaptureViewportBase(InMediaId, InClusterNodeId, InViewportId, InMediaOutput, SyncPolicy)
+{
+}
diff --git a/Engine/Plugins/Runtime/nDisplay/Source/DisplayClusterMedia/Private/Capture/DisplayClusterMediaCaptureViewportFull.h b/Engine/Plugins/Runtime/nDisplay/Source/DisplayClusterMedia/Private/Capture/DisplayClusterMediaCaptureViewportFull.h
new file mode 100644
--- /dev/null
+++ b/Engine/Plugins/Runtime/nDisplay/Source/DisplayClusterMedia/Private/Capture/DisplayClusterMediaCaptureViewportFull.h
@@ -0,0 +1,23 @@
+// Copyright Epic Games, Inc. All Rights Reserved.
+
+#pragma once
+
+#include "Capture/DisplayClusterMediaCaptureViewportBase.h"
+
+
+/**
+ * Viewport capture adapter (full frame)
+ */
+class FDisplayClusterMediaCaptureViewportFull
+	: public FDisplayClusterMediaCaptureViewportBase
+{
+public:
+
+	FDisplayClusterMediaCaptureViewportFull(
+		const FString& MediaId,
+		const FString& ClusterNodeId,
+		const FString& ViewportId,
+		UMediaOutput* MediaOutput,
+		UDisplayClusterMediaOutputSynchronizationPolicy* SyncPolicy
+	);
+};
diff --git a/Engine/Plugins/Runtime/nDisplay/Source/DisplayClusterMedia/Private/Capture/DisplayClusterMediaCaptureViewportTile.cpp b/Engine/Plugins/Runtime/nDisplay/Source/DisplayClusterMedia/Private/Capture/DisplayClusterMediaCaptureViewportTile.cpp
new file mode 100644
--- /dev/null
+++ b/Engine/Plugins/Runtime/nDisplay/Source/DisplayClusterMedia/Private/Capture/DisplayClusterMediaCaptureViewportTile.cpp
@@ -0,0 +1,50 @@
+// Copyright Epic Games, Inc. All Rights Reserved.
+
+#include "Capture/DisplayClusterMediaCaptureViewportTile.h"
+
+#include "DisplayClusterMediaHelpers.h"
+
+
+FDisplayClusterMediaCaptureViewportTile::FDisplayClusterMediaCaptureViewportTile(
+	const FString& InMediaId,
+	const FString& InClusterNodeId,
+	const FString& InViewportId,
+	const FIntPoint& InTilePosition,
+	UMediaOutput* InMediaOutput,
+	UDisplayClusterMediaOutputSynchronizationPolicy* SyncPolicy
+)
+	: FDisplayClusterMediaCaptureViewportBase(
+		InMediaId,
+		InClusterNodeId,
+		DisplayClusterMediaHelpers::GenerateTileViewportName(InViewportId, InTilePosition),
+		InMediaOutput,
+		SyncPolicy)
+{
+}
+
+bool FDisplayClusterMediaCaptureViewportTile::GetCaptureSizeFromConfig(FIntPoint& OutSize) const
+{
+	// The upper level logic always tries to get actual capture size from the game proxy. If the game proxy
+	// is not available (camera/viewport is disabled or deactivated on start and therefore not being rendered),
+	// it tries to acquire capture size from config. Being here means exactly this case.
+	//
+	// The problem is the tiles can't know their size until they start to render. There is a bunch of intermediate
+	// logic that affects the final size of the owning camera/viewport texture such as "adopt resolution",
+	// "screen percentage", "overscan", etc.
+	//
+	// But we need to provide something to start capture successfully. As a workaround, we return capture size
+	// of 64x64. It looks valid and small enough to initialize media capture pipeline with a little
+	// resource allocation. When tile rendering is started (e.g. camera/viewport is activated), its texture
+	// will have final size available in the game proxy. Most likely it will be different so the media capture
+	// pipeline will reconfigure for the new size.
+	//
+	// If game proxy is available, the capture size will be acquired by the higher level in
+	// FDisplayClusterMediaCaptureViewport::GetCaptureSizeFromGameProxy()
+	//
+	// This trick looks kind of legal as we allow to change let's say ICVFX camera texture size in runtime.
+	// The media pipeline is able to recognize the size changes and reconfigure.
+
+	OutSize = { 64, 64 };
+
+	return true;
+}
diff --git a/Engine/Plugins/Runtime/nDisplay/Source/DisplayClusterMedia/Private/Capture/DisplayClusterMediaCaptureViewportTile.h b/Engine/Plugins/Runtime/nDisplay/Source/DisplayClusterMedia/Private/Capture/DisplayClusterMediaCaptureViewportTile.h
new file mode 100644
--- /dev/null
+++ b/Engine/Plugins/Runtime/nDisplay/Source/DisplayClusterMedia/Private/Capture/DisplayClusterMediaCaptureViewportTile.h
@@ -0,0 +1,30 @@
+// Copyright Epic Games, Inc. All Rights Reserved.
+
+#pragma once
+
+#include "Capture/DisplayClusterMediaCaptureViewportBase.h"
+
+
+/**
+ * Viewport capture adapter (tile)
+ */
+class FDisplayClusterMediaCaptureViewportTile
+	: public FDisplayClusterMediaCaptureViewportBase
+{
+public:
+
+	FDisplayClusterMediaCaptureViewportTile(
+		const FString& MediaId,
+		const FString& ClusterNodeId,
+		const FString& ViewportId,
+		const FIntPoint& TilePosition,
+		UMediaOutput* MediaOutput,
+		UDisplayClusterMediaOutputSynchronizationPolicy* SyncPolicy
+	);
+
+protected:
+
+	//~ Begin FDisplayClusterMediaCaptureViewport
+	virtual bool GetCaptureSizeFromConfig(FIntPoint& OutSize) const override;
+	//~ End FDisplayClusterMediaCaptureViewport
+};
diff --git a/Engine/Plugins/Runtime/nDisplay/Source/DisplayClusterMedia/Private/Common/DisplayClusterMediaBase.cpp b/Engine/Plugins/Runtime/nDisplay/Source/DisplayClusterMedia/Private/Common/DisplayClusterMediaBase.cpp
new file mode 100644
--- /dev/null
+++ b/Engine/Plugins/Runtime/nDisplay/Source/DisplayClusterMedia/Private/Common/DisplayClusterMediaBase.cpp
@@ -0,0 +1,59 @@
+// Copyright Epic Games, Inc. All Rights Reserved.
+
+#include "DisplayClusterMediaBase.h"
+#include "RenderingThread.h"
+
+
+bool FDisplayClusterMediaBase::IsLateOCIO() const
+{
+	if (IsInGameThread())
+	{
+		return LateOCIOConfiguration.bLateOCIO;
+	}
+	else if (IsInRenderingThread())
+	{
+		return LateOCIOConfiguration_RT.bLateOCIO;
+	}
+
+	return false;
+}
+
+bool FDisplayClusterMediaBase::IsTransferPQ(bool bConsideringLateOCIOState) const
+{
+	if (IsInGameThread())
+	{
+		return bConsideringLateOCIOState ?
+			LateOCIOConfiguration.bLateOCIO && LateOCIOConfiguration.bTransferPQ :
+			LateOCIOConfiguration.bTransferPQ;
+	}
+	else if (IsInRenderingThread())
+	{
+		return bConsideringLateOCIOState ?
+			LateOCIOConfiguration_RT.bLateOCIO && LateOCIOConfiguration_RT.bTransferPQ :
+			LateOCIOConfiguration_RT.bTransferPQ;
+	}
+
+	return false;
+}
+
+void FDisplayClusterMediaBase::SetLateOCIO(const FLateOCIOData& NewLateOCIOConfiguration)
+{
+	if (IsInGameThread())
+	{
+		// Let children know the OCIO parameters have changed
+		if (NewLateOCIOConfiguration != LateOCIOConfiguration)
+		{
+			HandleLateOCIOChanged(NewLateOCIOConfiguration);
+		}
+
+		// Update configuration
+		LateOCIOConfiguration = NewLateOCIOConfiguration;
+
+		// And pass to the render thread
+		ENQUEUE_RENDER_COMMAND(DCMediaUpdateOCIOState)(
+			[this, SetNewLateOCIO = LateOCIOConfiguration](FRHICommandListImmediate& RHICmdList)
+			{
+				LateOCIOConfiguration_RT = SetNewLateOCIO;
+			});
+	}
+}
diff --git a/Engine/Plugins/Runtime/nDisplay/Source/DisplayClusterMedia/Private/Common/DisplayClusterMediaBase.h b/Engine/Plugins/Runtime/nDisplay/Source/DisplayClusterMedia/Private/Common/DisplayClusterMediaBase.h
new file mode 100644
--- /dev/null
+++ b/Engine/Plugins/Runtime/nDisplay/Source/DisplayClusterMedia/Private/Common/DisplayClusterMediaBase.h
@@ -0,0 +1,91 @@
+// Copyright Epic Games, Inc. All Rights Reserved.
+
+#pragma once
+
+#include "CoreMinimal.h"
+#include "OpenColorIOColorSpace.h"
+
+
+/**
+ * Base media adapter class
+ */
+class FDisplayClusterMediaBase
+{
+public:
+
+	FDisplayClusterMediaBase(const FString& InMediaId, const FString& InClusterNodeId)
+		: MediaId(InMediaId)
+		, ClusterNodeId(InClusterNodeId)
+	{ }
+
+	virtual ~FDisplayClusterMediaBase() = default;
+
+public:
+
+	/** Returns ID of this media adapter */
+	const FString& GetMediaId() const
+	{
+		return MediaId;
+	}
+
+	/** Returns current cluster node ID */
+	const FString& GetClusterNodeId() const
+	{
+		return ClusterNodeId;
+	}
+
+	/** Whether late OCIO is active on current frame */
+	bool IsLateOCIO() const;
+
+	/**
+	 * Is PQ transfer required?
+	 * 
+	 * @param bConsideringLateOCIOState - Should the OCIO on/off state be logically AND-ed with the PQ on/off state
+	 */
+	bool IsTransferPQ(bool bConsideringLateOCIOState = true) const;
+
+protected:
+
+	/** Auxiliary structure to keep late OCIO parameters */
+	struct FLateOCIOData
+	{
+		/** Late OCIO enabled/disabled flag */
+		bool bLateOCIO = false;
+
+		/** PQ transfer enabled/disabled */
+		bool bTransferPQ = false;
+
+		/** Comparison operator */
+		bool operator==(const FLateOCIOData& Other)
+		{
+			const bool bSameEnabledState = (this->bLateOCIO == Other.bLateOCIO);
+			const bool bSamePQ = (this->bTransferPQ == Other.bTransferPQ);
+
+			const bool bEqual = (bSameEnabledState && bSamePQ);
+
+			return bEqual;
+		}
+	};
+
+	/** Set late OCIO configuration by children */
+	void SetLateOCIO(const FLateOCIOData& NewLateOCIOConfiguration);
+
+	virtual void HandleLateOCIOChanged(const FLateOCIOData& NewLateOCIOConfiguration)
+	{ }
+
+private:
+
+	/** ID of this media adapter */
+	const FString MediaId;
+
+	/** Cluster node ID we're running on */
+	const FString ClusterNodeId;
+
+private:
+
+	/** Late OCIO configuration on current frame */
+	FLateOCIOData LateOCIOConfiguration;
+
+	/** Late OCIO configuration on current frame (render thread) */
+	FLateOCIOData LateOCIOConfiguration_RT;
+};
diff --git a/Engine/Plugins/Runtime/nDisplay/Source/DisplayClusterMedia/Private/Common/DisplayClusterMediaCameraCommon.cpp b/Engine/Plugins/Runtime/nDisplay/Source/DisplayClusterMedia/Private/Common/DisplayClusterMediaCameraCommon.cpp
new file mode 100644
--- /dev/null
+++ b/Engine/Plugins/Runtime/nDisplay/Source/DisplayClusterMedia/Private/Common/DisplayClusterMediaCameraCommon.cpp
@@ -0,0 +1,45 @@
+// Copyright Epic Games, Inc. All Rights Reserved.
+
+#include "Common/DisplayClusterMediaCameraCommon.h"
+
+#include "Components/DisplayClusterICVFXCameraComponent.h"
+#include "Game/IDisplayClusterGameManager.h"
+
+#include "DisplayClusterRootActor.h"
+#include "IDisplayCluster.h"
+
+
+FDisplayClusterMediaCameraCommon::FDisplayClusterMediaCameraCommon(const FString& InCameraId)
+	: CameraId(InCameraId)
+{
+}
+
+UDisplayClusterICVFXCameraComponent* FDisplayClusterMediaCameraCommon::GetCameraComponent() const
+{
+	if (const ADisplayClusterRootActor* const DCRA = IDisplayCluster::Get().GetGameMgr()->GetRootActor())
+	{
+		if (UDisplayClusterICVFXCameraComponent* ICVFXCameraComponent = DCRA->GetComponentByName<UDisplayClusterICVFXCameraComponent>(CameraId))
+		{
+			return ICVFXCameraComponent;
+		}
+	}
+
+	return nullptr;
+}
+
+void FDisplayClusterMediaCameraCommon::GetLateOCIOParameters(bool& bOutLateOCIOEnabled, bool& bOutTransferPQ) const
+{
+	if (const UDisplayClusterICVFXCameraComponent* const ICVFXCameraComponent = GetCameraComponent())
+	{
+		// Check if the main OCIO switch is on
+		const bool bEnabledOCIO = ICVFXCameraComponent->CameraSettings.CameraOCIO.AllNodesOCIOConfiguration.bIsEnabled;
+		// Check if media is enabled on this camera
+		const bool bEnabledMedia = ICVFXCameraComponent->CameraSettings.RenderSettings.Media.bEnable;
+		// Check if the late OCIO flag is on
+		const bool bEnabledLateOCIO = ICVFXCameraComponent->CameraSettings.RenderSettings.Media.bLateOCIOPass;
+
+		// Return OCIO configuration parameters
+		bOutLateOCIOEnabled = bEnabledOCIO && bEnabledMedia && bEnabledLateOCIO;
+		bOutTransferPQ = ICVFXCameraComponent->CameraSettings.RenderSettings.Media.bTransferPQ;
+	}
+}
diff --git a/Engine/Plugins/Runtime/nDisplay/Source/DisplayClusterMedia/Private/Common/DisplayClusterMediaCameraCommon.h b/Engine/Plugins/Runtime/nDisplay/Source/DisplayClusterMedia/Private/Common/DisplayClusterMediaCameraCommon.h
new file mode 100644
--- /dev/null
+++ b/Engine/Plugins/Runtime/nDisplay/Source/DisplayClusterMedia/Private/Common/DisplayClusterMediaCameraCommon.h
@@ -0,0 +1,32 @@
+// Copyright Epic Games, Inc. All Rights Reserved.
+
+#pragma once
+
+#include "CoreMinimal.h"
+
+class UDisplayClusterICVFXCameraComponent;
+
+
+/**
+ * Common camera logic for media
+ */
+class FDisplayClusterMediaCameraCommon
+{
+public:
+
+	FDisplayClusterMediaCameraCommon(const FString& CameraId);
+	virtual ~FDisplayClusterMediaCameraCommon() = default;
+
+protected:
+
+	/** Finds an ICVFX camera component by name, otherwise nullptr */
+	UDisplayClusterICVFXCameraComponent* GetCameraComponent() const;
+
+	/** Returns late OCIO parameters of the current camera */
+	void GetLateOCIOParameters(bool& bOutLateOCIOEnabled, bool& bOutTransferPQ) const;
+
+private:
+
+	/** Holds ICVFX camera component name associated with this media adapter */
+	const FString CameraId;
+};
diff --git a/Engine/Plugins/Runtime/nDisplay/Source/DisplayClusterMedia/Private/DisplayClusterMediaBase.h b/Engine/Plugins/Runtime/nDisplay/Source/DisplayClusterMedia/Private/DisplayClusterMediaBase.h
deleted file mode 100644
diff --git a/Engine/Plugins/Runtime/nDisplay/Source/DisplayClusterMedia/Private/DisplayClusterMediaHelpers.cpp b/Engine/Plugins/Runtime/nDisplay/Source/DisplayClusterMedia/Private/DisplayClusterMediaHelpers.cpp
--- a/Engine/Plugins/Runtime/nDisplay/Source/DisplayClusterMedia/Private/DisplayClusterMediaHelpers.cpp
+++ b/Engine/Plugins/Runtime/nDisplay/Source/DisplayClusterMedia/Private/DisplayClusterMediaHelpers.cpp
@@ -75,6 +75,14 @@ namespace DisplayClusterMediaHelpers
 		return FString::Printf(TEXT("%s_tile_%d_%d"), *InViewportId, InTilePos.X, InTilePos.Y);
 	}
 
+	FString GenerateICVFXTileViewportName(const FString& ClusterNodeId, const FString& ICVFXCameraName, const FIntPoint& InTilePos)
+	{
+		const FString BaseCameraViewportId = GenerateICVFXViewportName(ClusterNodeId, ICVFXCameraName);
+		const FString CameraTileViewportId = GenerateTileViewportName(BaseCameraViewportId, InTilePos);
+
+		return CameraTileViewportId;
+	}
+
 	template<class TScreenPixelShader>
 	void ResampleTextureImpl_RenderThread(FRHICommandListImmediate& RHICmdList, FRHITexture* SrcTexture, FRHITexture* DstTexture, const FIntRect& SrcRect, const FIntRect& DstRect)
 	{
diff --git a/Engine/Plugins/Runtime/nDisplay/Source/DisplayClusterMedia/Private/DisplayClusterMediaHelpers.h b/Engine/Plugins/Runtime/nDisplay/Source/DisplayClusterMedia/Private/DisplayClusterMediaHelpers.h
--- a/Engine/Plugins/Runtime/nDisplay/Source/DisplayClusterMedia/Private/DisplayClusterMediaHelpers.h
+++ b/Engine/Plugins/Runtime/nDisplay/Source/DisplayClusterMedia/Private/DisplayClusterMediaHelpers.h
@@ -37,6 +37,9 @@ namespace DisplayClusterMediaHelpers
 	/** Generates viewport tile IDs */
 	FString GenerateTileViewportName(const FString& InViewportId, const FIntPoint& InTilePos);
 
+	/** Generates viewport tile IDs for an ICVFX camera */
+	FString GenerateICVFXTileViewportName(const FString& ClusterNodeId, const FString& ICVFXCameraName, const FIntPoint& InTilePos);
+
 	/** Copy and resize an RHI texture */
 	void ResampleTexture_RenderThread(FRHICommandListImmediate& RHICmdList, FRHITexture* SrcTexture, FRHITexture* DstTexture, const FIntRect& SrcRect, const FIntRect& DstRect);
 }
diff --git a/Engine/Plugins/Runtime/nDisplay/Source/DisplayClusterMedia/Private/DisplayClusterMediaModule.cpp b/Engine/Plugins/Runtime/nDisplay/Source/DisplayClusterMedia/Private/DisplayClusterMediaModule.cpp
--- a/Engine/Plugins/Runtime/nDisplay/Source/DisplayClusterMedia/Private/DisplayClusterMediaModule.cpp
+++ b/Engine/Plugins/Runtime/nDisplay/Source/DisplayClusterMedia/Private/DisplayClusterMediaModule.cpp
@@ -17,11 +17,15 @@
 #include "Cluster/IDisplayClusterClusterManager.h"
 #include "Game/IDisplayClusterGameManager.h"
 
-#include "Capture/DisplayClusterMediaCaptureCamera.h"
+#include "Capture/DisplayClusterMediaCaptureCameraFull.h"
+#include "Capture/DisplayClusterMediaCaptureCameraTile.h"
 #include "Capture/DisplayClusterMediaCaptureNode.h"
-#include "Capture/DisplayClusterMediaCaptureTile.h"
-#include "Capture/DisplayClusterMediaCaptureViewport.h"
-#include "Input/DisplayClusterMediaInputViewport.h"
+#include "Capture/DisplayClusterMediaCaptureViewportFull.h"
+#include "Capture/DisplayClusterMediaCaptureViewportTile.h"
+#include "Input/DisplayClusterMediaInputCameraFull.h"
+#include "Input/DisplayClusterMediaInputCameraTile.h"
+#include "Input/DisplayClusterMediaInputViewportFull.h"
+#include "Input/DisplayClusterMediaInputViewportTile.h"
 
 #include "IMediaModule.h"
 #include "Misc/CoreDelegates.h"
@@ -133,59 +137,45 @@ void FDisplayClusterMediaModule::ReleaseMedia()
 	StopCapture();
 	StopMedia();
 
-	CaptureViewports.Reset();
-	CaptureNode.Reset();
-
-	InputViewports.Reset();
+	CaptureDevices.Reset();
+	InputDevices.Reset();
 
 	FrameQueue.Release();
 }
 
 void FDisplayClusterMediaModule::StartCapture()
 {
-	// Start viewports capture
-	for (TPair<FString, TSharedPtr<FDisplayClusterMediaCaptureViewport>>& Capture : CaptureViewports)
-	{
-		Capture.Value->StartCapture();
-	}
-
-	// Start backbuffer capture
-	for (TPair<FString, TSharedPtr<FDisplayClusterMediaCaptureNode>>& Capture : CaptureNode)
+	// Start all capture devices
+	for (TPair<FString, TSharedPtr<FDisplayClusterMediaCaptureBase>>& CaptureDevice : CaptureDevices)
 	{
-		Capture.Value->StartCapture();
+		CaptureDevice.Value->StartCapture();
 	}
 }
 
 void FDisplayClusterMediaModule::StopCapture()
 {
-	// Stop viewports capture
-	for (TPair<FString, TSharedPtr<FDisplayClusterMediaCaptureViewport>>& Capture : CaptureViewports)
+	// Stop all capture devices
+	for (TPair<FString, TSharedPtr<FDisplayClusterMediaCaptureBase>>& CaptureDevice : CaptureDevices)
 	{
-		Capture.Value->StopCapture();
-	}
-
-	// Stop backbuffer capture
-	for (TPair<FString, TSharedPtr<FDisplayClusterMediaCaptureNode>>& Capture : CaptureNode)
-	{
-		Capture.Value->StopCapture();
+		CaptureDevice.Value->StopCapture();
 	}
 }
 
 void FDisplayClusterMediaModule::PlayMedia()
 {
-	// Start playback on viewports
-	for (TPair<FString, TSharedPtr<FDisplayClusterMediaInputViewport>>& MediaInput : InputViewports)
+	// Start all input devices
+	for (TPair<FString, TSharedPtr<FDisplayClusterMediaInputBase>>& InputDevice : InputDevices)
 	{
-		MediaInput.Value->Play();
+		InputDevice.Value->Play();
 	}
 }
 
 void FDisplayClusterMediaModule::StopMedia()
 {
-	// Stop playback on viewports
-	for (TPair<FString, TSharedPtr<FDisplayClusterMediaInputViewport>>& MediaInput : InputViewports)
+	// Stop all input devices
+	for (TPair<FString, TSharedPtr<FDisplayClusterMediaInputBase>>& InputDevice : InputDevices)
 	{
-		MediaInput.Value->Stop();
+		InputDevice.Value->Stop();
 	}
 }
 
@@ -212,11 +202,12 @@ void FDisplayClusterMediaModule::InitializeBackbufferOutput(const UDisplayCluste
 					UE_LOG(LogDisplayClusterMedia, Log, TEXT("Initializing backbuffer media capture [%u]: '%s'"), CaptureIdx, *MediaCaptureId);
 
 					TSharedPtr<FDisplayClusterMediaCaptureNode> NewNodeCapture = MakeShared<FDisplayClusterMediaCaptureNode>(
-						MediaCaptureId, ClusterNodeId,
+						MediaCaptureId,
+						ClusterNodeId,
 						MediaOutputItem.MediaOutput,
 						MediaOutputItem.OutputSyncPolicy);
 
-					CaptureNode.Emplace(MediaCaptureId, NewNodeCapture);
+					CaptureDevices.Emplace(MediaCaptureId, MoveTemp(NewNodeCapture));
 				}
 
 				++CaptureIdx;
@@ -244,12 +235,13 @@ void FDisplayClusterMediaModule::InitializeViewportInput(const UDisplayClusterCo
 
 				UE_LOG(LogDisplayClusterMedia, Log, TEXT("Initializing viewport media input '%s' for viewport '%s'"), *MediaInputId, *ViewportId);
 
-				TSharedPtr<FDisplayClusterMediaInputViewport> NewViewportInput = MakeShared<FDisplayClusterMediaInputViewport>(
-					MediaInputId, ClusterNodeId,
+				TSharedPtr<FDisplayClusterMediaInputViewportFull> NewViewportInput = MakeShared<FDisplayClusterMediaInputViewportFull>(
+					MediaInputId,
+					ClusterNodeId,
 					ViewportId,
 					MediaSettings.MediaInput.MediaSource);
 
-				InputViewports.Emplace(MediaInputId, MoveTemp(NewViewportInput));
+				InputDevices.Emplace(MediaInputId, MoveTemp(NewViewportInput));
 			}
 		}
 	}
@@ -276,13 +268,14 @@ void FDisplayClusterMediaModule::InitializeViewportOutput(const UDisplayClusterC
 
 				UE_LOG(LogDisplayClusterMedia, Log, TEXT("Initializing viewport capture [%u]: '%s' for viewport '%s'"), CaptureIdx, *MediaCaptureId, *ViewportId);
 
-				TSharedPtr<FDisplayClusterMediaCaptureViewport> NewViewportCapture = MakeShared<FDisplayClusterMediaCaptureViewport>(
-					MediaCaptureId, ClusterNodeId,
+				TSharedPtr<FDisplayClusterMediaCaptureViewportFull> NewViewportCapture = MakeShared<FDisplayClusterMediaCaptureViewportFull>(
+					MediaCaptureId,
+					ClusterNodeId,
 					ViewportId,
 					MediaOutputItem.MediaOutput,
 					MediaOutputItem.OutputSyncPolicy);
 
-				CaptureViewports.Emplace(MediaCaptureId, MoveTemp(NewViewportCapture));
+				CaptureDevices.Emplace(MediaCaptureId, MoveTemp(NewViewportCapture));
 			}
 
 			++CaptureIdx;
@@ -290,7 +283,7 @@ void FDisplayClusterMediaModule::InitializeViewportOutput(const UDisplayClusterC
 	}
 }
 
-void FDisplayClusterMediaModule::InitializeICVFXCameraFullFrameInput(const UDisplayClusterICVFXCameraComponent* ICVFXCameraComponent, const FString& RootActorName, const FString& ClusterNodeId, const FString& ExplicitCameraName)
+void FDisplayClusterMediaModule::InitializeICVFXCameraFullFrameInput(const UDisplayClusterICVFXCameraComponent* ICVFXCameraComponent, const FString& RootActorName, const FString& ClusterNodeId)
 {
 	checkSlow(ICVFXCameraComponent);
 
@@ -300,15 +293,10 @@ void FDisplayClusterMediaModule::InitializeICVFXCameraFullFrameInput(const UDisp
 
 		if (MediaSettings.bEnable && MediaSettings.SplitType == EDisplayClusterConfigurationMediaSplitType::FullFrame)
 		{
-			const FString ICVFXCameraName = ExplicitCameraName.IsEmpty() ? ICVFXCameraComponent->GetName() : ExplicitCameraName;
-			const FString ICVFXViewportId = DisplayClusterMediaHelpers::GenerateICVFXViewportName(ClusterNodeId, ICVFXCameraName);
-
-			// Determine whether it's necessary to activate the late OCIO on media side
-			const bool bEnabledOCIO = ICVFXCameraComponent->CameraSettings.CameraOCIO.AllNodesOCIOConfiguration.bIsEnabled;
-			const bool bUseLateOCIO = MediaSettings.bLateOCIOPass && bEnabledOCIO;
-
 			if (UMediaSource* MediaSource = MediaSettings.GetMediaSource(ClusterNodeId))
 			{
+				const FString ICVFXCameraName = ICVFXCameraComponent->GetName();
+
 				const FString MediaInputId = DisplayClusterMediaHelpers::MediaId::GenerateMediaId(
 					DisplayClusterMediaHelpers::MediaId::EMediaDeviceType::Input,
 					DisplayClusterMediaHelpers::MediaId::EMediaOwnerType::ICVFXCamera,
@@ -316,19 +304,19 @@ void FDisplayClusterMediaModule::InitializeICVFXCameraFullFrameInput(const UDisp
 
 				UE_LOG(LogDisplayClusterMedia, Log, TEXT("Initializing ICVFX media input '%s' for camera '%s'"), *MediaInputId, *ICVFXCameraName);
 
-				TSharedPtr<FDisplayClusterMediaInputViewport> NewICVFXInput = MakeShared<FDisplayClusterMediaInputViewport>(
-					MediaInputId, ClusterNodeId,
-					ICVFXViewportId,
-					MediaSource,
-					bUseLateOCIO);
+				TSharedPtr<FDisplayClusterMediaInputCameraFull> NewICVFXInput = MakeShared<FDisplayClusterMediaInputCameraFull>(
+					MediaInputId,
+					ClusterNodeId,
+					ICVFXCameraName,
+					MediaSource);
 
-				InputViewports.Emplace(MediaInputId, MoveTemp(NewICVFXInput));
+				InputDevices.Emplace(MediaInputId, MoveTemp(NewICVFXInput));
 			}
 		}
 	}
 }
 
-void FDisplayClusterMediaModule::InitializeICVFXCameraFullFrameOutput(const UDisplayClusterICVFXCameraComponent* ICVFXCameraComponent, const FString& RootActorName, const FString& ClusterNodeId, const FString& ExplicitCameraName)
+void FDisplayClusterMediaModule::InitializeICVFXCameraFullFrameOutput(const UDisplayClusterICVFXCameraComponent* ICVFXCameraComponent, const FString& RootActorName, const FString& ClusterNodeId)
 {
 	checkSlow(ICVFXCameraComponent);
 
@@ -338,12 +326,7 @@ void FDisplayClusterMediaModule::InitializeICVFXCameraFullFrameOutput(const UDis
 
 		if (MediaSettings.bEnable && MediaSettings.SplitType == EDisplayClusterConfigurationMediaSplitType::FullFrame)
 		{
-			const FString ICVFXCameraName = ExplicitCameraName.IsEmpty() ? ICVFXCameraComponent->GetName() : ExplicitCameraName;
-			const FString ICVFXViewportId = DisplayClusterMediaHelpers::GenerateICVFXViewportName(ClusterNodeId, ICVFXCameraName);
-
-			// Determine whether it's necessary to activate the late OCIO on media side
-			const bool bEnabledOCIO = ICVFXCameraComponent->CameraSettings.CameraOCIO.AllNodesOCIOConfiguration.bIsEnabled;
-			const bool bUseLateOCIO = MediaSettings.bLateOCIOPass && bEnabledOCIO;
+			const FString ICVFXCameraName = ICVFXCameraComponent->GetName();
 
 			// Media capture (full frame)
 			const TArray<FDisplayClusterConfigurationMediaOutputGroup> MediaOutputItems = MediaSettings.GetMediaOutputGroups(ClusterNodeId);
@@ -359,14 +342,14 @@ void FDisplayClusterMediaModule::InitializeICVFXCameraFullFrameOutput(const UDis
 
 					UE_LOG(LogDisplayClusterMedia, Log, TEXT("Initializing ICVFX capture [%u]: '%s' for camera '%s'"), CaptureIdx, *MediaCaptureId, *ICVFXCameraName);
 
-					TSharedPtr<FDisplayClusterMediaCaptureViewport> NewICVFXCapture = MakeShared<FDisplayClusterMediaCaptureCamera>(
-						MediaCaptureId, ClusterNodeId,
-						ICVFXCameraName, ICVFXViewportId,
+					TSharedPtr<FDisplayClusterMediaCaptureCameraFull> NewICVFXCapture = MakeShared<FDisplayClusterMediaCaptureCameraFull>(
+						MediaCaptureId,
+						ClusterNodeId,
+						ICVFXCameraName,
 						MediaOutputItem.MediaOutput,
-						MediaOutputItem.OutputSyncPolicy,
-						bUseLateOCIO);
+						MediaOutputItem.OutputSyncPolicy);
 
-					CaptureViewports.Emplace(MediaCaptureId, MoveTemp(NewICVFXCapture));
+					CaptureDevices.Emplace(MediaCaptureId, MoveTemp(NewICVFXCapture));
 				}
 
 				++CaptureIdx;
@@ -375,7 +358,7 @@ void FDisplayClusterMediaModule::InitializeICVFXCameraFullFrameOutput(const UDis
 	}
 }
 
-void FDisplayClusterMediaModule::InitializeICVFXCameraUniformTilesInput(const UDisplayClusterICVFXCameraComponent* ICVFXCameraComponent, const FString& RootActorName, const FString& ClusterNodeId, const FString& ExplicitCameraName)
+void FDisplayClusterMediaModule::InitializeICVFXCameraUniformTilesInput(const UDisplayClusterICVFXCameraComponent* ICVFXCameraComponent, const FString& RootActorName, const FString& ClusterNodeId)
 {
 	checkSlow(ICVFXCameraComponent);
 
@@ -391,20 +374,13 @@ void FDisplayClusterMediaModule::InitializeICVFXCameraUniformTilesInput(const UD
 
 			if (bMediaInputGroupFound)
 			{
-				const FString ICVFXCameraName = ExplicitCameraName.IsEmpty() ? ICVFXCameraComponent->GetName() : ExplicitCameraName;
-				const FString ICVFXViewportId = DisplayClusterMediaHelpers::GenerateICVFXViewportName(ClusterNodeId, ICVFXCameraName);
-
-				// Determine whether it's necessary to activate the late OCIO on media side
-				const bool bEnabledOCIO = ICVFXCameraComponent->CameraSettings.CameraOCIO.AllNodesOCIOConfiguration.bIsEnabled;
-				const bool bUseLateOCIO = MediaSettings.bLateOCIOPass && bEnabledOCIO;
+				const FString ICVFXCameraName = ICVFXCameraComponent->GetName();
 
 				uint8 Index = 0;
 				for (const FDisplayClusterConfigurationMediaUniformTileInput& MediaInputTile : MediaInputTiles)
 				{
 					if (IsValid(MediaInputTile.MediaSource))
 					{
-						const FString ICVFXViewportTileId = DisplayClusterMediaHelpers::GenerateTileViewportName(ICVFXViewportId, MediaInputTile.Position);
-
 						const FString MediaInputId = DisplayClusterMediaHelpers::MediaId::GenerateMediaId(
 							DisplayClusterMediaHelpers::MediaId::EMediaDeviceType::Input,
 							DisplayClusterMediaHelpers::MediaId::EMediaOwnerType::ICVFXCamera,
@@ -414,13 +390,14 @@ void FDisplayClusterMediaModule::InitializeICVFXCameraUniformTilesInput(const UD
 						UE_LOG(LogDisplayClusterMedia, Log, TEXT("Initializing ICVFX media input '%s' for camera '%s' tile '%d,%d'"),
 							*MediaInputId, *ICVFXCameraName, MediaInputTile.Position.X, MediaInputTile.Position.Y);
 
-						TSharedPtr<FDisplayClusterMediaInputViewport> NewICVFXTileInput = MakeShared<FDisplayClusterMediaInputViewport>(
-							MediaInputId, ClusterNodeId,
-							ICVFXViewportTileId,
-							MediaInputTile.MediaSource,
-							bUseLateOCIO);
+						TSharedPtr<FDisplayClusterMediaInputCameraTile> NewICVFXTileInput = MakeShared<FDisplayClusterMediaInputCameraTile>(
+							MediaInputId,
+							ClusterNodeId,
+							ICVFXCameraName,
+							MediaInputTile.Position,
+							MediaInputTile.MediaSource);
 
-						InputViewports.Emplace(MediaInputId, MoveTemp(NewICVFXTileInput));
+						InputDevices.Emplace(MediaInputId, MoveTemp(NewICVFXTileInput));
 					}
 
 					++Index;
@@ -430,7 +407,7 @@ void FDisplayClusterMediaModule::InitializeICVFXCameraUniformTilesInput(const UD
 	}
 }
 
-void FDisplayClusterMediaModule::InitializeICVFXCameraUniformTilesOutput(const UDisplayClusterICVFXCameraComponent* ICVFXCameraComponent, const FString& RootActorName, const FString& ClusterNodeId, const FString& ExplicitCameraName)
+void FDisplayClusterMediaModule::InitializeICVFXCameraUniformTilesOutput(const UDisplayClusterICVFXCameraComponent* ICVFXCameraComponent, const FString& RootActorName, const FString& ClusterNodeId)
 {
 	checkSlow(ICVFXCameraComponent);
 
@@ -446,20 +423,13 @@ void FDisplayClusterMediaModule::InitializeICVFXCameraUniformTilesOutput(const U
 
 			if (bMediaOutputGroupFound)
 			{
-				const FString ICVFXCameraName = ExplicitCameraName.IsEmpty() ? ICVFXCameraComponent->GetName() : ExplicitCameraName;
-				const FString ICVFXViewportId = DisplayClusterMediaHelpers::GenerateICVFXViewportName(ClusterNodeId, ICVFXCameraName);
-
-				// Determine whether it's necessary to activate the late OCIO on media side
-				const bool bEnabledOCIO = ICVFXCameraComponent->CameraSettings.CameraOCIO.AllNodesOCIOConfiguration.bIsEnabled;
-				const bool bUseLateOCIO = MediaSettings.bLateOCIOPass && bEnabledOCIO;
+				const FString ICVFXCameraName = ICVFXCameraComponent->GetName();
 
 				uint8 Index = 0;
 				for (const FDisplayClusterConfigurationMediaUniformTileOutput& MediaOutputTile : MediaOutputTiles)
 				{
 					if (IsValid(MediaOutputTile.MediaOutput))
 					{
-						const FString ICVFXViewportTileId = DisplayClusterMediaHelpers::GenerateTileViewportName(ICVFXViewportId, MediaOutputTile.Position);
-
 						const FString MediaOutputId = DisplayClusterMediaHelpers::MediaId::GenerateMediaId(
 							DisplayClusterMediaHelpers::MediaId::EMediaDeviceType::Output,
 							DisplayClusterMediaHelpers::MediaId::EMediaOwnerType::ICVFXCamera,
@@ -469,14 +439,15 @@ void FDisplayClusterMediaModule::InitializeICVFXCameraUniformTilesOutput(const U
 						UE_LOG(LogDisplayClusterMedia, Log, TEXT("Initializing ICVFX media output '%s' for camera '%s' tile '%d,%d'"),
 							*MediaOutputId, *ICVFXCameraName, MediaOutputTile.Position.X, MediaOutputTile.Position.Y);
 
-						TSharedPtr<FDisplayClusterMediaCaptureTile> NewICVFXTileOutput = MakeShared<FDisplayClusterMediaCaptureTile>(
-							MediaOutputId, ClusterNodeId,
-							ICVFXViewportTileId,
+						TSharedPtr<FDisplayClusterMediaCaptureCameraTile> NewICVFXTileOutput = MakeShared<FDisplayClusterMediaCaptureCameraTile>(
+							MediaOutputId,
+							ClusterNodeId,
+							ICVFXCameraName,
+							MediaOutputTile.Position,
 							MediaOutputTile.MediaOutput,
-							nullptr, // no sync for tiles as it's basically for in-cluster use
-							bUseLateOCIO);
+							nullptr); // no sync for tiles as it's basically for in-cluster use
 
-						CaptureViewports.Emplace(MediaOutputId, MoveTemp(NewICVFXTileOutput));
+						CaptureDevices.Emplace(MediaOutputId, MoveTemp(NewICVFXTileOutput));
 					}
 
 					++Index;
diff --git a/Engine/Plugins/Runtime/nDisplay/Source/DisplayClusterMedia/Private/DisplayClusterMediaModule.h b/Engine/Plugins/Runtime/nDisplay/Source/DisplayClusterMedia/Private/DisplayClusterMediaModule.h
--- a/Engine/Plugins/Runtime/nDisplay/Source/DisplayClusterMedia/Private/DisplayClusterMediaModule.h
+++ b/Engine/Plugins/Runtime/nDisplay/Source/DisplayClusterMedia/Private/DisplayClusterMediaModule.h
@@ -7,11 +7,9 @@
 
 #include "Synchronization/LatencyQueue/DisplayClusterFrameQueue.h"
 
-class FDisplayClusterMediaCaptureNode;
-class FDisplayClusterMediaCaptureViewport;
-class FDisplayClusterMediaInputViewport;
+class FDisplayClusterMediaCaptureBase;
+class FDisplayClusterMediaInputBase;
 class FSceneViewFamilyContext;
-class IMediaPlayerFactory;
 class UDisplayClusterConfigurationClusterNode;
 class UDisplayClusterConfigurationViewport;
 class UDisplayClusterICVFXCameraComponent;
@@ -24,12 +22,14 @@ class FDisplayClusterMediaModule :
 	public IModuleInterface
 {
 public:
+
 	//~ Begin IModuleInterface Implementation
 	virtual void StartupModule() override;
 	virtual void ShutdownModule() override;
 	//~ End IModuleInterface Implementation
 
 protected:
+
 	/** Initialize all media internals (from instance) */
 	void InitializeMedia();
 
@@ -49,6 +49,7 @@ protected:
 	void StopMedia();
 
 private:
+
 	/** Backbuffer output initializer */
 	void InitializeBackbufferOutput(const UDisplayClusterConfigurationClusterNode* ClusterNode, const FString& RootActorName, const FString& ClusterNodeId);
 	
@@ -59,18 +60,19 @@ private:
 	void InitializeViewportOutput(const UDisplayClusterConfigurationViewport* Viewport, const FString& ViewportId, const FString& RootActorName, const FString& ClusterNodeId);
 
 	/** ICVFX camera input initializer (full frame) */
-	void InitializeICVFXCameraFullFrameInput(const UDisplayClusterICVFXCameraComponent* ICVFXCameraComponent, const FString& RootActorName, const FString& ClusterNodeId, const FString& ExplicitCameraName = FString());
+	void InitializeICVFXCameraFullFrameInput(const UDisplayClusterICVFXCameraComponent* ICVFXCameraComponent, const FString& RootActorName, const FString& ClusterNodeId);
 
 	/** ICVFX camera output initializer (full frame) */
-	void InitializeICVFXCameraFullFrameOutput(const UDisplayClusterICVFXCameraComponent* ICVFXCameraComponent, const FString& RootActorName, const FString& ClusterNodeId, const FString& ExplicitCameraName = FString());
+	void InitializeICVFXCameraFullFrameOutput(const UDisplayClusterICVFXCameraComponent* ICVFXCameraComponent, const FString& RootActorName, const FString& ClusterNodeId);
 
 	/** ICVFX camera input initializer (uniform tiles) */
-	void InitializeICVFXCameraUniformTilesInput(const UDisplayClusterICVFXCameraComponent* ICVFXCameraComponent, const FString& RootActorName, const FString& ClusterNodeId, const FString& ExplicitCameraName = FString());
+	void InitializeICVFXCameraUniformTilesInput(const UDisplayClusterICVFXCameraComponent* ICVFXCameraComponent, const FString& RootActorName, const FString& ClusterNodeId);
 
 	/** ICVFX camera output initializer (uniform tiles) */
-	void InitializeICVFXCameraUniformTilesOutput(const UDisplayClusterICVFXCameraComponent* ICVFXCameraComponent, const FString& RootActorName, const FString& ClusterNodeId, const FString& ExplicitCameraName = FString());
+	void InitializeICVFXCameraUniformTilesOutput(const UDisplayClusterICVFXCameraComponent* ICVFXCameraComponent, const FString& RootActorName, const FString& ClusterNodeId);
 
 private:
+
 	/** PreSubmitViewFamilies event handler. It's used to initialize media on start. */
 	void OnPreSubmitViewFamilies(TArray<FSceneViewFamilyContext*>&);
 
@@ -78,15 +80,15 @@ private:
 	void OnEnginePreExit();
 
 private:
-	/** Active media devices that capture viewports */
-	TMap<FString, TSharedPtr<FDisplayClusterMediaCaptureViewport>> CaptureViewports;
 
-	/** Active media devices that capture local backbuffer */
-	TMap<FString, TSharedPtr<FDisplayClusterMediaCaptureNode>>     CaptureNode;
+	/** Active media capture devices */
+	TMap<FString, TSharedPtr<FDisplayClusterMediaCaptureBase>> CaptureDevices;
+
+	/** Active media input devices */
+	TMap<FString, TSharedPtr<FDisplayClusterMediaInputBase>> InputDevices;
 
-	/** Active media devices for viewport playback */
-	TMap<FString, TSharedPtr<FDisplayClusterMediaInputViewport>>   InputViewports;
+private:
 
-	// Latency queue
+	/** Latency queue */
 	FDisplayClusterFrameQueue FrameQueue;
 };
diff --git a/Engine/Plugins/Runtime/nDisplay/Source/DisplayClusterMedia/Private/Input/DisplayClusterMediaInputBase.cpp b/Engine/Plugins/Runtime/nDisplay/Source/DisplayClusterMedia/Private/Input/DisplayClusterMediaInputBase.cpp
--- a/Engine/Plugins/Runtime/nDisplay/Source/DisplayClusterMedia/Private/Input/DisplayClusterMediaInputBase.cpp
+++ b/Engine/Plugins/Runtime/nDisplay/Source/DisplayClusterMedia/Private/Input/DisplayClusterMediaInputBase.cpp
@@ -5,6 +5,8 @@
 #include "DisplayClusterMediaHelpers.h"
 #include "DisplayClusterMediaLog.h"
 
+#include "IDisplayClusterShaders.h"
+
 #include "MediaPlayer.h"
 #include "MediaSource.h"
 #include "MediaTexture.h"
@@ -15,10 +17,13 @@
 #include "OpenColorIORendering.h"
 
 #include "RenderGraphBuilder.h"
+#include "RenderGraphUtils.h"
 #include "ScreenPass.h"
 
 #include "Engine/Engine.h"
 
+#include "ShaderParameters/DisplayClusterShaderParameters_Media.h"
+
 #include "UObject/UObjectGlobals.h"
 #include "UObject/Package.h"
 
@@ -51,10 +56,9 @@ TAutoConsoleVariable<int32> CVarTempRivermaxExtraPixelsRemove(
 FDisplayClusterMediaInputBase::FDisplayClusterMediaInputBase(
 	const FString& InMediaId,
 	const FString& InClusterNodeId,
-	UMediaSource* InMediaSource,
-	bool bInLateOCIO
+	UMediaSource* InMediaSource
 )
-	: FDisplayClusterMediaBase(InMediaId, InClusterNodeId, bInLateOCIO)
+	: FDisplayClusterMediaBase(InMediaId, InClusterNodeId)
 {
 	checkSlow(InMediaSource);
 	MediaSource = DuplicateObject(InMediaSource, GetTransientPackage());
@@ -114,9 +118,6 @@ void FDisplayClusterMediaInputBase::Stop()
 		MediaPlayer->OnMediaEvent().RemoveAll(this);
 	}
 
-	// Release internals
-	ReleaseInternals();
-
 	bRunningRivermaxMedia = false;
 }
 
@@ -155,15 +156,11 @@ void FDisplayClusterMediaInputBase::OverrideTextureRegions_RenderThread(FIntRect
 	}
 }
 
-void FDisplayClusterMediaInputBase::ReleaseInternals()
-{
-	OCIOAppliedTexture.SafeRelease();
-}
-
 void FDisplayClusterMediaInputBase::ImportMediaData_RenderThread(FRHICommandListImmediate& RHICmdList, const FMediaInputTextureInfo& TextureInfo)
 {
 	UE_LOG(LogDisplayClusterMedia, Verbose, TEXT("MediaInput '%s': importing texture on RT frame '%llu'..."), *GetMediaId(), GFrameCounterRenderThread);
 
+	// Render media texture
 	MediaTexture->JustInTimeRender();
 
 	FRHITexture* SrcTexture = MediaTexture->GetResource() ? MediaTexture->GetResource()->GetTextureRHI() : nullptr;
@@ -171,140 +168,100 @@ void FDisplayClusterMediaInputBase::ImportMediaData_RenderThread(FRHICommandList
 
 	if (!SrcTexture || !DstTexture)
 	{
+		UE_LOG(LogDisplayClusterMedia, Warning, TEXT("MediaInput '%s': wrong texture on RT frame '%llu'..."), *GetMediaId(), GFrameCounterRenderThread);
 		return;
 	}
 
-	// Apply OCIO if needed
+	// [Temp workaround]
+	// There is an extra pixel issue in Rivermax. Allow to work around it.
+	FIntRect SrcRect{ { 0, 0 }, SrcTexture->GetDesc().Extent };
+	FIntRect DstRect{ TextureInfo.Region };
+	OverrideTextureRegions_RenderThread(SrcRect, DstRect);
+
+	// Process import
 	if (IsLateOCIO())
 	{
-		const bool bApplied = ProcessLateOCIO(RHICmdList, SrcTexture, TextureInfo.OCIOPassResources);
-		if (bApplied)
-		{
-			// Redirect SrcTexture to the intermediate OCIO texture so we'll be importing this new one
-			SrcTexture = OCIOAppliedTexture;
-		}
+		ImportMediaDataOCIO_RenderThread(RHICmdList, SrcTexture, SrcRect, DstTexture, DstRect, TextureInfo.OCIOPassResources);
 	}
-
-	if (SrcTexture && DstTexture)
+	else
 	{
-		FIntRect SrcRect(FIntPoint(0, 0), SrcTexture->GetDesc().Extent);
-		FIntRect DstRect(TextureInfo.Region);
-		OverrideTextureRegions_RenderThread(SrcRect, DstRect);
-
-		const bool bSrcSrgb = EnumHasAnyFlags(SrcTexture->GetFlags(), TexCreate_SRGB);
-		const bool bDstSrgb = EnumHasAnyFlags(DstTexture->GetFlags(), TexCreate_SRGB);
-
-		if ((SrcTexture->GetDesc().Format == DstTexture->GetDesc().Format)
-			&& (SrcRect.Size() == DstRect.Size())
-			&& (bSrcSrgb == bDstSrgb)
-			)
-		{
-			FRHICopyTextureInfo CopyInfo;
-			CopyInfo.SourcePosition = FIntVector(SrcRect.Min.X, SrcRect.Min.Y, 0);
-			CopyInfo.DestPosition = FIntVector(DstRect.Min.X, DstRect.Min.Y, 0);
-			CopyInfo.Size = FIntVector(DstRect.Size().X, DstRect.Size().Y, 0);
-
-			TransitionAndCopyTexture(RHICmdList, SrcTexture, DstTexture, CopyInfo);
-		}
-		else
-		{
-			DisplayClusterMediaHelpers::ResampleTexture_RenderThread(RHICmdList, SrcTexture, DstTexture, SrcRect, DstRect);
-		}
+		ImportMediaDataDirect_RenderThread(RHICmdList, SrcTexture, SrcRect, DstTexture, DstRect);
 	}
 }
 
-bool FDisplayClusterMediaInputBase::ProcessLateOCIO(FRHICommandListImmediate& RHICmdList, FRHITexture* SrcTexture, const FOpenColorIORenderPassResources& OCIORenderPassResources)
+void FDisplayClusterMediaInputBase::ImportMediaDataDirect_RenderThread(FRHICommandListImmediate& RHICmdList, FRHITexture* SrcTexture, const FIntRect& SrcRect, FRHITexture* DstTexture, const FIntRect& DstRect)
 {
-	checkSlow(SrcTexture);
-
-	if (!SrcTexture)
+	if (!SrcTexture || !DstTexture)
 	{
-		return false;
+		return;
 	}
 
-	if (!OCIORenderPassResources.IsValid())
-	{
-		return false;
-	}
+	const bool bSrcSrgb = EnumHasAnyFlags(SrcTexture->GetFlags(), TexCreate_SRGB);
+	const bool bDstSrgb = EnumHasAnyFlags(DstTexture->GetFlags(), TexCreate_SRGB);
 
-	// Create an intermediate texture if not exists
-	if (!OCIOAppliedTexture)
-	{
-		OCIOAppliedTexture = CreateTexture(SrcTexture);
-	}
-	// Re-create it if parameters don't match
-	else
+	const bool bSameSrgb = (bSrcSrgb == bDstSrgb);
+	const bool bSameFormat = (SrcTexture->GetDesc().Format == DstTexture->GetDesc().Format);
+	const bool bSameSize = (SrcRect.Size() == DstRect.Size());
+
+	const bool bCanCopy = (bSameFormat && bSameSize && bSameSrgb);
+
+	// Based on the texture properties, copy it directly or resample
+	if (bCanCopy)
 	{
-		const bool bSameFormat = (OCIOAppliedTexture->GetDesc().Format == SrcTexture->GetDesc().Format);
-		const bool bSameSize   = (OCIOAppliedTexture->GetDesc().Extent == SrcTexture->GetDesc().Extent);
+		FRHICopyTextureInfo CopyInfo;
+		CopyInfo.SourcePosition = FIntVector(SrcRect.Min.X, SrcRect.Min.Y, 0);
+		CopyInfo.DestPosition = FIntVector(DstRect.Min.X, DstRect.Min.Y, 0);
+		CopyInfo.Size = FIntVector(DstRect.Size().X, DstRect.Size().Y, 0);
 
-		if (!bSameFormat || !bSameSize)
-		{
-			OCIOAppliedTexture = CreateTexture(SrcTexture);
-		}
+		TransitionAndCopyTexture(RHICmdList, SrcTexture, DstTexture, CopyInfo);
 	}
-
-	if (!OCIOAppliedTexture)
+	else
 	{
-		return false;
+		DisplayClusterMediaHelpers::ResampleTexture_RenderThread(RHICmdList, SrcTexture, DstTexture, SrcRect, DstRect);
 	}
+}
 
+void FDisplayClusterMediaInputBase::ImportMediaDataOCIO_RenderThread(FRHICommandListImmediate& RHICmdList, FRHITexture* InSrcTexture, const FIntRect& InSrcRect, FRHITexture* InDstTexture, const FIntRect& InDstRect, const FOpenColorIORenderPassResources& OCIOResources)
+{
 	FRDGBuilder GraphBuilder(RHICmdList);
 
-	FRDGTextureRef InputTexture = RegisterExternalTexture(GraphBuilder, SrcTexture, TEXT("DCMediaLateOCIOTexIn"));
-	FRDGTextureRef OutputTexture = RegisterExternalTexture(GraphBuilder, OCIOAppliedTexture, TEXT("DCMediaLateOCIOTexOut"));
+	// Register RHI textures for further processing
+	FRDGTextureRef SrcTexture = RegisterExternalTexture(GraphBuilder, InSrcTexture, TEXT("DC.MediaTextureSrc"));
+	FRDGTextureRef DstTexture = RegisterExternalTexture(GraphBuilder, InDstTexture, TEXT("DC.MediaTextureDst"));
 
-	const FIntPoint OutputResolution = OCIOAppliedTexture->GetDesc().Extent;
-	const FIntRect  OutputRect = FIntRect(FIntPoint::ZeroValue, OutputResolution);
+	// A helper reference to the actual input texture
+	FRDGTextureRef InputTexture = SrcTexture;
 
-	FScreenPassTexture Input = FScreenPassTexture(InputTexture);
-	FScreenPassRenderTarget Output = FScreenPassRenderTarget(OutputTexture, OutputRect, ERenderTargetLoadAction::EClear);
-
-	FOpenColorIORendering::AddPass_RenderThread(
-		GraphBuilder,
-		FScreenPassViewInfo(),
-		GEngine->GetDefaultWorldFeatureLevel(),
-		Input,
-		Output,
-		OCIORenderPassResources,
-		1.0f,
-		EOpenColorIOTransformAlpha::None
-	);
+	// Is PQ-decode pass required?
+	if (IsTransferPQ())
+	{
+		// An intermediate texture to process PQ decoding and OCIO in linear space
+		FRDGTextureDesc InterimTextureDesc = FRDGTextureDesc::Create2D(InDstRect.Size(), PF_FloatRGBA, FClearValueBinding::Black, TexCreate_ShaderResource | TexCreate_RenderTargetable);
+		FRDGTextureRef TextureLinear = GraphBuilder.CreateTexture(InterimTextureDesc, TEXT("DC.MediaTextureTempLinear"));
 
-	GraphBuilder.Execute();
+		// Add PQ-decode pass
+		FDisplayClusterShaderParameters_MediaPQ Parameters{ SrcTexture, InSrcRect, TextureLinear, InDstRect };
+		IDisplayClusterShaders::Get().AddPQToLinearPass(GraphBuilder, Parameters);
 
-	return true;
-}
+		// Use this texture in the OCIO pass
+		InputTexture = TextureLinear;
+	}
 
-FTextureRHIRef FDisplayClusterMediaInputBase::CreateTexture(const FRHITexture* ReferenceTexture)
-{
-	if (!ReferenceTexture)
+	// Now apply OCIO and store to the destination
 	{
-		return nullptr;
+		FOpenColorIORendering::AddPass_RenderThread(
+			GraphBuilder,
+			FScreenPassViewInfo(),
+			GEngine->GetDefaultWorldFeatureLevel(),
+			FScreenPassTexture(InputTexture),
+			FScreenPassRenderTarget(DstTexture, InDstRect, ERenderTargetLoadAction::EClear),
+			OCIOResources,
+			1.0f,
+			EOpenColorIOTransformAlpha::None
+		);
 	}
 
-	// Use original format and size
-	const int32 SizeX = ReferenceTexture->GetDesc().Extent.X;
-	const int32 SizeY = ReferenceTexture->GetDesc().Extent.Y;
-	const EPixelFormat Format = ReferenceTexture->GetFormat();
-
-	// Prepare description
-	FRHITextureCreateDesc Desc =
-		FRHITextureCreateDesc::Create2D(TEXT("DisplayClusterFrameQueueCacheTexture"), SizeX, SizeY, Format)
-		.SetClearValue(FClearValueBinding::Black)
-		.SetNumMips(1)
-		.SetFlags(ETextureCreateFlags::Dynamic)
-		.AddFlags(ETextureCreateFlags::MultiGPUGraphIgnore)
-		.SetInitialState(ERHIAccess::SRVMask);
-
-	// Leave original flags, but make sure it's ResolveTargetable but not RenderTargetable
-	ETextureCreateFlags Flags = ReferenceTexture->GetFlags();
-	Flags &= ~ETextureCreateFlags::RenderTargetable;
-	Flags |= ETextureCreateFlags::ResolveTargetable;
-	Desc.SetFlags(Flags);
-
-	// Create texture
-	return RHICreateTexture(Desc);
+	GraphBuilder.Execute();
 }
 
 void FDisplayClusterMediaInputBase::OnMediaEvent(EMediaEvent MediaEvent)
diff --git a/Engine/Plugins/Runtime/nDisplay/Source/DisplayClusterMedia/Private/Input/DisplayClusterMediaInputBase.h b/Engine/Plugins/Runtime/nDisplay/Source/DisplayClusterMedia/Private/Input/DisplayClusterMediaInputBase.h
--- a/Engine/Plugins/Runtime/nDisplay/Source/DisplayClusterMedia/Private/Input/DisplayClusterMediaInputBase.h
+++ b/Engine/Plugins/Runtime/nDisplay/Source/DisplayClusterMedia/Private/Input/DisplayClusterMediaInputBase.h
@@ -3,10 +3,12 @@
 #pragma once
 
 #include "CoreMinimal.h"
-#include "DisplayClusterMediaBase.h"
+
+#include "Common/DisplayClusterMediaBase.h"
+#include "UObject/GCObject.h"
+
 #include "IMediaEventSink.h"
 #include "OpenColorIORendering.h"
-#include "UObject/GCObject.h"
 
 class FRHICommandListImmediate;
 class FRHITexture;
@@ -17,21 +19,22 @@ class UMediaTexture;
 
 
 /**
- * Base media input class
+ * Base media input adapter class
  */
 class FDisplayClusterMediaInputBase
 	: public FDisplayClusterMediaBase
 	, public FGCObject
 {
 public:
+
 	FDisplayClusterMediaInputBase(
 		const FString& MediaId,
 		const FString& ClusterNodeId,
-		UMediaSource* MediaSource,
-		bool bInLateOCIO = false
+		UMediaSource* MediaSource
 	);
 
 public:
+
 	//~ Begin FGCObject interface
 	virtual void AddReferencedObjects(FReferenceCollector& Collector) override;
 	virtual FString GetReferencerName() const override
@@ -90,11 +93,11 @@ protected:
 
 private:
 
-	/** Applies OCIO transformation to the SRC texture, and stores result into DST */
-	bool ProcessLateOCIO(FRHICommandListImmediate& RHICmdList, FRHITexture* SrcTexture, const FOpenColorIORenderPassResources& OCIORenderPassResources);
+	/** Import implementation for non-OCIO path */
+	void ImportMediaDataDirect_RenderThread(FRHICommandListImmediate& RHICmdList, FRHITexture* SrcTexture, const FIntRect& SrcRect, FRHITexture* DstTexture, const FIntRect& DstRect);
 
-	/** Creates an internal intermediate texture */
-	FTextureRHIRef CreateTexture(const FRHITexture* ReferenceTexture);
+	/** Import implementation for OCIO path */
+	void ImportMediaDataOCIO_RenderThread(FRHICommandListImmediate& RHICmdList, FRHITexture* SrcTexture, const FIntRect& SrcRect, FRHITexture* DstTexture, const FIntRect& DstRect, const FOpenColorIORenderPassResources& OCIOResources);
 
 private:
 
@@ -110,10 +113,8 @@ private:
 	/** [TEMP] A temporary workaround to cut off extra pixels for Rivermax input streams */
 	void OverrideTextureRegions_RenderThread(FIntRect& InOutSrcRect, FIntRect& InOutDstRect) const;
 
-	/** Releases internal resources */
-	void ReleaseInternals();
-
 private:
+
 	//~ Begin GC by AddReferencedObjects
 	TObjectPtr<UMediaSource>  MediaSource;
 	TObjectPtr<UMediaPlayer>  MediaPlayer;
@@ -126,9 +127,6 @@ private:
 	/** Used to control the rate at which we try to restart the player */
 	double LastRestartTimestamp = 0;
 
-	/** Intermediate texture for late OCIO pass */
-	FTextureRHIRef OCIOAppliedTexture;
-
 	/** [Temp workaround] Whether current media is Rivermax */
 	bool bRunningRivermaxMedia = false;
 };
diff --git a/Engine/Plugins/Runtime/nDisplay/Source/DisplayClusterMedia/Private/Input/DisplayClusterMediaInputCameraBase.cpp b/Engine/Plugins/Runtime/nDisplay/Source/DisplayClusterMedia/Private/Input/DisplayClusterMediaInputCameraBase.cpp
new file mode 100644
--- /dev/null
+++ b/Engine/Plugins/Runtime/nDisplay/Source/DisplayClusterMedia/Private/Input/DisplayClusterMediaInputCameraBase.cpp
@@ -0,0 +1,36 @@
+// Copyright Epic Games, Inc. All Rights Reserved.
+
+#include "Input/DisplayClusterMediaInputCameraBase.h"
+
+#include "Components/DisplayClusterICVFXCameraComponent.h"
+#include "Game/IDisplayClusterGameManager.h"
+#include "Render/Viewport/IDisplayClusterViewport.h"
+
+#include "DisplayClusterRootActor.h"
+#include "IDisplayCluster.h"
+
+#include "DisplayClusterMediaLog.h"
+
+
+FDisplayClusterMediaInputCameraBase::FDisplayClusterMediaInputCameraBase(
+	const FString& InMediaId,
+	const FString& InClusterNodeId,
+	const FString& InCameraId,
+	const FString& InCameraViewportId,
+	UMediaSource* InMediaSource
+)
+	: FDisplayClusterMediaInputViewportBase(InMediaId, InClusterNodeId, InCameraViewportId, InMediaSource)
+	, FDisplayClusterMediaCameraCommon(InCameraId)
+{
+}
+
+void FDisplayClusterMediaInputCameraBase::UpdateLateOCIOState(IDisplayClusterViewport* Viewport)
+{
+	FLateOCIOData NewLateOCIOConfiguration;
+
+	// Get current late OCIO parameters from the camera component
+	FDisplayClusterMediaCameraCommon::GetLateOCIOParameters(NewLateOCIOConfiguration.bLateOCIO, NewLateOCIOConfiguration.bTransferPQ);
+
+	// And store/update for current frame
+	SetLateOCIO(NewLateOCIOConfiguration);
+}
diff --git a/Engine/Plugins/Runtime/nDisplay/Source/DisplayClusterMedia/Private/Input/DisplayClusterMediaInputCameraBase.h b/Engine/Plugins/Runtime/nDisplay/Source/DisplayClusterMedia/Private/Input/DisplayClusterMediaInputCameraBase.h
new file mode 100644
--- /dev/null
+++ b/Engine/Plugins/Runtime/nDisplay/Source/DisplayClusterMedia/Private/Input/DisplayClusterMediaInputCameraBase.h
@@ -0,0 +1,32 @@
+// Copyright Epic Games, Inc. All Rights Reserved.
+
+#pragma once
+
+#include "Common/DisplayClusterMediaCameraCommon.h"
+#include "Input/DisplayClusterMediaInputViewportBase.h"
+
+class UDisplayClusterICVFXCameraComponent;
+
+
+/**
+ * Base camera input adapter
+ */
+class FDisplayClusterMediaInputCameraBase
+	: public FDisplayClusterMediaInputViewportBase
+	, public FDisplayClusterMediaCameraCommon
+{
+public:
+
+	FDisplayClusterMediaInputCameraBase(
+		const FString& MediaId,
+		const FString& ClusterNodeId,
+		const FString& CameraId,
+		const FString& CameraViewportId,
+		UMediaSource* MediaSource
+	);
+
+protected:
+
+	/** Updates late OCIO state */
+	virtual void UpdateLateOCIOState(IDisplayClusterViewport* Viewport) override;
+};
diff --git a/Engine/Plugins/Runtime/nDisplay/Source/DisplayClusterMedia/Private/Input/DisplayClusterMediaInputCameraFull.cpp b/Engine/Plugins/Runtime/nDisplay/Source/DisplayClusterMedia/Private/Input/DisplayClusterMediaInputCameraFull.cpp
new file mode 100644
--- /dev/null
+++ b/Engine/Plugins/Runtime/nDisplay/Source/DisplayClusterMedia/Private/Input/DisplayClusterMediaInputCameraFull.cpp
@@ -0,0 +1,21 @@
+// Copyright Epic Games, Inc. All Rights Reserved.
+
+#include "Input/DisplayClusterMediaInputCameraFull.h"
+
+#include "DisplayClusterMediaHelpers.h"
+
+
+FDisplayClusterMediaInputCameraFull::FDisplayClusterMediaInputCameraFull(
+	const FString& InMediaId,
+	const FString& InClusterNodeId,
+	const FString& InCameraId,
+	UMediaSource* InMediaSource
+)
+	: FDisplayClusterMediaInputCameraBase(
+		InMediaId,
+		InClusterNodeId,
+		InCameraId,
+		DisplayClusterMediaHelpers::GenerateICVFXViewportName(InClusterNodeId, InCameraId),
+		InMediaSource)
+{
+}
diff --git a/Engine/Plugins/Runtime/nDisplay/Source/DisplayClusterMedia/Private/Input/DisplayClusterMediaInputCameraFull.h b/Engine/Plugins/Runtime/nDisplay/Source/DisplayClusterMedia/Private/Input/DisplayClusterMediaInputCameraFull.h
new file mode 100644
--- /dev/null
+++ b/Engine/Plugins/Runtime/nDisplay/Source/DisplayClusterMedia/Private/Input/DisplayClusterMediaInputCameraFull.h
@@ -0,0 +1,22 @@
+// Copyright Epic Games, Inc. All Rights Reserved.
+
+#pragma once
+
+#include "Input/DisplayClusterMediaInputCameraBase.h"
+
+
+/**
+ * Camera input adapter (full frame)
+ */
+class FDisplayClusterMediaInputCameraFull
+	: public FDisplayClusterMediaInputCameraBase
+{
+public:
+
+	FDisplayClusterMediaInputCameraFull(
+		const FString& MediaId,
+		const FString& ClusterNodeId,
+		const FString& CameraId,
+		UMediaSource* MediaSource
+	);
+};
diff --git a/Engine/Plugins/Runtime/nDisplay/Source/DisplayClusterMedia/Private/Input/DisplayClusterMediaInputCameraTile.cpp b/Engine/Plugins/Runtime/nDisplay/Source/DisplayClusterMedia/Private/Input/DisplayClusterMediaInputCameraTile.cpp
new file mode 100644
--- /dev/null
+++ b/Engine/Plugins/Runtime/nDisplay/Source/DisplayClusterMedia/Private/Input/DisplayClusterMediaInputCameraTile.cpp
@@ -0,0 +1,22 @@
+// Copyright Epic Games, Inc. All Rights Reserved.
+
+#include "Input/DisplayClusterMediaInputCameraTile.h"
+
+#include "DisplayClusterMediaHelpers.h"
+
+
+FDisplayClusterMediaInputCameraTile::FDisplayClusterMediaInputCameraTile(
+	const FString& InMediaId,
+	const FString& InClusterNodeId,
+	const FString& InCameraId,
+	const FIntPoint& InTilePosition,
+	UMediaSource* InMediaSource
+)
+	: FDisplayClusterMediaInputCameraBase(
+		InMediaId,
+		InClusterNodeId,
+		InCameraId,
+		DisplayClusterMediaHelpers::GenerateICVFXTileViewportName(InClusterNodeId, InCameraId, InTilePosition),
+		InMediaSource)
+{
+}
diff --git a/Engine/Plugins/Runtime/nDisplay/Source/DisplayClusterMedia/Private/Input/DisplayClusterMediaInputCameraTile.h b/Engine/Plugins/Runtime/nDisplay/Source/DisplayClusterMedia/Private/Input/DisplayClusterMediaInputCameraTile.h
new file mode 100644
--- /dev/null
+++ b/Engine/Plugins/Runtime/nDisplay/Source/DisplayClusterMedia/Private/Input/DisplayClusterMediaInputCameraTile.h
@@ -0,0 +1,23 @@
+// Copyright Epic Games, Inc. All Rights Reserved.
+
+#pragma once
+
+#include "Input/DisplayClusterMediaInputCameraBase.h"
+
+
+/**
+ * Camera input adapter (tile)
+ */
+class FDisplayClusterMediaInputCameraTile
+	: public FDisplayClusterMediaInputCameraBase
+{
+public:
+
+	FDisplayClusterMediaInputCameraTile(
+		const FString& MediaId,
+		const FString& InClusterNodeId,
+		const FString& CameraId,
+		const FIntPoint& TilePosition,
+		UMediaSource* MediaSource
+	);
+};
diff --git a/Engine/Plugins/Runtime/nDisplay/Source/DisplayClusterMedia/Private/Input/DisplayClusterMediaInputViewport.cpp b/Engine/Plugins/Runtime/nDisplay/Source/DisplayClusterMedia/Private/Input/DisplayClusterMediaInputViewport.cpp
deleted file mode 100644
diff --git a/Engine/Plugins/Runtime/nDisplay/Source/DisplayClusterMedia/Private/Input/DisplayClusterMediaInputViewport.h b/Engine/Plugins/Runtime/nDisplay/Source/DisplayClusterMedia/Private/Input/DisplayClusterMediaInputViewport.h
deleted file mode 100644
diff --git a/Engine/Plugins/Runtime/nDisplay/Source/DisplayClusterMedia/Private/Input/DisplayClusterMediaInputViewportBase.cpp b/Engine/Plugins/Runtime/nDisplay/Source/DisplayClusterMedia/Private/Input/DisplayClusterMediaInputViewportBase.cpp
new file mode 100644
--- /dev/null
+++ b/Engine/Plugins/Runtime/nDisplay/Source/DisplayClusterMedia/Private/Input/DisplayClusterMediaInputViewportBase.cpp
@@ -0,0 +1,151 @@
+// Copyright Epic Games, Inc. All Rights Reserved.
+
+#include "Input/DisplayClusterMediaInputViewportBase.h"
+
+#include "Engine/Engine.h"
+#include "Engine/World.h"
+
+#include "Render/IDisplayClusterRenderManager.h"
+#include "Render/Viewport/IDisplayClusterViewportManager.h"
+#include "Render/Viewport/IDisplayClusterViewportManagerProxy.h"
+#include "Render/Viewport/IDisplayClusterViewport.h"
+
+#include "IDisplayCluster.h"
+#include "IDisplayClusterCallbacks.h"
+
+#include "OpenColorIOColorSpace.h"
+#include "RenderGraphUtils.h"
+
+#include "RHICommandList.h"
+#include "RHIFeatureLevel.h"
+#include "RHIResources.h"
+
+
+FDisplayClusterMediaInputViewportBase::FDisplayClusterMediaInputViewportBase(
+	const FString& InMediaId,
+	const FString& InClusterNodeId,
+	const FString& InViewportId,
+	UMediaSource* InMediaSource
+)
+	: FDisplayClusterMediaInputBase(InMediaId, InClusterNodeId, InMediaSource)
+	, ViewportId(InViewportId)
+{
+}
+
+
+bool FDisplayClusterMediaInputViewportBase::Play()
+{
+	// If playback has started successfully, subscribe for rendering callbacks
+	if (FDisplayClusterMediaInputBase::Play())
+	{
+		IDisplayCluster::Get().GetCallbacks().OnDisplayClusterPreSubmitViewFamilies().AddRaw(this, &FDisplayClusterMediaInputViewportBase::OnPreSubmitViewFamilies);
+		IDisplayCluster::Get().GetCallbacks().OnDisplayClusterPostCrossGpuTransfer_RenderThread().AddRaw(this, &FDisplayClusterMediaInputViewportBase::OnPostCrossGpuTransfer_RenderThread);
+		IDisplayCluster::Get().GetCallbacks().OnDisplayClusterUpdateViewportMediaState().AddRaw(this, &FDisplayClusterMediaInputViewportBase::OnUpdateViewportMediaState);
+
+		return true;
+	}
+
+	return false;
+}
+
+void FDisplayClusterMediaInputViewportBase::Stop()
+{
+	// Unsubscribe from external events/callbacks
+	IDisplayCluster::Get().GetCallbacks().OnDisplayClusterPreSubmitViewFamilies().RemoveAll(this);
+	IDisplayCluster::Get().GetCallbacks().OnDisplayClusterPostCrossGpuTransfer_RenderThread().RemoveAll(this);
+	IDisplayCluster::Get().GetCallbacks().OnDisplayClusterUpdateViewportMediaState().RemoveAll(this);
+
+	// Stop playing
+	FDisplayClusterMediaInputBase::Stop();
+}
+
+void FDisplayClusterMediaInputViewportBase::UpdateLateOCIOState(IDisplayClusterViewport* Viewport)
+{
+	// So far, regular viewports dont't support late OCIO feature. But it can be overrided down the hierarchy.
+}
+
+void FDisplayClusterMediaInputViewportBase::HandleLateOCIOChanged(const FLateOCIOData& NewLateOCIOConfiguration)
+{
+	// Restart playback
+	Stop();
+	Play();
+}
+
+void FDisplayClusterMediaInputViewportBase::OnPreSubmitViewFamilies(TArray<FSceneViewFamilyContext*>&)
+{
+	// Get OCIO settings if there are any
+	if (const IDisplayClusterViewportManager* const ViewportMgr = IDisplayCluster::Get().GetRenderMgr()->GetViewportManager())
+	{
+		if (const IDisplayClusterViewport* const Viewport = ViewportMgr->FindViewport(ViewportId))
+		{
+			// Get OCIO settings assigned to this viewport
+			FOpenColorIOColorConversionSettings OCIOConversionSettings;
+			Viewport->GetOCIOConversionSettings(OCIOConversionSettings);
+
+			const UWorld* const CurrentWorld = ViewportMgr->GetConfiguration().GetCurrentWorld();
+
+			const ERHIFeatureLevel::Type FeatureLevel = (CurrentWorld ?
+				CurrentWorld->GetFeatureLevel() :
+				GEngine->GetDefaultWorldFeatureLevel());
+
+			// Get OCIO render pass resources
+			FOpenColorIORenderPassResources OCIOPassResources = FOpenColorIORendering::GetRenderPassResources(OCIOConversionSettings, FeatureLevel);
+
+			// And push it to the rendering thread
+			ENQUEUE_RENDER_COMMAND(DCMediaInputUpdateOCIOResources)(
+				[this, InOCIOPassResources = MoveTemp(OCIOPassResources)](FRHICommandListImmediate& RHICmdList)
+				{
+					OCIOPassResources_RT = InOCIOPassResources;
+				}
+			);
+		}
+	}
+}
+
+void FDisplayClusterMediaInputViewportBase::OnUpdateViewportMediaState(IDisplayClusterViewport* InViewport, EDisplayClusterViewportMediaState& InOutMediaState)
+{
+	// Note: Media currently supports only one DCRA.
+	// In the future, after the media redesign, the DCRA name will also need to be checked here.
+	if (InViewport && InViewport->GetId().Equals(GetViewportId(), ESearchCase::IgnoreCase))
+	{
+		// Raise flags that this viewport texture will be overridden by media.
+		InOutMediaState |= EDisplayClusterViewportMediaState::Input;
+
+		// Update late OCIO state on current frame
+		UpdateLateOCIOState(InViewport);
+
+		// Late OCIO flag
+		if (IsLateOCIO())
+		{
+			InOutMediaState |= EDisplayClusterViewportMediaState::InputLateOCIO;
+		}
+	}
+}
+
+void FDisplayClusterMediaInputViewportBase::OnPostCrossGpuTransfer_RenderThread(FRHICommandListImmediate& RHICmdList, const IDisplayClusterViewportManagerProxy* ViewportManagerProxy, FViewport* Viewport)
+{
+	checkSlow(ViewportManagerProxy);
+
+	if (const IDisplayClusterViewportProxy* const PlaybackViewport = ViewportManagerProxy->FindViewport_RenderThread(GetViewportId()))
+	{
+		const bool bShouldImportMedia = !PlaybackViewport->GetPostRenderSettings_RenderThread().Replace.IsEnabled();
+
+		if (bShouldImportMedia)
+		{
+			TArray<FRHITexture*> Textures;
+			TArray<FIntRect>     Regions;
+
+			if (PlaybackViewport->GetResourcesWithRects_RenderThread(EDisplayClusterViewportResourceType::InternalRenderTargetResource, Textures, Regions))
+			{
+				if (Textures.Num() > 0 && Regions.Num() > 0 && Textures[0])
+				{
+					// Prepare request data
+					FMediaInputTextureInfo TextureInfo{ Textures[0], Regions[0], MoveTemp(OCIOPassResources_RT) };
+
+					// Import texture from media input
+					ImportMediaData_RenderThread(RHICmdList, TextureInfo);
+				}
+			}
+		}
+	}
+}
diff --git a/Engine/Plugins/Runtime/nDisplay/Source/DisplayClusterMedia/Private/Input/DisplayClusterMediaInputViewportBase.h b/Engine/Plugins/Runtime/nDisplay/Source/DisplayClusterMedia/Private/Input/DisplayClusterMediaInputViewportBase.h
new file mode 100644
--- /dev/null
+++ b/Engine/Plugins/Runtime/nDisplay/Source/DisplayClusterMedia/Private/Input/DisplayClusterMediaInputViewportBase.h
@@ -0,0 +1,70 @@
+// Copyright Epic Games, Inc. All Rights Reserved.
+
+#pragma once
+
+#include "Input/DisplayClusterMediaInputBase.h"
+#include "Render/Viewport/Containers/DisplayClusterViewport_Enums.h"
+
+#include "OpenColorIORendering.h"
+
+class FRHICommandListImmediate;
+class FViewport;
+class FSceneViewFamilyContext;
+class IDisplayClusterViewport;
+class IDisplayClusterViewportManagerProxy;
+
+
+/**
+ * Base viewport input adapter
+ */
+class FDisplayClusterMediaInputViewportBase
+	: public FDisplayClusterMediaInputBase
+{
+public:
+	FDisplayClusterMediaInputViewportBase(
+		const FString& MediaId,
+		const FString& ClusterNodeId,
+		const FString& ViewportId,
+		UMediaSource* MediaSource
+	);
+
+public:
+	/** Start playback */
+	virtual bool Play() override;
+
+	/** Stop playback */
+	virtual void Stop() override;
+
+	/** Returns viewport ID bound for playback */
+	const FString& GetViewportId() const
+	{
+		return ViewportId;
+	}
+
+protected:
+
+	/** Updates late OCIO state */
+	virtual void UpdateLateOCIOState(IDisplayClusterViewport* Viewport);
+
+	/** Process late OCIO state updates */
+	virtual void HandleLateOCIOChanged(const FLateOCIOData& NewLateOCIOConfiguration) override;
+
+private:
+
+	/** PreSubmitViewFamilies event handler. It's used to initialize media on start. */
+	void OnPreSubmitViewFamilies(TArray<FSceneViewFamilyContext*>&);
+
+	/** PostCrossGpuTransfer callback handler where media data is pushed into nDisplay internal buffers */
+	void OnPostCrossGpuTransfer_RenderThread(FRHICommandListImmediate& RHICmdList, const IDisplayClusterViewportManagerProxy* ViewportManagerProxy, FViewport* Viewport);
+
+	/** UpdateViewportMediaState callback to configure media state for a viewoprt */
+	void OnUpdateViewportMediaState(IDisplayClusterViewport* InViewport, EDisplayClusterViewportMediaState& InOutMediaState);
+
+private:
+
+	/** Viewport ID assigned for this media input */
+	const FString ViewportId;
+
+	/** OCIO conversion pass resources (render thread data) */
+	FOpenColorIORenderPassResources OCIOPassResources_RT;
+};
diff --git a/Engine/Plugins/Runtime/nDisplay/Source/DisplayClusterMedia/Private/Input/DisplayClusterMediaInputViewportFull.cpp b/Engine/Plugins/Runtime/nDisplay/Source/DisplayClusterMedia/Private/Input/DisplayClusterMediaInputViewportFull.cpp
new file mode 100644
--- /dev/null
+++ b/Engine/Plugins/Runtime/nDisplay/Source/DisplayClusterMedia/Private/Input/DisplayClusterMediaInputViewportFull.cpp
@@ -0,0 +1,14 @@
+// Copyright Epic Games, Inc. All Rights Reserved.
+
+#include "Input/DisplayClusterMediaInputViewportFull.h"
+
+
+FDisplayClusterMediaInputViewportFull::FDisplayClusterMediaInputViewportFull(
+	const FString& InMediaId,
+	const FString& InClusterNodeId,
+	const FString& InViewportId,
+	UMediaSource* InMediaSource
+)
+	: FDisplayClusterMediaInputViewportBase(InMediaId, InClusterNodeId, InViewportId, InMediaSource)
+{
+}
diff --git a/Engine/Plugins/Runtime/nDisplay/Source/DisplayClusterMedia/Private/Input/DisplayClusterMediaInputViewportFull.h b/Engine/Plugins/Runtime/nDisplay/Source/DisplayClusterMedia/Private/Input/DisplayClusterMediaInputViewportFull.h
new file mode 100644
--- /dev/null
+++ b/Engine/Plugins/Runtime/nDisplay/Source/DisplayClusterMedia/Private/Input/DisplayClusterMediaInputViewportFull.h
@@ -0,0 +1,21 @@
+// Copyright Epic Games, Inc. All Rights Reserved.
+
+#pragma once
+
+#include "Input/DisplayClusterMediaInputViewportBase.h"
+
+
+/**
+ * Viewport input adapter (full frame)
+ */
+class FDisplayClusterMediaInputViewportFull
+	: public FDisplayClusterMediaInputViewportBase
+{
+public:
+	FDisplayClusterMediaInputViewportFull(
+		const FString& MediaId,
+		const FString& ClusterNodeId,
+		const FString& ViewportId,
+		UMediaSource* MediaSource
+	);
+};
diff --git a/Engine/Plugins/Runtime/nDisplay/Source/DisplayClusterMedia/Private/Input/DisplayClusterMediaInputViewportTile.cpp b/Engine/Plugins/Runtime/nDisplay/Source/DisplayClusterMedia/Private/Input/DisplayClusterMediaInputViewportTile.cpp
new file mode 100644
--- /dev/null
+++ b/Engine/Plugins/Runtime/nDisplay/Source/DisplayClusterMedia/Private/Input/DisplayClusterMediaInputViewportTile.cpp
@@ -0,0 +1,21 @@
+// Copyright Epic Games, Inc. All Rights Reserved.
+
+#include "Input/DisplayClusterMediaInputViewportTile.h"
+
+#include "DisplayClusterMediaHelpers.h"
+
+
+FDisplayClusterMediaInputViewportTile::FDisplayClusterMediaInputViewportTile(
+	const FString& InMediaId,
+	const FString& InClusterNodeId,
+	const FString& InViewportId,
+	const FIntPoint& InTilePosition,
+	UMediaSource* InMediaSource
+)
+	: FDisplayClusterMediaInputViewportBase(
+		InMediaId,
+		InClusterNodeId,
+		DisplayClusterMediaHelpers::GenerateTileViewportName(InViewportId, InTilePosition),
+		InMediaSource)
+{
+}
diff --git a/Engine/Plugins/Runtime/nDisplay/Source/DisplayClusterMedia/Private/Input/DisplayClusterMediaInputViewportTile.h b/Engine/Plugins/Runtime/nDisplay/Source/DisplayClusterMedia/Private/Input/DisplayClusterMediaInputViewportTile.h
new file mode 100644
--- /dev/null
+++ b/Engine/Plugins/Runtime/nDisplay/Source/DisplayClusterMedia/Private/Input/DisplayClusterMediaInputViewportTile.h
@@ -0,0 +1,22 @@
+// Copyright Epic Games, Inc. All Rights Reserved.
+
+#pragma once
+
+#include "Input/DisplayClusterMediaInputViewportBase.h"
+
+
+/**
+ * Viewport input adapter (tile)
+ */
+class FDisplayClusterMediaInputViewportTile
+	: public FDisplayClusterMediaInputViewportBase
+{
+public:
+	FDisplayClusterMediaInputViewportTile(
+		const FString& MediaId,
+		const FString& ClusterNodeId,
+		const FString& ViewportId,
+		const FIntPoint& TilePosition,
+		UMediaSource* MediaSource
+	);
+};
diff --git a/Engine/Plugins/Runtime/nDisplay/Source/DisplayClusterShaders/Private/DisplayClusterShadersModule.cpp b/Engine/Plugins/Runtime/nDisplay/Source/DisplayClusterShaders/Private/DisplayClusterShadersModule.cpp
--- a/Engine/Plugins/Runtime/nDisplay/Source/DisplayClusterShaders/Private/DisplayClusterShadersModule.cpp
+++ b/Engine/Plugins/Runtime/nDisplay/Source/DisplayClusterShaders/Private/DisplayClusterShadersModule.cpp
@@ -2,16 +2,18 @@
 
 #include "DisplayClusterShadersModule.h"
 
+#include "ShaderCore.h"
+
 #include "Interfaces/IPluginManager.h"
 #include "Misc/Paths.h"
-#include "ShaderCore.h"
 
+#include "Shaders/DisplayClusterShadersGenerateMips.h"
+#include "Shaders/DisplayClusterShadersMedia.h"
 #include "Shaders/DisplayClusterShadersPreprocess_UVLightCards.h"
 #include "Shaders/DisplayClusterShadersPostprocess_Blur.h"
 #include "Shaders/DisplayClusterShadersPostprocess_OutputRemap.h"
 #include "Shaders/DisplayClusterShadersWarpblend_ICVFX.h"
 #include "Shaders/DisplayClusterShadersWarpblend_MPCDI.h"
-#include "Shaders/DisplayClusterShadersGenerateMips.h"
 
 //////////////////////////////////////////////////////////////////////////////////////////////
 // IModuleInterface
@@ -62,6 +64,15 @@ bool FDisplayClusterShadersModule::GenerateMips(FRHICommandListImmediate& RHICmd
 	return FDisplayClusterShadersGenerateMips::GenerateMips(RHICmdList, InOutMipsTexture, InSettings);
 }
 
+void FDisplayClusterShadersModule::AddLinearToPQPass(FRDGBuilder& GraphBuilder, const FDisplayClusterShaderParameters_MediaPQ& Parameters) const
+{
+	FDisplayClusterShadersMedia::AddLinearToPQPass(GraphBuilder, Parameters);
+}
+
+void FDisplayClusterShadersModule::AddPQToLinearPass(FRDGBuilder& GraphBuilder, const FDisplayClusterShaderParameters_MediaPQ& Parameters) const
+{
+	FDisplayClusterShadersMedia::AddPQToLinearPass(GraphBuilder, Parameters);
+}
 
 IMPLEMENT_MODULE(FDisplayClusterShadersModule, DisplayClusterShaders);
 
diff --git a/Engine/Plugins/Runtime/nDisplay/Source/DisplayClusterShaders/Private/DisplayClusterShadersModule.h b/Engine/Plugins/Runtime/nDisplay/Source/DisplayClusterShaders/Private/DisplayClusterShadersModule.h
--- a/Engine/Plugins/Runtime/nDisplay/Source/DisplayClusterShaders/Private/DisplayClusterShadersModule.h
+++ b/Engine/Plugins/Runtime/nDisplay/Source/DisplayClusterShaders/Private/DisplayClusterShadersModule.h
@@ -22,4 +22,6 @@ public:
 	virtual bool RenderPostprocess_OutputRemap(FRHICommandListImmediate& RHICmdList, FRHITexture* InSourceTexture, FRHITexture* InRenderTargetableDestTexture, const IDisplayClusterRender_MeshComponentProxy& MeshProxy) const override;
 	virtual bool RenderPostprocess_Blur(FRHICommandListImmediate& RHICmdList, FRHITexture* InSourceTexture, FRHITexture* InRenderTargetableDestTexture, const FDisplayClusterShaderParameters_PostprocessBlur& InSettings) const override;
 	virtual bool GenerateMips(FRHICommandListImmediate& RHICmdList, FRHITexture* InOutMipsTexture, const FDisplayClusterShaderParameters_GenerateMips& InSettings) const override;
+	virtual void AddLinearToPQPass(FRDGBuilder& GraphBuilder, const FDisplayClusterShaderParameters_MediaPQ& Parameters) const override;
+	virtual void AddPQToLinearPass(FRDGBuilder& GraphBuilder, const FDisplayClusterShaderParameters_MediaPQ& Parameters) const override;
 };
diff --git a/Engine/Plugins/Runtime/nDisplay/Source/DisplayClusterShaders/Private/Shaders/DisplayClusterShadersMedia.cpp b/Engine/Plugins/Runtime/nDisplay/Source/DisplayClusterShaders/Private/Shaders/DisplayClusterShadersMedia.cpp
new file mode 100644
--- /dev/null
+++ b/Engine/Plugins/Runtime/nDisplay/Source/DisplayClusterShaders/Private/Shaders/DisplayClusterShadersMedia.cpp
@@ -0,0 +1,137 @@
+// Copyright Epic Games, Inc. All Rights Reserved.
+
+#include "Shaders/DisplayClusterShadersMedia.h"
+#include "ShaderParameters/DisplayClusterShaderParameters_Media.h"
+
+#include "PixelFormat.h"
+#include "RenderGraphBuilder.h"
+#include "RenderGraphUtils.h"
+#include "RHIResources.h"
+#include "RHIStaticStates.h"
+#include "SceneView.h"
+#include "ScreenPass.h"
+#include "TextureResource.h"
+
+#include "DisplayClusterShadersLog.h"
+
+
+namespace UE::DisplayClusterShaders::Private
+{
+	static const FString MediaShadersPath = TEXT("/Plugin/nDisplay/Private/MediaShaders.usf");
+
+	///////////////////////////////////////////////////////////////////////////////////
+	// Linear-To-PQ
+	class FLinearToPQPS : public FGlobalShader
+	{
+		DECLARE_SHADER_TYPE(FLinearToPQPS, Global);
+		SHADER_USE_PARAMETER_STRUCT(FLinearToPQPS, FGlobalShader);
+
+		BEGIN_SHADER_PARAMETER_STRUCT(FParameters, )
+			SHADER_PARAMETER_RDG_TEXTURE(Texture2D, InputTexture)
+			SHADER_PARAMETER_SAMPLER(SamplerState, InputSampler)
+			RENDER_TARGET_BINDING_SLOTS()
+		END_SHADER_PARAMETER_STRUCT()
+
+	public:
+
+		static bool ShouldCompilePermutation(const FGlobalShaderPermutationParameters& InParameters)
+		{
+			return true;
+		}
+
+		FParameters* AllocateAndSetParameters(FRDGBuilder& InGraphBuilder, FRDGTextureRef InLinearTexture, FRDGTextureRef InOutputPQTexture)
+		{
+			FParameters* Parameters = InGraphBuilder.AllocParameters<FParameters>();
+			Parameters->InputTexture = InLinearTexture;
+			Parameters->InputSampler = TStaticSamplerState<>::GetRHI();
+			Parameters->RenderTargets[0] = FRenderTargetBinding{ InOutputPQTexture, ERenderTargetLoadAction::ENoAction };
+			return Parameters;
+		}
+	};
+
+	IMPLEMENT_SHADER_TYPE(, FLinearToPQPS, *MediaShadersPath, TEXT("LinearToPQ_PS"), SF_Pixel);
+
+
+	///////////////////////////////////////////////////////////////////////////////////
+	// PQ-To-Linear
+	class FPQToLinearPS : public FGlobalShader
+	{
+		DECLARE_SHADER_TYPE(FPQToLinearPS, Global);
+		SHADER_USE_PARAMETER_STRUCT(FPQToLinearPS, FGlobalShader);
+
+		BEGIN_SHADER_PARAMETER_STRUCT(FParameters, )
+			SHADER_PARAMETER_RDG_TEXTURE(Texture2D, InputTexture)
+			SHADER_PARAMETER_SAMPLER(SamplerState, InputSampler)
+			RENDER_TARGET_BINDING_SLOTS()
+		END_SHADER_PARAMETER_STRUCT()
+
+	public:
+
+		static bool ShouldCompilePermutation(const FGlobalShaderPermutationParameters& InParameters)
+		{
+			return true;
+		}
+
+		FParameters* AllocateAndSetParameters(FRDGBuilder& InGraphBuilder, FRDGTextureRef InPQTexture, FRDGTextureRef InOutputLinearTexture)
+		{
+			FParameters* Parameters = InGraphBuilder.AllocParameters<FParameters>();
+			Parameters->InputTexture = InPQTexture;
+			Parameters->InputSampler = TStaticSamplerState<>::GetRHI();
+			Parameters->RenderTargets[0] = FRenderTargetBinding{ InOutputLinearTexture, ERenderTargetLoadAction::ENoAction };
+			return Parameters;
+		}
+	};
+
+	IMPLEMENT_SHADER_TYPE(, FPQToLinearPS, *MediaShadersPath, TEXT("PQToLinear_PS"), SF_Pixel);
+
+
+	/** Template function to handle both Linear-PQ and PQ-Linear conversions */
+	template<typename TPixelShaderType>
+	void AddPQPass(FRDGBuilder& GraphBuilder, const FString& PassName, const FDisplayClusterShaderParameters_MediaPQ& Parameters)
+	{
+		// Rectangle area to use from source
+		const FIntRect ViewRect(Parameters.InputRect);
+
+		//Dummy ViewFamily/ViewInfo created to use built in Draw Screen/Texture Pass
+		const FSceneViewFamily ViewFamily(FSceneViewFamily::ConstructionValues(nullptr, nullptr, FEngineShowFlags(ESFIM_Game)).SetTime(FGameTime()));
+		FSceneViewInitOptions ViewInitOptions;
+		ViewInitOptions.ViewFamily = &ViewFamily;
+		ViewInitOptions.SetViewRectangle(ViewRect);
+		ViewInitOptions.ViewOrigin = FVector::ZeroVector;
+		ViewInitOptions.ViewRotationMatrix = FMatrix::Identity;
+		ViewInitOptions.ProjectionMatrix = FMatrix::Identity;
+		FSceneView ViewInfo(ViewInitOptions);
+
+		const FGlobalShaderMap* GlobalShaderMap = GetGlobalShaderMap(GMaxRHIFeatureLevel);
+		const TShaderMapRef<FScreenPassVS> VertexShader(GlobalShaderMap);
+		const TShaderMapRef<TPixelShaderType> PixelShader(GlobalShaderMap);
+		class TPixelShaderType::FParameters* PixelShaderParameters = PixelShader->AllocateAndSetParameters(GraphBuilder, Parameters.InputTexture, Parameters.OutputTexture);
+
+		AddDrawScreenPass(
+			GraphBuilder,
+			FRDGEventName(*PassName),
+			ViewInfo,
+			FScreenPassTextureViewport{ Parameters.OutputTexture, Parameters.OutputRect },
+			FScreenPassTextureViewport{ Parameters.InputTexture,  Parameters.InputRect },
+			VertexShader,
+			PixelShader,
+			PixelShaderParameters
+		);
+	}
+}
+
+void FDisplayClusterShadersMedia::AddLinearToPQPass(FRDGBuilder& GraphBuilder, const FDisplayClusterShaderParameters_MediaPQ& Parameters)
+{
+	using namespace UE::DisplayClusterShaders::Private;
+
+	// Add Linear-To-PQ pass
+	AddPQPass<FLinearToPQPS>(GraphBuilder, TEXT("FDCShadersMedia::LinearToPQ"), Parameters);
+}
+
+void FDisplayClusterShadersMedia::AddPQToLinearPass(FRDGBuilder& GraphBuilder, const FDisplayClusterShaderParameters_MediaPQ& Parameters)
+{
+	using namespace UE::DisplayClusterShaders::Private;
+
+	// Add Pq-To-Linear pass
+	AddPQPass<FPQToLinearPS>(GraphBuilder, TEXT("FDCShadersMedia::PQToLinear"), Parameters);
+}
diff --git a/Engine/Plugins/Runtime/nDisplay/Source/DisplayClusterShaders/Private/Shaders/DisplayClusterShadersMedia.h b/Engine/Plugins/Runtime/nDisplay/Source/DisplayClusterShaders/Private/Shaders/DisplayClusterShadersMedia.h
new file mode 100644
--- /dev/null
+++ b/Engine/Plugins/Runtime/nDisplay/Source/DisplayClusterShaders/Private/Shaders/DisplayClusterShadersMedia.h
@@ -0,0 +1,23 @@
+// Copyright Epic Games, Inc. All Rights Reserved.
+
+#pragma once
+
+#include "CoreMinimal.h"
+
+class FRDGBuilder;
+struct FDisplayClusterShaderParameters_MediaPQ;
+
+
+/**
+ * API for media shaders
+ */
+class FDisplayClusterShadersMedia
+{
+public:
+
+	/** Adds Linear-To-PQ encoding pass (API wrapper) */
+	static void AddLinearToPQPass(FRDGBuilder& GraphBuilder, const FDisplayClusterShaderParameters_MediaPQ& Parameters);
+
+	/** Adds PQ-To-Linear decoding pass (API wrapper) */
+	static void AddPQToLinearPass(FRDGBuilder& GraphBuilder, const FDisplayClusterShaderParameters_MediaPQ& Parameters);
+};
diff --git a/Engine/Plugins/Runtime/nDisplay/Source/DisplayClusterShaders/Public/IDisplayClusterShaders.h b/Engine/Plugins/Runtime/nDisplay/Source/DisplayClusterShaders/Public/IDisplayClusterShaders.h
--- a/Engine/Plugins/Runtime/nDisplay/Source/DisplayClusterShaders/Public/IDisplayClusterShaders.h
+++ b/Engine/Plugins/Runtime/nDisplay/Source/DisplayClusterShaders/Public/IDisplayClusterShaders.h
@@ -13,13 +13,16 @@ class IDisplayClusterRender_MeshComponentProxy;
 class FDisplayClusterShaderParameters_ICVFX;
 class FDisplayClusterShaderParameters_Override;
 
+class FSceneInterface;
+class FRDGBuilder;
+class FRenderTarget;
+
 struct FDisplayClusterShaderParameters_GenerateMips;
+struct FDisplayClusterShaderParameters_MediaPQ;
 struct FDisplayClusterShaderParameters_PostprocessBlur;
 struct FDisplayClusterShaderParameters_WarpBlend;
 struct FDisplayClusterShaderParameters_UVLightCards;
 
-class FSceneInterface;
-class FRenderTarget;
 
 class IDisplayClusterShaders : public IModuleInterface
 {
@@ -132,4 +135,20 @@ public:
 	* @return - true if success
 	*/
 	virtual bool GenerateMips(FRHICommandListImmediate& RHICmdList, FRHITexture* InOutMipsTexture, const FDisplayClusterShaderParameters_GenerateMips& InSettings) const = 0;
+
+	/**
+	* Adds Linear-To-PQ encoding pass
+	*
+	* @param GraphBuilder - RDG builder
+	* @param Parameters   - Conversion parameters
+	*/
+	virtual void AddLinearToPQPass(FRDGBuilder& GraphBuilder, const FDisplayClusterShaderParameters_MediaPQ& Parameters) const = 0;
+
+	/**
+	* Adds PQ-To-Linear decoding pass (API wrapper)
+	*
+	* @param GraphBuilder - RDG builder
+	* @param Parameters   - Conversion parameters
+	*/
+	virtual void AddPQToLinearPass(FRDGBuilder& GraphBuilder, const FDisplayClusterShaderParameters_MediaPQ& Parameters) const = 0;
 };
diff --git a/Engine/Plugins/Runtime/nDisplay/Source/DisplayClusterShaders/Public/ShaderParameters/DisplayClusterShaderParameters_Media.h b/Engine/Plugins/Runtime/nDisplay/Source/DisplayClusterShaders/Public/ShaderParameters/DisplayClusterShaderParameters_Media.h
new file mode 100644
--- /dev/null
+++ b/Engine/Plugins/Runtime/nDisplay/Source/DisplayClusterShaders/Public/ShaderParameters/DisplayClusterShaderParameters_Media.h
@@ -0,0 +1,26 @@
+// Copyright Epic Games, Inc. All Rights Reserved.
+
+#pragma once
+
+#include "CoreMinimal.h"
+
+class FRDGTexture;
+
+
+/*
+ * A container to pass PQ coding & decoding parameters
+ */
+struct FDisplayClusterShaderParameters_MediaPQ
+{
+	/** Source texture for coding/decoding */
+	FRDGTexture* InputTexture = nullptr;
+
+	/** Source texture sub-region */
+	FIntRect InputRect = { FIntPoint::ZeroValue, FIntPoint::ZeroValue };
+
+	/** Output texture for coding/decoding */
+	FRDGTexture* OutputTexture = nullptr;
+
+	/** Output texture sub-region */
+	FIntRect OutputRect = { FIntPoint::ZeroValue, FIntPoint::ZeroValue };
+};
