[nDisplay] Suport for output backbuffer tiling

Pure tiled capture functionality. No tile configuration dialog or media configuration panel yet. Those will come over later in separate CLs. Reviewed at CL 34811572.

[RN] The node backbuffer can be captured as MxN tiles

diff --git a/Engine/Plugins/Runtime/nDisplay/Source/DisplayClusterConfiguration/Public/DisplayClusterConfigurationTypes_Media.h b/Engine/Plugins/Runtime/nDisplay/Source/DisplayClusterConfiguration/Public/DisplayClusterConfigurationTypes_Media.h
--- a/Engine/Plugins/Runtime/nDisplay/Source/DisplayClusterConfiguration/Public/DisplayClusterConfigurationTypes_Media.h
+++ b/Engine/Plugins/Runtime/nDisplay/Source/DisplayClusterConfiguration/Public/DisplayClusterConfigurationTypes_Media.h
@@ -8,13 +8,13 @@
 #include "DisplayClusterConfigurationTypes_MediaSync.h"
 #include "DisplayClusterConfigurationTypes_Tile.h"
 
-#include "Render/Viewport/Containers/DisplayClusterViewport_TileSettings.h"
-
 #include "MediaPlayer.h"
 #include "MediaSource.h"
 #include "MediaTexture.h"
 #include "MediaOutput.h"
 
+#include "Render/Viewport/Containers/DisplayClusterViewport_TileSettings.h"
+
 #include "DisplayClusterConfigurationTypes_Media.generated.h"
 
 
@@ -62,84 +62,6 @@
 	TObjectPtr<UDisplayClusterMediaOutputSynchronizationPolicy> OutputSyncPolicy;
 };
 
-/*
- * Media settings for node backbuffer
- */
-USTRUCT(Blueprintable)
-struct DISPLAYCLUSTERCONFIGURATION_API FDisplayClusterConfigurationMediaNodeBackbuffer
-{
-	GENERATED_BODY()
-
-public:
-	/** Enable/disable media */
-	UPROPERTY(EditAnywhere, BlueprintReadWrite, Category = "Media", meta = (DisplayName = "Enable Media"))
-	bool bEnable = false;
-
-	/** Media outputs to use */
-	UPROPERTY(EditAnywhere, BlueprintReadWrite, Category = "Media")
-	TArray<FDisplayClusterConfigurationMediaOutput> MediaOutputs;
-
-public:
-
-	/** Returns true if at least one media output assigned */
-	bool IsMediaOutputAssigned() const;
-};
-
-
-PRAGMA_DISABLE_DEPRECATION_WARNINGS
 
 /*
- * Media settings for viewports
- */
-USTRUCT(Blueprintable)
-struct DISPLAYCLUSTERCONFIGURATION_API FDisplayClusterConfigurationMediaViewport
-{
-	GENERATED_BODY()
-
-public:
-	/** Enable/disable media */
-	UPROPERTY(EditAnywhere, BlueprintReadWrite, Category = "Media", meta=(DisplayName="Enable Media"))
-	bool bEnable = false;
-
-	/** Media source to use */
-	UPROPERTY(EditAnywhere, BlueprintReadWrite, Category = "Media")
-	FDisplayClusterConfigurationMediaInput MediaInput;
-
-	/** Media outputs to use */
-	UPROPERTY(EditAnywhere, BlueprintReadWrite, Category = "Media")
-	TArray<FDisplayClusterConfigurationMediaOutput> MediaOutputs;
-
-public:
-	/** Returns true if a media source assigned */
-	bool IsMediaInputAssigned() const;
-
-	/** Returns true if at least one media output assigned */
-	bool IsMediaOutputAssigned() const;
-
-#if WITH_EDITORONLY_DATA
-
-protected:
-	UPROPERTY(meta = (DeprecatedProperty, DeprecationMessage = "This property has been deprecated"))
-	FString MediaSharingNode_DEPRECATED;
-
-public:
-	UE_DEPRECATED(5.3, "This property has been deprecated. Please refer new MediaInput property.")
-	UPROPERTY(meta = (DeprecatedProperty, DeprecationMessage = "This property has been deprecated. Please refer new MediaInput property."))
-	TObjectPtr<UMediaSource> MediaSource;
-
-	UE_DEPRECATED(5.3, "This property has been deprecated. Please refer new MediaOutputs property.")
-	UPROPERTY(meta = (DeprecatedProperty, DeprecationMessage = "This property has been deprecated. Please refer new MediaOutputs property."))
-	TObjectPtr<UMediaOutput> MediaOutput;
-
-	UE_DEPRECATED(5.3, "This property has been deprecated.")
-	UPROPERTY(meta = (DeprecatedProperty, DeprecationMessage = "This property has been deprecated."))
-	TObjectPtr<UDisplayClusterMediaOutputSynchronizationPolicy> OutputSyncPolicy;
-
-#endif // WITH_EDITORONLY_DATA
-};
-
-PRAGMA_ENABLE_DEPRECATION_WARNINGS
-
-
-/*
  * Media input group (ICVFX, Full frame)
@@ -208,6 +130,10 @@
 	/** Media output to use */
 	UPROPERTY(Instanced, EditAnywhere, BlueprintReadWrite, Category = "Media")
 	TObjectPtr<UMediaOutput> MediaOutput;
+
+	/** Media output synchronization policy */
+	UPROPERTY(Instanced, EditAnywhere, BlueprintReadWrite, Category = "Media", meta = (DisplayName = "Capture Synchronization"))
+	TObjectPtr<UDisplayClusterMediaOutputSynchronizationPolicy> OutputSyncPolicy;
 };
 
 
@@ -248,10 +174,98 @@
 	TArray<FDisplayClusterConfigurationMediaUniformTileOutput> Tiles;
 
 	/** Media output synchronization policy */
-	UPROPERTY(Instanced, EditAnywhere, BlueprintReadWrite, Category = "Media", meta = (DisplayName = "Capture Synchronization"))
+	UE_DEPRECATED(5.6, "This property has been deprecated. Please refer new MediaOutputs property.")
+	UPROPERTY(meta = (DeprecatedProperty, DeprecationMessage = "This property has been deprecated. Synchronization policy is bound to the actual media outputs."))
 	TObjectPtr<UDisplayClusterMediaOutputSynchronizationPolicy> OutputSyncPolicy;
 };
 
 
 /*
+ * Media settings for node backbuffer
+ */
+USTRUCT(Blueprintable)
+struct DISPLAYCLUSTERCONFIGURATION_API FDisplayClusterConfigurationMediaNodeBackbuffer
+{
+	GENERATED_BODY()
+
+public:
+	/** Enable/Disable Media */
+	UPROPERTY(EditAnywhere, BlueprintReadWrite, Category = "Media", meta = (DisplayName = "Enable Media"))
+	bool bEnable = false;
+
+	/** Full-Frame Media Outputs */
+	UPROPERTY(EditAnywhere, BlueprintReadWrite, Category = "Media")
+	TArray<FDisplayClusterConfigurationMediaOutput> MediaOutputs;
+
+	/** Split layout */
+	UPROPERTY(EditAnywhere, BlueprintReadWrite, Category = "Media", meta = (ClampMin = 1, ClampMax = 4))
+	FIntPoint TiledSplitLayout = { 1, 1 };
+
+	/** Tile mapping. Maps tiles to media outputs. */
+	UPROPERTY(EditAnywhere, BlueprintReadWrite, Category = "Media")
+	TArray<FDisplayClusterConfigurationMediaUniformTileOutput> TiledMediaOutputs;
+
+public:
+
+	/** Returns true if at least one media output assigned */
+	bool IsMediaOutputAssigned() const;
+};
+
+
+PRAGMA_DISABLE_DEPRECATION_WARNINGS
+
+/*
+ * Media settings for viewports
+ */
+USTRUCT(Blueprintable)
+struct DISPLAYCLUSTERCONFIGURATION_API FDisplayClusterConfigurationMediaViewport
+{
+	GENERATED_BODY()
+
+public:
+	/** Enable/disable media */
+	UPROPERTY(EditAnywhere, BlueprintReadWrite, Category = "Media", meta = (DisplayName = "Enable Media"))
+	bool bEnable = false;
+
+	/** Media source to use */
+	UPROPERTY(EditAnywhere, BlueprintReadWrite, Category = "Media")
+	FDisplayClusterConfigurationMediaInput MediaInput;
+
+	/** Media outputs to use */
+	UPROPERTY(EditAnywhere, BlueprintReadWrite, Category = "Media")
+	TArray<FDisplayClusterConfigurationMediaOutput> MediaOutputs;
+
+public:
+	/** Returns true if a media source assigned */
+	bool IsMediaInputAssigned() const;
+
+	/** Returns true if at least one media output assigned */
+	bool IsMediaOutputAssigned() const;
+
+#if WITH_EDITORONLY_DATA
+
+protected:
+	UPROPERTY(meta = (DeprecatedProperty, DeprecationMessage = "This property has been deprecated"))
+	FString MediaSharingNode_DEPRECATED;
+
+public:
+	UE_DEPRECATED(5.3, "This property has been deprecated. Please refer new MediaInput property.")
+	UPROPERTY(meta = (DeprecatedProperty, DeprecationMessage = "This property has been deprecated. Please refer new MediaInput property."))
+	TObjectPtr<UMediaSource> MediaSource;
+
+	UE_DEPRECATED(5.3, "This property has been deprecated. Please refer new MediaOutputs property.")
+	UPROPERTY(meta = (DeprecatedProperty, DeprecationMessage = "This property has been deprecated. Please refer new MediaOutputs property."))
+	TObjectPtr<UMediaOutput> MediaOutput;
+
+	UE_DEPRECATED(5.3, "This property has been deprecated.")
+	UPROPERTY(meta = (DeprecatedProperty, DeprecationMessage = "This property has been deprecated."))
+	TObjectPtr<UDisplayClusterMediaOutputSynchronizationPolicy> OutputSyncPolicy;
+
+#endif // WITH_EDITORONLY_DATA
+};
+
+PRAGMA_ENABLE_DEPRECATION_WARNINGS
+
+
+/*
  * Media settings for ICVFX cameras
diff --git a/Engine/Plugins/Runtime/nDisplay/Source/DisplayClusterMedia/Private/Capture/DisplayClusterMediaCaptureCameraBase.cpp b/Engine/Plugins/Runtime/nDisplay/Source/DisplayClusterMedia/Private/Capture/DisplayClusterMediaCaptureCameraBase.cpp
--- a/Engine/Plugins/Runtime/nDisplay/Source/DisplayClusterMedia/Private/Capture/DisplayClusterMediaCaptureCameraBase.cpp
+++ b/Engine/Plugins/Runtime/nDisplay/Source/DisplayClusterMedia/Private/Capture/DisplayClusterMediaCaptureCameraBase.cpp
@@ -2,14 +2,6 @@
 
 #include "Capture/DisplayClusterMediaCaptureCameraBase.h"
 
-#include "Components/DisplayClusterICVFXCameraComponent.h"
-#include "DisplayClusterConfigurationTypes.h"
-#include "DisplayClusterRootActor.h"
-#include "IDisplayCluster.h"
-#include "Game/IDisplayClusterGameManager.h"
-
-#include "DisplayClusterMediaLog.h"
-
 
 FDisplayClusterMediaCaptureCameraBase::FDisplayClusterMediaCaptureCameraBase(
 	const FString& InMediaId,
@@ -24,28 +16,6 @@ FDisplayClusterMediaCaptureCameraBase::FDisplayClusterMediaCaptureCameraBase(
 {
 }
 
-bool FDisplayClusterMediaCaptureCameraBase::GetCaptureSizeFromConfig(FIntPoint& OutSize) const
-{
-	const UDisplayClusterICVFXCameraComponent* const ICVFXCameraComponent = GetCameraComponent();
-	if (!ICVFXCameraComponent)
-	{
-		return false;
-	}
-
-	const FDisplayClusterConfigurationICVFX_CameraSettings& CameraSettings = ICVFXCameraComponent->GetCameraSettingsICVFX();
-
-	if (CameraSettings.CustomFrustum.bEnable)
-	{
-		OutSize = CameraSettings.CustomFrustum.EstimatedOverscanResolution;
-	}
-	else
-	{
-		OutSize = CameraSettings.CustomFrustum.InnerFrustumResolution;
-	}
-
-	return true;
-}
-
 
 void FDisplayClusterMediaCaptureCameraBase::UpdateLateOCIOState(IDisplayClusterViewport* Viewport)
 {
diff --git a/Engine/Plugins/Runtime/nDisplay/Source/DisplayClusterMedia/Private/Capture/DisplayClusterMediaCaptureCameraBase.h b/Engine/Plugins/Runtime/nDisplay/Source/DisplayClusterMedia/Private/Capture/DisplayClusterMediaCaptureCameraBase.h
--- a/Engine/Plugins/Runtime/nDisplay/Source/DisplayClusterMedia/Private/Capture/DisplayClusterMediaCaptureCameraBase.h
+++ b/Engine/Plugins/Runtime/nDisplay/Source/DisplayClusterMedia/Private/Capture/DisplayClusterMediaCaptureCameraBase.h
@@ -24,12 +24,6 @@ public:
 		UDisplayClusterMediaOutputSynchronizationPolicy* SyncPolicy
 	);
 
-public:
-
-	//~ Begin FDisplayClusterMediaCaptureViewport
-	virtual bool GetCaptureSizeFromConfig(FIntPoint& OutSize) const override;
-	//~ End FDisplayClusterMediaCaptureViewport
-
 protected:
 
 	/** Updates late OCIO state */
diff --git a/Engine/Plugins/Runtime/nDisplay/Source/DisplayClusterMedia/Private/Capture/DisplayClusterMediaCaptureCameraFull.cpp b/Engine/Plugins/Runtime/nDisplay/Source/DisplayClusterMedia/Private/Capture/DisplayClusterMediaCaptureCameraFull.cpp
--- a/Engine/Plugins/Runtime/nDisplay/Source/DisplayClusterMedia/Private/Capture/DisplayClusterMediaCaptureCameraFull.cpp
+++ b/Engine/Plugins/Runtime/nDisplay/Source/DisplayClusterMedia/Private/Capture/DisplayClusterMediaCaptureCameraFull.cpp
@@ -2,6 +2,9 @@
 
 #include "Capture/DisplayClusterMediaCaptureCameraFull.h"
 
+#include "Components/DisplayClusterICVFXCameraComponent.h"
+
+#include "DisplayClusterConfigurationTypes.h"
 #include "DisplayClusterMediaHelpers.h"
 
 
@@ -21,3 +24,25 @@ FDisplayClusterMediaCaptureCameraFull::FDisplayClusterMediaCaptureCameraFull(
 		InSyncPolicy)
 {
 }
+
+bool FDisplayClusterMediaCaptureCameraFull::GetCaptureSizeFromConfig(FIntPoint& OutSize) const
+{
+	const UDisplayClusterICVFXCameraComponent* const ICVFXCameraComponent = GetCameraComponent();
+	if (!ICVFXCameraComponent)
+	{
+		return false;
+	}
+
+	const FDisplayClusterConfigurationICVFX_CameraSettings& CameraSettings = ICVFXCameraComponent->GetCameraSettingsICVFX();
+
+	if (CameraSettings.CustomFrustum.bEnable)
+	{
+		OutSize = CameraSettings.CustomFrustum.EstimatedOverscanResolution;
+	}
+	else
+	{
+		OutSize = CameraSettings.CustomFrustum.InnerFrustumResolution;
+	}
+
+	return true;
+}
diff --git a/Engine/Plugins/Runtime/nDisplay/Source/DisplayClusterMedia/Private/Capture/DisplayClusterMediaCaptureCameraFull.h b/Engine/Plugins/Runtime/nDisplay/Source/DisplayClusterMedia/Private/Capture/DisplayClusterMediaCaptureCameraFull.h
--- a/Engine/Plugins/Runtime/nDisplay/Source/DisplayClusterMedia/Private/Capture/DisplayClusterMediaCaptureCameraFull.h
+++ b/Engine/Plugins/Runtime/nDisplay/Source/DisplayClusterMedia/Private/Capture/DisplayClusterMediaCaptureCameraFull.h
@@ -20,4 +20,10 @@ public:
 		UMediaOutput* MediaOutput,
 		UDisplayClusterMediaOutputSynchronizationPolicy* SyncPolicy
 	);
+
+public:
+
+	//~ Begin FDisplayClusterMediaCaptureViewport
+	virtual bool GetCaptureSizeFromConfig(FIntPoint& OutSize) const override;
+	//~ End FDisplayClusterMediaCaptureViewport
 };
diff --git a/Engine/Plugins/Runtime/nDisplay/Source/DisplayClusterMedia/Private/Capture/DisplayClusterMediaCaptureNode.cpp b/Engine/Plugins/Runtime/nDisplay/Source/DisplayClusterMedia/Private/Capture/DisplayClusterMediaCaptureNode.cpp
deleted file mode 100644
diff --git a/Engine/Plugins/Runtime/nDisplay/Source/DisplayClusterMedia/Private/Capture/DisplayClusterMediaCaptureNode.h b/Engine/Plugins/Runtime/nDisplay/Source/DisplayClusterMedia/Private/Capture/DisplayClusterMediaCaptureNode.h
deleted file mode 100644
diff --git a/Engine/Plugins/Runtime/nDisplay/Source/DisplayClusterMedia/Private/Capture/DisplayClusterMediaCaptureNodeBase.cpp b/Engine/Plugins/Runtime/nDisplay/Source/DisplayClusterMedia/Private/Capture/DisplayClusterMediaCaptureNodeBase.cpp
new file mode 100644
--- /dev/null
+++ b/Engine/Plugins/Runtime/nDisplay/Source/DisplayClusterMedia/Private/Capture/DisplayClusterMediaCaptureNodeBase.cpp
@@ -0,0 +1,53 @@
+// Copyright Epic Games, Inc. All Rights Reserved.
+
+#include "Capture/DisplayClusterMediaCaptureNodeBase.h"
+
+#include "Config/IDisplayClusterConfigManager.h"
+#include "Engine/Engine.h"
+
+#include "DisplayClusterConfigurationTypes.h"
+
+#include "IDisplayCluster.h"
+#include "IDisplayClusterCallbacks.h"
+
+#include "RHICommandList.h"
+#include "RHIResources.h"
+
+
+FDisplayClusterMediaCaptureNodeBase::FDisplayClusterMediaCaptureNodeBase(
+	const FString& InMediaId,
+	const FString& InClusterNodeId,
+	UMediaOutput* InMediaOutput,
+	UDisplayClusterMediaOutputSynchronizationPolicy* SyncPolicy
+)
+	: FDisplayClusterMediaCaptureBase(InMediaId, InClusterNodeId, InMediaOutput, SyncPolicy)
+{
+}
+
+
+bool FDisplayClusterMediaCaptureNodeBase::StartCapture()
+{
+	// If capturing initialized and started successfully, subscribe for rendering callbacks
+	if (FDisplayClusterMediaCaptureBase::StartCapture())
+	{
+		IDisplayCluster::Get().GetCallbacks().OnDisplayClusterPostBackbufferUpdated_RenderThread().AddRaw(this, &FDisplayClusterMediaCaptureNodeBase::OnPostBackbufferUpdated_RenderThread);
+		return true;
+	}
+
+	return false;
+}
+
+void FDisplayClusterMediaCaptureNodeBase::StopCapture()
+{
+	// Stop rendering notifications
+	IDisplayCluster::Get().GetCallbacks().OnDisplayClusterPostBackbufferUpdated_RenderThread().RemoveAll(this);
+
+	// Stop capturing
+	FDisplayClusterMediaCaptureBase::StopCapture();
+}
+
+void FDisplayClusterMediaCaptureNodeBase::OnPostBackbufferUpdated_RenderThread(FRHICommandListImmediate& RHICmdList, FViewport* Viewport)
+{
+	// Pass to the children
+	ProcessPostBackbufferUpdated_RenderThread(RHICmdList, Viewport);
+}
diff --git a/Engine/Plugins/Runtime/nDisplay/Source/DisplayClusterMedia/Private/Capture/DisplayClusterMediaCaptureNodeBase.h b/Engine/Plugins/Runtime/nDisplay/Source/DisplayClusterMedia/Private/Capture/DisplayClusterMediaCaptureNodeBase.h
new file mode 100644
--- /dev/null
+++ b/Engine/Plugins/Runtime/nDisplay/Source/DisplayClusterMedia/Private/Capture/DisplayClusterMediaCaptureNodeBase.h
@@ -0,0 +1,46 @@
+// Copyright Epic Games, Inc. All Rights Reserved.
+
+#pragma once
+
+#include "Capture/DisplayClusterMediaCaptureBase.h"
+
+class FRHICommandListImmediate;
+class FViewport;
+class UDisplayClusterMediaOutputSynchronizationPolicy;
+class UMediaOutput;
+
+
+/**
+ * Node backbuffer media capture base class
+ */
+class FDisplayClusterMediaCaptureNodeBase
+	: public FDisplayClusterMediaCaptureBase
+{
+public:
+
+	/** Constructor */
+	FDisplayClusterMediaCaptureNodeBase(
+		const FString& MediaId,
+		const FString& ClusterNodeId,
+		UMediaOutput* MediaOutput,
+		UDisplayClusterMediaOutputSynchronizationPolicy* SyncPolicy
+	);
+
+public:
+
+	/** Start backbuffer capture */
+	virtual bool StartCapture() override;
+
+	/** Stop backbuffer capture */
+	virtual void StopCapture() override;
+
+protected:
+
+	/** PostBackbufferUpdated implementation on the children side */
+	virtual void ProcessPostBackbufferUpdated_RenderThread(FRHICommandListImmediate& RHICmdList, FViewport* Viewport) = 0;
+
+private:
+
+	/** Main PostBackbufferUpdated event handler. It's passed as virtual ProcessPostBackbufferUpdated_RenderThread call to the children. */
+	void OnPostBackbufferUpdated_RenderThread(FRHICommandListImmediate& RHICmdList, FViewport* Viewport);
+};
diff --git a/Engine/Plugins/Runtime/nDisplay/Source/DisplayClusterMedia/Private/Capture/DisplayClusterMediaCaptureNodeFull.cpp b/Engine/Plugins/Runtime/nDisplay/Source/DisplayClusterMedia/Private/Capture/DisplayClusterMediaCaptureNodeFull.cpp
new file mode 100644
--- /dev/null
+++ b/Engine/Plugins/Runtime/nDisplay/Source/DisplayClusterMedia/Private/Capture/DisplayClusterMediaCaptureNodeFull.cpp
@@ -0,0 +1,69 @@
+// Copyright Epic Games, Inc. All Rights Reserved.
+
+#include "Capture/DisplayClusterMediaCaptureNodeFull.h"
+
+#include "DisplayClusterMediaLog.h"
+
+#include "RenderGraphBuilder.h"
+#include "RenderGraphUtils.h"
+#include "RHICommandList.h"
+#include "RHIResources.h"
+#include "UnrealClient.h"
+
+#include "Engine/Engine.h"
+#include "Engine/GameViewportClient.h"
+
+
+FDisplayClusterMediaCaptureNodeFull::FDisplayClusterMediaCaptureNodeFull(
+	const FString& InMediaId,
+	const FString& InClusterNodeId,
+	UMediaOutput* InMediaOutput,
+	UDisplayClusterMediaOutputSynchronizationPolicy* SyncPolicy
+)
+	: FDisplayClusterMediaCaptureNodeBase(InMediaId, InClusterNodeId, InMediaOutput, SyncPolicy)
+{
+}
+
+FIntPoint FDisplayClusterMediaCaptureNodeFull::GetCaptureSize() const
+{
+	// Return backbuffer runtime size
+	if (GEngine && GEngine->GameViewport && GEngine->GameViewport->Viewport)
+	{
+		const FIntPoint Size = GEngine->GameViewport->Viewport->GetSizeXY();
+
+		UE_LOG(LogDisplayClusterMedia, Log, TEXT("'%s' capture size is [%d, %d]"), *GetMediaId(), Size.X, Size.Y);
+
+		return Size;
+	}
+
+	UE_LOG(LogDisplayClusterMedia, Warning, TEXT("'%s' couldn't get viewport size"), *GetMediaId());
+
+	return FIntPoint();
+}
+
+void FDisplayClusterMediaCaptureNodeFull::ProcessPostBackbufferUpdated_RenderThread(FRHICommandListImmediate& RHICmdList, FViewport* Viewport)
+{
+	if (!Viewport)
+	{
+		UE_LOG(LogDisplayClusterMedia, Warning, TEXT("'%s' capture failed, got invalid viewport"), *GetMediaId());
+		return;
+	}
+
+	if (FRHITexture* const BackbufferTexture = Viewport->GetRenderTargetTexture())
+	{
+		FRDGBuilder GraphBuilder(RHICmdList);
+
+		// Prepare capture request data
+		FRDGTextureRef BackbufferTextureRef = RegisterExternalTexture(GraphBuilder, BackbufferTexture, TEXT("DCMediaOutBackbufferTex"));
+		const FIntRect TextureRegion = { FIntPoint::ZeroValue, BackbufferTextureRef->Desc.Extent };
+		const FMediaOutputTextureInfo TextureInfo{ BackbufferTextureRef, TextureRegion };
+
+		UE_LOG(LogDisplayClusterMedia, VeryVerbose, TEXT("'%s' capturing backbuffer of size %dx%d"),
+			*GetMediaId(), BackbufferTextureRef->Desc.Extent.X, BackbufferTextureRef->Desc.Extent.Y);
+
+		// Capture backbuffer
+		ExportMediaData_RenderThread(GraphBuilder, TextureInfo);
+
+		GraphBuilder.Execute();
+	}
+}
diff --git a/Engine/Plugins/Runtime/nDisplay/Source/DisplayClusterMedia/Private/Capture/DisplayClusterMediaCaptureNodeFull.h b/Engine/Plugins/Runtime/nDisplay/Source/DisplayClusterMedia/Private/Capture/DisplayClusterMediaCaptureNodeFull.h
new file mode 100644
--- /dev/null
+++ b/Engine/Plugins/Runtime/nDisplay/Source/DisplayClusterMedia/Private/Capture/DisplayClusterMediaCaptureNodeFull.h
@@ -0,0 +1,36 @@
+// Copyright Epic Games, Inc. All Rights Reserved.
+
+#pragma once
+
+#include "Capture/DisplayClusterMediaCaptureNodeBase.h"
+
+class FRHICommandListImmediate;
+class FViewport;
+class UDisplayClusterMediaOutputSynchronizationPolicy;
+class UMediaOutput;
+
+
+/**
+ * Node backbuffer media capture (whole buffer)
+ */
+class FDisplayClusterMediaCaptureNodeFull
+	: public FDisplayClusterMediaCaptureNodeBase
+{
+public:
+
+	/** Constructor */
+	FDisplayClusterMediaCaptureNodeFull(
+		const FString& MediaId,
+		const FString& ClusterNodeId,
+		UMediaOutput* MediaOutput,
+		UDisplayClusterMediaOutputSynchronizationPolicy* SyncPolicy = nullptr
+	);
+
+protected:
+
+	/** Returns backbuffer size */
+	virtual FIntPoint GetCaptureSize() const override;
+
+	/** PostBackbufferUpdated handler */
+	virtual void ProcessPostBackbufferUpdated_RenderThread(FRHICommandListImmediate& RHICmdList, FViewport* Viewport) override;
+};
diff --git a/Engine/Plugins/Runtime/nDisplay/Source/DisplayClusterMedia/Private/Capture/DisplayClusterMediaCaptureNodeTile.cpp b/Engine/Plugins/Runtime/nDisplay/Source/DisplayClusterMedia/Private/Capture/DisplayClusterMediaCaptureNodeTile.cpp
new file mode 100644
--- /dev/null
+++ b/Engine/Plugins/Runtime/nDisplay/Source/DisplayClusterMedia/Private/Capture/DisplayClusterMediaCaptureNodeTile.cpp
@@ -0,0 +1,127 @@
+// Copyright Epic Games, Inc. All Rights Reserved.
+
+#include "Capture/DisplayClusterMediaCaptureNodeTile.h"
+
+#include "DisplayClusterMediaHelpers.h"
+#include "DisplayClusterMediaLog.h"
+
+#include "RenderGraphBuilder.h"
+#include "RenderGraphUtils.h"
+#include "RHICommandList.h"
+#include "RHIResources.h"
+#include "UnrealClient.h"
+
+#include "Engine/Engine.h"
+#include "Engine/GameViewportClient.h"
+
+
+FDisplayClusterMediaCaptureNodeTile::FDisplayClusterMediaCaptureNodeTile(
+	const FString& InMediaId,
+	const FString& InClusterNodeId,
+	const FIntPoint& InTileLayout,
+	const FIntPoint& InTilePosition,
+	UMediaOutput* InMediaOutput,
+	UDisplayClusterMediaOutputSynchronizationPolicy* SyncPolicy
+)
+	: FDisplayClusterMediaCaptureNodeBase(InMediaId, InClusterNodeId, InMediaOutput, SyncPolicy)
+	, bValidTileSettings(
+		DisplayClusterMediaHelpers::IsValidLayout(InTileLayout, GetMaxTileLayout()) &&   // Valid layout
+		DisplayClusterMediaHelpers::IsValidTileCoordinate(InTilePosition, InTileLayout)) // Valid coordinates
+	, bEndingX(InTilePosition.X == InTileLayout.X - 1)
+	, bEndingY(InTilePosition.Y == InTileLayout.Y - 1)
+	, TileLayout(InTileLayout)
+	, TilePosition(InTilePosition)
+{
+	checkSlow(bValidTileSettings);
+}
+
+bool FDisplayClusterMediaCaptureNodeTile::StartCapture()
+{
+	// Don't start capture if there is something invalid
+	if (bValidTileSettings)
+	{
+		return FDisplayClusterMediaCaptureNodeBase::StartCapture();
+	}
+
+	return false;
+}
+
+FIntPoint FDisplayClusterMediaCaptureNodeTile::GetCaptureSize() const
+{
+	// Return backbuffer runtime size
+	if (!(GEngine && GEngine->GameViewport && GEngine->GameViewport->Viewport))
+	{
+		UE_LOG(LogDisplayClusterMedia, Warning, TEXT("'%s' couldn't get viewport size"), *GetMediaId());
+		return { };
+	}
+
+	const FIntPoint GameViewportSize = GEngine->GameViewport->Viewport->GetSizeXY();
+
+	// Non-ending tile width/height for reference
+	const int32 NonEndingTileWidth  = GameViewportSize.X / TileLayout.X;
+	const int32 NonEndingTileHeight = GameViewportSize.Y / TileLayout.Y;
+
+	// Ending tile width/height for reference
+	const int32 EndingTileWidth  = (GameViewportSize.X % TileLayout.X == 0 ? NonEndingTileWidth  : GameViewportSize.X % TileLayout.X);
+	const int32 EndingTileHeight = (GameViewportSize.Y % TileLayout.Y == 0 ? NonEndingTileHeight : GameViewportSize.Y % TileLayout.Y);
+
+	// This tile size
+	const int32 ThisTileWidth  = bEndingX ? EndingTileWidth  : NonEndingTileWidth;
+	const int32 ThisTileHeight = bEndingY ? EndingTileHeight : NonEndingTileHeight;
+
+	UE_LOG(LogDisplayClusterMedia, Log, TEXT("'%s' capture size is [%d, %d]"), *GetMediaId(), ThisTileWidth, ThisTileHeight);
+
+	return FIntPoint{ ThisTileWidth, ThisTileHeight };
+}
+
+void FDisplayClusterMediaCaptureNodeTile::ProcessPostBackbufferUpdated_RenderThread(FRHICommandListImmediate& RHICmdList, FViewport* Viewport)
+{
+	if (!Viewport)
+	{
+		UE_LOG(LogDisplayClusterMedia, Warning, TEXT("'%s' couldn't capture backbuffer tile [%d:%d], nullptr viewport"), *GetMediaId(), TilePosition.X, TilePosition.Y);
+		return;
+	}
+
+	if (!bValidTileSettings)
+	{
+		UE_LOG(LogDisplayClusterMedia, Warning, TEXT("'%s' couldn't capture backbuffer tile [%d:%d], layout [%d:%d] - invalid tile settings"), *GetMediaId(), TilePosition.X, TilePosition.Y, TileLayout.X, TileLayout.Y);
+		return;
+	}
+
+	if (FRHITexture* const BackbufferTexture = Viewport->GetRenderTargetTexture())
+	{
+		FRDGBuilder GraphBuilder(RHICmdList);
+
+		FRDGTextureRef BackbufferTextureRef = RegisterExternalTexture(GraphBuilder, BackbufferTexture, TEXT("DCMediaOutBackbufferTex"));
+
+		// Tile width/height (non-edge case)
+		const int32 TileWidth  = BackbufferTextureRef->Desc.Extent.X / TileLayout.X;
+		const int32 TileHeight = BackbufferTextureRef->Desc.Extent.Y / TileLayout.Y;
+
+		// Top-left of the tile sub-region
+		const FIntPoint RectStart{
+			TilePosition.X * TileWidth,
+			TilePosition.Y * TileHeight
+		};
+
+		// Bottom-right of the tile sub-region. Ending tiles might have a little bit more pixels then TileWidth or TileHeight.
+		const FIntPoint RectEnd{
+			bEndingX ? BackbufferTextureRef->Desc.Extent.X : (TilePosition.X + 1) * TileWidth,
+			bEndingY ? BackbufferTextureRef->Desc.Extent.Y : (TilePosition.Y + 1) * TileHeight
+		};
+
+		// Final tile sub-region
+		const FIntRect TextureRegion = { RectStart, RectEnd };
+
+		UE_LOG(LogDisplayClusterMedia, VeryVerbose, TEXT("'%s' capturing backbuffer tile [%d:%d], region [%d:%d - %d:%d] of size [%dx%d]"),
+			*GetMediaId(), TilePosition.X, TilePosition.Y,
+			RectStart.X, RectStart.Y, RectEnd.X, RectEnd.Y,
+			BackbufferTextureRef->Desc.Extent.X, BackbufferTextureRef->Desc.Extent.Y);
+
+		// Capture
+		const FMediaOutputTextureInfo TextureInfo{ BackbufferTextureRef, TextureRegion };
+		ExportMediaData_RenderThread(GraphBuilder, TextureInfo);
+
+		GraphBuilder.Execute();
+	}
+}
diff --git a/Engine/Plugins/Runtime/nDisplay/Source/DisplayClusterMedia/Private/Capture/DisplayClusterMediaCaptureNodeTile.h b/Engine/Plugins/Runtime/nDisplay/Source/DisplayClusterMedia/Private/Capture/DisplayClusterMediaCaptureNodeTile.h
new file mode 100644
--- /dev/null
+++ b/Engine/Plugins/Runtime/nDisplay/Source/DisplayClusterMedia/Private/Capture/DisplayClusterMediaCaptureNodeTile.h
@@ -0,0 +1,68 @@
+// Copyright Epic Games, Inc. All Rights Reserved.
+
+#pragma once
+
+#include "Capture/DisplayClusterMediaCaptureNodeBase.h"
+
+class FRHICommandListImmediate;
+class FViewport;
+class UDisplayClusterMediaOutputSynchronizationPolicy;
+class UMediaOutput;
+
+
+/**
+ * Node backbuffer media capture (tile)
+ */
+class FDisplayClusterMediaCaptureNodeTile
+	: public FDisplayClusterMediaCaptureNodeBase
+{
+public:
+
+	/** Returns maximum layout allowed for backbuffer capture */
+	static const FIntPoint GetMaxTileLayout()
+	{
+		return FIntPoint{ 4, 4 };
+	}
+
+public:
+
+	/** Constructor */
+	FDisplayClusterMediaCaptureNodeTile(
+		const FString& MediaId,
+		const FString& ClusterNodeId,
+		const FIntPoint& TileLayout,
+		const FIntPoint& TilePosition,
+		UMediaOutput* MediaOutput,
+		UDisplayClusterMediaOutputSynchronizationPolicy* SyncPolicy = nullptr
+	);
+
+public:
+
+	/** Start backbuffer capture */
+	virtual bool StartCapture() override;
+
+protected:
+
+	/** Returns tile size */
+	virtual FIntPoint GetCaptureSize() const override;
+
+	/** PostBackbufferUpdated handler */
+	virtual void ProcessPostBackbufferUpdated_RenderThread(FRHICommandListImmediate& RHICmdList, FViewport* Viewport) override;
+
+private:
+
+	/** Pre-computed optimization flag to avoid repetitive tile settings validation */
+	const bool bValidTileSettings;
+
+	/** Pre-computed optimization flag to know if it's the last tile in a raw */
+	const bool bEndingX;
+
+	/** Pre-computed optimization flag to know if it's the last tile in a column */
+	const bool bEndingY;
+
+	/** Output tile layout */
+	const FIntPoint TileLayout;
+
+	/** This tile XY coordinate */
+	const FIntPoint TilePosition;
+};
diff --git a/Engine/Plugins/Runtime/nDisplay/Source/DisplayClusterMedia/Private/Capture/DisplayClusterMediaCaptureViewportBase.cpp b/Engine/Plugins/Runtime/nDisplay/Source/DisplayClusterMedia/Private/Capture/DisplayClusterMediaCaptureViewportBase.cpp
--- a/Engine/Plugins/Runtime/nDisplay/Source/DisplayClusterMedia/Private/Capture/DisplayClusterMediaCaptureViewportBase.cpp
+++ b/Engine/Plugins/Runtime/nDisplay/Source/DisplayClusterMedia/Private/Capture/DisplayClusterMediaCaptureViewportBase.cpp
@@ -5,20 +5,14 @@
 #include "IDisplayCluster.h"
 #include "IDisplayClusterCallbacks.h"
 
-#include "DisplayClusterConfigurationTypes_Viewport.h"
-#include "DisplayClusterMediaCVars.h"
 #include "DisplayClusterMediaLog.h"
-#include "DisplayClusterRootActor.h"
-
-#include "Config/IDisplayClusterConfigManager.h"
-#include "Game/IDisplayClusterGameManager.h"
 
 #include "PostProcess/PostProcessMaterialInputs.h"
 
 #include "Render/IDisplayClusterRenderManager.h"
 #include "Render/Viewport/IDisplayClusterViewport.h"
 #include "Render/Viewport/IDisplayClusterViewportManager.h"
-#include "Render/Viewport/IDisplayClusterViewportManagerProxy.h"
+#include "Render/Viewport/IDisplayClusterViewportProxy.h"
 #include "Render/Viewport/Containers/DisplayClusterViewport_Context.h"
 #include "Render/Viewport/Containers/DisplayClusterViewport_RenderSettings.h"
 
@@ -34,7 +28,7 @@ FDisplayClusterMediaCaptureViewportBase::FDisplayClusterMediaCaptureViewportBase
 	UDisplayClusterMediaOutputSynchronizationPolicy* SyncPolicy
 )
 	: FDisplayClusterMediaCaptureBase(InMediaId, InClusterNodeId, InMediaOutput, SyncPolicy)
-	, ViewportId(InViewportId)
+	, ReferencedViewportId(InViewportId)
 {
 }
 
@@ -111,25 +105,6 @@ FIntPoint FDisplayClusterMediaCaptureViewportBase::GetViewportSize() const
 	return CaptureSize;
 }
 
-bool FDisplayClusterMediaCaptureViewportBase::GetCaptureSizeFromConfig(FIntPoint& OutSize) const
-{
-	if (const ADisplayClusterRootActor* const ActiveRootActor = IDisplayCluster::Get().GetGameMgr()->GetRootActor())
-	{
-		if (const UDisplayClusterConfigurationData* const ConfigData = ActiveRootActor->GetConfigData())
-		{
-			const FString& NodeId = GetClusterNodeId();
-			if (const UDisplayClusterConfigurationViewport* const ViewportCfg = ConfigData->GetViewport(NodeId, ViewportId))
-			{
-				const FIntRect ViewportRect = ViewportCfg->Region.ToRect();
-				OutSize = FIntPoint(ViewportRect.Width(), ViewportRect.Height());
-
-				return true;
-			}
-		}
-	}
-
-	return false;
-}
 
 bool FDisplayClusterMediaCaptureViewportBase::GetCaptureSizeFromGameProxy(FIntPoint& OutSize) const
 {
@@ -138,7 +113,7 @@ bool FDisplayClusterMediaCaptureViewportBase::GetCaptureSizeFromGameProxy(FIntPo
 	{
 		if (const IDisplayClusterViewportManager* const ViewportMgr = RenderMgr->GetViewportManager())
 		{
-			if (const IDisplayClusterViewport* const Viewport = ViewportMgr->FindViewport(ViewportId))
+			if (const IDisplayClusterViewport* const Viewport = ViewportMgr->FindViewport(ReferencedViewportId))
 			{
 				const TArray<FDisplayClusterViewport_Context>& Contexts = Viewport->GetContexts();
 				if (Contexts.Num() > 0)
diff --git a/Engine/Plugins/Runtime/nDisplay/Source/DisplayClusterMedia/Private/Capture/DisplayClusterMediaCaptureViewportBase.h b/Engine/Plugins/Runtime/nDisplay/Source/DisplayClusterMedia/Private/Capture/DisplayClusterMediaCaptureViewportBase.h
--- a/Engine/Plugins/Runtime/nDisplay/Source/DisplayClusterMedia/Private/Capture/DisplayClusterMediaCaptureViewportBase.h
+++ b/Engine/Plugins/Runtime/nDisplay/Source/DisplayClusterMedia/Private/Capture/DisplayClusterMediaCaptureViewportBase.h
@@ -41,14 +41,14 @@ public:
 	/** Returns viewport ID that is configured for capture */
 	const FString& GetViewportId() const
 	{
-		return ViewportId;
+		return ReferencedViewportId;
 	}
 
 	/** Returns texture size of a viewport assigned to capture (main thread) */
 	virtual FIntPoint GetCaptureSize() const override;
 
-	/** Provides default texture size from config */
-	virtual bool GetCaptureSizeFromConfig(FIntPoint& OutSize) const;
+	/** Ask children about capture size based on the configuration data */
+	virtual bool GetCaptureSizeFromConfig(FIntPoint& OutSize) const = 0;
 
 	/** Provides texture size from a game proxy (if available) */
 	bool GetCaptureSizeFromGameProxy(FIntPoint& OutSize) const;
@@ -80,5 +80,5 @@ private:
 private:
 
 	/** Viewport ID assigned to capture */
-	const FString ViewportId;
+	const FString ReferencedViewportId;
 };
diff --git a/Engine/Plugins/Runtime/nDisplay/Source/DisplayClusterMedia/Private/Capture/DisplayClusterMediaCaptureViewportFull.cpp b/Engine/Plugins/Runtime/nDisplay/Source/DisplayClusterMedia/Private/Capture/DisplayClusterMediaCaptureViewportFull.cpp
--- a/Engine/Plugins/Runtime/nDisplay/Source/DisplayClusterMedia/Private/Capture/DisplayClusterMediaCaptureViewportFull.cpp
+++ b/Engine/Plugins/Runtime/nDisplay/Source/DisplayClusterMedia/Private/Capture/DisplayClusterMediaCaptureViewportFull.cpp
@@ -2,6 +2,13 @@
 
 #include "Capture/DisplayClusterMediaCaptureViewportFull.h"
 
+#include "DisplayClusterConfigurationTypes_Viewport.h"
+#include "DisplayClusterRootActor.h"
+
+#include "IDisplayCluster.h"
+
+#include "Game/IDisplayClusterGameManager.h"
+
 
 FDisplayClusterMediaCaptureViewportFull::FDisplayClusterMediaCaptureViewportFull(
 	const FString& InMediaId,
@@ -13,3 +20,25 @@ FDisplayClusterMediaCaptureViewportFull::FDisplayClusterMediaCaptureViewportFull
 	: FDisplayClusterMediaCaptureViewportBase(InMediaId, InClusterNodeId, InViewportId, InMediaOutput, SyncPolicy)
 {
 }
+
+
+bool FDisplayClusterMediaCaptureViewportFull::GetCaptureSizeFromConfig(FIntPoint& OutSize) const
+{
+	if (const ADisplayClusterRootActor* const ActiveRootActor = IDisplayCluster::Get().GetGameMgr()->GetRootActor())
+	{
+		if (const UDisplayClusterConfigurationData* const ConfigData = ActiveRootActor->GetConfigData())
+		{
+			const FString& NodeId     = GetClusterNodeId();
+			const FString& ViewportId = GetViewportId();
+			if (const UDisplayClusterConfigurationViewport* const ViewportCfg = ConfigData->GetViewport(NodeId, ViewportId))
+			{
+				const FIntRect ViewportRect = ViewportCfg->Region.ToRect();
+				OutSize = FIntPoint(ViewportRect.Width(), ViewportRect.Height());
+
+				return true;
+			}
+		}
+	}
+
+	return false;
+}
diff --git a/Engine/Plugins/Runtime/nDisplay/Source/DisplayClusterMedia/Private/Capture/DisplayClusterMediaCaptureViewportFull.h b/Engine/Plugins/Runtime/nDisplay/Source/DisplayClusterMedia/Private/Capture/DisplayClusterMediaCaptureViewportFull.h
--- a/Engine/Plugins/Runtime/nDisplay/Source/DisplayClusterMedia/Private/Capture/DisplayClusterMediaCaptureViewportFull.h
+++ b/Engine/Plugins/Runtime/nDisplay/Source/DisplayClusterMedia/Private/Capture/DisplayClusterMediaCaptureViewportFull.h
@@ -20,4 +20,9 @@ public:
 		UMediaOutput* MediaOutput,
 		UDisplayClusterMediaOutputSynchronizationPolicy* SyncPolicy
 	);
+
+public:
+
+	/** Provides default capture size from config */
+	virtual bool GetCaptureSizeFromConfig(FIntPoint& OutSize) const;
 };
diff --git a/Engine/Plugins/Runtime/nDisplay/Source/DisplayClusterMedia/Private/DisplayClusterMediaHelpers.cpp b/Engine/Plugins/Runtime/nDisplay/Source/DisplayClusterMedia/Private/DisplayClusterMediaHelpers.cpp
--- a/Engine/Plugins/Runtime/nDisplay/Source/DisplayClusterMedia/Private/DisplayClusterMediaHelpers.cpp
+++ b/Engine/Plugins/Runtime/nDisplay/Source/DisplayClusterMedia/Private/DisplayClusterMediaHelpers.cpp
@@ -41,7 +41,9 @@ namespace DisplayClusterMediaHelpers
 				switch (OwnerType)
 				{
 				case EMediaOwnerType::Backbuffer:
-					return FString::Printf(TEXT("%s_%s_backbuffer_capture_%u"), *NodeId, *DCRAName, Index);
+					return InTilePos
+						? FString::Printf(TEXT("%s_%s_tile_%d_%d_backbuffer_capture_%u"), *NodeId, *DCRAName, InTilePos->X, InTilePos->Y, Index)
+						: FString::Printf(TEXT("%s_%s_backbuffer_capture_%u"), *NodeId, *DCRAName, Index);
 
 				case EMediaOwnerType::Viewport:
 					return InTilePos
@@ -162,4 +164,25 @@ namespace DisplayClusterMediaHelpers
 			ResampleTextureImpl_RenderThread<FScreenPS>(RHICmdList, SrcTexture, DstTexture, SrcRect, DstRect);
 		}
 	}
+
+	bool IsValidLayout(const FIntPoint& TileLayout, const FIntPoint& MaxLayout)
+	{
+		const bool bValidX  = TileLayout.X > 0 && TileLayout.X <= MaxLayout.X;
+		const bool bValidY  = TileLayout.Y > 0 && TileLayout.Y <= MaxLayout.Y;
+		const bool bNot1by1 = TileLayout.X > 1 || TileLayout.Y > 1;
+
+		const bool bIsValid = bValidX && bValidY && bNot1by1;
+
+		return bIsValid;
+	}
+
+	bool IsValidTileCoordinate(const FIntPoint& TilePosition, const FIntPoint& TileLayout)
+	{
+		const bool bValidX = TilePosition.X >= 0 && TilePosition.X < TileLayout.X;
+		const bool bValidY = TilePosition.Y >= 0 && TilePosition.Y < TileLayout.Y;
+
+		const bool bIsValid = bValidX && bValidY;
+
+		return bIsValid;
+	}
 }
diff --git a/Engine/Plugins/Runtime/nDisplay/Source/DisplayClusterMedia/Private/DisplayClusterMediaHelpers.h b/Engine/Plugins/Runtime/nDisplay/Source/DisplayClusterMedia/Private/DisplayClusterMediaHelpers.h
--- a/Engine/Plugins/Runtime/nDisplay/Source/DisplayClusterMedia/Private/DisplayClusterMediaHelpers.h
+++ b/Engine/Plugins/Runtime/nDisplay/Source/DisplayClusterMedia/Private/DisplayClusterMediaHelpers.h
@@ -42,4 +42,10 @@ namespace DisplayClusterMediaHelpers
 
 	/** Copy and resize an RHI texture */
 	void ResampleTexture_RenderThread(FRHICommandListImmediate& RHICmdList, FRHITexture* SrcTexture, FRHITexture* DstTexture, const FIntRect& SrcRect, const FIntRect& DstRect);
+
+	/** Checks whether a tile layout is valid */
+	bool IsValidLayout(const FIntPoint& TileLayout, const FIntPoint& MaxLayout);
+
+	/** Checks the correctness of the coordinates of the tile in the given layout */
+	bool IsValidTileCoordinate(const FIntPoint& TilePosition, const FIntPoint& TileLayout);
 }
diff --git a/Engine/Plugins/Runtime/nDisplay/Source/DisplayClusterMedia/Private/DisplayClusterMediaModule.cpp b/Engine/Plugins/Runtime/nDisplay/Source/DisplayClusterMedia/Private/DisplayClusterMediaModule.cpp
--- a/Engine/Plugins/Runtime/nDisplay/Source/DisplayClusterMedia/Private/DisplayClusterMediaModule.cpp
+++ b/Engine/Plugins/Runtime/nDisplay/Source/DisplayClusterMedia/Private/DisplayClusterMediaModule.cpp
@@ -19,7 +19,8 @@
 
 #include "Capture/DisplayClusterMediaCaptureCameraFull.h"
 #include "Capture/DisplayClusterMediaCaptureCameraTile.h"
-#include "Capture/DisplayClusterMediaCaptureNode.h"
+#include "Capture/DisplayClusterMediaCaptureNodeFull.h"
+#include "Capture/DisplayClusterMediaCaptureNodeTile.h"
 #include "Capture/DisplayClusterMediaCaptureViewportFull.h"
 #include "Capture/DisplayClusterMediaCaptureViewportTile.h"
 #include "Input/DisplayClusterMediaInputCameraFull.h"
@@ -94,7 +95,8 @@ void FDisplayClusterMediaModule::InitializeMedia()
 		{
 			// Node backbuffer media setup
 			{
-				InitializeBackbufferOutput(ClusterNode, RootActorName, ClusterNodeId);
+				InitializeBackbufferFullFrameOutput(ClusterNode, RootActorName, ClusterNodeId);
+				InitializeBackbufferUniformTilesOutput(ClusterNode, RootActorName, ClusterNodeId);
 			}
 
 			// Viewports media setup
@@ -179,10 +181,8 @@ void FDisplayClusterMediaModule::StopMedia()
 	}
 }
 
-void FDisplayClusterMediaModule::InitializeBackbufferOutput(const UDisplayClusterConfigurationClusterNode* ClusterNode, const FString& RootActorName, const FString& ClusterNodeId)
+void FDisplayClusterMediaModule::InitializeBackbufferFullFrameOutput(const UDisplayClusterConfigurationClusterNode* ClusterNode, const FString& RootActorName, const FString& ClusterNodeId)
 {
-	checkSlow(ClusterNode);
-
 	if (IsValid(ClusterNode))
 	{
 		const FDisplayClusterConfigurationMediaNodeBackbuffer& MediaSettings = ClusterNode->MediaSettings;
@@ -201,7 +201,7 @@ void FDisplayClusterMediaModule::InitializeBackbufferOutput(const UDisplayCluste
 
 					UE_LOG(LogDisplayClusterMedia, Log, TEXT("Initializing backbuffer media capture [%u]: '%s'"), CaptureIdx, *MediaCaptureId);
 
-					TSharedPtr<FDisplayClusterMediaCaptureNode> NewNodeCapture = MakeShared<FDisplayClusterMediaCaptureNode>(
+					TSharedPtr<FDisplayClusterMediaCaptureNodeFull> NewNodeCapture = MakeShared<FDisplayClusterMediaCaptureNodeFull>(
 						MediaCaptureId,
 						ClusterNodeId,
 						MediaOutputItem.MediaOutput,
@@ -216,10 +216,58 @@ void FDisplayClusterMediaModule::InitializeBackbufferOutput(const UDisplayCluste
 	}
 }
 
-void FDisplayClusterMediaModule::InitializeViewportInput(const UDisplayClusterConfigurationViewport* Viewport, const FString& ViewportId, const FString& RootActorName, const FString& ClusterNodeId)
+void FDisplayClusterMediaModule::InitializeBackbufferUniformTilesOutput(const UDisplayClusterConfigurationClusterNode* ClusterNode, const FString& RootActorName, const FString& ClusterNodeId)
 {
-	checkSlow(Viewport);
+	if (IsValid(ClusterNode))
+	{
+		const FDisplayClusterConfigurationMediaNodeBackbuffer& MediaSettings = ClusterNode->MediaSettings;
+
+		// Validate tile layout
+		const FIntPoint MaxLayout = FDisplayClusterMediaCaptureNodeTile::GetMaxTileLayout();
+		const bool bLayoutIsValid = DisplayClusterMediaHelpers::IsValidLayout(MediaSettings.TiledSplitLayout, MaxLayout);
+		if (!bLayoutIsValid)
+		{
+			UE_LOG(LogDisplayClusterMedia, Warning, TEXT("Invalid layout [%dx%d] was requested for backbuffer capture. Max layout is [%dx%d]."),
+				MediaSettings.TiledSplitLayout.X, MediaSettings.TiledSplitLayout.Y,
+				MaxLayout.X, MaxLayout.Y);
+
+			return;
+		}
+
+		if (MediaSettings.bEnable && bLayoutIsValid)
+		{
+			uint8 CaptureIdx = 0;
+			for (const FDisplayClusterConfigurationMediaUniformTileOutput& MediaOutputTile : MediaSettings.TiledMediaOutputs)
+			{
+				if (IsValid(MediaOutputTile.MediaOutput))
+				{
+					const FString MediaCaptureId = DisplayClusterMediaHelpers::MediaId::GenerateMediaId(
+						DisplayClusterMediaHelpers::MediaId::EMediaDeviceType::Output,
+						DisplayClusterMediaHelpers::MediaId::EMediaOwnerType::Backbuffer,
+						ClusterNodeId, RootActorName, FString(),
+						CaptureIdx, &MediaOutputTile.Position);
 
+					UE_LOG(LogDisplayClusterMedia, Log, TEXT("Initializing backbuffer media capture tile [%u]:'%d,%d': '%s'"),
+						CaptureIdx, MediaOutputTile.Position.X, MediaOutputTile.Position.Y, *MediaCaptureId);
+
+					TSharedPtr<FDisplayClusterMediaCaptureNodeTile> NewNodeCapture = MakeShared<FDisplayClusterMediaCaptureNodeTile>(
+						MediaCaptureId, ClusterNodeId,
+						MediaSettings.TiledSplitLayout,
+						MediaOutputTile.Position,
+						MediaOutputTile.MediaOutput,
+						MediaOutputTile.OutputSyncPolicy);
+
+					CaptureDevices.Emplace(MediaCaptureId, NewNodeCapture);
+				}
+
+				++CaptureIdx;
+			}
+		}
+	}
+}
+
+void FDisplayClusterMediaModule::InitializeViewportInput(const UDisplayClusterConfigurationViewport* Viewport, const FString& ViewportId, const FString& RootActorName, const FString& ClusterNodeId)
+{
 	if (IsValid(Viewport))
 	{
 		const FDisplayClusterConfigurationMediaViewport& MediaSettings = Viewport->RenderSettings.Media;
@@ -249,8 +297,6 @@ void FDisplayClusterMediaModule::InitializeViewportInput(const UDisplayClusterCo
 
 void FDisplayClusterMediaModule::InitializeViewportOutput(const UDisplayClusterConfigurationViewport* Viewport, const FString& ViewportId, const FString& RootActorName, const FString& ClusterNodeId)
 {
-	checkSlow(Viewport);
-
 	if (IsValid(Viewport))
 	{
 		const FDisplayClusterConfigurationMediaViewport& MediaSettings = Viewport->RenderSettings.Media;
@@ -285,8 +331,6 @@ void FDisplayClusterMediaModule::InitializeViewportOutput(const UDisplayClusterC
 
 void FDisplayClusterMediaModule::InitializeICVFXCameraFullFrameInput(const UDisplayClusterICVFXCameraComponent* ICVFXCameraComponent, const FString& RootActorName, const FString& ClusterNodeId)
 {
-	checkSlow(ICVFXCameraComponent);
-
 	if (IsValid(ICVFXCameraComponent))
 	{
 		const FDisplayClusterConfigurationMediaICVFX& MediaSettings = ICVFXCameraComponent->CameraSettings.RenderSettings.Media;
@@ -318,8 +362,6 @@ void FDisplayClusterMediaModule::InitializeICVFXCameraFullFrameInput(const UDisp
 
 void FDisplayClusterMediaModule::InitializeICVFXCameraFullFrameOutput(const UDisplayClusterICVFXCameraComponent* ICVFXCameraComponent, const FString& RootActorName, const FString& ClusterNodeId)
 {
-	checkSlow(ICVFXCameraComponent);
-
 	if (IsValid(ICVFXCameraComponent))
 	{
 		const FDisplayClusterConfigurationMediaICVFX& MediaSettings = ICVFXCameraComponent->CameraSettings.RenderSettings.Media;
@@ -360,8 +402,6 @@ void FDisplayClusterMediaModule::InitializeICVFXCameraFullFrameOutput(const UDis
 
 void FDisplayClusterMediaModule::InitializeICVFXCameraUniformTilesInput(const UDisplayClusterICVFXCameraComponent* ICVFXCameraComponent, const FString& RootActorName, const FString& ClusterNodeId)
 {
-	checkSlow(ICVFXCameraComponent);
-
 	if (IsValid(ICVFXCameraComponent))
 	{
 		const FDisplayClusterConfigurationMediaICVFX& MediaSettings = ICVFXCameraComponent->CameraSettings.RenderSettings.Media;
@@ -409,8 +449,6 @@ void FDisplayClusterMediaModule::InitializeICVFXCameraUniformTilesInput(const UD
 
 void FDisplayClusterMediaModule::InitializeICVFXCameraUniformTilesOutput(const UDisplayClusterICVFXCameraComponent* ICVFXCameraComponent, const FString& RootActorName, const FString& ClusterNodeId)
 {
-	checkSlow(ICVFXCameraComponent);
-
 	if (IsValid(ICVFXCameraComponent))
 	{
 		const FDisplayClusterConfigurationMediaICVFX& MediaSettings = ICVFXCameraComponent->CameraSettings.RenderSettings.Media;
diff --git a/Engine/Plugins/Runtime/nDisplay/Source/DisplayClusterMedia/Private/DisplayClusterMediaModule.h b/Engine/Plugins/Runtime/nDisplay/Source/DisplayClusterMedia/Private/DisplayClusterMediaModule.h
--- a/Engine/Plugins/Runtime/nDisplay/Source/DisplayClusterMedia/Private/DisplayClusterMediaModule.h
+++ b/Engine/Plugins/Runtime/nDisplay/Source/DisplayClusterMedia/Private/DisplayClusterMediaModule.h
@@ -50,9 +50,12 @@ protected:
 
 private:
 
-	/** Backbuffer output initializer */
-	void InitializeBackbufferOutput(const UDisplayClusterConfigurationClusterNode* ClusterNode, const FString& RootActorName, const FString& ClusterNodeId);
-	
+	/** Backbuffer output initializer (full frame) */
+	void InitializeBackbufferFullFrameOutput(const UDisplayClusterConfigurationClusterNode* ClusterNode, const FString& RootActorName, const FString& ClusterNodeId);
+
+	/** Backbuffer output initializer (uniform tiles) */
+	void InitializeBackbufferUniformTilesOutput(const UDisplayClusterConfigurationClusterNode* ClusterNode, const FString& RootActorName, const FString& ClusterNodeId);
+
 	/** Viewport input initializer */
 	void InitializeViewportInput(const UDisplayClusterConfigurationViewport* Viewport, const FString& ViewportId, const FString& RootActorName, const FString& ClusterNodeId);
 
