[nDisplay] Fixed CIS issues - added some depreacation 'wrapping'

#rnx

diff --git a/Engine/Plugins/Runtime/nDisplay/Source/DisplayClusterConfiguration/Public/DisplayClusterConfigurationTypes_Media.h b/Engine/Plugins/Runtime/nDisplay/Source/DisplayClusterConfiguration/Public/DisplayClusterConfigurationTypes_Media.h
--- a/Engine/Plugins/Runtime/nDisplay/Source/DisplayClusterConfiguration/Public/DisplayClusterConfigurationTypes_Media.h
+++ b/Engine/Plugins/Runtime/nDisplay/Source/DisplayClusterConfiguration/Public/DisplayClusterConfigurationTypes_Media.h
@@ -156,6 +156,8 @@ public:
 };
 
 
+PRAGMA_DISABLE_DEPRECATION_WARNINGS
+
 /*
  * Media output group (ICVFX, Tiled)
  */
@@ -173,12 +175,17 @@ public:
 	UPROPERTY(EditAnywhere, BlueprintReadWrite, Category = "Media")
 	TArray<FDisplayClusterConfigurationMediaUniformTileOutput> Tiles;
 
+#if WITH_EDITORONLY_DATA
+public:
 	/** Media output synchronization policy */
 	UE_DEPRECATED(5.6, "This property has been deprecated. Please refer new MediaOutputs property.")
 	UPROPERTY(meta = (DeprecatedProperty, DeprecationMessage = "This property has been deprecated. Synchronization policy is bound to the actual media outputs."))
 	TObjectPtr<UDisplayClusterMediaOutputSynchronizationPolicy> OutputSyncPolicy;
+#endif // WITH_EDITORONLY_DATA
 };
 
+PRAGMA_ENABLE_DEPRECATION_WARNINGS
+
 
 /*
  * Media settings for node backbuffer
