[nDisplay] Fixed editor crash when undo-ing some transactions

- ClusterChanged is fired now after UNDO. This properly updates the UI data, and therefore avoids any invalid weak pointers appearing.

[RN] Fixed crash when undo-ing some actions in the nDisplay configurator

diff --git a/Engine/Plugins/Runtime/nDisplay/Source/DisplayClusterConfigurator/Private/Views/TreeViews/Cluster/DisplayClusterConfiguratorViewCluster.cpp b/Engine/Plugins/Runtime/nDisplay/Source/DisplayClusterConfigurator/Private/Views/TreeViews/Cluster/DisplayClusterConfiguratorViewCluster.cpp
--- a/Engine/Plugins/Runtime/nDisplay/Source/DisplayClusterConfigurator/Private/Views/TreeViews/Cluster/DisplayClusterConfiguratorViewCluster.cpp
+++ b/Engine/Plugins/Runtime/nDisplay/Source/DisplayClusterConfigurator/Private/Views/TreeViews/Cluster/DisplayClusterConfiguratorViewCluster.cpp
@@ -35,8 +35,7 @@ FDisplayClusterConfiguratorViewCluster::FDisplayClusterConfiguratorViewCluster(c
 
 void FDisplayClusterConfiguratorViewCluster::PostUndo(bool bSuccess)
 {
-	ClearSelection();
-	ViewTree->Refresh();
+	ToolkitPtr.Pin()->ClusterChanged();
 }
 
 TSharedRef<SWidget> FDisplayClusterConfiguratorViewCluster::CreateWidget()
