nDisplay: Fixed a crash when FDisplayClusterActorRef tried to load UClass after UE shutdown.

diff --git a/Engine/Plugins/Runtime/nDisplay/Source/DisplayCluster/Public/Misc/DisplayClusterObjectRef.h b/Engine/Plugins/Runtime/nDisplay/Source/DisplayCluster/Public/Misc/DisplayClusterObjectRef.h
--- a/Engine/Plugins/Runtime/nDisplay/Source/DisplayCluster/Public/Misc/DisplayClusterObjectRef.h
+++ b/Engine/Plugins/Runtime/nDisplay/Source/DisplayCluster/Public/Misc/DisplayClusterObjectRef.h
@@ -64,6 +64,14 @@ public:
 		if (!ActorPtr.IsValid())
 		{
 			ActorPtr.Reset();
+
+			// Do not do anything after the exit request.
+			if (IsEngineExitRequested())
+			{
+				return nullptr;
+			}
+
+			// Reload the current world and actor class if necessary.
 			if (UpdateWorldPtr() && UpdateActorClassPtr())
 			{
 				// Find actor object
