nDisplay: Added features related to supporting using Ghost Frame

A new projection policy called 'Reference' has been added. [UE-229219]
Added per-viewport chromakey type [UE-230801]

Optimised resource management of nDisplay viewports: only used textures are allocated. 'FDisplayClusterViewport::IsResourceUsedImpl()'

The public API has been modified:
- IDisplayClusterProjectionPolicy
- IDisplayClusterViewport
- IDisplayClusterWarpPolicy

diff --git a/Engine/Plugins/Runtime/nDisplay/Source/DisplayCluster/Private/Render/Viewport/Configuration/DisplayClusterViewportConfigurationHelpers_ICVFX.cpp b/Engine/Plugins/Runtime/nDisplay/Source/DisplayCluster/Private/Render/Viewport/Configuration/DisplayClusterViewportConfigurationHelpers_ICVFX.cpp
--- a/Engine/Plugins/Runtime/nDisplay/Source/DisplayCluster/Private/Render/Viewport/Configuration/DisplayClusterViewportConfigurationHelpers_ICVFX.cpp
+++ b/Engine/Plugins/Runtime/nDisplay/Source/DisplayCluster/Private/Render/Viewport/Configuration/DisplayClusterViewportConfigurationHelpers_ICVFX.cpp
@@ -291,8 +291,10 @@ void FDisplayClusterViewportConfigurationHelpers_ICVFX::PreviewReuseInnerFrustum
 			if (ViewportIt.IsValid() && ViewportIt != InCameraViewport.AsShared() && ViewportIt->GetClusterNodeId() != InCameraViewport.GetClusterNodeId()
 				&& FDisplayClusterViewportConfigurationHelpers_Postprocess::IsInnerFrustumViewportSettingsEqual(*ViewportIt, InCameraViewport, InCameraSettings))
 			{
-				EDisplayClusterViewportOverrideMode ViewportOverrideMode = ViewportIt->IsOpenColorIOEquals(InCameraViewport) ?
-					EDisplayClusterViewportOverrideMode::All : EDisplayClusterViewportOverrideMode::InernalRTT;
+				const EDisplayClusterViewportOverrideMode ViewportOverrideMode =
+					ViewportIt->IsOpenColorIOEquals(InCameraViewport)
+					? EDisplayClusterViewportOverrideMode::All
+					: EDisplayClusterViewportOverrideMode::InternalRTT;
 
 				// Gain direct access to internal resources of the viewport:
 				FDisplayClusterViewport_RenderSettings& InOutRenderSettings = InCameraViewport.GetRenderSettingsImpl();
@@ -319,8 +321,10 @@ void FDisplayClusterViewportConfigurationHelpers_ICVFX::PreviewReuseChromakeyVie
 			if (ViewportIt.IsValid() && ViewportIt != InChromakeyViewport.AsShared() && ViewportIt->GetClusterNodeId() != InChromakeyViewport.GetClusterNodeId())
 			{
 				// Chromakey support OCIO
-				EDisplayClusterViewportOverrideMode ViewportOverrideMode = ViewportIt->IsOpenColorIOEquals(InChromakeyViewport) ?
-					EDisplayClusterViewportOverrideMode::All : EDisplayClusterViewportOverrideMode::InernalRTT;
+				const EDisplayClusterViewportOverrideMode ViewportOverrideMode = 
+					ViewportIt->IsOpenColorIOEquals(InChromakeyViewport)
+					? EDisplayClusterViewportOverrideMode::All
+					: EDisplayClusterViewportOverrideMode::InternalRTT;
 
 				// Gain direct access to internal resources of the viewport:
 				FDisplayClusterViewport_RenderSettings& InOutRenderSettings = InChromakeyViewport.GetRenderSettingsImpl();
@@ -484,12 +488,6 @@ void FDisplayClusterViewportConfigurationHelpers_ICVFX::UpdateCameraViewportSett
 	// incamera textrure used as overlay
 	InOutRenderSettings.bVisible = false;
 
-	// use framecolor instead of viewport rendering
-	if (InCameraSettings.Chromakey.GetChromakeyType(*StageSettings) == EDisplayClusterShaderParametersICVFX_ChromakeySource::FrameColor)
-	{
-		InOutRenderSettings.bSkipRendering = true;
-	}
-
 	// Update camera viewport projection policy settings
 	UpdateCameraProjectionSettingsICVFX(*DstViewport.Configuration, InCameraComponent, InCameraSettings, DstViewport.GetProjectionPolicy());
 
@@ -570,8 +568,23 @@ void FDisplayClusterViewportConfigurationHelpers_ICVFX::UpdateChromakeyViewportS
 
 		DstViewport.UpdateConfiguration_OverlayRenderSettings(InChromakeyRenderSettings->AdvancedRenderSettings);
 
-		// Resize chromakey RTT
-		const FIntPoint ChromakeyTextureSize = InOutRenderSettings.Rect.Size() * InChromakeyRenderSettings->ChromakeySizeMult;
+		// Use the size of the ICVFX camera viewport as a base value, then multiply it by ChromakeySizeMult.
+		const FIntPoint ChromakeyTextureSize = InOutCameraViewportRenderSettings.Rect.Size() * InChromakeyRenderSettings->ChromakeySizeMult;
+
+		// Check that the size of the CK viewport is zero and show the log message once
+		if (ChromakeyTextureSize.GetMin() <= 0)
+		{
+			if (DstViewport.CanShowLogMsgOnce(EDisplayClusterViewportShowLogMsgOnce::ChromakeyViewportHasZeroSize))
+			{
+				UE_LOG(LogDisplayClusterViewport, Error, TEXT("ICVFX Chromakey Viewport '%s': Has zero size. (ChromakeySizeMult=%.2f)"), *DstViewport.GetId(), InChromakeyRenderSettings->ChromakeySizeMult);
+			}
+			else
+			{
+				DstViewport.ResetShowLogMsgOnce(EDisplayClusterViewportShowLogMsgOnce::ChromakeyViewportHasZeroSize);
+			}
+		}
+
+		// Check the obtained value to be valid and use as the chromakey RTT size.
 		InOutRenderSettings.Rect = FDisplayClusterViewportHelpers::GetValidViewportRect(FIntRect(FIntPoint(0, 0), ChromakeyTextureSize), DstViewport.GetId(), TEXT("Configuration custom chromakey Frame Size"));
 
 		// Debug: override the texture of the target viewport from this chromakeyRTT
@@ -592,13 +605,18 @@ void FDisplayClusterViewportConfigurationHelpers_ICVFX::UpdateChromakeyViewportS
 	InOutRenderSettings.AssignParentViewport(InCameraViewport.GetId(), InOutCameraViewportRenderSettings);
 }
 
-void FDisplayClusterViewportConfigurationHelpers_ICVFX::UpdateCameraSettings_Chromakey(FDisplayClusterShaderParameters_ICVFX::FCameraSettings& InOutCameraSettings, const FDisplayClusterConfigurationICVFX_StageSettings& InStageSettings, const FDisplayClusterConfigurationICVFX_CameraSettings& InCameraSettings, bool bEnableChromakeyMarkers, const FString& InChromakeyViewportId)
+void FDisplayClusterViewportConfigurationHelpers_ICVFX::UpdateCameraSettings_Chromakey(FDisplayClusterShaderParameters_ICVFX::FCameraSettings& InOutCameraSettings,
+	const FDisplayClusterConfigurationICVFX_StageSettings& InStageSettings,
+	const FDisplayClusterConfigurationICVFX_CameraSettings& InCameraSettings,
+	const bool bEnableChromakey,
+	const bool bEnableChromakeyMarkers,
+	const FString& InChromakeyViewportId)
 {
 	// Set chromakey color
 	InOutCameraSettings.ChromakeyColor = InCameraSettings.Chromakey.GetChromakeyColor(InStageSettings);
 
 	// Set chromakey source
-	switch (InOutCameraSettings.ChromakeySource = InCameraSettings.Chromakey.GetChromakeyType(InStageSettings))
+	switch (InOutCameraSettings.ChromakeySource)
 	{
 	case EDisplayClusterShaderParametersICVFX_ChromakeySource::ChromakeyLayers:
 		if (InChromakeyViewportId.IsEmpty())
@@ -723,8 +741,16 @@ void FDisplayClusterViewportConfigurationHelpers_ICVFX::UpdateLightcardViewportS
 		DstViewport.UpdateConfiguration_OverlayRenderSettings(InRenderSettings.AdvancedRenderSettings);
 	}
 
+	FString ParentViewportId = BaseViewport.GetId();
+	if (!BaseViewport.GetRenderSettings().GetViewportOverrideId().IsEmpty())
+	{
+		// If the base viewport is already overridden (e.g. by the 'Reference' projection policy),
+		// use it as the parent viewport.
+		ParentViewportId = BaseViewport.GetRenderSettings().GetViewportOverrideId();
+	}
+
 	// Attach to parent viewport
-	InOutRenderSettings.AssignParentViewport(BaseViewport.GetId(), BaseViewport.GetRenderSettings());
+	InOutRenderSettings.AssignParentViewport(ParentViewportId, BaseViewport.GetRenderSettings());
 
 	if (EnumHasAnyFlags(DstViewport.GetRenderSettingsICVFX().RuntimeFlags, EDisplayClusterViewportRuntimeICVFXFlags::Lightcard))
 	{
diff --git a/Engine/Plugins/Runtime/nDisplay/Source/DisplayCluster/Private/Render/Viewport/Configuration/DisplayClusterViewportConfigurationHelpers_ICVFX.h b/Engine/Plugins/Runtime/nDisplay/Source/DisplayCluster/Private/Render/Viewport/Configuration/DisplayClusterViewportConfigurationHelpers_ICVFX.h
--- a/Engine/Plugins/Runtime/nDisplay/Source/DisplayCluster/Private/Render/Viewport/Configuration/DisplayClusterViewportConfigurationHelpers_ICVFX.h
+++ b/Engine/Plugins/Runtime/nDisplay/Source/DisplayCluster/Private/Render/Viewport/Configuration/DisplayClusterViewportConfigurationHelpers_ICVFX.h
@@ -62,7 +62,13 @@ public:
 	static bool UpdateCameraProjectionSettingsICVFX(FDisplayClusterViewportConfiguration& InConfiguration, UDisplayClusterICVFXCameraComponent& CameraComponent, const FDisplayClusterConfigurationICVFX_CameraSettings& InCameraSettings, const TSharedPtr<IDisplayClusterProjectionPolicy, ESPMode::ThreadSafe>& InProjectionPolicy);
 
 	/** Configures the camera's chromakey render settings to match the specified configuration settings  */
-	static void UpdateCameraSettings_Chromakey(FDisplayClusterShaderParameters_ICVFX::FCameraSettings& InOutCameraSettings, const FDisplayClusterConfigurationICVFX_StageSettings& InStageSettings, const FDisplayClusterConfigurationICVFX_CameraSettings& InCameraSettings, bool bEnableChromakeyMarkers, const FString& InChromakeyViewportId);
+	static void UpdateCameraSettings_Chromakey(
+		FDisplayClusterShaderParameters_ICVFX::FCameraSettings& InOutCameraSettings,
+		const FDisplayClusterConfigurationICVFX_StageSettings& InStageSettings,
+		const FDisplayClusterConfigurationICVFX_CameraSettings& InCameraSettings,
+		const bool bEnableChromakey,
+		const bool bEnableChromakeyMarkers,
+		const FString& InChromakeyViewportId);
 	static void UpdateCameraSettings_OverlapChromakey(FDisplayClusterShaderParameters_ICVFX::FCameraSettings& InOutCameraSettings, const FDisplayClusterConfigurationICVFX_StageSettings& InStageSettings, const FDisplayClusterConfigurationICVFX_CameraSettings& InCameraSettings, bool bEnableChromakeyMarkers);
 
 	static void UpdateCameraCustomFrustum(const FDisplayClusterConfigurationICVFX_CameraCustomFrustum& InCameraCustomFrustumConfiguration, FDisplayClusterViewport_CustomFrustumSettings& OutCustomFrustumSettings);
diff --git a/Engine/Plugins/Runtime/nDisplay/Source/DisplayCluster/Private/Render/Viewport/Configuration/DisplayClusterViewportConfigurationHelpers_RenderFrameSettings.cpp b/Engine/Plugins/Runtime/nDisplay/Source/DisplayCluster/Private/Render/Viewport/Configuration/DisplayClusterViewportConfigurationHelpers_RenderFrameSettings.cpp
--- a/Engine/Plugins/Runtime/nDisplay/Source/DisplayCluster/Private/Render/Viewport/Configuration/DisplayClusterViewportConfigurationHelpers_RenderFrameSettings.cpp
+++ b/Engine/Plugins/Runtime/nDisplay/Source/DisplayCluster/Private/Render/Viewport/Configuration/DisplayClusterViewportConfigurationHelpers_RenderFrameSettings.cpp
@@ -199,7 +199,7 @@ void FDisplayClusterViewportConfigurationHelpers_RenderFrameSettings::PostUpdate
 		FDisplayClusterRenderFrameSettings NewRenderFrameSettings = InOutConfiguration.GetRenderFrameSettings();
 
 		// Update global flags that control viewport resources
-		NewRenderFrameSettings.bShouldUseOutputTargetableResources = ViewportManager->ShouldUseOutputTargetableResources();
+		NewRenderFrameSettings.bShouldUseOutputFrameTargetableResources = ViewportManager->ShouldUseOutputFrameTargetableResources();
 		NewRenderFrameSettings.bShouldUseAdditionalFrameTargetableResource = ViewportManager->ShouldUseAdditionalFrameTargetableResource();
 		NewRenderFrameSettings.bShouldUseFullSizeFrameTargetableResource = ViewportManager->ShouldUseFullSizeFrameTargetableResource();
 
diff --git a/Engine/Plugins/Runtime/nDisplay/Source/DisplayCluster/Private/Render/Viewport/Configuration/DisplayClusterViewportConfiguration_ICVFX.cpp b/Engine/Plugins/Runtime/nDisplay/Source/DisplayCluster/Private/Render/Viewport/Configuration/DisplayClusterViewportConfiguration_ICVFX.cpp
--- a/Engine/Plugins/Runtime/nDisplay/Source/DisplayCluster/Private/Render/Viewport/Configuration/DisplayClusterViewportConfiguration_ICVFX.cpp
+++ b/Engine/Plugins/Runtime/nDisplay/Source/DisplayCluster/Private/Render/Viewport/Configuration/DisplayClusterViewportConfiguration_ICVFX.cpp
@@ -168,6 +168,7 @@ void FDisplayClusterViewportConfiguration_ICVFX::Update()
 		}
 	}
 
+
 	for (const TSharedPtr<FDisplayClusterViewport, ESPMode::ThreadSafe>& TargetIt : TargetViewports)
 	{
 		if (TargetIt.IsValid())
@@ -175,13 +176,8 @@ void FDisplayClusterViewportConfiguration_ICVFX::Update()
 			// Gain direct access to internal settings of the viewport:
 			FDisplayClusterViewport_RenderSettingsICVFX& InOutRenderSettingsICVFX = TargetIt->GetRenderSettingsICVFXImpl();
 
-			// Reverse camera order when the option is set in the viewport.
-			const UDisplayClusterConfigurationViewport* ViewportConfig = TargetIt->GetViewportConfigurationData();
-			ensureMsgf(ViewportConfig, TEXT("Viewport '%s' did not have a config set"), *TargetIt->ViewportId);
-			const bool bReverseCameraPriority = ViewportConfig ? ViewportConfig->ICVFX.bReverseCameraPriority : false;
-
 			// Sort cameras by render order for each target
-			InOutRenderSettingsICVFX.ICVFX.SortCamerasRenderOrder(bReverseCameraPriority);
+			InOutRenderSettingsICVFX.ICVFX.SortCamerasRenderOrder(EnumHasAnyFlags(InOutRenderSettingsICVFX.Flags, EDisplayClusterViewportICVFXFlags::ReverseCameraPriority));
 
 			// Setup incamera overlap mode
 			InOutRenderSettingsICVFX.ICVFX.CameraOverlappingRenderMode = StageSettings->GetCameraOverlappingRenderMode();
@@ -333,8 +329,10 @@ void FDisplayClusterViewportConfiguration_ICVFX::GetAndUpdateStageCameras(const
 	// Allocate and assign camera resources
 	if (ImplGetStageCameras())
 	{
+		const FDisplayClusterConfigurationICVFX_StageSettings* StageSettings = Configuration.GetStageSettings();
+
 		// Collect visible targets for cameras:
-		if (InTargetViewports)
+		if (InTargetViewports && StageSettings)
 		{
 			for (const TSharedPtr<FDisplayClusterViewport, ESPMode::ThreadSafe>& TargetIt : *InTargetViewports)
 			{
@@ -342,14 +340,24 @@ void FDisplayClusterViewportConfiguration_ICVFX::GetAndUpdateStageCameras(const
 				if (TargetIt.IsValid() && !EnumHasAnyFlags(TargetIt->GetRenderSettingsICVFX().Flags, EDisplayClusterViewportICVFXFlags::DisableCamera))
 				{
 					// Add this target to all cameras visible on it
-					for (FDisplayClusterViewportConfiguration_ICVFXCamera& CameraIt : StageCameras)
+					for (FDisplayClusterViewportConfiguration_ICVFXCamera& StageCameraIt : StageCameras)
 					{
-						const bool bCameraProjectionVisible  = CameraIt.IsCameraProjectionVisibleOnViewport(TargetIt.Get());
-						const bool bIsCameraHiddenOnViewport = CameraIt.GetCameraSettings().HiddenICVFXViewports.ItemNames.Contains(TargetIt->GetId());
+						const bool bCameraProjectionVisible  = StageCameraIt.IsCameraProjectionVisibleOnViewport(TargetIt.Get());
+						const bool bIsCameraHiddenOnViewport = StageCameraIt.GetCameraSettings().HiddenICVFXViewports.ItemNames.Contains(TargetIt->GetId());
 						
 						if (bCameraProjectionVisible && !bIsCameraHiddenOnViewport)
 						{
-							CameraIt.VisibleTargets.Add(TargetIt);
+							// This ICVXF camera is visible in this viewport.
+							FDisplayClusterViewportConfiguration_ICVFXCamera::FTargetViewport TargetViewport(TargetIt.ToSharedRef());
+
+							// Update per-viewport chromakey for this ICVFX camera
+							if (const UDisplayClusterConfigurationViewport* ViewportCfg = TargetIt->GetViewportConfigurationData())
+							{
+								TargetViewport.ChromakeySource =
+									ViewportCfg->GetViewportChromakeyType(*StageSettings, StageCameraIt.GetCameraUniqueId(), StageCameraIt.GetCameraSettings());
+							}
+
+							StageCameraIt.TargetViewports.Add(TargetViewport);
 						}
 					}
 				}
@@ -368,7 +376,7 @@ void FDisplayClusterViewportConfiguration_ICVFX::GetAndUpdateStageCameras(const
 				MediaSettings.HasAnyMediaInputAssigned(ClusterNodeId, MediaSettings.SplitType) ||
 				MediaSettings.HasAnyMediaOutputAssigned(ClusterNodeId, MediaSettings.SplitType));
 
-			if (InTargetViewports == nullptr || CameraIt.VisibleTargets.Num() > 0 || (bIsRunningCluster && bCameraHasMedia))
+			if (InTargetViewports == nullptr || !CameraIt.TargetViewports.IsEmpty() || (bIsRunningCluster && bCameraHasMedia))
 			{
 				CameraIt.Update();
 			}
diff --git a/Engine/Plugins/Runtime/nDisplay/Source/DisplayCluster/Private/Render/Viewport/Configuration/DisplayClusterViewportConfiguration_ICVFXCamera.cpp b/Engine/Plugins/Runtime/nDisplay/Source/DisplayCluster/Private/Render/Viewport/Configuration/DisplayClusterViewportConfiguration_ICVFXCamera.cpp
--- a/Engine/Plugins/Runtime/nDisplay/Source/DisplayCluster/Private/Render/Viewport/Configuration/DisplayClusterViewportConfiguration_ICVFXCamera.cpp
+++ b/Engine/Plugins/Runtime/nDisplay/Source/DisplayCluster/Private/Render/Viewport/Configuration/DisplayClusterViewportConfiguration_ICVFXCamera.cpp
@@ -76,6 +76,9 @@ void FDisplayClusterViewportConfiguration_ICVFXCamera::Update()
 	{
 		if (CameraViewport.IsValid())
 		{
+			// Performance: Do not render InnerFrustum if it is not visible.
+			CameraViewport->GetRenderSettingsImpl().bSkipRendering = !EnableInnerFrustumRendering();
+
 			ADisplayClusterRootActor* SceneRootActor = CameraViewport->Configuration->GetRootActor(EDisplayClusterRootActorType::Scene);
 			ADisplayClusterRootActor* ConfigurationRootActor = CameraViewport->Configuration->GetRootActor(EDisplayClusterRootActorType::Configuration);
 			const FDisplayClusterConfigurationICVFX_StageSettings* StageSettings = CameraViewport->Configuration->GetStageSettings();
@@ -95,15 +98,13 @@ void FDisplayClusterViewportConfiguration_ICVFXCamera::Update()
 				ShaderParametersCameraSettings.RenderOrder = (CameraRenderOrder < 0) ? CameraSettings.RenderSettings.RenderOrder : CameraRenderOrder;
 
 				// Add this camera data to all visible targets:
-				for (TSharedPtr<FDisplayClusterViewport, ESPMode::ThreadSafe>& ViewportIt : VisibleTargets)
+				for (const FTargetViewport& TargetViewportIt : TargetViewports)
 				{
-					if (ViewportIt.IsValid())
-					{
-						// Gain direct access to internal settings of the viewport:
-						FDisplayClusterViewport_RenderSettingsICVFX& InOutRenderSettingsICVFX = ViewportIt->GetRenderSettingsICVFXImpl();
+					// Per-viewport CK
+					ShaderParametersCameraSettings.ChromakeySource = TargetViewportIt.ChromakeySource;
 
-						InOutRenderSettingsICVFX.ICVFX.Cameras.Add(ShaderParametersCameraSettings);
-					}
+					// Gain direct access to internal settings of the viewport:
+					TargetViewportIt.Viewport->GetRenderSettingsICVFXImpl().ICVFX.Cameras.Add(ShaderParametersCameraSettings);
 				}
 			}
 		}
@@ -153,6 +154,40 @@ const FDisplayClusterConfigurationICVFX_CameraSettings& FDisplayClusterViewportC
 	return ConfigurationCameraComponent.GetCameraSettingsICVFX();
 }
 
+FString FDisplayClusterViewportConfiguration_ICVFXCamera::GetCameraUniqueId() const
+{
+	return CameraComponent.GetCameraUniqueId();
+}
+
+bool FDisplayClusterViewportConfiguration_ICVFXCamera::EnableChromakeyRendering() const
+{
+	// Performance: Rnder CK only when it is in use.
+	for (const FTargetViewport& TargetViewportIt : TargetViewports)
+	{
+		if (TargetViewportIt.ChromakeySource == EDisplayClusterShaderParametersICVFX_ChromakeySource::ChromakeyLayers)
+		{
+			return true;
+		}
+	}
+
+	return false;
+}
+
+bool FDisplayClusterViewportConfiguration_ICVFXCamera::EnableInnerFrustumRendering() const
+{
+	// Performance: If all chromakey sources have a ‘FrameColor’ value for all viewports on the current cluster node,
+	// we can skip rendering the InnerFrustum
+	for (const FTargetViewport& TargetViewportIt : TargetViewports)
+	{
+		if (TargetViewportIt.ChromakeySource != EDisplayClusterShaderParametersICVFX_ChromakeySource::FrameColor)
+		{
+			return true;
+		}
+	}
+
+	return false;
+}
+
 bool FDisplayClusterViewportConfiguration_ICVFXCamera::ImplCreateChromakeyViewport()
 {
 	check(CameraViewport.IsValid());
@@ -190,7 +225,7 @@ bool FDisplayClusterViewportConfiguration_ICVFXCamera::CreateAndSetupInnerCamera
 	// Try create chromakey render on demand
 	if (const FDisplayClusterConfigurationICVFX_ChromakeyRenderSettings* ChromakeyRenderSettings = GetCameraSettings().Chromakey.GetChromakeyRenderSettings(*StageSettings))
 	{
-		if (ChromakeyRenderSettings->ShouldUseChromakeyViewport(*StageSettings))
+		if (ChromakeyRenderSettings->ShouldUseChromakeyViewport(*StageSettings) && EnableChromakeyRendering())
 		{
 			ImplCreateChromakeyViewport();
 		}
@@ -200,26 +235,20 @@ bool FDisplayClusterViewportConfiguration_ICVFXCamera::CreateAndSetupInnerCamera
 	const FString ChromakeyViewportId(ChromakeyViewport.IsValid() ? ChromakeyViewport->GetId() : TEXT(""));
 
 	// Assign this chromakey to all supported targets
-	for (const TSharedPtr<FDisplayClusterViewport, ESPMode::ThreadSafe>& OuterViewportIt : VisibleTargets)
+	for (const FTargetViewport& TargetViewportIt : TargetViewports)
 	{
-		if (OuterViewportIt.IsValid())
-		{
-			const bool bEnableChromakey = !EnumHasAnyFlags(OuterViewportIt->GetRenderSettingsICVFX().Flags, EDisplayClusterViewportICVFXFlags::DisableChromakey);
-			const bool bEnableChromakeyMarkers = !EnumHasAnyFlags(OuterViewportIt->GetRenderSettingsICVFX().Flags, EDisplayClusterViewportICVFXFlags::DisableChromakeyMarkers);
+		const bool bEnableChromakey = TargetViewportIt.ChromakeySource != EDisplayClusterShaderParametersICVFX_ChromakeySource::Disabled;
+		const bool bEnableChromakeyMarkers = bEnableChromakey && !EnumHasAnyFlags(TargetViewportIt.Viewport->GetRenderSettingsICVFX().Flags, EDisplayClusterViewportICVFXFlags::DisableChromakeyMarkers);
 
-			// Gain direct access to internal settings of the viewport:
-			FDisplayClusterViewport_RenderSettingsICVFX& InOutOuterViewportRenderSettingsICVFX = OuterViewportIt->GetRenderSettingsICVFXImpl();
-			FDisplayClusterShaderParameters_ICVFX::FCameraSettings& DstCameraData = InOutOuterViewportRenderSettingsICVFX.ICVFX.Cameras.Last();
+		// Gain direct access to internal settings of the viewport:
+		FDisplayClusterViewport_RenderSettingsICVFX& InOutOuterViewportRenderSettingsICVFX = TargetViewportIt.Viewport->GetRenderSettingsICVFXImpl();
+		FDisplayClusterShaderParameters_ICVFX::FCameraSettings& DstCameraData = InOutOuterViewportRenderSettingsICVFX.ICVFX.Cameras.Last();
 
-			if (bEnableChromakey)
-			{
-				// Setup chromakey with markers
-				FDisplayClusterViewportConfigurationHelpers_ICVFX::UpdateCameraSettings_Chromakey(DstCameraData, *StageSettings, GetCameraSettings(), bEnableChromakeyMarkers, ChromakeyViewportId);
-			}
+		// Setup chromakey with markers
+		FDisplayClusterViewportConfigurationHelpers_ICVFX::UpdateCameraSettings_Chromakey(DstCameraData, *StageSettings, GetCameraSettings(), bEnableChromakey, bEnableChromakeyMarkers, ChromakeyViewportId);
 
-			// Setup overlap chromakey with markers
-			FDisplayClusterViewportConfigurationHelpers_ICVFX::UpdateCameraSettings_OverlapChromakey(DstCameraData, *StageSettings, GetCameraSettings(), bEnableChromakeyMarkers);
-		}
+		// Setup overlap chromakey with markers
+		FDisplayClusterViewportConfigurationHelpers_ICVFX::UpdateCameraSettings_OverlapChromakey(DstCameraData, *StageSettings, GetCameraSettings(), bEnableChromakeyMarkers);
 	}
 
 	return true;
diff --git a/Engine/Plugins/Runtime/nDisplay/Source/DisplayCluster/Private/Render/Viewport/Configuration/DisplayClusterViewportConfiguration_ICVFXCamera.h b/Engine/Plugins/Runtime/nDisplay/Source/DisplayCluster/Private/Render/Viewport/Configuration/DisplayClusterViewportConfiguration_ICVFXCamera.h
--- a/Engine/Plugins/Runtime/nDisplay/Source/DisplayCluster/Private/Render/Viewport/Configuration/DisplayClusterViewportConfiguration_ICVFXCamera.h
+++ b/Engine/Plugins/Runtime/nDisplay/Source/DisplayCluster/Private/Render/Viewport/Configuration/DisplayClusterViewportConfiguration_ICVFXCamera.h
@@ -31,6 +31,15 @@ public:
 	/** Gets a ICVFX camera settings. */
 	const FDisplayClusterConfigurationICVFX_CameraSettings& GetCameraSettings() const;
 
+	/** Gets the name of the ICVFX camera. */
+	FString GetCameraUniqueId() const;
+
+	// Performance: Rnder CK only when it is in use.
+	bool EnableChromakeyRendering() const;
+
+	// Performance: If all chromakey sources have a ‘FrameColor’ value for all viewports on the current cluster node, we can skip rendering the InnerFrustum
+	bool EnableInnerFrustumRendering() const;
+
 protected:
 	/** Creating a inner camera viewport for the ICVFX rendering stack. */
 	bool CreateAndSetupInnerCameraViewport();
@@ -59,8 +68,21 @@ public:
 	// The inner camera chromakey viewport ref
 	TSharedPtr<FDisplayClusterViewport, ESPMode::ThreadSafe> ChromakeyViewport;
 
+	struct FTargetViewport
+	{
+		FTargetViewport(const TSharedRef<FDisplayClusterViewport, ESPMode::ThreadSafe>& InViewport)
+			: Viewport(InViewport)
+		{ }
+
+		// The reference to the viewport
+		const TSharedRef<FDisplayClusterViewport, ESPMode::ThreadSafe> Viewport;
+
+		// The camera chromakey type from this viewport.
+		EDisplayClusterShaderParametersICVFX_ChromakeySource ChromakeySource = EDisplayClusterShaderParametersICVFX_ChromakeySource::Disabled;
+	};
+
 	// List of OuterViewports for this camera
-	TArray<TSharedPtr<FDisplayClusterViewport, ESPMode::ThreadSafe>> VisibleTargets;
+	TArray<FTargetViewport> TargetViewports;
 
 private:
 	FDisplayClusterViewportConfiguration& Configuration;
diff --git a/Engine/Plugins/Runtime/nDisplay/Source/DisplayCluster/Private/Render/Viewport/Configuration/DisplayClusterViewportConfiguration_ProjectionPolicy.cpp b/Engine/Plugins/Runtime/nDisplay/Source/DisplayCluster/Private/Render/Viewport/Configuration/DisplayClusterViewportConfiguration_ProjectionPolicy.cpp
--- a/Engine/Plugins/Runtime/nDisplay/Source/DisplayCluster/Private/Render/Viewport/Configuration/DisplayClusterViewportConfiguration_ProjectionPolicy.cpp
+++ b/Engine/Plugins/Runtime/nDisplay/Source/DisplayCluster/Private/Render/Viewport/Configuration/DisplayClusterViewportConfiguration_ProjectionPolicy.cpp
@@ -172,7 +172,10 @@ bool FDisplayClusterViewportConfiguration_ProjectionPolicy::UpdateCameraPolicy_I
 			if (FDisplayClusterViewport* ExistCameraViewport = FDisplayClusterViewportConfigurationHelpers_ICVFX::FindCameraViewport(Configuration, *SceneCameraComp))
 			{
 				// Re-use camera RTT, and support OCIO:
-				const EDisplayClusterViewportOverrideMode ViewportOverrideMode = DstViewport.IsOpenColorIOEquals(*ExistCameraViewport) ? EDisplayClusterViewportOverrideMode::All : EDisplayClusterViewportOverrideMode::InernalRTT;
+				const EDisplayClusterViewportOverrideMode ViewportOverrideMode =
+					DstViewport.IsOpenColorIOEquals(*ExistCameraViewport)
+					? EDisplayClusterViewportOverrideMode::All
+					: EDisplayClusterViewportOverrideMode::InternalRTT;
 
 				// Gain direct access to internal settings of the viewport:
 				FDisplayClusterViewport_RenderSettings& InOutRenderSettings = DstViewport.GetRenderSettingsImpl();
diff --git a/Engine/Plugins/Runtime/nDisplay/Source/DisplayCluster/Private/Render/Viewport/Configuration/DisplayClusterViewportConfiguration_Viewport.cpp b/Engine/Plugins/Runtime/nDisplay/Source/DisplayCluster/Private/Render/Viewport/Configuration/DisplayClusterViewportConfiguration_Viewport.cpp
--- a/Engine/Plugins/Runtime/nDisplay/Source/DisplayCluster/Private/Render/Viewport/Configuration/DisplayClusterViewportConfiguration_Viewport.cpp
+++ b/Engine/Plugins/Runtime/nDisplay/Source/DisplayCluster/Private/Render/Viewport/Configuration/DisplayClusterViewportConfiguration_Viewport.cpp
@@ -53,6 +53,15 @@ void FDisplayClusterViewportConfiguration_Viewport::CreateOrUpdateViewportInstan
 	InOutConfiguration.SetClusterNodeId(CurrentClusterNodeId);
 }
 
+void FDisplayClusterViewportConfiguration_Viewport::PostUpdateViewportInstance(FDisplayClusterViewportConfiguration& InOutConfiguration)
+{
+	// Call the post-update event for the projection policy
+	if (Viewport.IsValid() && Viewport->GetProjectionPolicy().IsValid())
+	{
+		Viewport->GetProjectionPolicy()->PostUpdateBaseConfiguration(Viewport.Get());
+	}
+}
+
 bool FDisplayClusterViewportConfiguration_Viewport::UpdateViewportConfiguration(FDisplayClusterViewport& DstViewport, const UDisplayClusterConfigurationViewport& ConfigurationViewport)
 {
 	check(IsInGameThread());
diff --git a/Engine/Plugins/Runtime/nDisplay/Source/DisplayCluster/Private/Render/Viewport/Configuration/DisplayClusterViewportConfiguration_Viewport.h b/Engine/Plugins/Runtime/nDisplay/Source/DisplayCluster/Private/Render/Viewport/Configuration/DisplayClusterViewportConfiguration_Viewport.h
--- a/Engine/Plugins/Runtime/nDisplay/Source/DisplayCluster/Private/Render/Viewport/Configuration/DisplayClusterViewportConfiguration_Viewport.h
+++ b/Engine/Plugins/Runtime/nDisplay/Source/DisplayCluster/Private/Render/Viewport/Configuration/DisplayClusterViewportConfiguration_Viewport.h
@@ -29,6 +29,9 @@ struct FDisplayClusterViewportConfiguration_Viewport
 	 */
 	void CreateOrUpdateViewportInstance(FDisplayClusterViewportConfiguration& InOutConfiguration);
 
+	/** Post update viewport. */
+	void PostUpdateViewportInstance(FDisplayClusterViewportConfiguration& InOutConfiguration);
+
 	/**
 	 * Update exists viewport runtime settings from configuration
 	 */
diff --git a/Engine/Plugins/Runtime/nDisplay/Source/DisplayCluster/Private/Render/Viewport/Configuration/DisplayClusterViewportConfiguration_ViewportManager.cpp b/Engine/Plugins/Runtime/nDisplay/Source/DisplayCluster/Private/Render/Viewport/Configuration/DisplayClusterViewportConfiguration_ViewportManager.cpp
--- a/Engine/Plugins/Runtime/nDisplay/Source/DisplayCluster/Private/Render/Viewport/Configuration/DisplayClusterViewportConfiguration_ViewportManager.cpp
+++ b/Engine/Plugins/Runtime/nDisplay/Source/DisplayCluster/Private/Render/Viewport/Configuration/DisplayClusterViewportConfiguration_ViewportManager.cpp
@@ -102,6 +102,12 @@ void FDisplayClusterViewportConfiguration_ViewportManager::ImplUpdateViewports()
 		ViewportData.CreateOrUpdateViewportInstance(Configuration);
 	}
 
+	// Post update viewports for current rendering frame
+	for (FDisplayClusterViewportConfiguration_Viewport& ViewportData : CurrentFrameViewports)
+	{
+		ViewportData.PostUpdateViewportInstance(Configuration);
+	}
+
 	// Clear viewports warp policies for entire cluster
 	for (const TSharedPtr<FDisplayClusterViewport, ESPMode::ThreadSafe>& Viewport : ViewportManager->ImplGetEntireClusterViewports())
 	{
diff --git a/Engine/Plugins/Runtime/nDisplay/Source/DisplayCluster/Private/Render/Viewport/Containers/DisplayClusterViewport_InternalEnums.h b/Engine/Plugins/Runtime/nDisplay/Source/DisplayCluster/Private/Render/Viewport/Containers/DisplayClusterViewport_InternalEnums.h
--- a/Engine/Plugins/Runtime/nDisplay/Source/DisplayCluster/Private/Render/Viewport/Containers/DisplayClusterViewport_InternalEnums.h
+++ b/Engine/Plugins/Runtime/nDisplay/Source/DisplayCluster/Private/Render/Viewport/Containers/DisplayClusterViewport_InternalEnums.h
@@ -29,5 +29,7 @@ enum class EDisplayClusterViewportShowLogMsgOnce : uint16
 	GetViewPointCameraComponent_HasAssignedViewPoint = 1 << 10,
 	GetViewPointCameraComponent_NotFound = 1 << 11,
 	GetViewPointCameraComponent = GetViewPointCameraComponent_NoRootActorFound | GetViewPointCameraComponent_NotFound,
+
+	ChromakeyViewportHasZeroSize = 1 << 15,
 };
 ENUM_CLASS_FLAGS(EDisplayClusterViewportShowLogMsgOnce);
diff --git a/Engine/Plugins/Runtime/nDisplay/Source/DisplayCluster/Private/Render/Viewport/DisplayClusterViewport.cpp b/Engine/Plugins/Runtime/nDisplay/Source/DisplayCluster/Private/Render/Viewport/DisplayClusterViewport.cpp
--- a/Engine/Plugins/Runtime/nDisplay/Source/DisplayCluster/Private/Render/Viewport/DisplayClusterViewport.cpp
+++ b/Engine/Plugins/Runtime/nDisplay/Source/DisplayCluster/Private/Render/Viewport/DisplayClusterViewport.cpp
@@ -155,6 +155,16 @@ FDisplayClusterViewportProxyData* FDisplayClusterViewport::CreateViewportProxyDa
 	return OutViewportProxyData;
 }
 
+bool FDisplayClusterViewport::UseSameOCIO(const IDisplayClusterViewport* InViewportPtr) const
+{
+	if (const FDisplayClusterViewport* InViewport = static_cast<const FDisplayClusterViewport*>(InViewportPtr))
+	{
+		return IsOpenColorIOEquals(*InViewport);
+	}
+
+	return false;
+}
+
 void FDisplayClusterViewport::BeginNewFrame(const FIntPoint& InRenderFrameSize)
 {
 	check(IsInGameThread());
@@ -536,7 +546,7 @@ bool FDisplayClusterViewport::UpdateFrameContexts(const uint32 InStereoViewIndex
 	if (PostRenderSettings.GenerateMips.IsEnabled())
 	{
 		//Check if current projection policy supports this feature
-		if (!ProjectionPolicy.IsValid() || !ProjectionPolicy->ShouldUseSourceTextureWithMips())
+		if (!ProjectionPolicy.IsValid() || !ProjectionPolicy->ShouldUseSourceTextureWithMips(this))
 		{
 			// Don't create unused mips texture
 			PostRenderSettings.GenerateMips.Reset();
@@ -754,32 +764,37 @@ bool FDisplayClusterViewport::UpdateFrameContexts(const uint32 InStereoViewIndex
 	}
 
 	// Reserve for resources
-	if (ShouldUseRenderTargetResource())
+	if (IsResourceUsed(EDisplayClusterViewportResource::RenderTargets))
 	{
 		Resources[EDisplayClusterViewportResource::RenderTargets].AddZeroed(FrameTargetsAmount);
 	}
 
-	if (ShouldUseInternalResources())
+	if (IsResourceUsed(EDisplayClusterViewportResource::InputShaderResources))
 	{
 		Resources[EDisplayClusterViewportResource::InputShaderResources].AddZeroed(FrameTargetsAmount);
+	}
 
-		if (ShouldUseAdditionalTargetableResource())
-		{
-			Resources[EDisplayClusterViewportResource::AdditionalTargetableResources].AddZeroed(FrameTargetsAmount);
-		}
-
+	if (IsResourceUsed(EDisplayClusterViewportResource::MipsShaderResources))
+	{
 		// Setup Mips resources:
 		for (FDisplayClusterViewport_Context& ContextIt : Contexts)
 		{
 			ContextIt.NumMips = FDisplayClusterViewportHelpers::GetMaxTextureNumMips(InFrameSettings, PostRenderSettings.GenerateMips.GetRequiredNumMips(ContextIt.ContextSize));
 			if (ContextIt.NumMips > 1)
 			{
-				Resources[EDisplayClusterViewportResource::MipsShaderResources].AddZeroed(1);
+				Resources[EDisplayClusterViewportResource::MipsShaderResources].AddZeroed(FrameTargetsAmount);
+				break;
 			}
 		}
 	}
 
-	if (InFrameSettings.IsPreviewRendering() && ShouldUseOutputTargetableResources())
+	// The AdditionalTargetableResource is used as a warpblend output
+	if (IsResourceUsed(EDisplayClusterViewportResource::AdditionalTargetableResources))
+	{
+		Resources[EDisplayClusterViewportResource::AdditionalTargetableResources].AddZeroed(FrameTargetsAmount);
+	}
+
+	if (InFrameSettings.IsPreviewRendering() && IsResourceUsed(EDisplayClusterViewportResource::OutputPreviewTargetableResources))
 	{
 		// reserve preview texture resource for all visible viewports
 		Resources[EDisplayClusterViewportResource::OutputPreviewTargetableResources].AddZeroed(FrameTargetsAmount);
diff --git a/Engine/Plugins/Runtime/nDisplay/Source/DisplayCluster/Private/Render/Viewport/DisplayClusterViewport.h b/Engine/Plugins/Runtime/nDisplay/Source/DisplayCluster/Private/Render/Viewport/DisplayClusterViewport.h
--- a/Engine/Plugins/Runtime/nDisplay/Source/DisplayCluster/Private/Render/Viewport/DisplayClusterViewport.h
+++ b/Engine/Plugins/Runtime/nDisplay/Source/DisplayCluster/Private/Render/Viewport/DisplayClusterViewport.h
@@ -161,6 +161,8 @@ public:
 		return CustomPostProcessSettings;
 	}
 
+	virtual bool UseSameOCIO(const IDisplayClusterViewport* InViewportPtr) const override;
+
 	// Setup scene view for rendering specified Context
 	virtual void SetupSceneView(uint32 ContextNum, class UWorld* World, class FSceneViewFamily& InViewFamily, FSceneView& InView) const override;
 
@@ -173,12 +175,7 @@ public:
 
 
 	// Get from logic request for additional targetable resource
-	bool ShouldUseAdditionalTargetableResource() const;
-	bool ShouldUseAdditionalFrameTargetableResource() const;
 	bool ShouldUseFullSizeFrameTargetableResource() const;
-	
-	// Return true if this viewport requires to use any of output resources (OutputPreviewTargetableResources or OutputFrameTargetableResources)
-	bool ShouldUseOutputTargetableResources() const;
 
 	void SetViewportBufferRatio(const float InBufferRatio);
 
@@ -336,11 +333,8 @@ public:
 	/** Returns true if per-view Lumen scene is enabled. */
 	bool ShouldUseLumenPerView() const;
 
-	/** Returns true if this viewport should use the RTT. */
-	bool ShouldUseRenderTargetResource() const;
-
-	/** Returns true if this viewport should use internal resources such as Input, Mips, Additional. */
-	bool ShouldUseInternalResources() const;
+	/** Returns true if this viewport should use this type of resource. */
+	bool IsResourceUsed(const EDisplayClusterViewportResource InResourceType) const;
 
 	/** Start a new frame with the specified size. The associated objects and geometries will be updated accordingly. */
 	void BeginNewFrame(const FIntPoint& InRenderFrameSize);
@@ -423,6 +417,11 @@ private:
 	FIntPoint GetDesiredContextSize(const FIntPoint& InSize, const FDisplayClusterRenderFrameSettings& InFrameSettings) const;
 	float GetCustomBufferRatio(const FDisplayClusterRenderFrameSettings& InFrameSettings) const;
 
+	/** Returns true if this viewport should use this type of resource.
+	* This is function can be used only by the IsResourceUsed();
+	*/
+	bool IsResourceUsedImpl(const EDisplayClusterViewportResource InResourceType) const;
+
 public:
 	// Configuration of the current cluster node
 	const TSharedRef<FDisplayClusterViewportConfiguration, ESPMode::ThreadSafe> Configuration;
diff --git a/Engine/Plugins/Runtime/nDisplay/Source/DisplayCluster/Private/Render/Viewport/DisplayClusterViewportManager.cpp b/Engine/Plugins/Runtime/nDisplay/Source/DisplayCluster/Private/Render/Viewport/DisplayClusterViewportManager.cpp
--- a/Engine/Plugins/Runtime/nDisplay/Source/DisplayCluster/Private/Render/Viewport/DisplayClusterViewportManager.cpp
+++ b/Engine/Plugins/Runtime/nDisplay/Source/DisplayCluster/Private/Render/Viewport/DisplayClusterViewportManager.cpp
@@ -461,13 +461,13 @@ bool FDisplayClusterViewportManager::ShouldUseFullSizeFrameTargetableResource()
 	return false;
 }
 
-bool FDisplayClusterViewportManager::ShouldUseOutputTargetableResources() const
+bool FDisplayClusterViewportManager::ShouldUseOutputFrameTargetableResources() const
 {
 	check(IsInGameThread());
 
 	for (const TSharedPtr<FDisplayClusterViewport, ESPMode::ThreadSafe>& Viewport : ImplGetCurrentRenderFrameViewports())
 	{
-		if (Viewport.IsValid() && Viewport->ShouldUseOutputTargetableResources())
+		if (Viewport.IsValid() && Viewport->IsResourceUsed(EDisplayClusterViewportResource::OutputFrameTargetableResources))
 		{
 			return true;
 		}
@@ -487,7 +487,7 @@ bool FDisplayClusterViewportManager::ShouldUseAdditionalFrameTargetableResource(
 
 	for (const TSharedPtr<FDisplayClusterViewport, ESPMode::ThreadSafe>& Viewport : ImplGetCurrentRenderFrameViewports())
 	{
-		if (Viewport.IsValid() && Viewport->ShouldUseAdditionalFrameTargetableResource())
+		if (Viewport.IsValid() && Viewport->IsResourceUsed(EDisplayClusterViewportResource::AdditionalFrameTargetableResources))
 		{
 			return true;
 		}
diff --git a/Engine/Plugins/Runtime/nDisplay/Source/DisplayCluster/Private/Render/Viewport/DisplayClusterViewportManager.h b/Engine/Plugins/Runtime/nDisplay/Source/DisplayCluster/Private/Render/Viewport/DisplayClusterViewportManager.h
--- a/Engine/Plugins/Runtime/nDisplay/Source/DisplayCluster/Private/Render/Viewport/DisplayClusterViewportManager.h
+++ b/Engine/Plugins/Runtime/nDisplay/Source/DisplayCluster/Private/Render/Viewport/DisplayClusterViewportManager.h
@@ -185,9 +185,8 @@ public:
 	TSharedPtr<class FDisplayClusterViewportManagerViewExtension, ESPMode::ThreadSafe> GetViewportManagerViewExtension() const
 	{ return ViewportManagerViewExtension; }
 
-	/** Returns true if any of the viewports are visible and should use the output RTT resources. */
-	bool ShouldUseOutputTargetableResources() const;
-
+	/** Returns true if at least one of the viewports are visible on the frame RTT. */
+	bool ShouldUseOutputFrameTargetableResources() const;
 	bool ShouldUseAdditionalFrameTargetableResource() const;
 	bool ShouldUseFullSizeFrameTargetableResource() const;
 
diff --git a/Engine/Plugins/Runtime/nDisplay/Source/DisplayCluster/Private/Render/Viewport/DisplayClusterViewportProxy.cpp b/Engine/Plugins/Runtime/nDisplay/Source/DisplayCluster/Private/Render/Viewport/DisplayClusterViewportProxy.cpp
--- a/Engine/Plugins/Runtime/nDisplay/Source/DisplayCluster/Private/Render/Viewport/DisplayClusterViewportProxy.cpp
+++ b/Engine/Plugins/Runtime/nDisplay/Source/DisplayCluster/Private/Render/Viewport/DisplayClusterViewportProxy.cpp
@@ -113,10 +113,8 @@ bool FDisplayClusterViewportProxy::GetResources_RenderThread(const EDisplayClust
 
 const FDisplayClusterViewportProxy& FDisplayClusterViewportProxy::GetRenderingViewportProxy() const
 {
-	switch (RenderSettings.GetViewportOverrideMode())
+	if(RenderSettings.GetViewportOverrideMode() != EDisplayClusterViewportOverrideMode::None)
 	{
-	case EDisplayClusterViewportOverrideMode::All:
-	case EDisplayClusterViewportOverrideMode::InernalRTT:
 		if (FDisplayClusterViewportManagerProxy* ViewportManagerProxy = ConfigurationProxy->GetViewportManagerProxyImpl())
 		{
 			if (FDisplayClusterViewportProxy const* OverrideViewportProxy = ViewportManagerProxy->ImplFindViewportProxy_RenderThread(RenderSettings.GetViewportOverrideId()))
@@ -124,10 +122,7 @@ const FDisplayClusterViewportProxy& FDisplayClusterViewportProxy::GetRenderingVi
 				return *OverrideViewportProxy;
 			}
 		}
-
-		break;
 	}
-
 	return *this;
 }
 
@@ -283,10 +278,15 @@ void FDisplayClusterViewportProxy::UpdateDeferredResources(FRHICommandListImmedi
 		return;
 	}
 
-	if (RenderSettings.GetViewportOverrideMode() == EDisplayClusterViewportOverrideMode::All)
+	switch (RenderSettings.GetViewportOverrideMode())
 	{
+		case EDisplayClusterViewportOverrideMode::All:
+		case EDisplayClusterViewportOverrideMode::InternalViewportResources:
 		// Disable deferred update for clone viewports
 		return;
+
+		default:
+			break;
 	}
 
 	const FDisplayClusterViewportProxy& SourceViewportProxy = GetRenderingViewportProxy();
diff --git a/Engine/Plugins/Runtime/nDisplay/Source/DisplayCluster/Private/Render/Viewport/DisplayClusterViewportProxy_Rules.cpp b/Engine/Plugins/Runtime/nDisplay/Source/DisplayCluster/Private/Render/Viewport/DisplayClusterViewportProxy_Rules.cpp
--- a/Engine/Plugins/Runtime/nDisplay/Source/DisplayCluster/Private/Render/Viewport/DisplayClusterViewportProxy_Rules.cpp
+++ b/Engine/Plugins/Runtime/nDisplay/Source/DisplayCluster/Private/Render/Viewport/DisplayClusterViewportProxy_Rules.cpp
@@ -171,7 +171,7 @@ bool FDisplayClusterViewportProxy::ShouldApplyWarpBlend_RenderThread() const
 	}
 
 	// Ask current projection policy if it's warp&blend compatible
-	return ProjectionPolicy.IsValid() && ProjectionPolicy->IsWarpBlendSupported();
+	return ProjectionPolicy.IsValid() && ProjectionPolicy->IsWarpBlendSupported_RenderThread(this);
 }
 
 bool FDisplayClusterViewportProxy::ShouldOverrideViewportResource(const EDisplayClusterViewportResourceType InExtResourceType) const
@@ -197,7 +197,22 @@ bool FDisplayClusterViewportProxy::ShouldOverrideViewportResource(const EDisplay
 		}
 		break;
 
-		case EDisplayClusterViewportOverrideMode::InernalRTT:
+		case EDisplayClusterViewportOverrideMode::InternalViewportResources:
+		{
+			switch (GetResourceType_RenderThread(InExtResourceType))
+			{
+			case EDisplayClusterViewportResourceType::InternalRenderTargetResource:
+			case EDisplayClusterViewportResourceType::InputShaderResource:
+			case EDisplayClusterViewportResourceType::MipsShaderResource:
+				return true;
+
+			default:
+				break;
+			}
+		}
+		break;
+
+		case EDisplayClusterViewportOverrideMode::InternalRTT:
 		{
 			switch (GetResourceType_RenderThread(InExtResourceType))
 			{
diff --git a/Engine/Plugins/Runtime/nDisplay/Source/DisplayCluster/Private/Render/Viewport/DisplayClusterViewport_Rules.cpp b/Engine/Plugins/Runtime/nDisplay/Source/DisplayCluster/Private/Render/Viewport/DisplayClusterViewport_Rules.cpp
--- a/Engine/Plugins/Runtime/nDisplay/Source/DisplayCluster/Private/Render/Viewport/DisplayClusterViewport_Rules.cpp
+++ b/Engine/Plugins/Runtime/nDisplay/Source/DisplayCluster/Private/Render/Viewport/DisplayClusterViewport_Rules.cpp
@@ -288,29 +288,11 @@ bool FDisplayClusterViewport::CanSplitIntoTiles() const
 	return true;
 }
 
-bool FDisplayClusterViewport::ShouldUseRenderTargetResource() const
+bool FDisplayClusterViewport::IsResourceUsedImpl(const EDisplayClusterViewportResource InResourceType) const
 {
 	if (RenderSettings.bSkipRendering)
 	{
-		// When rendering is skipped, the RTT resources are not used.
-		// For example this feature is used when the ICVFX camera uses full-frame chromakey colour,
-		// thereby eliminating the ICVFX camera's viewport rendering resources for optimization.
-		return false;
-	}
-
-	if (IsExternalRendering())
-	{
-		return false;
-	}
-
-	return true;
-}
-
-bool FDisplayClusterViewport::ShouldUseInternalResources() const
-{
-	if (RenderSettings.bSkipRendering)
-	{
-		// When rendering is skipped, internal resources are not used.
+		// When rendering is skipped, the resources isn't used
 		// For example this feature is used when the ICVFX camera uses full-frame chromakey colour,
 		// thereby eliminating the ICVFX camera's viewport rendering resources for optimization.
 		return false;
@@ -320,70 +302,121 @@ bool FDisplayClusterViewport::ShouldUseInternalResources() const
 	{
 		// The tiled viewport does not use internal resources.
 		// Since rendering results are copied between RTTs.
-		return false;
+		if (InResourceType != EDisplayClusterViewportResource::RenderTargets)
+		{
+			return false;
+		}
 	}
 
+	if (InResourceType == EDisplayClusterViewportResource::RenderTargets)
+	{
+		// When external rendering is used, this means that the RTT is a reference, not a texture.
+		if (IsExternalRendering())
+		{
+			return false;
+		}
+	}
 
 	if (RenderSettings.IsViewportOverridden())
 	{
 		switch (RenderSettings.GetViewportOverrideMode())
 		{
+		// Override all resources
 		case EDisplayClusterViewportOverrideMode::All:
+			switch (InResourceType)
+			{
+			case EDisplayClusterViewportResource::RenderTargets:
+			case EDisplayClusterViewportResource::InputShaderResources:
+			case EDisplayClusterViewportResource::MipsShaderResources:
+			case EDisplayClusterViewportResource::AdditionalTargetableResources:
+				return false;
+
+			default:
+				break;
+			}
+			break;
+
+		// Override all internal resources except the texture for warpblend
+		case EDisplayClusterViewportOverrideMode::InternalViewportResources:
+			switch (InResourceType)
+			{
+			case EDisplayClusterViewportResource::RenderTargets:
+			case EDisplayClusterViewportResource::InputShaderResources:
+			case EDisplayClusterViewportResource::MipsShaderResources:
+				return false;
+
+			default:
+				break;
+			}
+			break;
+
+		case EDisplayClusterViewportOverrideMode::InternalRTT:
 			// When all resources in a viewport are overridden from another viewport.
-			return false;
+			if (InResourceType == EDisplayClusterViewportResource::RenderTargets)
+			{
+				return false;
+			}
+			break;
 
 		default:
 			break;
 		}
 	}
 
-	return true;
-}
-
-bool FDisplayClusterViewport::ShouldUseAdditionalTargetableResource() const
-{
-	check(IsInGameThread());
-
-	// PostRender Blur require additional RTT for shader
-	if (PostRenderSettings.PostprocessBlur.IsEnabled())
-	{
-		return true;
-	}
-
-	// Supoport projection policy additional resource
-	if (ProjectionPolicy.IsValid() && ProjectionPolicy->ShouldUseAdditionalTargetableResource())
+	// These are exceptions to the rules, for each case we must return true or false:
+	switch (InResourceType)
 	{
-		return true;
-	}
+	case EDisplayClusterViewportResource::AdditionalTargetableResources:
+		// PostRender Blur require additional RTT for shader
+		if (PostRenderSettings.PostprocessBlur.IsEnabled())
+		{
+			return true;
+		}
 
-	return false;
-}
+		// Supoport projection policy additional resource
+		if (ProjectionPolicy.IsValid() && ProjectionPolicy->ShouldUseAdditionalTargetableResource(this))
+		{
+			return true;
+		}
 
-bool FDisplayClusterViewport::ShouldUseOutputTargetableResources() const
-{
-	// Do not create output RTTs for internal viewport
-	if (IsInternalViewport())
-	{
 		return false;
+
+		// Return true if this viewport requires to use any of output resources (OutputPreviewTargetableResources or OutputFrameTargetableResources)
+	case EDisplayClusterViewportResource::OutputPreviewTargetableResources:
+	case EDisplayClusterViewportResource::OutputFrameTargetableResources:
+		// Only if this viewport is enabled and visible on the final frame.
+		return RenderSettings.bEnable && RenderSettings.bVisible;
+
+	default:
+		break;
 	}
 
-	// Only if this viewport is enabled and visible on the final frame.
-	return RenderSettings.bEnable && RenderSettings.bVisible;
+	return true;
 }
 
-bool FDisplayClusterViewport::ShouldUseAdditionalFrameTargetableResource() const
+bool FDisplayClusterViewport::IsResourceUsed(const EDisplayClusterViewportResource InResourceType) const
 {
-	if (ShouldUseOutputTargetableResources())
+	switch (InResourceType)
 	{
+	case EDisplayClusterViewportResource::AdditionalFrameTargetableResources:
 		// OutputFrameTargetableResources must be used for AdditionalFrameTargetableResource
-		if (ViewportRemap.IsUsed())
+		if (IsResourceUsedImpl(EDisplayClusterViewportResource::OutputFrameTargetableResources))
 		{
-			return true;
+			// viewport remap requires AdditionalFrameTargetableResources.
+			// should be reviewed in the future.
+			if (ViewportRemap.IsUsed())
+			{
+				return true;
+			}
 		}
+		return false;
+
+	default:
+		break;
 	}
 
-	return false;
-}
+	return IsResourceUsedImpl(InResourceType);
+};
 
 bool FDisplayClusterViewport::ShouldUseFullSizeFrameTargetableResource() const
 {
diff --git a/Engine/Plugins/Runtime/nDisplay/Source/DisplayCluster/Private/Render/Viewport/Preview/DisplayClusterViewportManagerPreview.cpp b/Engine/Plugins/Runtime/nDisplay/Source/DisplayCluster/Private/Render/Viewport/Preview/DisplayClusterViewportManagerPreview.cpp
--- a/Engine/Plugins/Runtime/nDisplay/Source/DisplayCluster/Private/Render/Viewport/Preview/DisplayClusterViewportManagerPreview.cpp
+++ b/Engine/Plugins/Runtime/nDisplay/Source/DisplayCluster/Private/Render/Viewport/Preview/DisplayClusterViewportManagerPreview.cpp
@@ -249,7 +249,7 @@ const TArray<TSharedPtr<FDisplayClusterViewportPreview, ESPMode::ThreadSafe>> FD
 	{
 		for (const TSharedPtr<FDisplayClusterViewport, ESPMode::ThreadSafe>& ViewportIt : ViewportManager->ImplGetEntireClusterViewports())
 		{
-			if (ViewportIt.IsValid() && ViewportIt->ShouldUseOutputTargetableResources())
+			if (ViewportIt.IsValid() && ViewportIt->IsResourceUsed(EDisplayClusterViewportResource::OutputPreviewTargetableResources))
 			{
 				OutViewports.Add(ViewportIt->ViewportPreview);
 			}
diff --git a/Engine/Plugins/Runtime/nDisplay/Source/DisplayCluster/Private/Render/Viewport/RenderFrame/DisplayClusterRenderFrameSettings.cpp b/Engine/Plugins/Runtime/nDisplay/Source/DisplayCluster/Private/Render/Viewport/RenderFrame/DisplayClusterRenderFrameSettings.cpp
--- a/Engine/Plugins/Runtime/nDisplay/Source/DisplayCluster/Private/Render/Viewport/RenderFrame/DisplayClusterRenderFrameSettings.cpp
+++ b/Engine/Plugins/Runtime/nDisplay/Source/DisplayCluster/Private/Render/Viewport/RenderFrame/DisplayClusterRenderFrameSettings.cpp
@@ -168,7 +168,7 @@ bool FDisplayClusterRenderFrameSettings::ShouldUseOutputFrameTargetableResources
 		break;
 	}
 
-	return bShouldUseOutputTargetableResources;
+	return bShouldUseOutputFrameTargetableResources;
 }
 
 bool FDisplayClusterRenderFrameSettings::ShouldUseStereoRenderingOnMonoscopicDisplay() const
diff --git a/Engine/Plugins/Runtime/nDisplay/Source/DisplayCluster/Private/Render/Viewport/RenderFrame/DisplayClusterRenderFrameSettings.h b/Engine/Plugins/Runtime/nDisplay/Source/DisplayCluster/Private/Render/Viewport/RenderFrame/DisplayClusterRenderFrameSettings.h
--- a/Engine/Plugins/Runtime/nDisplay/Source/DisplayCluster/Private/Render/Viewport/RenderFrame/DisplayClusterRenderFrameSettings.h
+++ b/Engine/Plugins/Runtime/nDisplay/Source/DisplayCluster/Private/Render/Viewport/RenderFrame/DisplayClusterRenderFrameSettings.h
@@ -43,7 +43,7 @@ struct FDisplayClusterRenderFrameSettings
 	bool bShouldUseFullSizeFrameTargetableResource = false;
 
 	// Create output resources only for visible viewports. This will save GPU memory space.
-	bool bShouldUseOutputTargetableResources = false;
+	bool bShouldUseOutputFrameTargetableResources = false;
 
 
 	// Multiply all viewports RTT size's for whole cluster by this value
diff --git a/Engine/Plugins/Runtime/nDisplay/Source/DisplayCluster/Private/Render/Viewport/RenderTarget/DisplayClusterRenderTargetResource.h b/Engine/Plugins/Runtime/nDisplay/Source/DisplayCluster/Private/Render/Viewport/RenderTarget/DisplayClusterRenderTargetResource.h
--- a/Engine/Plugins/Runtime/nDisplay/Source/DisplayCluster/Private/Render/Viewport/RenderTarget/DisplayClusterRenderTargetResource.h
+++ b/Engine/Plugins/Runtime/nDisplay/Source/DisplayCluster/Private/Render/Viewport/RenderTarget/DisplayClusterRenderTargetResource.h
@@ -74,7 +74,7 @@ public:
 
 public:
 	// OutputFrameTargetableResources frame offset on backbuffer (special for 'side_by_side' and 'top_bottom' DCRenderDevices)
-	FIntPoint BackbufferFrameOffset;
+	FIntPoint BackbufferFrameOffset = FIntPoint::ZeroValue;
 };
 
 /**
diff --git a/Engine/Plugins/Runtime/nDisplay/Source/DisplayCluster/Public/Render/Projection/IDisplayClusterProjectionPolicy.h b/Engine/Plugins/Runtime/nDisplay/Source/DisplayCluster/Public/Render/Projection/IDisplayClusterProjectionPolicy.h
--- a/Engine/Plugins/Runtime/nDisplay/Source/DisplayCluster/Public/Render/Projection/IDisplayClusterProjectionPolicy.h
+++ b/Engine/Plugins/Runtime/nDisplay/Source/DisplayCluster/Public/Render/Projection/IDisplayClusterProjectionPolicy.h
@@ -82,6 +82,14 @@ public:
 	virtual void HandleEndScene(IDisplayClusterViewport* InViewport)
 	{ }
 
+	/**
+	 * Called after the base configuration is completed, but before the ICVFX
+	 *
+	 * @param InViewport - a owner viewport
+	 */
+	virtual void PostUpdateBaseConfiguration(IDisplayClusterViewport* InViewport) const
+	{ }
+
 	/**
 	 * Called before FDisplayClusterViewport::UpdateFrameContexts()
 	 * From this function, the policy can override any viewport settings (custom overscan, etc).
@@ -124,7 +132,7 @@ public:
 	}
 
 	// Handle request for additional render targetable resource inside viewport api for projection policy
-	virtual bool ShouldUseAdditionalTargetableResource() const
+	virtual bool ShouldUseAdditionalTargetableResource(const IDisplayClusterViewport* InViewport) const
 	{ 
 		return false; 
 	}
@@ -135,7 +143,7 @@ public:
 	*
 	* @return - true, if mip-texture is supported by the policy implementation
 	*/
-	virtual bool ShouldUseSourceTextureWithMips() const
+	virtual bool ShouldUseSourceTextureWithMips(const IDisplayClusterViewport* InViewport) const
 	{
 		return false;
 	}
@@ -144,7 +152,7 @@ public:
 	 * 
 	 * @param InViewport - a owner viewport
 	 */
-	virtual bool ShouldSupportICVFX(IDisplayClusterViewport* InViewport) const
+	virtual bool ShouldSupportICVFX(const IDisplayClusterViewport* InViewport) const
 	{
 		return false;
 	}
@@ -230,7 +238,18 @@ public:
 	*
 	* @return - True if warp&blend operations are supported by the policy implementation
 	*/
-	virtual bool IsWarpBlendSupported()
+	virtual bool IsWarpBlendSupported(const IDisplayClusterViewport* InViewport) const
+	{
+		return false;
+	}
+
+	/**
+	* Returns if a policy provides warp&blend feature
+	* (
+	*
+	* @return - True if warp&blend operations are supported by the policy implementation
+	*/
+	virtual bool IsWarpBlendSupported_RenderThread(const IDisplayClusterViewportProxy* InViewportProxy) const
 	{
 		return false;
 	}
@@ -263,7 +282,8 @@ public:
 	}
 
 	/**
-	* Initializing the projection policy logic for the current frame before applying warp blending. Called if IsWarpBlendSupported() returns true
+	* Initializing the projection policy logic for the current frame before applying warp blending.
+	* Have to be called only if IsWarpBlendSupported_RenderThread() returns true
 	*
 	* @param RHICmdList      - RHI commands
 	* @param InViewportProxy - viewport proxy
@@ -273,7 +293,8 @@ public:
 	{ }
 
 	/**
-	* Performs warp&blend. Called if IsWarpBlendSupported() returns true
+	* Performs Warp and Blend on the rendering thread
+	* Have to be called only if IsWarpBlendSupported_RenderThread() returns true
 	*
 	* @param RHICmdList      - RHI commands
 	* @param InViewportProxy - viewport proxy
@@ -283,7 +304,8 @@ public:
 	{ }
 
 	/**
-	* Completing the projection policy logic for the current frame after applying warp blending. Called if IsWarpBlendSupported() returns true
+	* Completing the projection policy logic for the current frame after applying warp blending.
+	* Have to be called only if IsWarpBlendSupported_RenderThread() returns true
 	*
 	* @param ViewIdx      - Index of view that is being processed for this viewport
 	* @param RHICmdList   - RHI commands
@@ -404,7 +426,7 @@ public:
 
 	// This policy can support ICVFX rendering
 	UE_DEPRECATED(5.3, "This function has been deprecated. Please use 'ShouldSupportICVFX(IDisplayClusterViewport*)'.")
-		virtual bool ShouldSupportICVFX() const
+	virtual bool ShouldSupportICVFX() const
 	{
 		return false;
 	}
@@ -421,4 +443,35 @@ public:
 	{
 		return false;
 	}
+
+	/**
+	* Returns if a policy provides warp&blend feature
+	*
+	* @return - True if warp&blend operations are supported by the policy implementation
+	*/
+	UE_DEPRECATED(5.6, "This function has been deprecated. Please use 'IsWarpBlendSupported(IDisplayClusterViewport*) or IsWarpBlendSupported_RenderThread(IDisplayClusterViewportProxy*)'.")
+	virtual bool IsWarpBlendSupported()
+	{
+		return false;
+	}
+
+	// Handle request for additional render targetable resource inside viewport api for projection policy
+	UE_DEPRECATED(5.6, "This function has been deprecated. Please use 'ShouldUseAdditionalTargetableResource(IDisplayClusterViewport*).")
+	virtual bool ShouldUseAdditionalTargetableResource() const
+	{
+		return false;
+	}
+
+	/**
+	* Returns true if the policy supports input mip-textures.
+	* Use a mip texture for smoother deformation on curved surfaces.
+	*
+	* @return - true, if mip-texture is supported by the policy implementation
+	*/
+	UE_DEPRECATED(5.6, "This function has been deprecated. Please use 'ShouldUseSourceTextureWithMips(IDisplayClusterViewport*).")
+	virtual bool ShouldUseSourceTextureWithMips() const
+	{
+		return false;
+	}
+
 };
diff --git a/Engine/Plugins/Runtime/nDisplay/Source/DisplayCluster/Public/Render/Viewport/Containers/DisplayClusterViewport_Enums.h b/Engine/Plugins/Runtime/nDisplay/Source/DisplayCluster/Public/Render/Viewport/Containers/DisplayClusterViewport_Enums.h
--- a/Engine/Plugins/Runtime/nDisplay/Source/DisplayCluster/Public/Render/Viewport/Containers/DisplayClusterViewport_Enums.h
+++ b/Engine/Plugins/Runtime/nDisplay/Source/DisplayCluster/Public/Render/Viewport/Containers/DisplayClusterViewport_Enums.h
@@ -114,7 +114,11 @@ enum class EDisplayClusterViewportOverrideMode : uint8
 
 	// Override internalRTT from the other viewport (Don't render this viewport; Don't create RTT resource)
 	// Useful for custom PP on the same InRTT. (OCIO per-viewport\node)
-	InernalRTT,
+	InternalRTT,
+
+	// Overrides only internal viewport resources
+	// The output RTT is not overridden and must use rendering.
+	InternalViewportResources,
 
 	// Override all - clone viewport (Dont render; Don't create resources;)
 	All
diff --git a/Engine/Plugins/Runtime/nDisplay/Source/DisplayCluster/Public/Render/Viewport/Containers/DisplayClusterViewport_EnumsICVFX.h b/Engine/Plugins/Runtime/nDisplay/Source/DisplayCluster/Public/Render/Viewport/Containers/DisplayClusterViewport_EnumsICVFX.h
--- a/Engine/Plugins/Runtime/nDisplay/Source/DisplayCluster/Public/Render/Viewport/Containers/DisplayClusterViewport_EnumsICVFX.h
+++ b/Engine/Plugins/Runtime/nDisplay/Source/DisplayCluster/Public/Render/Viewport/Containers/DisplayClusterViewport_EnumsICVFX.h
@@ -7,7 +7,7 @@
 /**
  * Runtime configuration from DCRA.
  */
-enum class EDisplayClusterViewportICVFXFlags : uint8
+enum class EDisplayClusterViewportICVFXFlags : uint16
 {
 	None = 0,
 
@@ -37,6 +37,9 @@ enum class EDisplayClusterViewportICVFXFlags : uint8
 
 	// Mask that get only flags that define the light card render mode.
 	LightcardRenderModeMask = LightcardAlwaysUnder | LightcardAlwaysOver | LightcardUseStageSettings,
+
+	/** The order in which the ICVFX cameras are composited over is reversed. Useful for time-multiplexed displays. */
+	ReverseCameraPriority = 1 << 8,
 };
 ENUM_CLASS_FLAGS(EDisplayClusterViewportICVFXFlags);
 
diff --git a/Engine/Plugins/Runtime/nDisplay/Source/DisplayCluster/Public/Render/Viewport/IDisplayClusterViewport.h b/Engine/Plugins/Runtime/nDisplay/Source/DisplayCluster/Public/Render/Viewport/IDisplayClusterViewport.h
--- a/Engine/Plugins/Runtime/nDisplay/Source/DisplayCluster/Public/Render/Viewport/IDisplayClusterViewport.h
+++ b/Engine/Plugins/Runtime/nDisplay/Source/DisplayCluster/Public/Render/Viewport/IDisplayClusterViewport.h
@@ -60,6 +60,9 @@ public:
 	/** Get post-render settings. */
 	virtual const FDisplayClusterViewport_PostRenderSettings&  GetPostRenderSettings() const = 0;
 
+	/** Returns true if both viewports use the same OCIO. */
+	virtual bool UseSameOCIO(const IDisplayClusterViewport* InViewport) const = 0;
+
 	/** Calculate projection matrices for viewport context
 	 * (also created overscan matrices, etc)
 	 * 
diff --git a/Engine/Plugins/Runtime/nDisplay/Source/DisplayCluster/Public/Render/Warp/IDisplayClusterWarpPolicy.h b/Engine/Plugins/Runtime/nDisplay/Source/DisplayCluster/Public/Render/Warp/IDisplayClusterWarpPolicy.h
--- a/Engine/Plugins/Runtime/nDisplay/Source/DisplayCluster/Public/Render/Warp/IDisplayClusterWarpPolicy.h
+++ b/Engine/Plugins/Runtime/nDisplay/Source/DisplayCluster/Public/Render/Warp/IDisplayClusterWarpPolicy.h
@@ -43,7 +43,7 @@ public:
 	 * 
 	 * @return - true if this warp policy supports ICVFX
 	 */
-	virtual bool ShouldSupportICVFX(IDisplayClusterViewport* InViewport) const
+	virtual bool ShouldSupportICVFX(const IDisplayClusterViewport* InViewport) const
 	{
 		return false;
 	}
diff --git a/Engine/Plugins/Runtime/nDisplay/Source/DisplayClusterConfiguration/Private/DisplayClusterConfigurationTypes.cpp b/Engine/Plugins/Runtime/nDisplay/Source/DisplayClusterConfiguration/Private/DisplayClusterConfigurationTypes.cpp
--- a/Engine/Plugins/Runtime/nDisplay/Source/DisplayClusterConfiguration/Private/DisplayClusterConfigurationTypes.cpp
+++ b/Engine/Plugins/Runtime/nDisplay/Source/DisplayClusterConfiguration/Private/DisplayClusterConfigurationTypes.cpp
@@ -172,6 +172,7 @@ const TSet<FString> UDisplayClusterConfigurationData::ProjectionPolicies =
 	TEXT("DomeProjection"),
 	TEXT("VIOSO"),
 	TEXT("Manual"),
+	TEXT("Reference"),
 };
 
 #endif
@@ -181,6 +182,13 @@ FDisplayClusterConfigurationProjection::FDisplayClusterConfigurationProjection()
 	Type = TEXT("simple");
 }
 
+bool FDisplayClusterConfigurationProjection::ShouldDisplayICVFXCategory() const
+{
+	return Type.Equals(DisplayClusterProjectionStrings::projection::Mesh, ESearchCase::IgnoreCase)
+		|| Type.Equals(DisplayClusterProjectionStrings::projection::MPCDI, ESearchCase::IgnoreCase)
+		|| Type.Equals(DisplayClusterProjectionStrings::projection::Reference, ESearchCase::IgnoreCase);
+}
+
 constexpr float UDisplayClusterConfigurationViewport::ViewportMinimumSize = 1.0f;
 constexpr float UDisplayClusterConfigurationViewport::ViewportMaximumSize = 15360.0f;
 
diff --git a/Engine/Plugins/Runtime/nDisplay/Source/DisplayClusterConfiguration/Private/DisplayClusterConfigurationTypes_ICVFX.cpp b/Engine/Plugins/Runtime/nDisplay/Source/DisplayClusterConfiguration/Private/DisplayClusterConfigurationTypes_ICVFX.cpp
--- a/Engine/Plugins/Runtime/nDisplay/Source/DisplayClusterConfiguration/Private/DisplayClusterConfigurationTypes_ICVFX.cpp
+++ b/Engine/Plugins/Runtime/nDisplay/Source/DisplayClusterConfiguration/Private/DisplayClusterConfigurationTypes_ICVFX.cpp
@@ -296,43 +296,37 @@ EDisplayClusterShaderParametersICVFX_CameraOverlappingRenderMode FDisplayCluster
 ///////////////////////////////////////////////////////////////////////////////////////
 EDisplayClusterShaderParametersICVFX_ChromakeySource FDisplayClusterConfigurationICVFX_ChromakeySettings::GetChromakeyType(const FDisplayClusterConfigurationICVFX_StageSettings& InStageSettings) const
 {
-	if (!bEnable)
+	if (bEnable)
 	{
-		return EDisplayClusterShaderParametersICVFX_ChromakeySource::Disabled;
-	}
+		switch (ChromakeyType)
+		{
+		case EDisplayClusterConfigurationICVFX_ChromakeyType::InnerFrustum:
+			return EDisplayClusterShaderParametersICVFX_ChromakeySource::FrameColor;
 
-	switch (ChromakeyType)
-	{
-	case EDisplayClusterConfigurationICVFX_ChromakeyType::CustomChromakey:
-		return EDisplayClusterShaderParametersICVFX_ChromakeySource::ChromakeyLayers;
+		case EDisplayClusterConfigurationICVFX_ChromakeyType::CustomChromakey:
+			return EDisplayClusterShaderParametersICVFX_ChromakeySource::ChromakeyLayers;
 
-	default:
-		break;
+		case EDisplayClusterConfigurationICVFX_ChromakeyType::Disabled:
+			return EDisplayClusterShaderParametersICVFX_ChromakeySource::Default;
+
+		default:
+			break;
+		}
 	}
 
-	return EDisplayClusterShaderParametersICVFX_ChromakeySource::FrameColor;
+	return EDisplayClusterShaderParametersICVFX_ChromakeySource::Disabled;
 }
 
 FDisplayClusterConfigurationICVFX_ChromakeyRenderSettings* FDisplayClusterConfigurationICVFX_ChromakeySettings::GetWritableChromakeyRenderSettings(const FDisplayClusterConfigurationICVFX_StageSettings& InStageSettings)
 {
 	// Note: Here we can add an override of the CK rendering settings from StageSetings
-	if (GetChromakeyType(InStageSettings) == EDisplayClusterShaderParametersICVFX_ChromakeySource::ChromakeyLayers)
-	{
-		return &ChromakeyRenderTexture;
-	}
-
-	return nullptr;
+	return &ChromakeyRenderTexture;
 }
 
 const FDisplayClusterConfigurationICVFX_ChromakeyRenderSettings* FDisplayClusterConfigurationICVFX_ChromakeySettings::GetChromakeyRenderSettings(const FDisplayClusterConfigurationICVFX_StageSettings& InStageSettings) const
 {
 	// Note: Here we can add an override of the CK rendering settings from StageSetings
-	if (GetChromakeyType(InStageSettings) == EDisplayClusterShaderParametersICVFX_ChromakeySource::ChromakeyLayers)
-	{
-		return &ChromakeyRenderTexture;
-	}
-
-	return nullptr;
+	return &ChromakeyRenderTexture;
 }
 
 const FLinearColor& FDisplayClusterConfigurationICVFX_ChromakeySettings::GetChromakeyColor(const FDisplayClusterConfigurationICVFX_StageSettings& InStageSettings) const
diff --git a/Engine/Plugins/Runtime/nDisplay/Source/DisplayClusterConfiguration/Private/DisplayClusterConfigurationTypes_Viewport.cpp b/Engine/Plugins/Runtime/nDisplay/Source/DisplayClusterConfiguration/Private/DisplayClusterConfigurationTypes_Viewport.cpp
--- a/Engine/Plugins/Runtime/nDisplay/Source/DisplayClusterConfiguration/Private/DisplayClusterConfigurationTypes_Viewport.cpp
+++ b/Engine/Plugins/Runtime/nDisplay/Source/DisplayClusterConfiguration/Private/DisplayClusterConfigurationTypes_Viewport.cpp
@@ -67,6 +67,62 @@ EDisplayClusterViewportICVFXFlags UDisplayClusterConfigurationViewport::GetViewp
 		break;
 	}
 
+	// Reverse camera order when the option is set in the viewport.
+	if (ICVFX.bReverseCameraPriority)
+	{
+		EnumAddFlags(OutFlags, EDisplayClusterViewportICVFXFlags::ReverseCameraPriority);
+	}
 
 	return OutFlags;
 }
+
+EDisplayClusterShaderParametersICVFX_ChromakeySource UDisplayClusterConfigurationViewport::GetViewportChromakeyType(
+	const FDisplayClusterConfigurationICVFX_StageSettings& InStageSettings,
+	const FString& InCameraId,
+	const struct FDisplayClusterConfigurationICVFX_CameraSettings& InCameraSettings) const
+{
+	// Get ICVFX camera settings:
+	const EDisplayClusterShaderParametersICVFX_ChromakeySource CameraChromakeySource = InCameraSettings.Chromakey.GetChromakeyType(InStageSettings);
+
+	const EDisplayClusterViewportICVFXFlags ICVFXFlags = GetViewportICVFXFlags(InStageSettings);
+	const EDisplayClusterConfigurationICVFX_OverrideChromakeyType OverrideChromakeySource = ICVFX.GetOverrideChromakeyType(InCameraId);
+
+	if (CameraChromakeySource == EDisplayClusterShaderParametersICVFX_ChromakeySource::Disabled
+	|| OverrideChromakeySource == EDisplayClusterConfigurationICVFX_OverrideChromakeyType::Disabled
+	|| EnumHasAnyFlags(ICVFXFlags, EDisplayClusterViewportICVFXFlags::DisableCamera | EDisplayClusterViewportICVFXFlags::DisableChromakey)
+	)
+	{
+		return EDisplayClusterShaderParametersICVFX_ChromakeySource::Disabled;
+	}
+
+	// Override chromakey from viewport
+	switch (OverrideChromakeySource)
+	{
+		case EDisplayClusterConfigurationICVFX_OverrideChromakeyType::InnerFrustum:
+			return EDisplayClusterShaderParametersICVFX_ChromakeySource::FrameColor;
+
+		case EDisplayClusterConfigurationICVFX_OverrideChromakeyType::CustomChromakey:
+			return EDisplayClusterShaderParametersICVFX_ChromakeySource::ChromakeyLayers;
+
+		default:
+			break;
+	}
+
+	// Default By default CK is disabled.
+	if(CameraChromakeySource == EDisplayClusterShaderParametersICVFX_ChromakeySource::Default)
+	{
+		return EDisplayClusterShaderParametersICVFX_ChromakeySource::Disabled;
+	}
+
+	return CameraChromakeySource;
+}
+
+EDisplayClusterConfigurationICVFX_OverrideChromakeyType FDisplayClusterConfigurationViewport_ICVFX::GetOverrideChromakeyType(const FString& CameraId) const
+{
+	if (const EDisplayClusterConfigurationICVFX_OverrideChromakeyType* OverrideChromakeyTypePtr = PerCameraOverrideChromakeyType.Find(CameraId))
+	{
+		return *OverrideChromakeyTypePtr;
+	}
+
+	return OverrideChromakeyType;
+}
\ No newline at end of file
diff --git a/Engine/Plugins/Runtime/nDisplay/Source/DisplayClusterConfiguration/Public/DisplayClusterConfigurationTypes_Base.h b/Engine/Plugins/Runtime/nDisplay/Source/DisplayClusterConfiguration/Public/DisplayClusterConfigurationTypes_Base.h
--- a/Engine/Plugins/Runtime/nDisplay/Source/DisplayClusterConfiguration/Public/DisplayClusterConfigurationTypes_Base.h
+++ b/Engine/Plugins/Runtime/nDisplay/Source/DisplayClusterConfiguration/Public/DisplayClusterConfigurationTypes_Base.h
@@ -173,6 +173,11 @@ struct DISPLAYCLUSTERCONFIGURATION_API FDisplayClusterConfigurationProjection
 
 public:
 	FDisplayClusterConfigurationProjection();
+
+	/**
+	* The ICVFX properties of the viewport will only be displayed in the UI if this function returns true.
+	*/
+	bool ShouldDisplayICVFXCategory() const;
 };
 
 USTRUCT(Blueprintable)
diff --git a/Engine/Plugins/Runtime/nDisplay/Source/DisplayClusterConfiguration/Public/DisplayClusterConfigurationTypes_Enums.h b/Engine/Plugins/Runtime/nDisplay/Source/DisplayClusterConfiguration/Public/DisplayClusterConfigurationTypes_Enums.h
--- a/Engine/Plugins/Runtime/nDisplay/Source/DisplayClusterConfiguration/Public/DisplayClusterConfigurationTypes_Enums.h
+++ b/Engine/Plugins/Runtime/nDisplay/Source/DisplayClusterConfiguration/Public/DisplayClusterConfigurationTypes_Enums.h
@@ -58,6 +58,13 @@ enum class EDisplayClusterConfigurationICVFX_ChromakeyType : uint8
 
 	/** Only actors specified in the custom chromakey actor list are rendered as chromakey */
 	CustomChromakey    UMETA(DisplayName = "Custom Chromakey"),
+
+	/**
+	* Disable chromakey.
+	* This value can be overridden by the viewport ICVFX settings,
+	* by changing the ‘Override Chromakey Type’ property.
+	*/
+	Disabled            UMETA(DisplayName = "Disabled"),
 };
 
 /** Indicates the source of the chromakey settings when an ICVFX camera's inner frustum is rendered as chromakey */
@@ -126,6 +133,23 @@ enum class EDisplayClusterConfigurationICVFX_OverrideLightcardRenderMode : uint8
 	Under   UMETA(DisplayName = "Lightcard Under Frustum"),
 };
 
+/** Override the chromakey type for the viewport. */
+UENUM()
+enum class EDisplayClusterConfigurationICVFX_OverrideChromakeyType : uint8
+{
+	// Use the chromakey type from ICVFX camera
+	Default    UMETA(DisplayName = "Default"),
+
+	// The entire inner frustum is rendered as chromakey for this viewport
+	InnerFrustum       UMETA(DisplayName = "Inner Frustum"),
+
+	/** Only actors specified in the custom chromakey actor list are rendered as chromakey */
+	CustomChromakey    UMETA(DisplayName = "Custom Chromakey"),
+
+	// Disable chromakey
+	Disabled UMETA(DisplayName = "Disabled Chromakey"),
+};
+
 UENUM()
 enum class EDisplayClusterConfigurationICVFX_OverrideCameraRenderMode : uint8
 {
diff --git a/Engine/Plugins/Runtime/nDisplay/Source/DisplayClusterConfiguration/Public/DisplayClusterConfigurationTypes_Viewport.h b/Engine/Plugins/Runtime/nDisplay/Source/DisplayClusterConfiguration/Public/DisplayClusterConfigurationTypes_Viewport.h
--- a/Engine/Plugins/Runtime/nDisplay/Source/DisplayClusterConfiguration/Public/DisplayClusterConfigurationTypes_Viewport.h
+++ b/Engine/Plugins/Runtime/nDisplay/Source/DisplayClusterConfiguration/Public/DisplayClusterConfigurationTypes_Viewport.h
@@ -25,6 +25,12 @@
 
 struct FDisplayClusterConfigurationICVFX_StageSettings;
 
+/**
+* Unique ICVFX customisation for each viewport.
+* 
+* Must be processed in UDisplayClusterConfigurationViewport::GetViewportICVFXFlags().
+* This will result in some EDisplayClusterViewportICVFXFlags being raised.
+*/
 USTRUCT(BlueprintType)
 struct DISPLAYCLUSTERCONFIGURATION_API FDisplayClusterConfigurationViewport_ICVFX
 {
@@ -45,6 +51,8 @@ public:
 		return EDisplayClusterViewportICVFXFlags::None;
 	}
 
+	EDisplayClusterConfigurationICVFX_OverrideChromakeyType GetOverrideChromakeyType(const FString& CameraId) const;
+
 public:
 	/** Enable in-camera VFX for this Viewport (works only with supported Projection Policies) */
 	UPROPERTY(EditAnywhere, BlueprintReadWrite, Category = "In Camera VFX")
@@ -62,6 +70,14 @@ public:
 	UPROPERTY(EditAnywhere, BlueprintReadWrite, Category = "In Camera VFX")
 	EDisplayClusterConfigurationICVFX_OverrideLightcardRenderMode LightcardRenderMode = EDisplayClusterConfigurationICVFX_OverrideLightcardRenderMode::Default;
 
+	/** Use unique chromakey type for this viewport */
+	UPROPERTY(EditAnywhere, BlueprintReadWrite, Category = "In Camera VFX")
+	EDisplayClusterConfigurationICVFX_OverrideChromakeyType OverrideChromakeyType = EDisplayClusterConfigurationICVFX_OverrideChromakeyType::Default;
+
+	/** Determines the chromakey override per-camera in this viewport.  */
+	UPROPERTY(EditAnywhere, BlueprintReadWrite, Category = "In Camera VFX")
+	TMap<FString, EDisplayClusterConfigurationICVFX_OverrideChromakeyType> PerCameraOverrideChromakeyType;
+
 	/** The order in which the ICVFX cameras are composited over is reversed. Useful for time-multiplexed displays. */
 	UPROPERTY(EditAnywhere, BlueprintReadWrite, Category = "In Camera VFX")
 	bool bReverseCameraPriority = false;
@@ -168,6 +184,12 @@ public:
 	/** Get ICVFX settings flags for viewport*/
 	EDisplayClusterViewportICVFXFlags GetViewportICVFXFlags(const FDisplayClusterConfigurationICVFX_StageSettings& InStageSettings) const;
 
+	/** Returns the chromakey type for this viewport used by the specified ICVFX camera. */
+	EDisplayClusterShaderParametersICVFX_ChromakeySource GetViewportChromakeyType(
+		const FDisplayClusterConfigurationICVFX_StageSettings& InStageSettings,
+		const FString& InCameraId,
+		const struct FDisplayClusterConfigurationICVFX_CameraSettings& InCameraSettings) const;
+
 #if WITH_EDITOR
 	/** Enable the preview texture. Only should be called by the object managing the preview texture state. */
 	void EnablePreviewTexture();
diff --git a/Engine/Plugins/Runtime/nDisplay/Source/DisplayClusterConfigurator/Private/Views/Details/Cluster/DisplayClusterConfiguratorViewportDetailsCustomization.cpp b/Engine/Plugins/Runtime/nDisplay/Source/DisplayClusterConfigurator/Private/Views/Details/Cluster/DisplayClusterConfiguratorViewportDetailsCustomization.cpp
--- a/Engine/Plugins/Runtime/nDisplay/Source/DisplayClusterConfigurator/Private/Views/Details/Cluster/DisplayClusterConfiguratorViewportDetailsCustomization.cpp
+++ b/Engine/Plugins/Runtime/nDisplay/Source/DisplayClusterConfigurator/Private/Views/Details/Cluster/DisplayClusterConfiguratorViewportDetailsCustomization.cpp
@@ -42,11 +42,7 @@ void FDisplayClusterConfiguratorViewportDetailsCustomization::CustomizeDetails(I
 	}
 	check(ConfigurationViewportPtr != nullptr);
 
-	const bool bDisplayICVFX =
-			ConfigurationViewportPtr->ProjectionPolicy.Type.Equals(DisplayClusterProjectionStrings::projection::Mesh, ESearchCase::IgnoreCase)
-		||  ConfigurationViewportPtr->ProjectionPolicy.Type.Equals(DisplayClusterProjectionStrings::projection::MPCDI, ESearchCase::IgnoreCase);
-
-	if (!bDisplayICVFX)
+	if (!ConfigurationViewportPtr->ProjectionPolicy.ShouldDisplayICVFXCategory())
 	{
 		IDetailCategoryBuilder& Category = InLayoutBuilder.EditCategory(DisplayClusterConfigurationStrings::categories::InCameraVFXCategory);
 		Category.SetCategoryVisibility(false);
diff --git a/Engine/Plugins/Runtime/nDisplay/Source/DisplayClusterConfigurator/Private/Views/Details/Policies/DisplayClusterConfiguratorPolicyDetailCustomization.cpp b/Engine/Plugins/Runtime/nDisplay/Source/DisplayClusterConfigurator/Private/Views/Details/Policies/DisplayClusterConfiguratorPolicyDetailCustomization.cpp
--- a/Engine/Plugins/Runtime/nDisplay/Source/DisplayClusterConfigurator/Private/Views/Details/Policies/DisplayClusterConfiguratorPolicyDetailCustomization.cpp
+++ b/Engine/Plugins/Runtime/nDisplay/Source/DisplayClusterConfigurator/Private/Views/Details/Policies/DisplayClusterConfiguratorPolicyDetailCustomization.cpp
@@ -391,6 +391,13 @@ void FDisplayClusterConfiguratorProjectionCustomization::BuildParametersForPolic
 	{
 		CreateEasyBlendPolicy(Blueprint);
 	}
+	/*
+	 * Reference
+	 */
+	else if (PolicyLower == DisplayClusterProjectionStrings::projection::Reference)
+	{
+		CreateReferencePolicy(Blueprint);
+	}
 	/*
 	 * Manual
 	 */
@@ -619,6 +626,15 @@ void FDisplayClusterConfiguratorProjectionCustomization::CreateEasyBlendPolicy(U
 		ConfigurationViewports));
 }
 
+void FDisplayClusterConfiguratorProjectionCustomization::CreateReferencePolicy(UDisplayClusterBlueprint* Blueprint)
+{
+	CustomPolicyParameters.Add(MakeShared<FPolicyParameterInfoText>(
+		"Source Viewport",
+		DisplayClusterProjectionStrings::cfg::reference::ViewportId,
+		Blueprint,
+		ConfigurationViewports));
+}
+
 void FDisplayClusterConfiguratorProjectionCustomization::CreateManualPolicy(UDisplayClusterBlueprint* Blueprint)
 {
 	check(Blueprint);
diff --git a/Engine/Plugins/Runtime/nDisplay/Source/DisplayClusterConfigurator/Private/Views/Details/Policies/DisplayClusterConfiguratorPolicyDetailCustomization.h b/Engine/Plugins/Runtime/nDisplay/Source/DisplayClusterConfigurator/Private/Views/Details/Policies/DisplayClusterConfiguratorPolicyDetailCustomization.h
--- a/Engine/Plugins/Runtime/nDisplay/Source/DisplayClusterConfigurator/Private/Views/Details/Policies/DisplayClusterConfiguratorPolicyDetailCustomization.h
+++ b/Engine/Plugins/Runtime/nDisplay/Source/DisplayClusterConfigurator/Private/Views/Details/Policies/DisplayClusterConfiguratorPolicyDetailCustomization.h
@@ -76,6 +76,7 @@ private:
 	void CreateEasyBlendPolicy(UDisplayClusterBlueprint* Blueprint);
 	void CreateManualPolicy(UDisplayClusterBlueprint* Blueprint);
 	void CreateMPCDIPolicy(UDisplayClusterBlueprint* Blueprint);
+	void CreateReferencePolicy(UDisplayClusterBlueprint* Blueprint);
 
 private:
 	/** All custom parameters for the selected policy. */
diff --git a/Engine/Plugins/Runtime/nDisplay/Source/DisplayClusterMoviePipeline/Private/DisplayClusterMoviePipelineViewportPass.cpp b/Engine/Plugins/Runtime/nDisplay/Source/DisplayClusterMoviePipeline/Private/DisplayClusterMoviePipelineViewportPass.cpp
--- a/Engine/Plugins/Runtime/nDisplay/Source/DisplayClusterMoviePipeline/Private/DisplayClusterMoviePipelineViewportPass.cpp
+++ b/Engine/Plugins/Runtime/nDisplay/Source/DisplayClusterMoviePipeline/Private/DisplayClusterMoviePipelineViewportPass.cpp
@@ -30,7 +30,7 @@ namespace UE::DisplayCluster::MoviePipelineViewportPass
 			if (IDisplayClusterViewportProxy* InViewportProxy = InViewportManagerProxy->FindViewport_RenderThread(InViewportId))
 			{
 				const TSharedPtr<IDisplayClusterProjectionPolicy, ESPMode::ThreadSafe>& InPolicy = InViewportProxy->GetProjectionPolicy_RenderThread();
-				if (InPolicy.IsValid() && InPolicy->IsWarpBlendSupported())
+				if (InPolicy.IsValid() && InPolicy->IsWarpBlendSupported_RenderThread(InViewportProxy))
 				{
 					TArray<FRHITexture*> WarpInputTextures, WarpOutputTextures;
 					TArray<FIntRect> WarpInputRects, WarpOutputRects;
diff --git a/Engine/Plugins/Runtime/nDisplay/Source/DisplayClusterProjection/Private/DisplayClusterProjectionLog.cpp b/Engine/Plugins/Runtime/nDisplay/Source/DisplayClusterProjection/Private/DisplayClusterProjectionLog.cpp
--- a/Engine/Plugins/Runtime/nDisplay/Source/DisplayClusterProjection/Private/DisplayClusterProjectionLog.cpp
+++ b/Engine/Plugins/Runtime/nDisplay/Source/DisplayClusterProjection/Private/DisplayClusterProjectionLog.cpp
@@ -12,3 +12,4 @@ DEFINE_LOG_CATEGORY(LogDisplayClusterProjectionMesh);
 DEFINE_LOG_CATEGORY(LogDisplayClusterProjectionMPCDI);
 DEFINE_LOG_CATEGORY(LogDisplayClusterProjectionSimple);
 DEFINE_LOG_CATEGORY(LogDisplayClusterProjectionVIOSO);
+DEFINE_LOG_CATEGORY(LogDisplayClusterProjectionReference);
diff --git a/Engine/Plugins/Runtime/nDisplay/Source/DisplayClusterProjection/Private/DisplayClusterProjectionLog.h b/Engine/Plugins/Runtime/nDisplay/Source/DisplayClusterProjection/Private/DisplayClusterProjectionLog.h
--- a/Engine/Plugins/Runtime/nDisplay/Source/DisplayClusterProjection/Private/DisplayClusterProjectionLog.h
+++ b/Engine/Plugins/Runtime/nDisplay/Source/DisplayClusterProjection/Private/DisplayClusterProjectionLog.h
@@ -13,3 +13,4 @@ DECLARE_LOG_CATEGORY_EXTERN(LogDisplayClusterProjectionMesh,           Log, All)
 DECLARE_LOG_CATEGORY_EXTERN(LogDisplayClusterProjectionMPCDI,          Log, All);
 DECLARE_LOG_CATEGORY_EXTERN(LogDisplayClusterProjectionSimple,         Log, All);
 DECLARE_LOG_CATEGORY_EXTERN(LogDisplayClusterProjectionVIOSO,          Log, All);
+DECLARE_LOG_CATEGORY_EXTERN(LogDisplayClusterProjectionReference,      Log, All);
diff --git a/Engine/Plugins/Runtime/nDisplay/Source/DisplayClusterProjection/Private/DisplayClusterProjectionModule.cpp b/Engine/Plugins/Runtime/nDisplay/Source/DisplayClusterProjection/Private/DisplayClusterProjectionModule.cpp
--- a/Engine/Plugins/Runtime/nDisplay/Source/DisplayClusterProjection/Private/DisplayClusterProjectionModule.cpp
+++ b/Engine/Plugins/Runtime/nDisplay/Source/DisplayClusterProjection/Private/DisplayClusterProjectionModule.cpp
@@ -6,11 +6,12 @@
 #include "DisplayClusterProjectionStrings.h"
 
 #include "Policy/Camera/DisplayClusterProjectionCameraPolicyFactory.h"
-#include "Policy/EasyBlend/DisplayClusterProjectionEasyBlendPolicyFactory.h"
 #include "Policy/Domeprojection/DisplayClusterProjectionDomeprojectionPolicyFactory.h"
+#include "Policy/EasyBlend/DisplayClusterProjectionEasyBlendPolicyFactory.h"
 #include "Policy/Link/DisplayClusterProjectionLinkPolicyFactory.h"
 #include "Policy/Manual/DisplayClusterProjectionManualPolicyFactory.h"
 #include "Policy/MPCDI/DisplayClusterProjectionMPCDIPolicyFactory.h"
+#include "Policy/Reference/DisplayClusterProjectionReferencePolicyFactory.h"
 #include "Policy/Simple/DisplayClusterProjectionSimplePolicyFactory.h"
 #include "Policy/VIOSO/DisplayClusterProjectionVIOSOPolicyFactory.h"
 
@@ -50,6 +51,10 @@ FDisplayClusterProjectionModule::FDisplayClusterProjectionModule()
 	Factory = MakeShared<FDisplayClusterProjectionLinkPolicyFactory>();
 	ProjectionPolicyFactories.Emplace(DisplayClusterProjectionStrings::projection::Link, Factory);
 
+	// Reference projection
+	Factory = MakeShared<FDisplayClusterProjectionReferencePolicyFactory>();
+	ProjectionPolicyFactories.Emplace(DisplayClusterProjectionStrings::projection::Reference, Factory);
+
 	// Manual projection
 	Factory = MakeShared<FDisplayClusterProjectionManualPolicyFactory>();
 	ProjectionPolicyFactories.Emplace(DisplayClusterProjectionStrings::projection::Manual, Factory);
diff --git a/Engine/Plugins/Runtime/nDisplay/Source/DisplayClusterProjection/Private/Policy/Camera/DisplayClusterProjectionCameraPolicy.h b/Engine/Plugins/Runtime/nDisplay/Source/DisplayClusterProjection/Private/Policy/Camera/DisplayClusterProjectionCameraPolicy.h
--- a/Engine/Plugins/Runtime/nDisplay/Source/DisplayClusterProjection/Private/Policy/Camera/DisplayClusterProjectionCameraPolicy.h
+++ b/Engine/Plugins/Runtime/nDisplay/Source/DisplayClusterProjection/Private/Policy/Camera/DisplayClusterProjectionCameraPolicy.h
@@ -31,7 +31,7 @@ public:
 
 	virtual void UpdatePostProcessSettings(IDisplayClusterViewport* InViewport) override;
 
-	virtual bool ShouldUseSourceTextureWithMips() const override
+	virtual bool ShouldUseSourceTextureWithMips(const IDisplayClusterViewport* InViewport) const override
 	{
 		return true;
 	}
diff --git a/Engine/Plugins/Runtime/nDisplay/Source/DisplayClusterProjection/Private/Policy/Domeprojection/DisplayClusterProjectionDomeprojectionPolicyBase.cpp b/Engine/Plugins/Runtime/nDisplay/Source/DisplayClusterProjection/Private/Policy/Domeprojection/DisplayClusterProjectionDomeprojectionPolicyBase.cpp
--- a/Engine/Plugins/Runtime/nDisplay/Source/DisplayClusterProjection/Private/Policy/Domeprojection/DisplayClusterProjectionDomeprojectionPolicyBase.cpp
+++ b/Engine/Plugins/Runtime/nDisplay/Source/DisplayClusterProjection/Private/Policy/Domeprojection/DisplayClusterProjectionDomeprojectionPolicyBase.cpp
@@ -157,7 +157,12 @@ bool FDisplayClusterProjectionDomeprojectionPolicyBase::GetProjectionMatrix(IDis
 	return ViewAdapter->GetProjectionMatrix(InViewport, InContextNum, DomeprojectionChannel, OutPrjMatrix);
 }
 
-bool FDisplayClusterProjectionDomeprojectionPolicyBase::IsWarpBlendSupported()
+bool FDisplayClusterProjectionDomeprojectionPolicyBase::IsWarpBlendSupported(const IDisplayClusterViewport* InViewport) const
+{
+	return true;
+}
+
+bool FDisplayClusterProjectionDomeprojectionPolicyBase::IsWarpBlendSupported_RenderThread(const IDisplayClusterViewportProxy* InViewportProxy) const
 {
 	return true;
 }
diff --git a/Engine/Plugins/Runtime/nDisplay/Source/DisplayClusterProjection/Private/Policy/Domeprojection/DisplayClusterProjectionDomeprojectionPolicyBase.h b/Engine/Plugins/Runtime/nDisplay/Source/DisplayClusterProjection/Private/Policy/Domeprojection/DisplayClusterProjectionDomeprojectionPolicyBase.h
--- a/Engine/Plugins/Runtime/nDisplay/Source/DisplayClusterProjection/Private/Policy/Domeprojection/DisplayClusterProjectionDomeprojectionPolicyBase.h
+++ b/Engine/Plugins/Runtime/nDisplay/Source/DisplayClusterProjection/Private/Policy/Domeprojection/DisplayClusterProjectionDomeprojectionPolicyBase.h
@@ -16,9 +16,7 @@ public:
 	FDisplayClusterProjectionDomeprojectionPolicyBase(const FString& ProjectionPolicyId, const FDisplayClusterConfigurationProjection* InConfigurationProjectionPolicy);
 
 public:
-	//////////////////////////////////////////////////////////////////////////////////////////////
-	// IDisplayClusterProjectionPolicy
-	//////////////////////////////////////////////////////////////////////////////////////////////
+	//~Begin IDisplayClusterProjectionPolicy
 	virtual const FString& GetType() const override;
 
 	virtual bool HandleStartScene(IDisplayClusterViewport* InViewport) override;
@@ -27,14 +25,17 @@ public:
 	virtual bool CalculateView(IDisplayClusterViewport* InViewport, const uint32 InContextNum, FVector& InOutViewLocation, FRotator& InOutViewRotation, const FVector& ViewOffset, const float WorldToMeters, const float NCP, const float FCP) override;
 	virtual bool GetProjectionMatrix(IDisplayClusterViewport* InViewport, const uint32 InContextNum, FMatrix& OutPrjMatrix) override;
 
-	virtual bool IsWarpBlendSupported() override;
+	virtual bool IsWarpBlendSupported(const IDisplayClusterViewport* InViewport) const override;
+	virtual bool IsWarpBlendSupported_RenderThread(const IDisplayClusterViewportProxy* InViewportProxy) const override;
+
 	virtual void ApplyWarpBlend_RenderThread(FRHICommandListImmediate& RHICmdList, const IDisplayClusterViewportProxy* InViewportProxy) override;
 	
 	// Request additional targetable resources for domeprojection  external warpblend
-	virtual bool ShouldUseAdditionalTargetableResource() const override
+	virtual bool ShouldUseAdditionalTargetableResource(const IDisplayClusterViewport* InViewport) const override
 	{
 		return true;
 	}
+	//~End IDisplayClusterProjectionPolicy
 
 protected:
 	// Delegate view adapter instantiation to the RHI specific children
diff --git a/Engine/Plugins/Runtime/nDisplay/Source/DisplayClusterProjection/Private/Policy/EasyBlend/DisplayClusterProjectionEasyBlendPolicy.cpp b/Engine/Plugins/Runtime/nDisplay/Source/DisplayClusterProjection/Private/Policy/EasyBlend/DisplayClusterProjectionEasyBlendPolicy.cpp
--- a/Engine/Plugins/Runtime/nDisplay/Source/DisplayClusterProjection/Private/Policy/EasyBlend/DisplayClusterProjectionEasyBlendPolicy.cpp
+++ b/Engine/Plugins/Runtime/nDisplay/Source/DisplayClusterProjection/Private/Policy/EasyBlend/DisplayClusterProjectionEasyBlendPolicy.cpp
@@ -208,9 +208,14 @@ const FString& FDisplayClusterProjectionEasyBlendPolicy::GetType() const
 	return Type;
 }
 
-bool FDisplayClusterProjectionEasyBlendPolicy::IsWarpBlendSupported()
+bool FDisplayClusterProjectionEasyBlendPolicy::IsWarpBlendSupported(const IDisplayClusterViewport* InViewport) const
 {
-	return PolicyViewData.IsValid() || PolicyViewDataProxy.IsValid();
+	return PolicyViewData.IsValid();
+}
+
+bool FDisplayClusterProjectionEasyBlendPolicy::IsWarpBlendSupported_RenderThread(const IDisplayClusterViewportProxy* InViewportProxy) const
+{
+	return PolicyViewDataProxy.IsValid();
 }
 
 void FDisplayClusterProjectionEasyBlendPolicy::ApplyWarpBlend_RenderThread(FRHICommandListImmediate& RHICmdList, const IDisplayClusterViewportProxy* InViewportProxy)
diff --git a/Engine/Plugins/Runtime/nDisplay/Source/DisplayClusterProjection/Private/Policy/EasyBlend/DisplayClusterProjectionEasyBlendPolicy.h b/Engine/Plugins/Runtime/nDisplay/Source/DisplayClusterProjection/Private/Policy/EasyBlend/DisplayClusterProjectionEasyBlendPolicy.h
--- a/Engine/Plugins/Runtime/nDisplay/Source/DisplayClusterProjection/Private/Policy/EasyBlend/DisplayClusterProjectionEasyBlendPolicy.h
+++ b/Engine/Plugins/Runtime/nDisplay/Source/DisplayClusterProjection/Private/Policy/EasyBlend/DisplayClusterProjectionEasyBlendPolicy.h
@@ -30,11 +30,13 @@ public:
 	virtual bool CalculateView(IDisplayClusterViewport* InViewport, const uint32 InContextNum, FVector& InOutViewLocation, FRotator& InOutViewRotation, const FVector& ViewOffset, const float WorldToMeters, const float NCP, const float FCP) override;
 	virtual bool GetProjectionMatrix(IDisplayClusterViewport* InViewport, const uint32 InContextNum, FMatrix& OutPrjMatrix) override;
 
-	virtual bool IsWarpBlendSupported() override;
+	virtual bool IsWarpBlendSupported(const IDisplayClusterViewport* InViewport) const override;
+	virtual bool IsWarpBlendSupported_RenderThread(const IDisplayClusterViewportProxy* InViewportProxy) const override;
+
 	virtual void ApplyWarpBlend_RenderThread(FRHICommandListImmediate& RHICmdList, const IDisplayClusterViewportProxy* InViewportProxy) override;
 
 	// Request additional targetable resources for easyblend external warpblend
-	virtual bool ShouldUseAdditionalTargetableResource() const override
+	virtual bool ShouldUseAdditionalTargetableResource(const IDisplayClusterViewport* InViewport) const override
 	{
 		return true;
 	}
diff --git a/Engine/Plugins/Runtime/nDisplay/Source/DisplayClusterProjection/Private/Policy/Link/DisplayClusterProjectionLinkPolicy.h b/Engine/Plugins/Runtime/nDisplay/Source/DisplayClusterProjection/Private/Policy/Link/DisplayClusterProjectionLinkPolicy.h
--- a/Engine/Plugins/Runtime/nDisplay/Source/DisplayClusterProjection/Private/Policy/Link/DisplayClusterProjectionLinkPolicy.h
+++ b/Engine/Plugins/Runtime/nDisplay/Source/DisplayClusterProjection/Private/Policy/Link/DisplayClusterProjectionLinkPolicy.h
@@ -28,7 +28,7 @@ public:
 	virtual bool GetProjectionMatrix(IDisplayClusterViewport* InViewport, const uint32 InContextNum, FMatrix& OutPrjMatrix) override;
 	virtual void SetupProjectionViewPoint(IDisplayClusterViewport* InViewport, const float InDeltaTime, FMinimalViewInfo& InOutViewInfo, float* OutCustomNearClippingPlane = nullptr) override;
 
-	virtual bool ShouldUseSourceTextureWithMips() const override
+	virtual bool ShouldUseSourceTextureWithMips(const IDisplayClusterViewport* InViewport) const override
 	{
 		return true;
 	}
diff --git a/Engine/Plugins/Runtime/nDisplay/Source/DisplayClusterProjection/Private/Policy/MPCDI/DisplayClusterProjectionMPCDIPolicy.cpp b/Engine/Plugins/Runtime/nDisplay/Source/DisplayClusterProjection/Private/Policy/MPCDI/DisplayClusterProjectionMPCDIPolicy.cpp
--- a/Engine/Plugins/Runtime/nDisplay/Source/DisplayClusterProjection/Private/Policy/MPCDI/DisplayClusterProjectionMPCDIPolicy.cpp
+++ b/Engine/Plugins/Runtime/nDisplay/Source/DisplayClusterProjection/Private/Policy/MPCDI/DisplayClusterProjectionMPCDIPolicy.cpp
@@ -84,7 +84,7 @@ IDisplayClusterWarpPolicy* FDisplayClusterProjectionMPCDIPolicy::GetWarpPolicy_R
 	return WarpPolicyInterface_Proxy.Get();
 }
 
-bool FDisplayClusterProjectionMPCDIPolicy::ShouldSupportICVFX(IDisplayClusterViewport* InViewport) const
+bool FDisplayClusterProjectionMPCDIPolicy::ShouldSupportICVFX(const IDisplayClusterViewport* InViewport) const
 {
 	check(IsInGameThread());
 	check(InViewport);
@@ -109,18 +109,16 @@ void FDisplayClusterProjectionMPCDIPolicy::UpdateProxyData(IDisplayClusterViewpo
 	check(InViewport);
 
 	ENQUEUE_RENDER_COMMAND(DisplayClusterProjectionMPCDIPolicy_UpdateProxyData)(
-		[ProjectionPolicy = InViewport->GetProjectionPolicy(), WarpBlendInterfacePtr = WarpBlendInterface, WarPolicyInterfacePtr = WarpPolicyInterface, Contexts = WarpBlendContexts](FRHICommandListImmediate& RHICmdList)
-	{
-		if (ProjectionPolicy.IsValid())
-		{
-			FDisplayClusterProjectionMPCDIPolicy* MPCDIPolicy = static_cast<FDisplayClusterProjectionMPCDIPolicy*>(ProjectionPolicy.Get());
-			if (MPCDIPolicy)
-			{
-				MPCDIPolicy->WarpBlendInterface_Proxy = WarpBlendInterfacePtr;
-				MPCDIPolicy->WarpPolicyInterface_Proxy = WarPolicyInterfacePtr;
-				MPCDIPolicy->WarpBlendContexts_Proxy = Contexts;
-			}
-		}
+		[
+			This = SharedThis(this),
+			WarpBlendInterfacePtr = WarpBlendInterface,
+			WarPolicyInterfacePtr = WarpPolicyInterface,
+			Contexts = WarpBlendContexts
+		](FRHICommandListImmediate& RHICmdList)
+	{
+		This->WarpBlendInterface_Proxy = WarpBlendInterfacePtr;
+		This->WarpPolicyInterface_Proxy = WarPolicyInterfacePtr;
+		This->WarpBlendContexts_Proxy = Contexts;
 	});
 }
 
@@ -320,7 +318,12 @@ bool FDisplayClusterProjectionMPCDIPolicy::IsFrustumRotatedToFitContextSize(IDis
 	return false;
 }
 
-bool FDisplayClusterProjectionMPCDIPolicy::IsWarpBlendSupported()
+bool FDisplayClusterProjectionMPCDIPolicy::IsWarpBlendSupported(const IDisplayClusterViewport* InViewport) const
+{
+	return true;
+}
+
+bool FDisplayClusterProjectionMPCDIPolicy::IsWarpBlendSupported_RenderThread(const IDisplayClusterViewportProxy* InViewportProxy) const
 {
 	return true;
 }
diff --git a/Engine/Plugins/Runtime/nDisplay/Source/DisplayClusterProjection/Private/Policy/MPCDI/DisplayClusterProjectionMPCDIPolicy.h b/Engine/Plugins/Runtime/nDisplay/Source/DisplayClusterProjection/Private/Policy/MPCDI/DisplayClusterProjectionMPCDIPolicy.h
--- a/Engine/Plugins/Runtime/nDisplay/Source/DisplayClusterProjection/Private/Policy/MPCDI/DisplayClusterProjectionMPCDIPolicy.h
+++ b/Engine/Plugins/Runtime/nDisplay/Source/DisplayClusterProjection/Private/Policy/MPCDI/DisplayClusterProjectionMPCDIPolicy.h
@@ -5,6 +5,7 @@
 #include "Policy/DisplayClusterProjectionPolicyBase.h"
 #include "Containers/DisplayClusterWarpContext.h"
 #include "IDisplayClusterWarpBlend.h"
+#include "Templates/SharedPointer.h"
 
 class UMeshComponent;
 
@@ -14,6 +15,7 @@ class UMeshComponent;
  */
 class FDisplayClusterProjectionMPCDIPolicy
 	: public FDisplayClusterProjectionPolicyBase
+	, public TSharedFromThis<FDisplayClusterProjectionMPCDIPolicy, ESPMode::ThreadSafe>
 {
 public:
 	FDisplayClusterProjectionMPCDIPolicy(const FString& ProjectionPolicyId, const FDisplayClusterConfigurationProjection* InConfigurationProjectionPolicy);
@@ -21,7 +23,7 @@ public:
 
 public:
 	// This policy can support ICVFX rendering
-	virtual bool ShouldSupportICVFX(IDisplayClusterViewport* InViewport) const override;
+	virtual bool ShouldSupportICVFX(const IDisplayClusterViewport* InViewport) const override;
 
 public:
 	//~Begin IDisplayClusterProjectionPolicy
@@ -34,19 +36,22 @@ public:
 	virtual bool GetProjectionMatrix(IDisplayClusterViewport* InViewport, const uint32 InContextNum, FMatrix& OutPrjMatrix) override;
 
 	virtual bool IsFrustumRotatedToFitContextSize(IDisplayClusterViewport* InViewport, const uint32 InContextNum) override;
-	virtual bool IsWarpBlendSupported() override;
+
+	virtual bool IsWarpBlendSupported(const IDisplayClusterViewport* InViewport) const override;
+	virtual bool IsWarpBlendSupported_RenderThread(const IDisplayClusterViewportProxy* InViewportProxy) const override;
+
 	virtual void ApplyWarpBlend_RenderThread(FRHICommandListImmediate& RHICmdList, const IDisplayClusterViewportProxy* InViewportProxy) override;
 
 	virtual bool GetWarpBlendInterface(TSharedPtr<class IDisplayClusterWarpBlend, ESPMode::ThreadSafe>& OutWarpBlendInterface) const override;
 	virtual bool GetWarpBlendInterface_RenderThread(TSharedPtr<class IDisplayClusterWarpBlend, ESPMode::ThreadSafe>& OutWarpBlendInterfaceProxy) const override;
 
-	virtual bool ShouldUseSourceTextureWithMips() const override
+	virtual bool ShouldUseSourceTextureWithMips(const IDisplayClusterViewport* InViewport) const override
 	{
 		// Support input texture with mips
 		return true;
 	}
 
-	virtual bool ShouldUseAdditionalTargetableResource() const override
+	virtual bool ShouldUseAdditionalTargetableResource(const IDisplayClusterViewport* InViewport) const override
 	{
 		// Request additional targetable resources for warp&blend output
 		return true;
diff --git a/Engine/Plugins/Runtime/nDisplay/Source/DisplayClusterProjection/Private/Policy/Reference/DisplayClusterProjectionReferencePolicy.cpp b/Engine/Plugins/Runtime/nDisplay/Source/DisplayClusterProjection/Private/Policy/Reference/DisplayClusterProjectionReferencePolicy.cpp
new file mode 100644
--- /dev/null
+++ b/Engine/Plugins/Runtime/nDisplay/Source/DisplayClusterProjection/Private/Policy/Reference/DisplayClusterProjectionReferencePolicy.cpp
@@ -0,0 +1,315 @@
+// Copyright Epic Games, Inc. All Rights Reserved.
+
+#include "Policy/Reference/DisplayClusterProjectionReferencePolicy.h"
+
+#include "Misc/DisplayClusterHelpers.h"
+
+#include "DisplayClusterProjectionLog.h"
+#include "DisplayClusterProjectionStrings.h"
+
+#include "Render/Viewport/IDisplayClusterViewportManager.h"
+#include "Render/Viewport/IDisplayClusterViewportManagerProxy.h"
+#include "Render/Viewport/IDisplayClusterViewport.h"
+#include "Render/Viewport/IDisplayClusterViewportProxy.h"
+
+
+FDisplayClusterProjectionReferencePolicy::FDisplayClusterProjectionReferencePolicy(
+	const FString& ProjectionPolicyId,
+	const FDisplayClusterConfigurationProjection* InConfigurationProjectionPolicy)
+	: FDisplayClusterProjectionPolicyBase(ProjectionPolicyId, InConfigurationProjectionPolicy)
+{
+	if (DisplayClusterHelpers::map::template ExtractValue(GetParameters(),
+		DisplayClusterProjectionStrings::cfg::reference::ViewportId,
+		ReferencedViewportId))
+	{
+		UE_LOG(LogDisplayClusterProjectionReference, Verbose, TEXT("Found Argument '%s'='%s'"),
+			DisplayClusterProjectionStrings::cfg::reference::ViewportId, *ReferencedViewportId);
+	}
+}
+
+const FString& FDisplayClusterProjectionReferencePolicy::GetType() const
+{
+	static const FString Type(DisplayClusterProjectionStrings::projection::Reference);
+
+	return Type;
+}
+
+void FDisplayClusterProjectionReferencePolicy::UpdateProxyData(IDisplayClusterViewport* InViewport)
+{
+	FString ReferencedViewportIdProxy;
+	bool bSupportICVFX = false;
+	if (IDisplayClusterViewport* SourceViewport = GetSourceViewport(InViewport))
+	{
+		bSupportICVFX = SourceViewport->GetProjectionPolicy().IsValid() && SourceViewport->GetProjectionPolicy()->ShouldSupportICVFX(SourceViewport);
+		ReferencedViewportIdProxy = ReferencedViewportId;
+	}
+
+	// Send view data to rendering thread.
+	ENQUEUE_RENDER_COMMAND(DisplayClusterProjectionReferencePolicy_UpdateProxyData)(
+		[
+			This = SharedThis(this),
+			InReferencedViewportId = ReferencedViewportIdProxy,
+			bSupportICVFX = bSupportICVFX
+		](FRHICommandListImmediate& RHICmdList)
+		{
+			// Update rendering thread resoures
+			This->ReferencedViewportId_RenderThread = InReferencedViewportId;
+			This->bSupportICVFXProxy = bSupportICVFX;
+		});
+}
+
+bool FDisplayClusterProjectionReferencePolicy::ShouldSupportICVFX(const IDisplayClusterViewport* InViewport) const
+{
+	if (IDisplayClusterViewport* SourceViewport = GetSourceViewport(InViewport))
+	{
+		// return source viewport settings
+		return SourceViewport->GetProjectionPolicy().IsValid()
+			&& SourceViewport->GetProjectionPolicy()->ShouldSupportICVFX(SourceViewport);
+	}
+
+	return false;
+}
+
+bool FDisplayClusterProjectionReferencePolicy::ShouldUseSourceTextureWithMips(const IDisplayClusterViewport* InViewport) const
+{
+	if (IDisplayClusterViewport* SourceViewport = GetSourceViewport(InViewport))
+	{
+		// Only ICVFX requires composition. Other projection policies do not. So we can reuse all textures and not do warpblend.
+		if (SourceViewport->GetProjectionPolicy().IsValid() && SourceViewport->GetProjectionPolicy()->ShouldSupportICVFX(SourceViewport))
+		{
+			return SourceViewport->GetProjectionPolicy()->ShouldUseSourceTextureWithMips(SourceViewport);
+		}
+	}
+
+	return false;
+}
+
+bool FDisplayClusterProjectionReferencePolicy::ShouldUseAdditionalTargetableResource(const IDisplayClusterViewport* InViewport) const
+{
+	if (IDisplayClusterViewport* SourceViewport = GetSourceViewport(InViewport))
+	{
+		// Only ICVFX requires composition. Other projection policies do not. So we can reuse all textures and not do warpblend.
+		if (SourceViewport->GetProjectionPolicy().IsValid() && SourceViewport->GetProjectionPolicy()->ShouldSupportICVFX(SourceViewport))
+		{
+			return SourceViewport->GetProjectionPolicy()->ShouldUseAdditionalTargetableResource(SourceViewport);
+		}
+	}
+
+	return false;
+}
+
+bool FDisplayClusterProjectionReferencePolicy::IsWarpBlendSupported(const IDisplayClusterViewport* InViewport) const
+{
+	if (IDisplayClusterViewport* SourceViewport = GetSourceViewport(InViewport))
+	{
+		// Only ICVFX requires composition. Other projection policies do not. So we can reuse all textures and not do warpblend.
+		if (SourceViewport->GetProjectionPolicy().IsValid() && SourceViewport->GetProjectionPolicy()->ShouldSupportICVFX(SourceViewport))
+		{
+			return SourceViewport->GetProjectionPolicy()->IsWarpBlendSupported(SourceViewport);
+		}
+	}
+
+	return false;
+}
+
+bool FDisplayClusterProjectionReferencePolicy::IsWarpBlendSupported_RenderThread(const IDisplayClusterViewportProxy* InViewportProxy) const
+{
+	// Only ICVFX requires composition. Other projection policies do not. So we can reuse all textures and not do warpblend.
+	if (!bSupportICVFXProxy)
+	{
+		return false;
+	}
+
+	if (IDisplayClusterViewportProxy* SourceViewportProxy = GetSourceViewport_RenderThread(InViewportProxy))
+	{
+		return SourceViewportProxy->GetProjectionPolicy_RenderThread().IsValid() && SourceViewportProxy->GetProjectionPolicy_RenderThread()->IsWarpBlendSupported_RenderThread(SourceViewportProxy);
+	}
+
+	return false;
+}
+
+void FDisplayClusterProjectionReferencePolicy::ApplyWarpBlend_RenderThread(FRHICommandListImmediate& RHICmdList, const IDisplayClusterViewportProxy* InViewportProxy)
+{
+	// Only ICVFX requires composition. Other projection policies do not. So we can reuse all textures and not do warpblend.
+	if (!bSupportICVFXProxy)
+	{
+		return;
+	}
+
+	if (IDisplayClusterViewportProxy* SourceViewportProxy = GetSourceViewport_RenderThread(InViewportProxy))
+	{
+		if (SourceViewportProxy->GetProjectionPolicy_RenderThread().IsValid())
+		{
+			// Rendering a this (cloned) viewport using the logic of the source viewport
+			SourceViewportProxy->GetProjectionPolicy_RenderThread()->ApplyWarpBlend_RenderThread(RHICmdList, InViewportProxy);
+		}
+	}
+}
+
+void FDisplayClusterProjectionReferencePolicy::PostUpdateBaseConfiguration(IDisplayClusterViewport* InViewport) const
+{
+	// Get the current rendering setting we want to change
+	FDisplayClusterViewport_RenderSettings InOutRenderSettings = InViewport->GetRenderSettings();
+
+	// If a source viewport exists, use it as the source of textures for rendering
+	if (IDisplayClusterViewport* SourceViewport = GetSourceViewport(InViewport))
+	{
+		if (InViewport->UseSameOCIO(SourceViewport))
+		{
+			// Rule #1: By default, all resources except the output RTT are cloned.
+			// The output texture must be rendered using the projection policy from the parent viewport.
+			InOutRenderSettings.SetViewportOverride(ReferencedViewportId, EDisplayClusterViewportOverrideMode::InternalViewportResources);
+		}
+		else
+		{
+			// Rule#2: If the OCIO on the cloned viewport is different, only the input RTT is cloned. The custom OCIO must then be applied.
+			InOutRenderSettings.SetViewportOverride(ReferencedViewportId, EDisplayClusterViewportOverrideMode::InternalRTT);
+		}
+
+		// Apply the changes we made to the current viewport.
+		InViewport->SetRenderSettings(InOutRenderSettings);
+	}
+	else
+	{
+		// Setup is not ready, just disable this viewport
+		InOutRenderSettings.bEnable = false;
+	}
+
+	// Apply the changes we made to the current viewport.
+	InViewport->SetRenderSettings(InOutRenderSettings);
+}
+
+IDisplayClusterViewport* FDisplayClusterProjectionReferencePolicy::GetSourceViewport(const IDisplayClusterViewport* InViewport) const
+{
+	check(IsInGameThread());
+	check(InViewport);
+
+	IDisplayClusterViewportManager* ViewportManager = InViewport ? InViewport->GetConfiguration().GetViewportManager() : nullptr;
+
+	// These are settings we change to alter the viewport behaviour.
+	if (ReferencedViewportId.IsEmpty() || !ViewportManager)
+	{
+		return nullptr;
+	}
+
+	// Check if the viewport name we referenced to exists
+	IDisplayClusterViewport* SourceViewport = ViewportManager->FindViewport(ReferencedViewportId);
+	if (!SourceViewport)
+	{
+		// The viewport referenced does not exist.
+		if (!EnumHasAnyFlags(LogMsgState, EDisplayClusterProjectionReferencePolicyLogMsgState::InvalidName))
+		{
+			EnumAddFlags(LogMsgState, EDisplayClusterProjectionReferencePolicyLogMsgState::InvalidName);
+
+			UE_LOG(LogDisplayClusterProjectionReference, Warning,
+				TEXT("The source viewport '%s' that referenced by the viewport '%s' not exist."),
+				*ReferencedViewportId, *InViewport->GetId());
+		}
+
+		return nullptr;
+	}
+
+	// The source viewport should be located on the same node of the cluster.
+	if (SourceViewport->GetClusterNodeId() != InViewport->GetClusterNodeId())
+	{
+		if (!EnumHasAnyFlags(LogMsgState, EDisplayClusterProjectionReferencePolicyLogMsgState::LocalNode))
+		{
+			EnumAddFlags(LogMsgState, EDisplayClusterProjectionReferencePolicyLogMsgState::LocalNode);
+
+			UE_LOG(LogDisplayClusterProjectionReference, Warning,
+				TEXT("The source viewport '%s' that referenced by the viewport '%s' must be on the same node."),
+				*ReferencedViewportId, *InViewport->GetId());
+		}
+
+		return nullptr;
+	}
+
+	// The viewport we are referencing cannot reference another viewport or use this projection policy.
+	if (SourceViewport->GetRenderSettings().IsViewportOverridden()
+	|| (SourceViewport->GetProjectionPolicy().IsValid() && SourceViewport->GetProjectionPolicy()->GetType() == GetType()))
+	{
+		if (!EnumHasAnyFlags(LogMsgState, EDisplayClusterProjectionReferencePolicyLogMsgState::Recursion))
+		{
+			EnumAddFlags(LogMsgState, EDisplayClusterProjectionReferencePolicyLogMsgState::Recursion);
+
+			UE_LOG(LogDisplayClusterProjectionReference, Warning,
+				TEXT("The source viewport '%s' that referenced by the viewport '%s' can't be referenced."),
+				*ReferencedViewportId, *InViewport->GetId());
+		}
+
+		return nullptr;
+	}
+
+	// No errors - let's reset the log flags.
+	LogMsgState = EDisplayClusterProjectionReferencePolicyLogMsgState::None;
+
+	return SourceViewport;
+}
+
+IDisplayClusterViewportProxy* FDisplayClusterProjectionReferencePolicy::GetSourceViewport_RenderThread(const IDisplayClusterViewportProxy* InViewportProxy) const
+{
+	check(IsInRenderingThread());
+	check(InViewportProxy);
+
+	IDisplayClusterViewportManagerProxy* ViewportManagerProxy = InViewportProxy ? InViewportProxy->GetConfigurationProxy().GetViewportManagerProxy_RenderThread() : nullptr;
+
+	// These are settings we change to alter the viewport behaviour.
+	if (ReferencedViewportId_RenderThread.IsEmpty() || !ViewportManagerProxy)
+	{
+		return nullptr;
+	}
+
+	// Check if the viewport name we referenced to exists
+	IDisplayClusterViewportProxy* SourceViewportProxy = ViewportManagerProxy->FindViewport_RenderThread(ReferencedViewportId_RenderThread);
+	if (!SourceViewportProxy)
+	{
+		// The viewport referenced does not exist.
+		if (!EnumHasAnyFlags(LogMsgState_RenderThread, EDisplayClusterProjectionReferencePolicyLogMsgState::InvalidName))
+		{
+			EnumAddFlags(LogMsgState_RenderThread, EDisplayClusterProjectionReferencePolicyLogMsgState::InvalidName);
+
+			UE_LOG(LogDisplayClusterProjectionReference, Warning,
+				TEXT("The source viewport proxy '%s' that referenced by the viewport '%s' not exist."),
+				*ReferencedViewportId, *InViewportProxy->GetId());
+		}
+
+		return nullptr;
+	}
+
+	// The source viewport should be located on the same node of the cluster.
+	if (SourceViewportProxy->GetClusterNodeId() != InViewportProxy->GetClusterNodeId())
+	{
+		if (!EnumHasAnyFlags(LogMsgState_RenderThread, EDisplayClusterProjectionReferencePolicyLogMsgState::LocalNode))
+		{
+			EnumAddFlags(LogMsgState_RenderThread, EDisplayClusterProjectionReferencePolicyLogMsgState::LocalNode);
+
+			UE_LOG(LogDisplayClusterProjectionReference, Warning,
+				TEXT("The source viewport proxy '%s' that referenced by the viewport '%s' must be on the same node."),
+				*ReferencedViewportId, *InViewportProxy->GetId());
+		}
+
+		return nullptr;
+	}
+
+	// The viewport we are referencing cannot reference another viewport or use this projection policy.
+	if (SourceViewportProxy->GetRenderSettings_RenderThread().IsViewportOverridden()||
+		(  SourceViewportProxy->GetProjectionPolicy_RenderThread().IsValid()
+		&& SourceViewportProxy->GetProjectionPolicy_RenderThread()->GetType() == GetType()))
+	{
+		if (!EnumHasAnyFlags(LogMsgState_RenderThread, EDisplayClusterProjectionReferencePolicyLogMsgState::Recursion))
+		{
+			EnumAddFlags(LogMsgState_RenderThread, EDisplayClusterProjectionReferencePolicyLogMsgState::Recursion);
+
+			UE_LOG(LogDisplayClusterProjectionReference, Warning,
+				TEXT("The source viewport proxy '%s' that referenced by the viewport '%s' can't be referenced."),
+				*ReferencedViewportId, *InViewportProxy->GetId());
+		}
+
+		return nullptr;
+	}
+
+	// No errors - let's reset the log flags.
+	LogMsgState_RenderThread = EDisplayClusterProjectionReferencePolicyLogMsgState::None;
+
+	return SourceViewportProxy;
+}
diff --git a/Engine/Plugins/Runtime/nDisplay/Source/DisplayClusterProjection/Private/Policy/Reference/DisplayClusterProjectionReferencePolicy.h b/Engine/Plugins/Runtime/nDisplay/Source/DisplayClusterProjection/Private/Policy/Reference/DisplayClusterProjectionReferencePolicy.h
new file mode 100644
--- /dev/null
+++ b/Engine/Plugins/Runtime/nDisplay/Source/DisplayClusterProjection/Private/Policy/Reference/DisplayClusterProjectionReferencePolicy.h
@@ -0,0 +1,95 @@
+// Copyright Epic Games, Inc. All Rights Reserved.
+
+#pragma once
+
+#include "CoreMinimal.h"
+#include "Policy/DisplayClusterProjectionPolicyBase.h"
+#include "Templates/SharedPointer.h"
+
+/** Enum flags used to mark log messages after they have been displayed. */
+enum class EDisplayClusterProjectionReferencePolicyLogMsgState : uint8
+{
+	None = 0,
+
+	// The referenced viewport should be on the same cluster node.
+	LocalNode = 1 << 0,
+
+	// The viewport we are referencing cannot reference another viewport or use this projection policy.
+	Recursion= 1 << 1,
+
+	// The viewport referenced does not exist.
+	InvalidName = 1 << 2,
+};
+ENUM_CLASS_FLAGS(EDisplayClusterProjectionReferencePolicyLogMsgState);
+
+/**
+ * Link projection policy (for internal usage)
+ */
+class FDisplayClusterProjectionReferencePolicy
+	: public FDisplayClusterProjectionPolicyBase
+	, public TSharedFromThis<FDisplayClusterProjectionReferencePolicy, ESPMode::ThreadSafe>
+{
+public:
+	FDisplayClusterProjectionReferencePolicy(
+		const FString& ProjectionPolicyId,
+		const FDisplayClusterConfigurationProjection* InConfigurationProjectionPolicy);
+
+public:
+	//~Begin IDisplayClusterProjectionPolicy
+	virtual const FString& GetType() const override;
+
+	virtual bool ShouldSupportICVFX(const IDisplayClusterViewport* InViewport) const override;
+	virtual void PostUpdateBaseConfiguration(IDisplayClusterViewport* InViewport) const override;
+
+	virtual bool IsWarpBlendSupported(const IDisplayClusterViewport* InViewport) const override;
+	virtual bool IsWarpBlendSupported_RenderThread(const IDisplayClusterViewportProxy* InViewportProxy) const override;
+
+	virtual void ApplyWarpBlend_RenderThread(FRHICommandListImmediate& RHICmdList, const IDisplayClusterViewportProxy* InViewportProxy) override;
+
+	virtual bool ShouldUseSourceTextureWithMips(const IDisplayClusterViewport* InViewport) const override;
+	virtual bool ShouldUseAdditionalTargetableResource(const IDisplayClusterViewport* InViewport) const override;
+
+	virtual void UpdateProxyData(IDisplayClusterViewport* InViewport) override;
+
+	virtual bool CalculateView(
+		IDisplayClusterViewport* InViewport,
+		const uint32 InContextNum,
+		FVector& InOutViewLocation,
+		FRotator& InOutViewRotation,
+		const FVector& ViewOffset,
+		const float WorldToMeters,
+		const float NCP,
+		const float FCP) override
+	{
+		// This projection policy does not use rendering
+		return false;
+	}
+
+	virtual bool GetProjectionMatrix(
+		IDisplayClusterViewport* InViewport,
+		const uint32 InContextNum,
+		FMatrix& OutPrjMatrix) override
+	{
+		// This projection policy does not use rendering
+		return false;
+	}
+	//~~End IDisplayClusterProjectionPolicy
+
+protected:
+	/** Return viewport referenced by this policy. */
+	IDisplayClusterViewport* GetSourceViewport(const IDisplayClusterViewport* InViewport) const;
+
+	/** Return viewport proxy referenced by this policy. */
+	IDisplayClusterViewportProxy* GetSourceViewport_RenderThread(const IDisplayClusterViewportProxy* InViewportProxy) const;
+
+private:
+	// The name of the viewport referenced by this policy.
+	FString ReferencedViewportId;
+	FString ReferencedViewportId_RenderThread;
+
+	// Allow a message to be displayed in the log only once.
+	mutable EDisplayClusterProjectionReferencePolicyLogMsgState LogMsgState = EDisplayClusterProjectionReferencePolicyLogMsgState::None;
+	mutable EDisplayClusterProjectionReferencePolicyLogMsgState LogMsgState_RenderThread = EDisplayClusterProjectionReferencePolicyLogMsgState::None;
+
+	bool bSupportICVFXProxy = false;
+};
diff --git a/Engine/Plugins/Runtime/nDisplay/Source/DisplayClusterProjection/Private/Policy/Reference/DisplayClusterProjectionReferencePolicyFactory.cpp b/Engine/Plugins/Runtime/nDisplay/Source/DisplayClusterProjection/Private/Policy/Reference/DisplayClusterProjectionReferencePolicyFactory.cpp
new file mode 100644
--- /dev/null
+++ b/Engine/Plugins/Runtime/nDisplay/Source/DisplayClusterProjection/Private/Policy/Reference/DisplayClusterProjectionReferencePolicyFactory.cpp
@@ -0,0 +1,18 @@
+// Copyright Epic Games, Inc. All Rights Reserved.
+
+#include "Policy/Reference/DisplayClusterProjectionReferencePolicyFactory.h"
+#include "Policy/Reference/DisplayClusterProjectionReferencePolicy.h"
+#include "DisplayClusterConfigurationTypes.h"
+#include "DisplayClusterProjectionLog.h"
+
+TSharedPtr<IDisplayClusterProjectionPolicy, ESPMode::ThreadSafe> FDisplayClusterProjectionReferencePolicyFactory::Create(
+	const FString& ProjectionPolicyId,
+	const FDisplayClusterConfigurationProjection* InConfigurationProjectionPolicy)
+{
+	check(InConfigurationProjectionPolicy != nullptr);
+
+	UE_LOG(LogDisplayClusterProjectionReference, Verbose,
+		TEXT("Instantiating projection policy <%s> id='%s'"), *InConfigurationProjectionPolicy->Type, *ProjectionPolicyId);
+
+	return  MakeShared<FDisplayClusterProjectionReferencePolicy, ESPMode::ThreadSafe>(ProjectionPolicyId, InConfigurationProjectionPolicy);
+}
diff --git a/Engine/Plugins/Runtime/nDisplay/Source/DisplayClusterProjection/Private/Policy/Reference/DisplayClusterProjectionReferencePolicyFactory.h b/Engine/Plugins/Runtime/nDisplay/Source/DisplayClusterProjection/Private/Policy/Reference/DisplayClusterProjectionReferencePolicyFactory.h
new file mode 100644
--- /dev/null
+++ b/Engine/Plugins/Runtime/nDisplay/Source/DisplayClusterProjection/Private/Policy/Reference/DisplayClusterProjectionReferencePolicyFactory.h
@@ -0,0 +1,19 @@
+// Copyright Epic Games, Inc. All Rights Reserved.
+
+#pragma once
+
+#include "Render/Projection/IDisplayClusterProjectionPolicyFactory.h"
+
+/**
+ * Implements projection policy factory for the 'Reference' policy
+ */
+class FDisplayClusterProjectionReferencePolicyFactory
+	: public IDisplayClusterProjectionPolicyFactory
+{
+public:
+	//~Begin IDisplayClusterProjectionPolicyFactory
+	virtual TSharedPtr<IDisplayClusterProjectionPolicy, ESPMode::ThreadSafe> Create(
+		const FString& ProjectionPolicyId,
+		const FDisplayClusterConfigurationProjection* InConfigurationProjectionPolicy) override;
+	//~~End IDisplayClusterProjectionPolicyFactory
+};
diff --git a/Engine/Plugins/Runtime/nDisplay/Source/DisplayClusterProjection/Private/Policy/VIOSO/Windows/DisplayClusterProjectionVIOSOPolicy.cpp b/Engine/Plugins/Runtime/nDisplay/Source/DisplayClusterProjection/Private/Policy/VIOSO/Windows/DisplayClusterProjectionVIOSOPolicy.cpp
--- a/Engine/Plugins/Runtime/nDisplay/Source/DisplayClusterProjection/Private/Policy/VIOSO/Windows/DisplayClusterProjectionVIOSOPolicy.cpp
+++ b/Engine/Plugins/Runtime/nDisplay/Source/DisplayClusterProjection/Private/Policy/VIOSO/Windows/DisplayClusterProjectionVIOSOPolicy.cpp
@@ -144,7 +144,12 @@ bool FDisplayClusterProjectionVIOSOPolicy::GetProjectionMatrix(IDisplayClusterVi
 	return true;
 }
 
-bool FDisplayClusterProjectionVIOSOPolicy::IsWarpBlendSupported()
+bool FDisplayClusterProjectionVIOSOPolicy::IsWarpBlendSupported(const IDisplayClusterViewport* InViewport) const
+{
+	return true;
+}
+
+bool FDisplayClusterProjectionVIOSOPolicy::IsWarpBlendSupported_RenderThread(const IDisplayClusterViewportProxy* InViewportProxy) const
 {
 	return true;
 }
diff --git a/Engine/Plugins/Runtime/nDisplay/Source/DisplayClusterProjection/Private/Policy/VIOSO/Windows/DisplayClusterProjectionVIOSOPolicy.h b/Engine/Plugins/Runtime/nDisplay/Source/DisplayClusterProjection/Private/Policy/VIOSO/Windows/DisplayClusterProjectionVIOSOPolicy.h
--- a/Engine/Plugins/Runtime/nDisplay/Source/DisplayClusterProjection/Private/Policy/VIOSO/Windows/DisplayClusterProjectionVIOSOPolicy.h
+++ b/Engine/Plugins/Runtime/nDisplay/Source/DisplayClusterProjection/Private/Policy/VIOSO/Windows/DisplayClusterProjectionVIOSOPolicy.h
@@ -25,16 +25,18 @@ public:
 	virtual bool CalculateView(IDisplayClusterViewport* InViewport, const uint32 InContextNum, FVector& InOutViewLocation, FRotator& InOutViewRotation, const FVector& ViewOffset, const float WorldToMeters, const float NCP, const float FCP) override;
 	virtual bool GetProjectionMatrix(IDisplayClusterViewport* InViewport, const uint32 InContextNum, FMatrix& OutPrjMatrix) override;
 
-	virtual bool IsWarpBlendSupported() override;
+	virtual bool IsWarpBlendSupported(const IDisplayClusterViewport* InViewport) const override;
+	virtual bool IsWarpBlendSupported_RenderThread(const IDisplayClusterViewportProxy* InViewportProxy) const override;
+
 	virtual void ApplyWarpBlend_RenderThread(FRHICommandListImmediate& RHICmdList, const IDisplayClusterViewportProxy* InViewportProxy) override;
 
 	// Request additional targetable resources for domeprojection  external warpblend
-	virtual bool ShouldUseAdditionalTargetableResource() const override
+	virtual bool ShouldUseAdditionalTargetableResource(const IDisplayClusterViewport* InViewport) const override
 	{
 		return true;
 	}
 
-	virtual bool ShouldUseSourceTextureWithMips() const override
+	virtual bool ShouldUseSourceTextureWithMips(const IDisplayClusterViewport* InViewport) const override
 	{
 		return true;
 	}
diff --git a/Engine/Plugins/Runtime/nDisplay/Source/DisplayClusterProjection/Public/DisplayClusterProjectionStrings.h b/Engine/Plugins/Runtime/nDisplay/Source/DisplayClusterProjection/Public/DisplayClusterProjectionStrings.h
--- a/Engine/Plugins/Runtime/nDisplay/Source/DisplayClusterProjection/Public/DisplayClusterProjectionStrings.h
+++ b/Engine/Plugins/Runtime/nDisplay/Source/DisplayClusterProjection/Public/DisplayClusterProjectionStrings.h
@@ -82,6 +82,11 @@ namespace DisplayClusterProjectionStrings
 			static constexpr const TCHAR* Native    = TEXT("native");
 		}
 
+		namespace reference
+		{
+			static constexpr const TCHAR* ViewportId = TEXT("viewport");
+		}
+
 		namespace mesh
 		{
 			static constexpr const TCHAR* Component        = TEXT("mesh_component");
@@ -171,6 +176,7 @@ namespace DisplayClusterProjectionStrings
 		static constexpr const TCHAR* Manual         = TEXT("manual");
 		static constexpr const TCHAR* Mesh           = TEXT("mesh");
 		static constexpr const TCHAR* MPCDI          = TEXT("mpcdi");
+		static constexpr const TCHAR* Reference      = TEXT("reference");
 		static constexpr const TCHAR* Simple         = TEXT("simple");
 		static constexpr const TCHAR* VIOSO          = TEXT("vioso");
 	}
diff --git a/Engine/Plugins/Runtime/nDisplay/Source/DisplayClusterShaders/Private/Shaders/DisplayClusterShadersWarpblend_ICVFX.cpp b/Engine/Plugins/Runtime/nDisplay/Source/DisplayClusterShaders/Private/Shaders/DisplayClusterShadersWarpblend_ICVFX.cpp
--- a/Engine/Plugins/Runtime/nDisplay/Source/DisplayClusterShaders/Private/Shaders/DisplayClusterShadersWarpblend_ICVFX.cpp
+++ b/Engine/Plugins/Runtime/nDisplay/Source/DisplayClusterShaders/Private/Shaders/DisplayClusterShadersWarpblend_ICVFX.cpp
@@ -768,7 +768,7 @@ public:
 	EDisplayClusterShaderParametersICVFX_ChromakeySource GetCameraOverlapParameters(FIcvfxRenderPassData& RenderPassData)
 	{
 		int32 OverlapCameraNum = 0;
-		EDisplayClusterShaderParametersICVFX_ChromakeySource OutChromakeySource = EDisplayClusterShaderParametersICVFX_ChromakeySource::Unknown;
+		EDisplayClusterShaderParametersICVFX_ChromakeySource OutChromakeySource = EDisplayClusterShaderParametersICVFX_ChromakeySource::Default;
 
 		if (EnumHasAnyFlags(RenderState, EIcvfxPassRenderState::EnableInnerFrustumChromakeyOverlap))
 		{
@@ -859,11 +859,30 @@ public:
 		}
 		
 		const FDisplayClusterShaderParameters_ICVFX::FCameraSettings& Camera = ICVFXParameters.Cameras[ActiveCameraIndex];
-		
-		switch (InChromakeySource == EDisplayClusterShaderParametersICVFX_ChromakeySource::Unknown ? Camera.ChromakeySource : InChromakeySource)
+
+		EDisplayClusterShaderParametersICVFX_ChromakeySource ChromakeySource = Camera.ChromakeySource;
+
+		// Override by overlap
+		if (InChromakeySource != EDisplayClusterShaderParametersICVFX_ChromakeySource::Default)
+		{
+			ChromakeySource = InChromakeySource;
+		};
+
+
+		switch (ChromakeySource)
 		{
 		case EDisplayClusterShaderParametersICVFX_ChromakeySource::ChromakeyLayers:
 		{
+			// Protection against invalid parameter combinations
+			if (RenderPassData.PSParameters.InnerCameraTexture == nullptr || Camera.Chromakey.Texture == nullptr)
+			{
+				// There is no InCamera texture and no frame colour, we have nothing to render.
+				// Disable this camera, wrong combination
+				RenderPassData.PSPermutationVector.Set<IcvfxShaderPermutation::FIcvfxShaderInnerCamera>(false);
+
+				return false;
+			}
+
 			RenderPassData.PSParameters.ChromakeyColor = Camera.ChromakeyColor;
 			RenderPassData.PSParameters.ChromakeyCameraTexture = Camera.Chromakey.Texture;
 			RenderPassData.PSParameters.ChromakeyCameraSampler = TStaticSamplerState<SF_Trilinear, AM_Clamp, AM_Clamp, AM_Clamp>::GetRHI();
@@ -889,11 +908,22 @@ public:
 			return true;
 		}
 
+		case EDisplayClusterShaderParametersICVFX_ChromakeySource::Default:
+			// By default chromakey is disabled
 		case EDisplayClusterShaderParametersICVFX_ChromakeySource::Disabled:
 		default:
 			break;
 		}
 
+		// Protection against invalid parameter combinations
+		if (RenderPassData.PSParameters.InnerCameraTexture == nullptr)
+		{
+			// Disable this camera, wrong combination
+			RenderPassData.PSPermutationVector.Set<IcvfxShaderPermutation::FIcvfxShaderInnerCamera>(false);
+
+			return false;
+		}
+
 		return false;
 	}
 
@@ -968,7 +998,8 @@ public:
 
 			if (GetCameraParameters(RenderPassData))
 			{
-				const EDisplayClusterShaderParametersICVFX_ChromakeySource ChromakeySourceOverride = GetCameraOverlapParameters(RenderPassData);
+				// Overlap is first priority
+				EDisplayClusterShaderParametersICVFX_ChromakeySource ChromakeySourceOverride = GetCameraOverlapParameters(RenderPassData);
 
 				// Chromakey inside cam:
 				if (GetCameraChromakeyParameters(RenderPassData, ChromakeySourceOverride))
diff --git a/Engine/Plugins/Runtime/nDisplay/Source/DisplayClusterShaders/Public/Containers/DisplayClusterShader_Enums.h b/Engine/Plugins/Runtime/nDisplay/Source/DisplayClusterShaders/Public/Containers/DisplayClusterShader_Enums.h
--- a/Engine/Plugins/Runtime/nDisplay/Source/DisplayClusterShaders/Public/Containers/DisplayClusterShader_Enums.h
+++ b/Engine/Plugins/Runtime/nDisplay/Source/DisplayClusterShaders/Public/Containers/DisplayClusterShader_Enums.h
@@ -18,7 +18,8 @@ enum class EDisplayClusterShaderParametersICVFX_ChromakeySource : uint8
 	// Render specified layer from scene
 	ChromakeyLayers,
 
-	Unknown,
+	// Use default value.
+	Default,
 };
 
 /**
diff --git a/Engine/Plugins/VirtualProduction/TextureShare/Source/TextureShareDisplayCluster/Private/PostProcess/TextureSharePostprocess_FrameProxy.cpp b/Engine/Plugins/VirtualProduction/TextureShare/Source/TextureShareDisplayCluster/Private/PostProcess/TextureSharePostprocess_FrameProxy.cpp
--- a/Engine/Plugins/VirtualProduction/TextureShare/Source/TextureShareDisplayCluster/Private/PostProcess/TextureSharePostprocess_FrameProxy.cpp
+++ b/Engine/Plugins/VirtualProduction/TextureShare/Source/TextureShareDisplayCluster/Private/PostProcess/TextureSharePostprocess_FrameProxy.cpp
@@ -27,7 +27,7 @@ namespace UE::TextureShare::PostProcess_FrameProxy
 		const TSharedPtr<IDisplayClusterProjectionPolicy, ESPMode::ThreadSafe>& PrjPolicy = ViewportProxy->GetProjectionPolicy_RenderThread();
 
 		// Projection policy must support warp blend op
-		return PrjPolicy.IsValid() && PrjPolicy->IsWarpBlendSupported();
+		return PrjPolicy.IsValid() && PrjPolicy->IsWarpBlendSupported_RenderThread(ViewportProxy);
 	}
 
 	static float GetFrameTargetsGamma(const IDisplayClusterViewportManagerProxy* InViewportManagerProxy)
@@ -55,14 +55,10 @@ void FTextureSharePostprocess::ShareViewport_RenderThread(FRHICommandListImmedia
 				EDisplayClusterViewportResourceType ResourceType = InResourceType;
 				if (bAfterWarpBlend && ShouldApplyWarpBlend(ViewportProxyIt.Get()))
 				{
-					const TSharedPtr<IDisplayClusterProjectionPolicy, ESPMode::ThreadSafe>& PrjPolicy = ViewportProxyIt->GetProjectionPolicy_RenderThread();
-					if (PrjPolicy->ShouldUseAdditionalTargetableResource())
+					TArray<FRHITexture*> Resources;
+					if (ViewportProxyIt->GetResources_RenderThread(EDisplayClusterViewportResourceType::AdditionalTargetableResource, Resources) && Resources.Num())
 					{
-						TArray<FRHITexture*> Resources;
-						if (ViewportProxyIt->GetResources_RenderThread(EDisplayClusterViewportResourceType::AdditionalTargetableResource, Resources) && Resources.Num())
-						{
-							ResourceType = EDisplayClusterViewportResourceType::AdditionalTargetableResource;
-						}
+						ResourceType = EDisplayClusterViewportResourceType::AdditionalTargetableResource;
 					}
 				}
 
diff --git a/Engine/Plugins/VirtualProduction/TextureShare/Source/TextureShareDisplayCluster/Private/Projection/TextureShareProjectionPolicy.h b/Engine/Plugins/VirtualProduction/TextureShare/Source/TextureShareDisplayCluster/Private/Projection/TextureShareProjectionPolicy.h
--- a/Engine/Plugins/VirtualProduction/TextureShare/Source/TextureShareDisplayCluster/Private/Projection/TextureShareProjectionPolicy.h
+++ b/Engine/Plugins/VirtualProduction/TextureShare/Source/TextureShareDisplayCluster/Private/Projection/TextureShareProjectionPolicy.h
@@ -25,10 +25,13 @@ public:
 	virtual bool CalculateView(IDisplayClusterViewport* InViewport, const uint32 InContextNum, FVector& InOutViewLocation, FRotator& InOutViewRotation, const FVector& ViewOffset, const float WorldToMeters, const float NCP, const float FCP) override;
 	virtual bool GetProjectionMatrix(IDisplayClusterViewport* InViewport, const uint32 InContextNum, FMatrix& OutPrjMatrix) override;
 
-	virtual bool IsWarpBlendSupported() override
+	virtual bool IsWarpBlendSupported(const IDisplayClusterViewport* InViewport) const override
 	{ return false; }
 
-	virtual bool ShouldUseSourceTextureWithMips() const override
+	virtual bool IsWarpBlendSupported_RenderThread(const IDisplayClusterViewportProxy* InViewportProxy) const override
+	{ return false; }
+
+	virtual bool ShouldUseSourceTextureWithMips(const IDisplayClusterViewport* InViewport) const override
 	{ return false; }
 
 public:
