nDisplay: Fixed In-Camera viewport rendering for headless mode

diff --git a/Engine/Plugins/Runtime/nDisplay/Source/DisplayCluster/Private/Render/Viewport/Configuration/DisplayClusterViewportConfiguration_ICVFXCamera.cpp b/Engine/Plugins/Runtime/nDisplay/Source/DisplayCluster/Private/Render/Viewport/Configuration/DisplayClusterViewportConfiguration_ICVFXCamera.cpp
--- a/Engine/Plugins/Runtime/nDisplay/Source/DisplayCluster/Private/Render/Viewport/Configuration/DisplayClusterViewportConfiguration_ICVFXCamera.cpp
+++ b/Engine/Plugins/Runtime/nDisplay/Source/DisplayCluster/Private/Render/Viewport/Configuration/DisplayClusterViewportConfiguration_ICVFXCamera.cpp
@@ -175,6 +175,12 @@ bool FDisplayClusterViewportConfiguration_ICVFXCamera::EnableChromakeyRendering(
 
 bool FDisplayClusterViewportConfiguration_ICVFXCamera::EnableInnerFrustumRendering() const
 {
+	if (TargetViewports.IsEmpty())
+	{
+		// Headless node
+		return true;
+	}
+
 	// Performance: If all chromakey sources have a ‘FrameColor’ value for all viewports on the current cluster node,
 	// we can skip rendering the InnerFrustum
 	for (const FTargetViewport& TargetViewportIt : TargetViewports)
