nDisplay: Not require techvis checked for the display device component to apply OCIO render pass.

Allows OCIO inversion without techvis extra stuff.

#rnx

diff --git a/Engine/Plugins/Runtime/nDisplay/Source/DisplayCluster/Private/Render/DisplayDevice/Components/DisplayClusterDisplayDeviceComponent.cpp b/Engine/Plugins/Runtime/nDisplay/Source/DisplayCluster/Private/Render/DisplayDevice/Components/DisplayClusterDisplayDeviceComponent.cpp
--- a/Engine/Plugins/Runtime/nDisplay/Source/DisplayCluster/Private/Render/DisplayDevice/Components/DisplayClusterDisplayDeviceComponent.cpp
+++ b/Engine/Plugins/Runtime/nDisplay/Source/DisplayCluster/Private/Render/DisplayDevice/Components/DisplayClusterDisplayDeviceComponent.cpp
@@ -38,7 +38,7 @@ void UDisplayClusterDisplayDeviceComponent::UpdateDisplayDeviceProxyImpl(IDispla
 		return;
 	}
 
-	if (!bEnableRenderPass || !ColorConversionSettings.IsValid() || !InConfiguration.IsTechvisEnabled())
+	if (!bEnableRenderPass || !ColorConversionSettings.IsValid())
 	{
 		// OCIO for preview is disabled
 		DisplayDeviceProxy.Reset();
