nDisplay: Added option to disable OCIO in preview

diff --git a/Engine/Plugins/Runtime/nDisplay/Source/DisplayCluster/Private/Game/EngineClasses/Scene/DisplayClusterRootActor.cpp b/Engine/Plugins/Runtime/nDisplay/Source/DisplayCluster/Private/Game/EngineClasses/Scene/DisplayClusterRootActor.cpp
--- a/Engine/Plugins/Runtime/nDisplay/Source/DisplayCluster/Private/Game/EngineClasses/Scene/DisplayClusterRootActor.cpp
+++ b/Engine/Plugins/Runtime/nDisplay/Source/DisplayCluster/Private/Game/EngineClasses/Scene/DisplayClusterRootActor.cpp
@@ -1302,6 +1302,8 @@ FDisplayClusterViewport_PreviewSettings ADisplayClusterRootActor::GetPreviewSett
 
 	OutPreviewSettings.bEnablePreviewTechvis     = bEnablePreviewTechvis;
 	OutPreviewSettings.bPreviewEnablePostProcess = bPreviewEnablePostProcess;
+	OutPreviewSettings.bPreviewEnableOCIO        = bPreviewEnableOCIO;
+
 	OutPreviewSettings.bPreviewEnableOverlayMaterial = bPreviewEnableOverlayMaterial;
 
 	OutPreviewSettings.bPreviewICVFXFrustums           = bPreviewICVFXFrustums;
@@ -1593,6 +1595,7 @@ void ADisplayClusterRootActor::OverrideRootActorProperties(const FDisplayCluster
 	OVERRIDE_PROPERTY(bPreviewEnable);
 	OVERRIDE_PROPERTY(PreviewRenderTargetRatioMult);
 	OVERRIDE_PROPERTY(bPreviewEnablePostProcess);
+	OVERRIDE_PROPERTY(bPreviewEnableOCIO);
 	OVERRIDE_PROPERTY(bPreviewEnableOverlayMaterial);
 	OVERRIDE_PROPERTY(bEnablePreviewTechvis);
 	OVERRIDE_PROPERTY(bEnablePreviewMesh);
diff --git a/Engine/Plugins/Runtime/nDisplay/Source/DisplayCluster/Private/Render/Viewport/Configuration/DisplayClusterViewportConfigurationHelpers_OpenColorIO.cpp b/Engine/Plugins/Runtime/nDisplay/Source/DisplayCluster/Private/Render/Viewport/Configuration/DisplayClusterViewportConfigurationHelpers_OpenColorIO.cpp
--- a/Engine/Plugins/Runtime/nDisplay/Source/DisplayCluster/Private/Render/Viewport/Configuration/DisplayClusterViewportConfigurationHelpers_OpenColorIO.cpp
+++ b/Engine/Plugins/Runtime/nDisplay/Source/DisplayCluster/Private/Render/Viewport/Configuration/DisplayClusterViewportConfigurationHelpers_OpenColorIO.cpp
@@ -88,6 +88,14 @@ bool FDisplayClusterViewportConfigurationHelpers_OpenColorIO::UpdateChromakeyVie
 
 void FDisplayClusterViewportConfigurationHelpers_OpenColorIO::ApplyOCIOConfiguration(FDisplayClusterViewport& DstViewport, const FOpenColorIOColorConversionSettings& InConversionSettings)
 {
+	if (!DstViewport.Configuration->GetRenderFrameSettings().IsEnabledOpenColorIO())
+	{
+		// OCIO is disabled for nDisplay.
+		DisableOCIOConfiguration(DstViewport);
+
+		return;
+	}
+
 	if (DstViewport.GetOpenColorIO().IsValid() && DstViewport.GetOpenColorIO()->IsConversionSettingsEqual(InConversionSettings))
 	{
 		// Already assigned
diff --git a/Engine/Plugins/Runtime/nDisplay/Source/DisplayCluster/Private/Render/Viewport/RenderFrame/DisplayClusterRenderFrameSettings.cpp b/Engine/Plugins/Runtime/nDisplay/Source/DisplayCluster/Private/Render/Viewport/RenderFrame/DisplayClusterRenderFrameSettings.cpp
--- a/Engine/Plugins/Runtime/nDisplay/Source/DisplayCluster/Private/Render/Viewport/RenderFrame/DisplayClusterRenderFrameSettings.cpp
+++ b/Engine/Plugins/Runtime/nDisplay/Source/DisplayCluster/Private/Render/Viewport/RenderFrame/DisplayClusterRenderFrameSettings.cpp
@@ -83,6 +83,18 @@ bool FDisplayClusterRenderFrameSettings::ShouldUseLinearGamma() const
 	return false;
 }
 
+bool FDisplayClusterRenderFrameSettings::IsEnabledOpenColorIO() const
+{
+	if (IsPreviewRendering() && !PreviewSettings.bPreviewEnableOCIO)
+	{
+		// Disable OCIO on preview
+		return false;
+	}
+
+	// OCIO is enabled by default
+	return true;
+}
+
 bool FDisplayClusterRenderFrameSettings::IsPostProcessDisabled() const
 {
 	if (IsPreviewRendering())
diff --git a/Engine/Plugins/Runtime/nDisplay/Source/DisplayCluster/Private/Render/Viewport/RenderFrame/DisplayClusterRenderFrameSettings.h b/Engine/Plugins/Runtime/nDisplay/Source/DisplayCluster/Private/Render/Viewport/RenderFrame/DisplayClusterRenderFrameSettings.h
--- a/Engine/Plugins/Runtime/nDisplay/Source/DisplayCluster/Private/Render/Viewport/RenderFrame/DisplayClusterRenderFrameSettings.h
+++ b/Engine/Plugins/Runtime/nDisplay/Source/DisplayCluster/Private/Render/Viewport/RenderFrame/DisplayClusterRenderFrameSettings.h
@@ -107,6 +107,9 @@ public:
 	*/
 	const FIntPoint* GetPreviewMultiGPURendering() const;
 
+	/** True if OCIO is enabled for nDisplay. */
+	bool IsEnabledOpenColorIO() const;
+
 	/** PostProcess should be disabled for the current frame. */
 	bool IsPostProcessDisabled() const;
 
diff --git a/Engine/Plugins/Runtime/nDisplay/Source/DisplayCluster/Public/DisplayClusterRootActor.h b/Engine/Plugins/Runtime/nDisplay/Source/DisplayCluster/Public/DisplayClusterRootActor.h
--- a/Engine/Plugins/Runtime/nDisplay/Source/DisplayCluster/Public/DisplayClusterRootActor.h
+++ b/Engine/Plugins/Runtime/nDisplay/Source/DisplayCluster/Public/DisplayClusterRootActor.h
@@ -459,6 +459,10 @@ public:
 	UPROPERTY(EditAnywhere, BlueprintReadWrite, Category = "Preview", meta = (DisplayName = "Enable Post Process"), BlueprintSetter = SetPreviewEnablePostProcess)
 	bool bPreviewEnablePostProcess = false;
 
+	/** Enable OCIO on preview. */
+	UPROPERTY(EditAnywhere, BlueprintReadWrite, Category = "Preview", meta = (DisplayName = "Enable OCIO on Preview"))
+	bool bPreviewEnableOCIO = false;
+
 	/** Show overlay material on the preview mesh when preview rendering is enabled (UMeshComponent::OverlayMaterial). */
 	UPROPERTY(EditAnywhere, BlueprintReadWrite, Category = "Preview", AdvancedDisplay, meta = (DisplayName = "Enable Preview Overlay"))
 	bool bPreviewEnableOverlayMaterial = true;
diff --git a/Engine/Plugins/Runtime/nDisplay/Source/DisplayCluster/Public/DisplayClusterRootActorContainers.h b/Engine/Plugins/Runtime/nDisplay/Source/DisplayCluster/Public/DisplayClusterRootActorContainers.h
--- a/Engine/Plugins/Runtime/nDisplay/Source/DisplayCluster/Public/DisplayClusterRootActorContainers.h
+++ b/Engine/Plugins/Runtime/nDisplay/Source/DisplayCluster/Public/DisplayClusterRootActorContainers.h
@@ -25,6 +25,9 @@ struct FDisplayClusterRootActorPropertyOverrides
 	/** Enable PostProcess for preview. */
 	TOptional<bool> bPreviewEnablePostProcess;
 
+	/** Enable OCIO on preview. */
+	TOptional<bool> bPreviewEnableOCIO;
+
 	/** Show overlay material on the preview mesh when preview rendering is enabled (UMeshComponent::OverlayMaterial). */
 	TOptional<bool> bPreviewEnableOverlayMaterial;
 
diff --git a/Engine/Plugins/Runtime/nDisplay/Source/DisplayCluster/Public/Render/Viewport/Containers/DisplayClusterViewport_PreviewSettings.h b/Engine/Plugins/Runtime/nDisplay/Source/DisplayCluster/Public/Render/Viewport/Containers/DisplayClusterViewport_PreviewSettings.h
--- a/Engine/Plugins/Runtime/nDisplay/Source/DisplayCluster/Public/Render/Viewport/Containers/DisplayClusterViewport_PreviewSettings.h
+++ b/Engine/Plugins/Runtime/nDisplay/Source/DisplayCluster/Public/Render/Viewport/Containers/DisplayClusterViewport_PreviewSettings.h
@@ -33,6 +33,9 @@ struct FDisplayClusterViewport_PreviewSettings
 	// But in it, PostProcess is also rendered with the flag turned off.
 	bool bPreviewEnablePostProcess = false;
 
+	// Enable OCIO on preview
+	bool bPreviewEnableOCIO = false;
+
 	// Show overlay material on the preview mesh when preview rendering is enabled (UMeshComponent::OverlayMaterial)
 	bool bPreviewEnableOverlayMaterial = true;
 
