[nDisplay] Fixed an issue that prevented the inner cameras from appearing on the backbuffer of offscreen nodes

#rnx

diff --git a/Engine/Plugins/Runtime/nDisplay/Source/DisplayClusterConfiguration/Private/DisplayClusterConfigurationTypes_Media.cpp b/Engine/Plugins/Runtime/nDisplay/Source/DisplayClusterConfiguration/Private/DisplayClusterConfigurationTypes_Media.cpp
--- a/Engine/Plugins/Runtime/nDisplay/Source/DisplayClusterConfiguration/Private/DisplayClusterConfigurationTypes_Media.cpp
+++ b/Engine/Plugins/Runtime/nDisplay/Source/DisplayClusterConfiguration/Private/DisplayClusterConfigurationTypes_Media.cpp
@@ -8,12 +8,21 @@
 
 bool FDisplayClusterConfigurationMediaNodeBackbuffer::IsMediaOutputAssigned() const
 {
-	// Just see if any media output instances are assigned
-	const bool bAnyMediaAssigned = MediaOutputs.ContainsByPredicate([](const FDisplayClusterConfigurationMediaOutput& Item)
+	// Just see if any FullFrame media output instances are assigned
+	bool bAnyMediaAssigned = MediaOutputs.ContainsByPredicate([](const FDisplayClusterConfigurationMediaOutput& Item)
 		{
 			return IsValid(Item.MediaOutput);
 		});
 
+	// If no FullFrame, let's see if there are any tiles configured
+	if (!bAnyMediaAssigned)
+	{
+		bAnyMediaAssigned = TiledMediaOutputs.ContainsByPredicate([](const FDisplayClusterConfigurationMediaUniformTileOutput& Item)
+			{
+				return IsValid(Item.MediaOutput);
+			});
+	}
+
 	return bAnyMediaAssigned;
 }
 
