[nDisplay] Fixed internal cluster events processing in PIE

[RN] Fixed cluster events in PIE

diff --git a/Engine/Plugins/Runtime/nDisplay/Source/DisplayCluster/Private/Cluster/Controller/DisplayClusterClusterNodeCtrlEditor.cpp b/Engine/Plugins/Runtime/nDisplay/Source/DisplayCluster/Private/Cluster/Controller/DisplayClusterClusterNodeCtrlEditor.cpp
--- a/Engine/Plugins/Runtime/nDisplay/Source/DisplayCluster/Private/Cluster/Controller/DisplayClusterClusterNodeCtrlEditor.cpp
+++ b/Engine/Plugins/Runtime/nDisplay/Source/DisplayCluster/Private/Cluster/Controller/DisplayClusterClusterNodeCtrlEditor.cpp
@@ -2,6 +2,7 @@
 
 #include "Cluster/Controller/DisplayClusterClusterNodeCtrlEditor.h"
 
+#include "Cluster/IPDisplayClusterClusterManager.h"
 #include "Config/IPDisplayClusterConfigManager.h"
 
 #include "Network/Service/ClusterEventsBinary/DisplayClusterClusterEventsBinaryService.h"
@@ -23,6 +24,21 @@ FDisplayClusterClusterNodeCtrlEditor::~FDisplayClusterClusterNodeCtrlEditor()
 }
 
 
+//////////////////////////////////////////////////////////////////////////////////////////////
+// IDisplayClusterProtocolClusterSync
+//////////////////////////////////////////////////////////////////////////////////////////////
+EDisplayClusterCommResult FDisplayClusterClusterNodeCtrlEditor::GetEventsData(TArray<TSharedPtr<FDisplayClusterClusterEventJson, ESPMode::ThreadSafe>>& OutJsonEvents, TArray<TSharedPtr<FDisplayClusterClusterEventBinary, ESPMode::ThreadSafe>>& OutBinaryEvents)
+{
+	if (IPDisplayClusterClusterManager* const ClusterMgr = GDisplayCluster->GetPrivateClusterMgr())
+	{
+		ClusterMgr->ExportEventsData(OutJsonEvents, OutBinaryEvents);
+		return EDisplayClusterCommResult::Ok;
+	}
+
+	return EDisplayClusterCommResult::InternalError;
+}
+
+
 //////////////////////////////////////////////////////////////////////////////////////////////
 // FDisplayClusterNodeCtrlBase
 //////////////////////////////////////////////////////////////////////////////////////////////
diff --git a/Engine/Plugins/Runtime/nDisplay/Source/DisplayCluster/Private/Cluster/Controller/DisplayClusterClusterNodeCtrlEditor.h b/Engine/Plugins/Runtime/nDisplay/Source/DisplayCluster/Private/Cluster/Controller/DisplayClusterClusterNodeCtrlEditor.h
--- a/Engine/Plugins/Runtime/nDisplay/Source/DisplayCluster/Private/Cluster/Controller/DisplayClusterClusterNodeCtrlEditor.h
+++ b/Engine/Plugins/Runtime/nDisplay/Source/DisplayCluster/Private/Cluster/Controller/DisplayClusterClusterNodeCtrlEditor.h
@@ -24,9 +24,15 @@ public:
 	//////////////////////////////////////////////////////////////////////////////////////////////
 	virtual EDisplayClusterNodeRole GetClusterRole() const override
 	{
-		return EDisplayClusterNodeRole::None;
+		return EDisplayClusterNodeRole::Primary;
 	}
 
+public:
+	//////////////////////////////////////////////////////////////////////////////////////////////
+	// IDisplayClusterProtocolClusterSync
+	//////////////////////////////////////////////////////////////////////////////////////////////
+	virtual EDisplayClusterCommResult GetEventsData(TArray<TSharedPtr<FDisplayClusterClusterEventJson, ESPMode::ThreadSafe>>& OutJsonEvents, TArray<TSharedPtr<FDisplayClusterClusterEventBinary, ESPMode::ThreadSafe>>& OutBinaryEvents) override;
+
 protected:
 	//////////////////////////////////////////////////////////////////////////////////////////////
 	// FDisplayClusterNodeCtrlBase
