nDisplay: Support hiding actors from specific outer viewports.

diff --git a/Engine/Plugins/Runtime/nDisplay/Source/DisplayCluster/Private/Render/Viewport/Configuration/DisplayClusterViewportConfigurationHelpers_Visibility.cpp b/Engine/Plugins/Runtime/nDisplay/Source/DisplayCluster/Private/Render/Viewport/Configuration/DisplayClusterViewportConfigurationHelpers_Visibility.cpp
--- a/Engine/Plugins/Runtime/nDisplay/Source/DisplayCluster/Private/Render/Viewport/Configuration/DisplayClusterViewportConfigurationHelpers_Visibility.cpp
+++ b/Engine/Plugins/Runtime/nDisplay/Source/DisplayCluster/Private/Render/Viewport/Configuration/DisplayClusterViewportConfigurationHelpers_Visibility.cpp
@@ -195,6 +195,17 @@ void FDisplayClusterViewportConfigurationHelpers_Visibility::UpdateHideList_ICVF
 				{
 					ViewportIt->GetVisibilitySettingsImpl().AppendVisibilityComponentsList(EDisplayClusterViewport_VisibilityMode::Hide, OuterComponentsList);
 				}
+
+				if (const UDisplayClusterConfigurationData* ConfigData = ConfigurationRootActor->GetConfigData())
+				{
+					// Hide actors from specific viewports.
+					if (const UDisplayClusterConfigurationViewport* SourceViewport = ConfigData->GetViewport(ViewportIt->ClusterNodeId, ViewportIt->ViewportId))
+					{
+						TSet<FPrimitiveComponentId> ViewportSpecificComponentList;
+						UE::DisplayCluster::Configuration::VisibilityHelpers::ImplCollectComponentsFromVisibilityList(InConfiguration, SourceViewport->RenderSettings.HiddenContent, ViewportSpecificComponentList);
+						ViewportIt->GetVisibilitySettingsImpl().AppendVisibilityComponentsList(EDisplayClusterViewport_VisibilityMode::Hide, ViewportSpecificComponentList);
+					}
+				}
 			}
 		}
 	}
diff --git a/Engine/Plugins/Runtime/nDisplay/Source/DisplayCluster/Public/DisplayClusterRootActor.h b/Engine/Plugins/Runtime/nDisplay/Source/DisplayCluster/Public/DisplayClusterRootActor.h
--- a/Engine/Plugins/Runtime/nDisplay/Source/DisplayCluster/Public/DisplayClusterRootActor.h
+++ b/Engine/Plugins/Runtime/nDisplay/Source/DisplayCluster/Public/DisplayClusterRootActor.h
@@ -296,6 +296,9 @@ private:
 	UPROPERTY(EditInstanceOnly, Transient, Category = Viewports, meta = (DisplayName = "Viewport Overscan", PropertyPath = "CurrentConfigData.Cluster.Nodes.Viewports.RenderSettings.Overscan", ToolTip = "Render a larger frame than specified in the configuration to achieve continuity across displays when using post-processing effects."))
 	FDisplayClusterEditorPropertyReference ViewportOverscanRef;
 
+	UPROPERTY(EditInstanceOnly, Transient, Category = Viewports, meta = (PropertyPath = "CurrentConfigData.Cluster.Nodes.Viewports.RenderSettings.HiddenContent", ToolTip = "Hide actors from a specific viewport."))
+	FDisplayClusterEditorPropertyReference ViewportHiddenActorsRef;
+	
 	UPROPERTY(EditAnywhere, Transient, Category = Viewports, meta = (PropertyPath = "CurrentConfigData.StageSettings.bFreezeRenderOuterViewports"))
 	FDisplayClusterEditorPropertyReference FreezeRenderOuterViewportsRef;
 
diff --git a/Engine/Plugins/Runtime/nDisplay/Source/DisplayClusterConfiguration/Public/DisplayClusterConfigurationTypes_Viewport.h b/Engine/Plugins/Runtime/nDisplay/Source/DisplayClusterConfiguration/Public/DisplayClusterConfigurationTypes_Viewport.h
--- a/Engine/Plugins/Runtime/nDisplay/Source/DisplayClusterConfiguration/Public/DisplayClusterConfigurationTypes_Viewport.h
+++ b/Engine/Plugins/Runtime/nDisplay/Source/DisplayClusterConfiguration/Public/DisplayClusterConfigurationTypes_Viewport.h
@@ -9,6 +9,7 @@
 
 #include "DisplayClusterConfigurationStrings.h"
 #include "DisplayClusterConfigurationTypes_Base.h"
+#include "DisplayClusterConfigurationTypes_ICVFX.h"
 #include "DisplayClusterConfigurationTypes_Media.h"
 #include "DisplayClusterConfigurationTypes_PostRender.h"
 #include "DisplayClusterConfigurationTypes_Postprocess.h"
@@ -136,6 +137,10 @@ public:
 	UPROPERTY(EditAnywhere, BlueprintReadWrite, Category = "Configuration")
 	FDisplayClusterConfigurationViewport_Overscan Overscan;
 
+	/** Override actor visibility for this viewport. */
+	UPROPERTY(EditAnywhere, BlueprintReadWrite, Category = "Configuration")
+	FDisplayClusterConfigurationICVFX_VisibilityList HiddenContent;
+	
 	// Media settings
 	UPROPERTY(EditAnywhere, BlueprintReadWrite, Category = "Media", meta = (DisplayName = "Media", ShowOnlyInnerProperties))
 	FDisplayClusterConfigurationMediaViewport Media;
