nDisplay: Update property names for hiding content from viewports.

diff --git a/Engine/Plugins/Runtime/nDisplay/Source/DisplayCluster/Public/DisplayClusterRootActor.h b/Engine/Plugins/Runtime/nDisplay/Source/DisplayCluster/Public/DisplayClusterRootActor.h
--- a/Engine/Plugins/Runtime/nDisplay/Source/DisplayCluster/Public/DisplayClusterRootActor.h
+++ b/Engine/Plugins/Runtime/nDisplay/Source/DisplayCluster/Public/DisplayClusterRootActor.h
@@ -295,9 +295,6 @@ private:
 
 	UPROPERTY(EditInstanceOnly, Transient, Category = Viewports, meta = (DisplayName = "Viewport Overscan", PropertyPath = "CurrentConfigData.Cluster.Nodes.Viewports.RenderSettings.Overscan", ToolTip = "Render a larger frame than specified in the configuration to achieve continuity across displays when using post-processing effects."))
 	FDisplayClusterEditorPropertyReference ViewportOverscanRef;
-
-	UPROPERTY(EditInstanceOnly, Transient, Category = Viewports, meta = (PropertyPath = "CurrentConfigData.Cluster.Nodes.Viewports.RenderSettings.HiddenContent", ToolTip = "Hide actors from a specific viewport."))
-	FDisplayClusterEditorPropertyReference ViewportHiddenActorsRef;
 	
 	UPROPERTY(EditAnywhere, Transient, Category = Viewports, meta = (PropertyPath = "CurrentConfigData.StageSettings.bFreezeRenderOuterViewports"))
 	FDisplayClusterEditorPropertyReference FreezeRenderOuterViewportsRef;
@@ -308,6 +305,9 @@ private:
 	UPROPERTY(EditAnywhere, Transient, Category = Viewports, meta = (PropertyPath = "CurrentConfigData.StageSettings.OuterViewportHideList"))
 	FDisplayClusterEditorPropertyReference OuterHideListRef;
 
+	UPROPERTY(EditInstanceOnly, Transient, Category = Viewports, meta = (DisplayName = "Content Hidden Per-Viewport", PropertyPath = "CurrentConfigData.Cluster.Nodes.Viewports.RenderSettings.HiddenContent", ToolTip = "Hide actors from a specific viewport."))
+	FDisplayClusterEditorPropertyReference PerOuterHideListRef;
+	
 	UPROPERTY(EditAnywhere, Transient, Category = "In Camera VFX", meta = (PropertyPath = "CurrentConfigData.StageSettings.bEnableInnerFrustums"))
 	FDisplayClusterEditorPropertyReference EnableInnerFrustumsRef;
 
diff --git a/Engine/Plugins/Runtime/nDisplay/Source/DisplayClusterConfiguration/Public/DisplayClusterConfigurationTypes_ICVFX.h b/Engine/Plugins/Runtime/nDisplay/Source/DisplayClusterConfiguration/Public/DisplayClusterConfigurationTypes_ICVFX.h
--- a/Engine/Plugins/Runtime/nDisplay/Source/DisplayClusterConfiguration/Public/DisplayClusterConfigurationTypes_ICVFX.h
+++ b/Engine/Plugins/Runtime/nDisplay/Source/DisplayClusterConfiguration/Public/DisplayClusterConfigurationTypes_ICVFX.h
@@ -938,7 +938,7 @@ public:
 	FDisplayClusterConfigurationICVFX_VisibilityList HideList;
 
 	/** Special hide list for Outer viewports */
-	UPROPERTY(EditAnywhere, BlueprintReadWrite, Category = Viewports, meta = (DisplayName = "Content Hidden from Viewports", ToolTip = "Content specified here will not appear in the nDisplay viewports, but can appear in the inner frustum.", Substitutions = "LayersTooltip = Layers hidden from the nDisplay viewports, ActorsTooltip = Actors hidden from the nDisplay viewports"))
+	UPROPERTY(EditAnywhere, BlueprintReadWrite, Category = Viewports, meta = (DisplayName = "Content Hidden from All Viewports", ToolTip = "Content specified here will not appear in the nDisplay viewports, but can appear in the inner frustum.", Substitutions = "LayersTooltip = Layers hidden from the nDisplay viewports, ActorsTooltip = Actors hidden from the nDisplay viewports"))
 	FDisplayClusterConfigurationICVFX_VisibilityList OuterViewportHideList;
 
 	/** Viewport Color Grading */
