[nDisplay] Fixed bug: InFrustumFit should use external camera location instead of ViewPoint component.

By default, the ViewPoint component ignores the external camera position and only gets post processing etc.
Except when using the `bUseICVFXCameraComponentTracking` parameter for the ICVX camera.

But the InFrustumFit component should always use only the camera position.
So we need to override this logic for this component.
And at the same time, the observer's eye is in the position of the InFrustumFit component.

diff --git a/Engine/Plugins/Runtime/nDisplay/Source/DisplayClusterWarp/Private/Components/DisplayClusterInFrustumFitCameraComponent.cpp b/Engine/Plugins/Runtime/nDisplay/Source/DisplayClusterWarp/Private/Components/DisplayClusterInFrustumFitCameraComponent.cpp
--- a/Engine/Plugins/Runtime/nDisplay/Source/DisplayClusterWarp/Private/Components/DisplayClusterInFrustumFitCameraComponent.cpp
+++ b/Engine/Plugins/Runtime/nDisplay/Source/DisplayClusterWarp/Private/Components/DisplayClusterInFrustumFitCameraComponent.cpp
@@ -137,6 +137,30 @@ bool UDisplayClusterInFrustumFitCameraComponent::IsICVFXCameraBeingUsed() const
 	return Super::IsICVFXCameraBeingUsed();
 }
 
+bool UDisplayClusterInFrustumFitCameraComponent::IsViewPointOverrideCameraPosition() const
+{
+	// If InFrustumFit uses an external camera, use it as a ViewPoint.
+	if (IsEnabled())
+	{
+		return false;
+	}
+
+	return Super::IsViewPointOverrideCameraPosition();
+}
+
+
+void UDisplayClusterInFrustumFitCameraComponent::GetEyePosition(const IDisplayClusterViewportConfiguration& InViewportConfiguration, FVector& OutViewLocation, FRotator& OutViewRotation)
+{
+	Super::GetEyePosition(InViewportConfiguration, OutViewLocation, OutViewRotation);
+
+	// The observer's eye is located inside the InFrustumFit component.
+	if (IsEnabled())
+	{
+		OutViewLocation = GetComponentLocation();
+		OutViewRotation = GetComponentRotation();
+	}
+}
+
 void UDisplayClusterInFrustumFitCameraComponent::TickComponent(float DeltaTime, ELevelTick TickType, FActorComponentTickFunction* ThisTickFunction)
 {
 	// Tick warp policy instance
diff --git a/Engine/Plugins/Runtime/nDisplay/Source/DisplayClusterWarp/Public/Components/DisplayClusterInFrustumFitCameraComponent.h b/Engine/Plugins/Runtime/nDisplay/Source/DisplayClusterWarp/Public/Components/DisplayClusterInFrustumFitCameraComponent.h
--- a/Engine/Plugins/Runtime/nDisplay/Source/DisplayClusterWarp/Public/Components/DisplayClusterInFrustumFitCameraComponent.h
+++ b/Engine/Plugins/Runtime/nDisplay/Source/DisplayClusterWarp/Public/Components/DisplayClusterInFrustumFitCameraComponent.h
@@ -44,6 +44,9 @@ public:
 	virtual TObjectPtr<UMaterial> GetDisplayDeviceMaterial(const EDisplayClusterDisplayDeviceMeshType InMeshType, const EDisplayClusterDisplayDeviceMaterialType InMaterialType) const override;
 	virtual void OnUpdateDisplayDeviceMeshAndMaterialInstance(IDisplayClusterViewportPreview& InViewportPreview, const EDisplayClusterDisplayDeviceMeshType InMeshType, const EDisplayClusterDisplayDeviceMaterialType InMaterialType, UMeshComponent* InMeshComponent, UMaterialInstanceDynamic* InMeshMaterialInstance) const override;
 
+	virtual bool IsViewPointOverrideCameraPosition() const override;
+	virtual void GetEyePosition(const IDisplayClusterViewportConfiguration& InViewportConfiguration, FVector& OutViewLocation, FRotator& OutViewRotation) override;
+
 protected:
 	virtual bool IsICVFXCameraBeingUsed() const override;
 	//~~End UDisplayClusterCameraComponent
