[nDisplay] HoldoutComposite integration 1/3: Reworked rendering pipeline to initialize from ViewExtension.

The order of modifying the SceneView structure has been changed:
- Currently, IDisplayClusterViewport::SetupSceneView() is called after the scene view has been set up by LocalPlayer and VE.
But too late - because the FHoldoutCompositeSceneViewExtension::SetupView() expects that `ShowOnlyPrimitives` and `HiddenPrimitives` are already properly configured.
- In CL, this function is called from `FDisplayClusterViewportManagerViewPointExtension::SetupView()`. Other calls to this function have been deleted;

The `FDisplayClusterViewportManagerViewPointExtension` related changes:
- Changed sort priority to 1000, meaning this VE is called before other VEs.
- Since this VE requires the StereoViewIndex variable to be set, it was handled by FDisplayClusterViewportManager::SetCurrentStereoViewIndexForViewPointExtension().
- Changed FDisplayClusterViewport::GatherActiveExtensions() logic.

- Fixed ViewInitOptions initialization in DisplayClusterViewportPreview_Render.cpp: assigned ViewInitOptions.ViewLocation and StereoViewLocation.

diff --git a/Engine/Plugins/Runtime/nDisplay/Source/DisplayCluster/Private/Game/EngineClasses/Basics/DisplayClusterViewportClient.cpp b/Engine/Plugins/Runtime/nDisplay/Source/DisplayCluster/Private/Game/EngineClasses/Basics/DisplayClusterViewportClient.cpp
--- a/Engine/Plugins/Runtime/nDisplay/Source/DisplayCluster/Private/Game/EngineClasses/Basics/DisplayClusterViewportClient.cpp
+++ b/Engine/Plugins/Runtime/nDisplay/Source/DisplayCluster/Private/Game/EngineClasses/Basics/DisplayClusterViewportClient.cpp
@@ -573,9 +573,6 @@ void UDisplayClusterViewportClient::Draw(FViewport* InViewport, FCanvas* SceneCa
 				{
 					Views.Add(View);
 
-					// Apply viewport context settings to view (crossGPU, visibility, etc)
-					DCView.Viewport->SetupSceneView(DCView.ContextNum, MyWorld, ViewFamily, *View);
-
 					// We don't allow instanced stereo currently
 					View->bIsInstancedStereoEnabled = false;
 					View->bShouldBindInstancedViewUB = false;
diff --git a/Engine/Plugins/Runtime/nDisplay/Source/DisplayCluster/Private/Render/Viewport/DisplayClusterViewport.cpp b/Engine/Plugins/Runtime/nDisplay/Source/DisplayCluster/Private/Render/Viewport/DisplayClusterViewport.cpp
--- a/Engine/Plugins/Runtime/nDisplay/Source/DisplayCluster/Private/Render/Viewport/DisplayClusterViewport.cpp
+++ b/Engine/Plugins/Runtime/nDisplay/Source/DisplayCluster/Private/Render/Viewport/DisplayClusterViewport.cpp
@@ -7,6 +7,7 @@
 #include "Render/Viewport/DisplayClusterViewportManagerProxy.h"
 #include "Render/Viewport/DisplayClusterViewport_OpenColorIO.h"
 #include "Render/Viewport/DisplayClusterViewportManagerViewExtension.h"
+#include "Render/Viewport/DisplayClusterViewportManagerViewPointExtension.h"
 
 #include "Render/Viewport/Configuration/DisplayClusterViewportConfiguration.h"
 #include "Render/Viewport/Configuration/DisplayClusterViewportConfigurationProxy.h"
@@ -189,7 +190,30 @@ void FDisplayClusterViewport::FinalizeNewFrame()
 	RenderSettings.FinishUpdateSettings();
 }
 
-const TArray<FSceneViewExtensionRef> FDisplayClusterViewport::GatherActiveExtensions(FViewport* InViewport) const
+const TArray<FSceneViewExtensionRef> FDisplayClusterViewport::GatherActiveExtensions(const int32 ViewIndex, FViewport* InViewport) const
+{
+	FDisplayClusterViewportManager* ViewportManager = Configuration->GetViewportManagerImpl();
+	if (ViewportManager && Contexts.IsValidIndex(ViewIndex))
+	{
+		// Configure FDisplayClusterViewportManagerViewPointExtension for this viewport.
+		ViewportManager->SetCurrentStereoViewIndexForViewPointExtension(Contexts[ViewIndex].StereoViewIndex);
+	}
+
+	TArray<FSceneViewExtensionRef> OutExtensions;
+	ImplGatherActiveExtensions(ViewIndex, InViewport, OutExtensions);
+
+	if (ViewportManager)
+	{
+		ViewportManager->SetCurrentStereoViewIndexForViewPointExtension(INDEX_NONE);
+	}
+
+	// Sort extensions in order of priority (copied from FSceneViewExtensions::GatherActiveExtensions)
+	Algo::SortBy(OutExtensions, &ISceneViewExtension::GetPriority, TGreater<>());
+
+	return OutExtensions;
+}
+
+void FDisplayClusterViewport::ImplGatherActiveExtensions(const int32 ViewIndex, FViewport* InViewport, TArray<FSceneViewExtensionRef>& OutExtensions) const
 {
 	// Use VE from engine for default render and MRQ:
 	switch (RenderSettings.CaptureMode)
@@ -199,53 +223,44 @@ const TArray<FSceneViewExtensionRef> FDisplayClusterViewport::GatherActiveExtens
 		if (InViewport)
 		{
 			FDisplayClusterSceneViewExtensionContext ViewExtensionContext(InViewport, Configuration, GetId());
-			return GEngine->ViewExtensions->GatherActiveExtensions(ViewExtensionContext);
+			OutExtensions = GEngine->ViewExtensions->GatherActiveExtensions(ViewExtensionContext);
+			return;
 		}
 		else
 		{
 			if (UWorld* CurrentWorld = Configuration->GetCurrentWorld())
 			{
 				FDisplayClusterSceneViewExtensionContext ViewExtensionContext(CurrentWorld->Scene, Configuration, GetId());
-				return GEngine->ViewExtensions->GatherActiveExtensions(ViewExtensionContext);
+				OutExtensions = GEngine->ViewExtensions->GatherActiveExtensions(ViewExtensionContext);
+				return;
 			}
 		}
 
 		// No extension found.
-		return TArray<FSceneViewExtensionRef>();
+		return;
 		
-	default:
-		break;
-	}
-
-	// Get custom VE:
-	TArray<FSceneViewExtensionRef> OutCustomExtensions;
-
-	// Initialize custom extensions:
-	switch (RenderSettings.CaptureMode)
-	{
 	case EDisplayClusterViewportCaptureMode::Chromakey:
 	case EDisplayClusterViewportCaptureMode::Lightcard:
 	{
-		// Chromakey and LightCard use only nDisplay VE for callback purposes (preserve alpha channel, etc).
+		// Chromakey and LightCard only use some internal nDisplay ViewExtensions:
 		if (FDisplayClusterViewportManager* ViewportManager = Configuration->GetViewportManagerImpl())
 		{
-			TSharedPtr<FDisplayClusterViewportManagerViewExtension, ESPMode::ThreadSafe> ViewportManagerViewExtension = ViewportManager->GetViewportManagerViewExtension();
-			if (ViewportManagerViewExtension.IsValid())
+			if (ViewportManager->GetViewportManagerViewPointExtension().IsValid())
 			{
-				OutCustomExtensions.Add(ViewportManagerViewExtension->AsShared());
+				// To call SetupView() as early as possible: integration with other VEs
+				OutExtensions.Add(ViewportManager->GetViewportManagerViewPointExtension()->AsShared());
+			}
+			if (ViewportManager->GetViewportManagerViewExtension().IsValid())
+			{
+				// For callback purposes (preserving alpha channel, etc.)
+				OutExtensions.Add(ViewportManager->GetViewportManagerViewExtension()->AsShared());
 			}
 		}
+		break;
 	}
-	break;
-
 	default:
 		break;
 	}
-
-	// Sort extensions in order of priority (copied from FSceneViewExtensions::GatherActiveExtensions)
-	Algo::SortBy(OutCustomExtensions, &ISceneViewExtension::GetPriority, TGreater<>());
-
-	return OutCustomExtensions;
 }
 
 void FDisplayClusterViewport::Initialize()
diff --git a/Engine/Plugins/Runtime/nDisplay/Source/DisplayCluster/Private/Render/Viewport/DisplayClusterViewport.h b/Engine/Plugins/Runtime/nDisplay/Source/DisplayCluster/Private/Render/Viewport/DisplayClusterViewport.h
--- a/Engine/Plugins/Runtime/nDisplay/Source/DisplayCluster/Private/Render/Viewport/DisplayClusterViewport.h
+++ b/Engine/Plugins/Runtime/nDisplay/Source/DisplayCluster/Private/Render/Viewport/DisplayClusterViewport.h
@@ -207,8 +207,11 @@ public:
 	/** This function MUST always be called before configuring the viewport at the beginning of each frame. */
 	void ResetRuntimeParameters(const UDisplayClusterConfigurationViewport* InConfigurationViewport = nullptr);
 
-	// Active view extension for this viewport
-	const TArray<FSceneViewExtensionRef> GatherActiveExtensions(FViewport* InViewport) const;
+	/** Collect the view extension active for the desired context of this viewport
+	* @param ViewIndex - view index
+	* @param InViewport - the viewport display
+	*/
+	const TArray<FSceneViewExtensionRef> GatherActiveExtensions(const int32 ViewIndex, FViewport* InViewport) const;
 
 	/** Initialize viewport contexts and resources for new frame
 	 *
@@ -422,6 +425,13 @@ private:
 	*/
 	bool IsResourceUsedImpl(const EDisplayClusterViewportResource InResourceType) const;
 
+	/** Collect the view extension active for the desired context of this viewport
+	* @param ViewIndex     - view index
+	* @param InViewport    - the viewport display
+	* @param OutExtensions - (out) Extension that can be used by this viewport context
+	*/
+	void ImplGatherActiveExtensions(const int32 ViewIndex, FViewport* InViewport, TArray<FSceneViewExtensionRef>& OutExtensions) const;
+
 public:
 	// Configuration of the current cluster node
 	const TSharedRef<FDisplayClusterViewportConfiguration, ESPMode::ThreadSafe> Configuration;
diff --git a/Engine/Plugins/Runtime/nDisplay/Source/DisplayCluster/Private/Render/Viewport/DisplayClusterViewportManager.cpp b/Engine/Plugins/Runtime/nDisplay/Source/DisplayCluster/Private/Render/Viewport/DisplayClusterViewportManager.cpp
--- a/Engine/Plugins/Runtime/nDisplay/Source/DisplayCluster/Private/Render/Viewport/DisplayClusterViewportManager.cpp
+++ b/Engine/Plugins/Runtime/nDisplay/Source/DisplayCluster/Private/Render/Viewport/DisplayClusterViewportManager.cpp
@@ -1187,6 +1187,14 @@ void FDisplayClusterViewportManager::AddReferencedObjects(FReferenceCollector& C
 	}
 }
 
+void FDisplayClusterViewportManager::SetCurrentStereoViewIndexForViewPointExtension(const int32 InStereoViewIndex) const
+{
+	if (ViewportManagerViewPointExtension.IsValid())
+	{
+		ViewportManagerViewPointExtension->SetCurrentStereoViewIndex(InStereoViewIndex);
+	}
+}
+
 FSceneView* FDisplayClusterViewportManager::CalcSceneView(ULocalPlayer* LocalPlayer, FSceneViewFamily* ViewFamily, FVector& OutViewLocation, FRotator& OutViewRotation, FViewport* InViewport, FViewElementDrawer* InViewDrawer, int32 InStereoViewIndex)
 {
 	check(ViewportManagerViewPointExtension.IsValid());
@@ -1195,13 +1203,13 @@ FSceneView* FDisplayClusterViewportManager::CalcSceneView(ULocalPlayer* LocalPla
 	// Set current view index for this VE:
 	// This function must be used because LocalPlayer::GetViewPoint() calls the ISceneViewExtension::SetupViewPoint() function from this VE.
 	// And at this point the VE must know the current StereoViewIndex value in order to understand which viewport will be used for this SetupViewPoint() call.
-	ViewportManagerViewPointExtension->SetCurrentStereoViewIndex(InStereoViewIndex);
+	SetCurrentStereoViewIndexForViewPointExtension(InStereoViewIndex);
 
 	// Calculate view inside LocalPlayer: VE used to override view point for this StereoViewIndex
 	FSceneView* View = LocalPlayer->CalcSceneView(ViewFamily, OutViewLocation, OutViewRotation, InViewport, InViewDrawer, InStereoViewIndex);
 
 	// Set INDEX_NONE when this VE no longer needs to be used.
-	ViewportManagerViewPointExtension->SetCurrentStereoViewIndex(INDEX_NONE);
+	SetCurrentStereoViewIndexForViewPointExtension(INDEX_NONE);
 
 	return View;
 }
diff --git a/Engine/Plugins/Runtime/nDisplay/Source/DisplayCluster/Private/Render/Viewport/DisplayClusterViewportManager.h b/Engine/Plugins/Runtime/nDisplay/Source/DisplayCluster/Private/Render/Viewport/DisplayClusterViewportManager.h
--- a/Engine/Plugins/Runtime/nDisplay/Source/DisplayCluster/Private/Render/Viewport/DisplayClusterViewportManager.h
+++ b/Engine/Plugins/Runtime/nDisplay/Source/DisplayCluster/Private/Render/Viewport/DisplayClusterViewportManager.h
@@ -182,9 +182,16 @@ public:
 	bool ShouldRenderFinalColor() const;
 
 	/** Return nDisplay VE object. */
-	TSharedPtr<class FDisplayClusterViewportManagerViewExtension, ESPMode::ThreadSafe> GetViewportManagerViewExtension() const
+	const TSharedPtr<class FDisplayClusterViewportManagerViewExtension, ESPMode::ThreadSafe>& GetViewportManagerViewExtension() const
 	{ return ViewportManagerViewExtension; }
 
+	/** Return ViewPoint VE instance. */
+	const TSharedPtr<class FDisplayClusterViewportManagerViewPointExtension, ESPMode::ThreadSafe>& GetViewportManagerViewPointExtension() const
+	{ return ViewportManagerViewPointExtension; }
+
+	/** Set current view index for FDisplayClusterViewportManagerViewPointExtension */
+	void SetCurrentStereoViewIndexForViewPointExtension(const int32 InStereoViewIndex) const;
+
 	/** Returns true if at least one of the viewports are visible on the frame RTT. */
 	bool ShouldUseOutputFrameTargetableResources() const;
 	bool ShouldUseAdditionalFrameTargetableResource() const;
diff --git a/Engine/Plugins/Runtime/nDisplay/Source/DisplayCluster/Private/Render/Viewport/DisplayClusterViewportManagerViewPointExtension.cpp b/Engine/Plugins/Runtime/nDisplay/Source/DisplayCluster/Private/Render/Viewport/DisplayClusterViewportManagerViewPointExtension.cpp
--- a/Engine/Plugins/Runtime/nDisplay/Source/DisplayCluster/Private/Render/Viewport/DisplayClusterViewportManagerViewPointExtension.cpp
+++ b/Engine/Plugins/Runtime/nDisplay/Source/DisplayCluster/Private/Render/Viewport/DisplayClusterViewportManagerViewPointExtension.cpp
@@ -15,10 +15,10 @@ FDisplayClusterViewportManagerViewPointExtension::FDisplayClusterViewportManager
 
 bool FDisplayClusterViewportManagerViewPointExtension::IsActiveThisFrame_Internal(const FSceneViewExtensionContext& Context) const
 {
-	return IsActive() && Context.IsStereoSupported();
+	// This VE can be used in the editor to preview in the scene
+	return IsActive();
 }
 
-/** True, if VE can be used at the moment. */
 bool FDisplayClusterViewportManagerViewPointExtension::IsActive() const
 {
 	return Configuration->GetViewportManager() != nullptr && CurrentStereoViewIndex != INDEX_NONE;
@@ -32,3 +32,21 @@ void FDisplayClusterViewportManagerViewPointExtension::SetupViewPoint(APlayerCon
 		DCViewport->SetupViewPoint(ContextNum, InOutViewInfo);
 	}
 }
+
+void FDisplayClusterViewportManagerViewPointExtension::SetupView(FSceneViewFamily& InViewFamily, FSceneView& InView)
+{
+	uint32 ContextNum = 0;
+	if (IDisplayClusterViewport* DCViewport = IsActive() ? Configuration->GetViewportManager()->FindViewport(CurrentStereoViewIndex, &ContextNum) : nullptr)
+	{
+		UWorld* CurrentWorld = InViewFamily.Scene ? InViewFamily.Scene->GetWorld() : nullptr;
+		if (!CurrentWorld)
+		{
+			CurrentWorld = Configuration->GetCurrentWorld();
+		}
+		if (CurrentWorld)
+		{
+			// Apply viewport context settings to view (crossGPU, visibility, etc)
+			DCViewport->SetupSceneView(ContextNum, CurrentWorld, InViewFamily, InView);
+		}
+	}
+}
diff --git a/Engine/Plugins/Runtime/nDisplay/Source/DisplayCluster/Private/Render/Viewport/DisplayClusterViewportManagerViewPointExtension.h b/Engine/Plugins/Runtime/nDisplay/Source/DisplayCluster/Private/Render/Viewport/DisplayClusterViewportManagerViewPointExtension.h
--- a/Engine/Plugins/Runtime/nDisplay/Source/DisplayCluster/Private/Render/Viewport/DisplayClusterViewportManagerViewPointExtension.h
+++ b/Engine/Plugins/Runtime/nDisplay/Source/DisplayCluster/Private/Render/Viewport/DisplayClusterViewportManagerViewPointExtension.h
@@ -4,7 +4,8 @@
 #include "SceneViewExtension.h"
 #include "Render/Viewport/Configuration/DisplayClusterViewportConfiguration.h"
 
-#define DISPLAYCLUSTER_SCENE_VIEWPOINT_EXTENSION_PRIORITY -1
+// This VE should be used before all others.
+#define DISPLAYCLUSTER_SCENE_VIEWPOINT_EXTENSION_PRIORITY 1000
 
 /**
  * View extension applying an DC Viewport features
@@ -18,6 +19,7 @@ public:
 public:
 	//~ Begin ISceneViewExtension interface
 	virtual void SetupViewPoint(APlayerController* Player, FMinimalViewInfo& InOutViewInfo) override;
+	virtual void SetupView(FSceneViewFamily& InViewFamily, FSceneView& InView) override;
 
 	virtual int32 GetPriority() const override { return DISPLAYCLUSTER_SCENE_VIEWPOINT_EXTENSION_PRIORITY; }
 
diff --git a/Engine/Plugins/Runtime/nDisplay/Source/DisplayCluster/Private/Render/Viewport/Preview/DisplayClusterViewportManagerPreview_Render.cpp b/Engine/Plugins/Runtime/nDisplay/Source/DisplayCluster/Private/Render/Viewport/Preview/DisplayClusterViewportManagerPreview_Render.cpp
--- a/Engine/Plugins/Runtime/nDisplay/Source/DisplayCluster/Private/Render/Viewport/Preview/DisplayClusterViewportManagerPreview_Render.cpp
+++ b/Engine/Plugins/Runtime/nDisplay/Source/DisplayCluster/Private/Render/Viewport/Preview/DisplayClusterViewportManagerPreview_Render.cpp
@@ -177,9 +177,6 @@ bool FDisplayClusterViewportManagerPreview::InitializeClusterNodePreview(const E
 								if (FSceneView* View = ViewportPtr->ViewportPreview->CalcSceneView(*ViewFamily, ViewIt.ContextNum))
 								{
 									Views.Add(View);
-
-									// Apply viewport context settings to view (crossGPU, visibility, etc)
-									ViewportPtr->SetupSceneView(ViewIt.ContextNum, PreviewScene->GetWorld(), *ViewFamily, *View);
 								}
 							}
 							else
diff --git a/Engine/Plugins/Runtime/nDisplay/Source/DisplayCluster/Private/Render/Viewport/Preview/DisplayClusterViewportPreview.h b/Engine/Plugins/Runtime/nDisplay/Source/DisplayCluster/Private/Render/Viewport/Preview/DisplayClusterViewportPreview.h
--- a/Engine/Plugins/Runtime/nDisplay/Source/DisplayCluster/Private/Render/Viewport/Preview/DisplayClusterViewportPreview.h
+++ b/Engine/Plugins/Runtime/nDisplay/Source/DisplayCluster/Private/Render/Viewport/Preview/DisplayClusterViewportPreview.h
@@ -130,7 +130,12 @@ protected:
 	TSharedPtr<FDisplayClusterViewportResource, ESPMode::ThreadSafe> GetOutputPreviewTargetableResources() const;
 
 	/** Calculating stereoviewer offset for preview. */
-	bool CalculateStereoViewOffset(FDisplayClusterViewport& InViewport, const uint32 InContextNum, FRotator& ViewRotation, FVector& ViewLocation);
+	bool CalculateStereoViewOffset(
+		FDisplayClusterViewport& InViewport,
+		const uint32 InContextNum,
+		FRotator& ViewRotation,
+		FVector& ViewLocation,
+		FVector& StereoViewLocation);
 
 	/** Get viewport context projection matrix. */
 	FMatrix GetStereoProjectionMatrix(FDisplayClusterViewport& InViewport, const uint32 InContextNum);
diff --git a/Engine/Plugins/Runtime/nDisplay/Source/DisplayCluster/Private/Render/Viewport/Preview/DisplayClusterViewportPreview_Render.cpp b/Engine/Plugins/Runtime/nDisplay/Source/DisplayCluster/Private/Render/Viewport/Preview/DisplayClusterViewportPreview_Render.cpp
--- a/Engine/Plugins/Runtime/nDisplay/Source/DisplayCluster/Private/Render/Viewport/Preview/DisplayClusterViewportPreview_Render.cpp
+++ b/Engine/Plugins/Runtime/nDisplay/Source/DisplayCluster/Private/Render/Viewport/Preview/DisplayClusterViewportPreview_Render.cpp
@@ -26,12 +26,13 @@ bool FDisplayClusterViewportPreview::CalculateViewportContext(uint32 InContextNu
 {
 	if (FDisplayClusterViewport* InViewport = GetViewportImpl())
 	{
-		FVector  ViewLocation;
+		FVector ViewLocation;
+		FVector StereoViewLocation;
 		FRotator ViewRotation;
 		FMatrix ProjectionMatrix;
 
 		return Configuration->IsSceneOpened()
-			&& CalculateStereoViewOffset(*InViewport, InContextNum, ViewRotation, ViewLocation)
+			&& CalculateStereoViewOffset(*InViewport, InContextNum, ViewRotation, ViewLocation, StereoViewLocation)
 			&& InViewport->GetProjectionMatrix(InContextNum, ProjectionMatrix);
 	}
 
@@ -54,10 +55,11 @@ FSceneView* FDisplayClusterViewportPreview::CalcSceneView(FSceneViewFamilyContex
 
 	const AActor* ViewOwner = nullptr;
 
-	FVector  ViewLocation;
+	FVector ViewLocation;
+	FVector StereoViewLocation;
 	FRotator ViewRotation;
 
-	if (CalculateStereoViewOffset(*InViewport, InContextNum, ViewRotation, ViewLocation))
+	if (CalculateStereoViewOffset(*InViewport, InContextNum, ViewRotation, ViewLocation, StereoViewLocation))
 	{
 		FMatrix ProjectionMatrix = GetStereoProjectionMatrix(*InViewport, InContextNum);
 
@@ -79,7 +81,10 @@ FSceneView* FDisplayClusterViewportPreview::CalcSceneView(FSceneViewFamilyContex
 		ViewInitOptions.SceneViewStateInterface = InViewport->GetViewState(InContextNum);
 		ViewInitOptions.ViewActor = ViewOwner;
 
-		ViewInitOptions.ViewOrigin = ViewLocation;
+		ViewInitOptions.ViewOrigin = StereoViewLocation;
+		ViewInitOptions.ViewLocation = ViewLocation;
+		ViewInitOptions.ViewRotation = ViewRotation;
+
 		ViewInitOptions.ViewRotationMatrix = ViewRotationMatrix;
 		ViewInitOptions.ProjectionMatrix = ProjectionMatrix;
 
@@ -123,12 +128,25 @@ FSceneView* FDisplayClusterViewportPreview::CalcSceneView(FSceneViewFamilyContex
 		View->EndFinalPostprocessSettings(ViewInitOptions);
 		}
 
+		FDisplayClusterViewportManager* ViewportManager = Configuration->GetViewportManagerImpl();
+		if (ViewportManager && InViewport->GetContexts().IsValidIndex(InContextNum))
+		{
+			// Configure FDisplayClusterViewportManagerViewPointExtension for this viewport.
+			ViewportManager->SetCurrentStereoViewIndexForViewPointExtension(InViewport->GetContexts()[InContextNum].StereoViewIndex);
+		}
+
 		// Setup view extension for this view
 		for (int32 ViewExt = 0; ViewExt < InOutViewFamily.ViewExtensions.Num(); ViewExt++)
 		{
 			InOutViewFamily.ViewExtensions[ViewExt]->SetupView(InOutViewFamily, *View);
 		}
 
+		if (ViewportManager)
+		{
+			// Configure FDisplayClusterViewportManagerViewPointExtension for this viewport.
+			ViewportManager->SetCurrentStereoViewIndexForViewPointExtension(INDEX_NONE);
+		}
+
 		return View;
 	}
 
@@ -153,7 +171,12 @@ FMatrix FDisplayClusterViewportPreview::GetStereoProjectionMatrix(FDisplayCluste
 	return PrjMatrix;
 }
 
-bool FDisplayClusterViewportPreview::CalculateStereoViewOffset(FDisplayClusterViewport& InViewport, const uint32 InContextNum, FRotator& ViewRotation, FVector& ViewLocation)
+bool FDisplayClusterViewportPreview::CalculateStereoViewOffset(
+	FDisplayClusterViewport& InViewport,
+	const uint32 InContextNum,
+	FRotator& ViewRotation,
+	FVector& ViewLocation,
+	FVector& StereoViewLocation)
 {
 	check(IsInGameThread());
 
@@ -169,7 +192,8 @@ bool FDisplayClusterViewportPreview::CalculateStereoViewOffset(FDisplayClusterVi
 
 	// Obtaining the offset of the stereo eye and the values of the projection clipping plane for the given viewport was moved inside CalculateView().
 	// Perform view calculations on a policy side
-	if (!InViewport.CalculateView(InContextNum, ViewLocation, ViewRotation, Configuration->GetWorldToMeters()))
+	StereoViewLocation = ViewInfo.Location;
+	if (!InViewport.CalculateView(InContextNum, StereoViewLocation, ViewRotation, Configuration->GetWorldToMeters()))
 	{
 		if (CanShowLogMsgOnce(EDisplayClusterViewportPreviewShowLogMsgOnce::CalculateViewIsFailed))
 		{
diff --git a/Engine/Plugins/Runtime/nDisplay/Source/DisplayCluster/Private/Render/Viewport/RenderFrame/DisplayClusterRenderFrameManager.cpp b/Engine/Plugins/Runtime/nDisplay/Source/DisplayCluster/Private/Render/Viewport/RenderFrame/DisplayClusterRenderFrameManager.cpp
--- a/Engine/Plugins/Runtime/nDisplay/Source/DisplayCluster/Private/Render/Viewport/RenderFrame/DisplayClusterRenderFrameManager.cpp
+++ b/Engine/Plugins/Runtime/nDisplay/Source/DisplayCluster/Private/Render/Viewport/RenderFrame/DisplayClusterRenderFrameManager.cpp
@@ -99,7 +99,7 @@ bool FDisplayClusterRenderFrameManager::BuildSimpleFrame(FViewport* InViewport,
 				{
 					FrameViewFamily.Views.Add(FrameView);
 					FrameViewFamily.CustomBufferRatio = ContextIt.CustomBufferRatio;
-					FrameViewFamily.ViewExtensions = ViewportIt->GatherActiveExtensions(InViewport);
+					FrameViewFamily.ViewExtensions = ViewportIt->GatherActiveExtensions(ContextIt.ContextNum, InViewport);
 				}
 
 				FDisplayClusterRenderFrameTarget FrameRenderTarget;
