[nDisplay] HoldoutComposite integration 2/3: New TextureUtils API (RDG/RHI) and color encoding (Gamma, Premultiply).

Added color encoding container `FDisplayClusterColorEncoding`
- Encoding type (Linear, Gamma, sRGB, MediaPQ)
- Premultiply (None, Alpha, InvertedAlpha)

Added a new public function IDisplayClusterViewportProxy::GetResourceColorEncoding_RenderThread()
nDisplay. It returns the color encoding of the viewport resource.

Added `IDisplayClusterShadersTextureUtils` API, here is features:
- use RDG or RHI
- textures as parameters can be any RHI or RDG types
- Color mask (to copy only alpha or RGB/RGBA)
- Alpha override (none, invert, one, zero)
- Color encoding conversion (Linear, Gamma, sRGB, MediaPQ, ...)
- Premultiply color (None, Alpha, InvertAlpha).
- Unpremultiply when using Color encoding.
- Fitting texture rects to real texture sizes.
- Use CopyTexture, if possible, to dispense with the shader.
- Internal transition RHI->RDG (if necessary) and its performance at the end

Reworked all related code.
- functions FDisplayClusterViewportProxy::CopyResource_RenderThread(...) have been removed. The corresponding code has been reworked.
- Reworked IDisplayClusterDisplayDeviceProxy::AddFinalPass_RenderThread() function.

Reworked ICVFX shader:
- Premultiplied textures (LightCard) must be Unpremultiplied before any gamma changes.
- LC + UVLC blending is done in linear color space.
- Removed alpha channel inversion from UVLC texture color, this has been moved to UVLightCardGammaCorrection.usf
- The LightCard gamma correction has been moved to the very end.

Fixed bugs:
- Preview without post processing: when applying OCIO to the viewport, LC and UVLC colors were different.
- UVLC gamma was incorrect

diff --git a/Engine/Plugins/Runtime/nDisplay/Shaders/Private/DisplayClusterColorCorrection.ush b/Engine/Plugins/Runtime/nDisplay/Shaders/Private/DisplayClusterColorCorrection.ush
new file mode 100644
--- /dev/null
+++ b/Engine/Plugins/Runtime/nDisplay/Shaders/Private/DisplayClusterColorCorrection.ush
@@ -0,0 +1,36 @@
+// Copyright Epic Games, Inc. All Rights Reserved.
+
+#include "/Engine/Public/Platform.ush"
+#include "/Engine/Private/GammaCorrectionCommon.ush"
+
+// This function copied from the `\\Engine\Plugins\Compositing\OpenColorIO\Shaders\Private\OpenColorIOShader.usf`
+// Note: For our (zero-tolerance) tests to succeed against a no-op, precise is required with the pow function.
+float3 DisplayGammaCorrection(precise float3 InColor, precise float InGamma)
+{
+	return sign(InColor) * pow(abs(InColor), InGamma);
+}
+
+float SafeRcp(float x)
+{
+	return x > 0.0 ? rcp(x) : 0.0;
+}
+
+float4 Unpremultiply(float4 InColor)
+{
+	return float4(InColor.xyz * SafeRcp(InColor.w), InColor.w);
+}
+
+float4 Premultiply(float4 InColor)
+{
+	return float4(InColor.xyz * InColor.w, InColor.w);
+}
+
+float4 InvertedUnpremultiply(float4 InColor)
+{
+	return float4(InColor.xyz * SafeRcp(1.0f-InColor.w), InColor.w);
+}
+
+float4 InvertedPremultiply(float4 InColor)
+{
+	return float4(InColor.xyz * (1.0f-InColor.w), InColor.w);
+}
diff --git a/Engine/Plugins/Runtime/nDisplay/Shaders/Private/MPCDIOverlayShaders.usf b/Engine/Plugins/Runtime/nDisplay/Shaders/Private/MPCDIOverlayShaders.usf
--- a/Engine/Plugins/Runtime/nDisplay/Shaders/Private/MPCDIOverlayShaders.usf
+++ b/Engine/Plugins/Runtime/nDisplay/Shaders/Private/MPCDIOverlayShaders.usf
@@ -37,6 +37,21 @@ float3 ComposeOverlayColor(float4 InColor, float4 InOverlayColor)
 	return InOverlayColor.xyz + InColor.xyz * InOverlayColor.w;
 }
 
+/** Apply final gamma for LightCard color. */
+float4 GetFinalLightcardColor(float4 InColor)
+{
+    // Lightcard color is premultiplied. Unpremultiply before gamma correction.
+	float4 OutColor = InvertedUnpremultiply(InColor);
+
+	// Apply gamma conversion for unpremultiplied color
+	OutColor.rgb = DisplayGammaCorrection(OutColor.rgb, LightCardGamma);
+
+	// And finally premultiply it back
+	OutColor = InvertedPremultiply(InColor);
+	
+	return OutColor;
+}
+
 // Additive render pass with 2 overlays
 // a = in\out, b,c = ordered overlays (a + b + c)
 // c.xyz + (b.xyz + a.xyz * b.w)*c.w = a.xyz * (b.w *c.w) + b.xyz*c.w + c.xyz
@@ -48,36 +63,26 @@ float4 ComposeRenderPassTwoOverlays(float4 InOverlay1, float4 InOverlay2)
 ////////////////////////////////////////////////////////////////////////////////
 // LightCard
 ////////////////////////////////////////////////////////////////////////////////
-float4 NormalizeLightCardColor(float4 LightCardColorLin)
-{
-	return float4(pow(LightCardColorLin.xyz, LightCardGamma), LightCardColorLin.w);
-}
-
-float4 NormalizeLightCardColorInvAlpha(float4 LightCardColorLin)
-{
-	return float4(pow(LightCardColorLin.xyz, LightCardGamma), 1 - LightCardColorLin.w);
-}
-
 float4 GetLightCardUnderTextureColor(float2 InUV)
 {
-	return NormalizeLightCardColor(LightCardUnderTexture.Sample(LightCardUnderSampler, InUV));
+	return LightCardUnderTexture.Sample(LightCardUnderSampler, InUV);
 }
 
 float4 GetLightCardOverTextureColor(float2 InUV)
 {
-	return NormalizeLightCardColor(LightCardOverTexture.Sample(LightCardOverSampler, InUV));
+	return LightCardOverTexture.Sample(LightCardOverSampler, InUV);
 }
 
 float4 GetUVLightCardUnderTextureColor(float2 InUV)
 {
-	float4 UVLightCardUnderColor = NormalizeLightCardColorInvAlpha(UVLightCardUnderTexture.Sample(UVLightCardUnderSampler, InUV));
+	float4 UVLightCardUnderColor = UVLightCardUnderTexture.Sample(UVLightCardUnderSampler, InUV);
 	
 	if (LightCardMode & UVLIGHTCARD_MERGE)
 	{
 		// UVLIGHTCARD_MERGE - special rendering mode, for Outer viewports that are forced to render LightCards only in "over" or "under" mode.
 		// both UV lightcard textures are rendered under or over the inner frustum.
 
-		float4 UVLightCardOverColor= NormalizeLightCardColorInvAlpha(UVLightCardOverTexture.Sample(UVLightCardOverSampler, InUV));
+		float4 UVLightCardOverColor= UVLightCardOverTexture.Sample(UVLightCardOverSampler, InUV);
 
 		return ComposeRenderPassTwoOverlays(UVLightCardUnderColor, UVLightCardOverColor);
 	}
@@ -87,14 +92,14 @@ float4 GetUVLightCardUnderTextureColor(float2 InUV)
 
 float4 GetUVLightCardOverTextureColor(float2 InUV)
 {
-	float4 UVLightCardOverColor = NormalizeLightCardColorInvAlpha(UVLightCardOverTexture.Sample(UVLightCardOverSampler, InUV));
+	float4 UVLightCardOverColor = UVLightCardOverTexture.Sample(UVLightCardOverSampler, InUV);
 
 	if (LightCardMode & UVLIGHTCARD_MERGE)
 	{
 		// UVLIGHTCARD_MERGE - special rendering mode, for Outer viewports that are forced to render LightCards only in "over" or "under" mode.
 		// both UV lightcard textures are rendered under or over the inner frustum.
 
-		float4 UVLightCardUnderColor = NormalizeLightCardColorInvAlpha(UVLightCardUnderTexture.Sample(UVLightCardUnderSampler, InUV));
+		float4 UVLightCardUnderColor = UVLightCardUnderTexture.Sample(UVLightCardUnderSampler, InUV);
 
 		return ComposeRenderPassTwoOverlays(UVLightCardUnderColor, UVLightCardOverColor);
 	}
@@ -147,6 +152,9 @@ float3 ComposeLightCardUnderColor(float4 InColor, OutputVS IN, float2 OverlayUV)
 	{
 		float4 LightCardUnderColor = GetLightCardUnderColor(IN, OverlayUV);
 
+		// Convert final LC color from linear to the output color space
+		LightCardUnderColor = GetFinalLightcardColor(LightCardUnderColor);
+
 		return ComposeOverlayColor(InColor, LightCardUnderColor);
 	}
 	
@@ -159,6 +167,9 @@ float3 ComposeLightCardOverColor(float4 InColor, OutputVS IN, float2 OverlayUV)
 	{
 		float4 LightCardOverColor = GetLightCardOverColor(IN, OverlayUV);
 
+		// Convert final LC color from linear to the output color space
+		LightCardOverColor = GetFinalLightcardColor(LightCardOverColor);
+
 		return ComposeOverlayColor(InColor, LightCardOverColor);
 	}
 	
@@ -413,7 +424,7 @@ OutputPS IcvfxWarpPS(OutputVS IN)
 	#endif
 
 	
-	// Compose a “Under” lightcard, if used.
+	// Compose a "Under" lightcard, if used.
 	OutColor.xyz = ComposeLightCardUnderColor(OutColor, IN, OverlayUV);
 	
 	#if INNER_CAMERA
@@ -422,7 +433,7 @@ OutputPS IcvfxWarpPS(OutputVS IN)
 	#endif
 
 
-	// Compose a “Over” lightcard, if used.
+	// Compose a "Over" lightcard, if used.
 	OutColor.xyz = ComposeLightCardOverColor(OutColor,IN, OverlayUV);
 
 #else 
@@ -435,6 +446,9 @@ OutputPS IcvfxWarpPS(OutputVS IN)
 		if( LightCardMode & (LIGHTCARD_OVER | UVLIGHTCARD_OVER))
 		{
 			OutColor = GetLightCardOverColor(IN, OverlayUV);
+
+			// Convert final LC color from linear to the output color space
+			OutColor = GetFinalLightcardColor(OutColor);
 		}
 		else
 		{
diff --git a/Engine/Plugins/Runtime/nDisplay/Shaders/Private/MPCDIUtils.ush b/Engine/Plugins/Runtime/nDisplay/Shaders/Private/MPCDIUtils.ush
--- a/Engine/Plugins/Runtime/nDisplay/Shaders/Private/MPCDIUtils.ush
+++ b/Engine/Plugins/Runtime/nDisplay/Shaders/Private/MPCDIUtils.ush
@@ -1,8 +1,7 @@
 // Copyright Epic Games, Inc. All Rights Reserved.
 
 #include "/Engine/Public/Platform.ush"
-#include "/Engine/Private/GammaCorrectionCommon.ush"
-
+#include "DisplayClusterColorCorrection.ush"
 
 //VS Parameters
 float4 DrawRectanglePosScaleBias;
diff --git a/Engine/Plugins/Runtime/nDisplay/Shaders/Private/ResourceUtils.usf b/Engine/Plugins/Runtime/nDisplay/Shaders/Private/ResourceUtils.usf
new file mode 100644
--- /dev/null
+++ b/Engine/Plugins/Runtime/nDisplay/Shaders/Private/ResourceUtils.usf
@@ -0,0 +1,141 @@
+// Copyright Epic Games, Inc. All Rights Reserved.
+
+#include "/Engine/Private/Common.ush"
+#include "DisplayClusterColorCorrection.ush"
+
+// Input texture and sampler
+Texture2D InputTexture;
+SamplerState InputTextureSampler;
+
+// Definition for the color encoding and alpha overrides
+// Encodings.X = Input Encoding
+// Encodings.Y = Output Encoding
+// Encodings.Z = Override Alpha
+
+#define ESourceEncoding_Linear  0
+#define ESourceEncoding_Gamma   1
+#define ESourceEncoding_sRGB    2
+#define ESourceEncoding_MediaPQ 3
+
+#define EOverrideAlpha_None   0
+#define EOverrideAlpha_Invert 1
+#define EOverrideAlpha_One    2
+#define EOverrideAlpha_Zero   3
+
+// Color encoding parameters:
+uint3 Encodings;
+float3 DisplayGamma;
+
+// Definition for the ColorPremultiply
+// ColorPremultiply.X = Input texture Color premultiply
+// ColorPremultiply.Y = Output texture Color premultiply
+
+#define EColorPremultiply_None          0
+#define EColorPremultiply_Alpha         1
+#define EColorPremultiply_InvertedAlpha 2
+
+// Color modifier
+uint3 ColorPremultiply;
+
+/** Implements color encoding. */
+float4 DisplayClusterEncodeColorAndAlpha(float4 InColorArg)
+{
+	float4 OutColor = InColorArg;
+
+#if COLOR_PREMULTIPLY	
+	// Unpremultiply color
+	BRANCH
+    if(ColorPremultiply.x == EColorPremultiply_Alpha) // Unpremultiply
+    {
+        OutColor = Unpremultiply(OutColor);
+	}
+	else if(ColorPremultiply.x == EColorPremultiply_InvertedAlpha) // Unpremultiply inverted alpha
+	{
+        OutColor = InvertedUnpremultiply(OutColor);
+	}
+#endif
+	
+#if OVERRIDE_ALPHA
+	// Override alpha
+	BRANCH
+    if(Encodings.z == EOverrideAlpha_Invert)
+    {
+        OutColor.a = 1.0f - OutColor.a;
+    }
+    else if(Encodings.z == EOverrideAlpha_One)
+    {
+        OutColor.a = 1.0f;
+    }
+    else if(Encodings.z == EOverrideAlpha_Zero)
+    {
+        OutColor.a = 0.0f;
+    }
+#endif
+
+#if ENCODE_INPUT || ENCODE_OUTPUT
+
+#if ENCODE_INPUT
+	// Convert Input color to linear
+	BRANCH
+    if(Encodings.x == ESourceEncoding_Gamma) // Gamma -> Linear
+    {
+        OutColor.rgb = DisplayGammaCorrection(OutColor.rgb, DisplayGamma.x);
+    }
+    else if(Encodings.x == ESourceEncoding_sRGB) // sRGB -> Linear
+    {
+        OutColor.rgb = sRGBToLinear(OutColor.rgb);
+	}
+	else if(Encodings.x == ESourceEncoding_MediaPQ) // MediaPQ -> Linear
+    {
+		// Pixel shader for Linear-PQ conversion
+		float3 colorPQ = LinearToST2084(OutColor.rgb);
+        OutColor.rgb = saturate(colorPQ);
+	}
+
+#endif // ENCODE_INPUT
+
+#if ENCODE_OUTPUT
+	// Convert the output color to a output color space
+	BRANCH
+    if(Encodings.y == ESourceEncoding_Gamma) // Linear -> Gamma
+    {
+        OutColor.rgb = DisplayGammaCorrection(OutColor.rgb, DisplayGamma.y);
+    }
+    else if(Encodings.y == ESourceEncoding_sRGB) // Linear -> sRGB
+    {
+        OutColor.rgb = LinearToSrgb(OutColor.rgb);
+    }
+	else if(Encodings.y == ESourceEncoding_MediaPQ) // Linear -> MediaPQ
+    {
+		// Pixel shader for PQ-Linear conversion
+		OutColor.rgb = ST2084ToLinear(OutColor.rgb);
+    }
+#endif // ENCODE_OUTPUT
+
+#endif // ENCODE_INPUT || ENCODE_OUTPUT
+
+#if COLOR_PREMULTIPLY
+	// Premultiply color
+	BRANCH
+    if(ColorPremultiply.y == EColorPremultiply_Alpha) // Premultiply
+    {
+        OutColor = Premultiply(OutColor);
+	}
+	else if(ColorPremultiply.y == EColorPremultiply_InvertedAlpha) // Premultiply inverted alpha
+	{
+        OutColor = InvertedPremultiply(OutColor);
+	}
+#endif
+
+	return OutColor;
+}
+
+void Main(
+	FScreenVertexOutput Input,
+	out float4 OutColor : SV_Target0
+	)
+{
+	float4 InColor = Texture2DSample(InputTexture, InputTextureSampler, Input.UV);
+
+	OutColor = DisplayClusterEncodeColorAndAlpha(InColor);
+}
\ No newline at end of file
diff --git a/Engine/Plugins/Runtime/nDisplay/Shaders/Private/UVLightCardGammaCorrection.usf b/Engine/Plugins/Runtime/nDisplay/Shaders/Private/UVLightCardGammaCorrection.usf
--- a/Engine/Plugins/Runtime/nDisplay/Shaders/Private/UVLightCardGammaCorrection.usf
+++ b/Engine/Plugins/Runtime/nDisplay/Shaders/Private/UVLightCardGammaCorrection.usf
@@ -1,14 +1,23 @@
 // Copyright Epic Games, Inc. All Rights Reserved.
 
 #include "/Engine/Private/Common.ush"
+#include "DisplayClusterColorCorrection.ush"
 
 Texture2D InputTexture;
 SamplerState InputSampler;
 float InverseGamma = 1.0;
 
+
 float4 GammaCorrectionPS(noperspective float4 UVAndScreenPos : TEXCOORD0) : SV_Target0
 {
-	float2 UV = UVAndScreenPos.xy;
-	float4 TextureColor = Texture2DSample(InputTexture, InputSampler, UV);
-	return float4(pow(TextureColor.rgb, InverseGamma), TextureColor.a);
+	float4 OutColor = Texture2DSample(InputTexture, InputSampler, UVAndScreenPos.xy);
+	
+	// LightCard textures are premultiplied, so we need to unpremultiply before any gamma correction
+	OutColor = Unpremultiply(OutColor);
+
+	OutColor.rgb = DisplayGammaCorrection(OutColor.rgb, InverseGamma);
+
+	OutColor = Premultiply(OutColor);
+
+	return float4(OutColor.rgb, 1.0 - OutColor.a);
 }
\ No newline at end of file
diff --git a/Engine/Plugins/Runtime/nDisplay/Source/DisplayCluster/Private/Render/DisplayDevice/Components/DisplayClusterDisplayDeviceComponent.cpp b/Engine/Plugins/Runtime/nDisplay/Source/DisplayCluster/Private/Render/DisplayDevice/Components/DisplayClusterDisplayDeviceComponent.cpp
--- a/Engine/Plugins/Runtime/nDisplay/Source/DisplayCluster/Private/Render/DisplayDevice/Components/DisplayClusterDisplayDeviceComponent.cpp
+++ b/Engine/Plugins/Runtime/nDisplay/Source/DisplayCluster/Private/Render/DisplayDevice/Components/DisplayClusterDisplayDeviceComponent.cpp
@@ -60,7 +60,12 @@ void UDisplayClusterDisplayDeviceComponent::UpdateDisplayDeviceProxyImpl(IDispla
 
 	if (!DisplayDeviceProxy.IsValid())
 	{
-		// create a new proxy object for the current settings
-		DisplayDeviceProxy = MakeShared<FDisplayClusterDisplayDeviceProxy_OpenColorIO, ESPMode::ThreadSafe>(ColorConversionSettings);
+		// OCIO shaders may not be ready at the moment, so we should check for them before creating our own proxy.
+		FOpenColorIORenderPassResources OCIOPassResources = FOpenColorIORendering::GetRenderPassResources(ColorConversionSettings, GMaxRHIFeatureLevel);
+		if(OCIOPassResources.IsValid())
+		{
+			// Create a DisplayDevice proxy only when OCIO render pass resources are available for use.
+			DisplayDeviceProxy = MakeShared<FDisplayClusterDisplayDeviceProxy_OpenColorIO, ESPMode::ThreadSafe>(ColorConversionSettings.ToString(), OCIOPassResources);
+		}
 	}
 }
diff --git a/Engine/Plugins/Runtime/nDisplay/Source/DisplayCluster/Private/Render/DisplayDevice/Proxy/DisplayClusterDisplayDeviceProxy_OpenColorIO.cpp b/Engine/Plugins/Runtime/nDisplay/Source/DisplayCluster/Private/Render/DisplayDevice/Proxy/DisplayClusterDisplayDeviceProxy_OpenColorIO.cpp
--- a/Engine/Plugins/Runtime/nDisplay/Source/DisplayCluster/Private/Render/DisplayDevice/Proxy/DisplayClusterDisplayDeviceProxy_OpenColorIO.cpp
+++ b/Engine/Plugins/Runtime/nDisplay/Source/DisplayCluster/Private/Render/DisplayDevice/Proxy/DisplayClusterDisplayDeviceProxy_OpenColorIO.cpp
@@ -2,6 +2,9 @@
 
 #include "Render/DisplayDevice/Proxy/DisplayClusterDisplayDeviceProxy_OpenColorIO.h"
 #include "Render/Viewport/Containers/DisplayClusterViewport_Context.h"
+#include "Render/Viewport/DisplayClusterViewport_OpenColorIO.h"
+
+#include "IDisplayClusterShadersTextureUtils.h"
 
 #include "Shader.h"
 #include "SceneView.h"
@@ -10,50 +13,61 @@
 #include "PostProcess/PostProcessMaterialInputs.h"
 #include "Engine/TextureRenderTarget.h"
 
-namespace UE::DisplayCluster
-{
-	static inline float GetDisplayGamma(const FDisplayClusterViewport_Context& InViewportContext)
+bool FDisplayClusterDisplayDeviceProxy_OpenColorIO::HasFinalPass_RenderThread() const
 	{
-		const float DefaultDisplayGamma = UTextureRenderTarget::GetDefaultDisplayGamma();
-		const float EngineDisplayGamma = InViewportContext.RenderThreadData.EngineDisplayGamma;
-
-		// There is a special case where post processing and tonemapper are disabled. In this case tonemapper applies a static display Inverse of Gamma which defaults to 2.2.
-		// In the case when Both PostProcessing and ToneMapper are disabled we apply gamma manually. In every other case we apply inverse gamma before applying OCIO.
-		const float DisplayGamma = (InViewportContext.RenderThreadData.EngineShowFlags.Tonemapper == 0) || (InViewportContext.RenderThreadData.EngineShowFlags.PostProcessing == 0) ? DefaultDisplayGamma : DefaultDisplayGamma / EngineDisplayGamma;
-
-		return DisplayGamma;
+	return true;
 	}
-};
 
-FDisplayClusterDisplayDeviceProxy_OpenColorIO::FDisplayClusterDisplayDeviceProxy_OpenColorIO(FOpenColorIOColorConversionSettings& InColorConversionSettings)
-	: OCIOPassId(InColorConversionSettings.ToString())
-	, OCIOPassResources(FOpenColorIORendering::GetRenderPassResources(InColorConversionSettings, GMaxRHIFeatureLevel))
-{ }
-
-bool FDisplayClusterDisplayDeviceProxy_OpenColorIO::HasFinalPass_RenderThread() const
+bool FDisplayClusterDisplayDeviceProxy_OpenColorIO::AddFinalPass_RenderThread(
+	const FDisplayClusterShadersTextureUtilsSettings& InTextureUtilsSettings,
+	const TSharedRef<IDisplayClusterShadersTextureUtils, ESPMode::ThreadSafe>& InTextureUtils) const
 {
+	if (!OCIOPassResources.IsValid())
+{
+		// The OCIO shader is not ready at this point, use the default resolve method.
+		InTextureUtils->Resolve(InTextureUtilsSettings);
+
 	return true;
 }
 
-bool FDisplayClusterDisplayDeviceProxy_OpenColorIO::AddFinalPass_RenderThread(FRDGBuilder& GraphBuilder, const FDisplayClusterViewport_Context& InViewportContext,
-	FRHITexture* InputTextureRHI, const FIntRect& InputRect, FRHITexture* OutputTextureRHI, const FIntRect& OutputRect) const
-{
-	FRDGTextureRef InputTexture = GraphBuilder.RegisterExternalTexture(CreateRenderTarget(InputTextureRHI, TEXT("FDisplayClusterDisplayDeviceOCIOInputTexture")));
-	FRDGTextureRef OutputTexture = GraphBuilder.RegisterExternalTexture(CreateRenderTarget(OutputTextureRHI, TEXT("FDisplayClusterDisplayDeviceOCIORenderTargetTexture")));
-
-	FScreenPassRenderTarget Output = FScreenPassRenderTarget(OutputTexture, OutputRect, ERenderTargetLoadAction::EClear);
-	EOpenColorIOTransformAlpha TransformAlpha = EOpenColorIOTransformAlpha::None;
-
-	FOpenColorIORendering::AddPass_RenderThread(
-		GraphBuilder,
-		FScreenPassViewInfo(),
-		GMaxRHIFeatureLevel,
-		FScreenPassTexture(InputTexture),
-		Output,
-		OCIOPassResources,
-		UE::DisplayCluster::GetDisplayGamma(InViewportContext),
-		TransformAlpha
-	);
+	// When requesting RDG from the TextureUtils API, it switches from RHI to RDG.
+	FRDGBuilder& GraphBuilder = InTextureUtils->GetOrCreateRDGBuilder();
 
-	return true;
+	bool bFinalPassResult = false;
+	InTextureUtils->ForEachContextByPredicate([&](
+		const FDisplayClusterShadersTextureViewportContext& Input,
+		const FDisplayClusterShadersTextureViewportContext& Output)
+		{
+			// returns true if OCIO is applied.
+			bFinalPassResult = true;
+
+			// The OCIO output is in sRGB
+			FOpenColorIORendering::AddPass_RenderThread(
+				GraphBuilder,
+				FScreenPassViewInfo(),
+				GMaxRHIFeatureLevel,
+				Input.ToScreenPassTexture(),
+				FScreenPassRenderTarget(Output.ToScreenPassTexture(), ERenderTargetLoadAction::EClear),
+				OCIOPassResources,
+				FDisplayClusterViewport_OpenColorIO::GetGammaCorrection(Input.ColorEncoding),
+				FDisplayClusterViewport_OpenColorIO::GetTransformAlpha(Input.ColorEncoding)
+			);
+
+			// Convert to output color space if necessary.
+			if (Output.ColorEncoding.GetEqualEncoding() != EDisplayClusterColorEncoding::sRGB)
+			{
+				// Convert OutputTexture color encoding from sRGB to Linear(Holdout):
+				InTextureUtils->ResolveTextureContext(
+					// Customize the settings by requesting a temporary input texture cloned from the output texture.
+					{ InTextureUtilsSettings , EDisplayClusterShaderTextureUtilsFlags::UseOutputTextureAsInput },
+					// The input texture is empty, a temporary texture created from the output will be used instead.
+					// Holdout uses the Liner->sRGB transform at the end, so use sRGB as the input encoding here even if it's different.
+					{ {}, {EDisplayClusterColorEncoding::sRGB, Input.ColorEncoding.Premultiply } },
+					// The output texture already contains the OCIO result and will be used as input in this pass.
+					Output
+				);
+			}
+		});
+
+	return bFinalPassResult;
 }
diff --git a/Engine/Plugins/Runtime/nDisplay/Source/DisplayCluster/Private/Render/DisplayDevice/Proxy/DisplayClusterDisplayDeviceProxy_OpenColorIO.h b/Engine/Plugins/Runtime/nDisplay/Source/DisplayCluster/Private/Render/DisplayDevice/Proxy/DisplayClusterDisplayDeviceProxy_OpenColorIO.h
--- a/Engine/Plugins/Runtime/nDisplay/Source/DisplayCluster/Private/Render/DisplayDevice/Proxy/DisplayClusterDisplayDeviceProxy_OpenColorIO.h
+++ b/Engine/Plugins/Runtime/nDisplay/Source/DisplayCluster/Private/Render/DisplayDevice/Proxy/DisplayClusterDisplayDeviceProxy_OpenColorIO.h
@@ -5,7 +5,6 @@
 #include "Render/DisplayDevice/IDisplayClusterDisplayDeviceProxy.h"
 
 #include "OpenColorIORendering.h"
-#include "OpenColorIOColorSpace.h"
 
 /**
  * Display Device Proxy object (OCIO render pass)
@@ -15,13 +14,18 @@ class FDisplayClusterDisplayDeviceProxy_OpenColorIO
 	: public IDisplayClusterDisplayDeviceProxy
 {
 public:
-	FDisplayClusterDisplayDeviceProxy_OpenColorIO(FOpenColorIOColorConversionSettings& InColorConversionSettings);
+	FDisplayClusterDisplayDeviceProxy_OpenColorIO(const FString& InOCIOPassId, FOpenColorIORenderPassResources& InOCIOPassResources)
+		: OCIOPassId(InOCIOPassId)
+		, OCIOPassResources(InOCIOPassResources)
+	{ }
+
 	virtual ~FDisplayClusterDisplayDeviceProxy_OpenColorIO() = default;
 
 public:
+	//~Begin IDisplayClusterDisplayDeviceProxy
 	virtual bool HasFinalPass_RenderThread() const override;
-	virtual bool AddFinalPass_RenderThread(FRDGBuilder& GraphBuilder, const FDisplayClusterViewport_Context& InViewportContext,
-		FRHITexture* InputTextureRHI, const FIntRect& InputRect, FRHITexture* OutputTextureRHI, const FIntRect& OutputRect) const override;
+	virtual bool AddFinalPass_RenderThread(const FDisplayClusterShadersTextureUtilsSettings& InTextureUtilsSettings, const TSharedRef<IDisplayClusterShadersTextureUtils, ESPMode::ThreadSafe>& InTextureUtils) const override;
+	//~~End IDisplayClusterDisplayDeviceProxy
 
 public:
 	const FString OCIOPassId;
diff --git a/Engine/Plugins/Runtime/nDisplay/Source/DisplayCluster/Private/Render/Viewport/Configuration/DisplayClusterViewportConfigurationHelpers_ICVFX.cpp b/Engine/Plugins/Runtime/nDisplay/Source/DisplayCluster/Private/Render/Viewport/Configuration/DisplayClusterViewportConfigurationHelpers_ICVFX.cpp
--- a/Engine/Plugins/Runtime/nDisplay/Source/DisplayCluster/Private/Render/Viewport/Configuration/DisplayClusterViewportConfigurationHelpers_ICVFX.cpp
+++ b/Engine/Plugins/Runtime/nDisplay/Source/DisplayCluster/Private/Render/Viewport/Configuration/DisplayClusterViewportConfigurationHelpers_ICVFX.cpp
@@ -789,10 +789,6 @@ void FDisplayClusterViewportConfigurationHelpers_ICVFX::UpdateLightcardViewportS
 
 	if (EnumHasAnyFlags(DstViewport.GetRenderSettingsICVFX().RuntimeFlags, EDisplayClusterViewportRuntimeICVFXFlags::Lightcard | EDisplayClusterViewportRuntimeICVFXFlags::UVLightcard))
 	{
-		// Set the light card gamma used to linearize the light card textures before blending during final composite
-		// The OCIO pass will have already linearized the light card renders, so set the gamma to 1 if using OCIO on the light cards
-		InOutBaseViewportRenderSettingsICVFX.ICVFX.LightCardGamma = DstViewport.GetOpenColorIO().IsValid() ? 1.0 : 2.2;
-
 		if (EnumHasAnyFlags(InOutBaseViewportRenderSettingsICVFX.Flags, EDisplayClusterViewportICVFXFlags::LightcardAlwaysUnder))
 		{
 			InOutBaseViewportRenderSettingsICVFX.ICVFX.OverrideUVLightCardRenderMode = EDisplayClusterShaderParametersICVFX_LightCardRenderMode::Under;
diff --git a/Engine/Plugins/Runtime/nDisplay/Source/DisplayCluster/Private/Render/Viewport/DisplayClusterViewportProxy.cpp b/Engine/Plugins/Runtime/nDisplay/Source/DisplayCluster/Private/Render/Viewport/DisplayClusterViewportProxy.cpp
--- a/Engine/Plugins/Runtime/nDisplay/Source/DisplayCluster/Private/Render/Viewport/DisplayClusterViewportProxy.cpp
+++ b/Engine/Plugins/Runtime/nDisplay/Source/DisplayCluster/Private/Render/Viewport/DisplayClusterViewportProxy.cpp
@@ -22,8 +22,11 @@
 #include "IDisplayCluster.h"
 #include "IDisplayClusterCallbacks.h"
 #include "IDisplayClusterShaders.h"
+#include "IDisplayClusterShadersTextureUtils.h"
 #include "TextureResource.h"
 
+#include "Engine/TextureRenderTarget.h"
+
 #include "RHIStaticStates.h"
 
 #include "RenderResource.h"
@@ -38,22 +41,19 @@
 #include "ScreenRendering.h"
 #include "RenderGraphUtils.h"
 #include "RenderGraphResources.h"
-#include "PostProcess/DrawRectangle.h"
 
 #include "ScreenPass.h"
 #include "SceneTextures.h"
 #include "PostProcess/PostProcessMaterialInputs.h"
 
+
 ///////////////////////////////////////////////////////////////////////////////////////
-namespace UE::DisplayCluster::ViewportProxy
-{
-	static IDisplayClusterShaders& GetShadersAPI()
+IDisplayClusterShaders& FDisplayClusterViewportProxy::GetShadersAPI()
 	{
 		static IDisplayClusterShaders& ShadersAPISingleton = IDisplayClusterShaders::Get();
 
 		return ShadersAPISingleton;
 	}
-};
 
 ///////////////////////////////////////////////////////////////////////////////////////
 FDisplayClusterViewportProxy::FDisplayClusterViewportProxy(const TSharedRef<FDisplayClusterViewportConfiguration, ESPMode::ThreadSafe>& InConfiguration, const FString& InViewportId, const TSharedPtr<IDisplayClusterProjectionPolicy, ESPMode::ThreadSafe>& InProjectionPolicy)
@@ -155,8 +155,6 @@ void FDisplayClusterViewportProxy::PostResolveViewport_RenderThread(FRHICommandL
 
 void FDisplayClusterViewportProxy::ImplViewportRemap_RenderThread(FRHICommandListImmediate& RHICmdList) const
 {
-	using namespace UE::DisplayCluster::ViewportProxy;
-
 	// Preview in editor not support this feature
 	if (ConfigurationProxy->IsPreviewRendering_RenderThread())
 	{
@@ -196,63 +194,8 @@ bool FDisplayClusterViewportProxy::GetResourcesWithRects_RenderThread(const EDis
 	return ImplGetResourcesWithRects_RenderThread(InExtResourceType, OutResources, OutResourceRects, false);
 }
 
-bool FDisplayClusterViewportProxy::ApplyOCIO_RenderThread(FRHICommandListImmediate& RHICmdList, const FDisplayClusterViewportProxy& InSrcViewportProxy, const EDisplayClusterViewportResourceType InSrcResourceType) const
-{
-	if (GetOpenColorIOMode() != EDisplayClusterViewportOpenColorIOMode::Resolved)
-	{
-		return false;
-	}
-
-	const EDisplayClusterViewportResourceType DestResourceType = (InSrcResourceType == EDisplayClusterViewportResourceType::InternalRenderTargetResource)
-		? EDisplayClusterViewportResourceType::InputShaderResource
-		: EDisplayClusterViewportResourceType::AdditionalTargetableResource;
-
-	TArray<FRHITexture*> Input, Output;
-	TArray<FIntRect> InputRects, OutputRects;
-	if (!InSrcViewportProxy.GetResourcesWithRects_RenderThread(InSrcResourceType, Input, InputRects)
-		|| !GetResourcesWithRects_RenderThread(DestResourceType, Output, OutputRects)
-		|| Input.Num() != Output.Num())
-	{
-		return false;
-	};
-
-	FRDGBuilder GraphBuilder(RHICmdList);
-
-	bool bResult = false;
-	for (int32 ContextNum = 0; ContextNum < Input.Num(); ContextNum++)
-	{
-		const bool bUnpremultiply = EnumHasAnyFlags(RenderSettingsICVFX.RuntimeFlags, EDisplayClusterViewportRuntimeICVFXFlags::UVLightcard | EDisplayClusterViewportRuntimeICVFXFlags::Lightcard);
-		const bool bInvertAlpha = !EnumHasAnyFlags(RenderSettingsICVFX.RuntimeFlags, EDisplayClusterViewportRuntimeICVFXFlags::UVLightcard);
-
-		if (OpenColorIO->AddPass_RenderThread(
-			GraphBuilder,
-			InSrcViewportProxy.GetContexts_RenderThread()[ContextNum],
-			Input[ContextNum], InputRects[ContextNum],
-			Output[ContextNum], OutputRects[ContextNum],
-			bUnpremultiply,
-			bInvertAlpha))
-		{
-			bResult = true;
-		}
-	}
-
-	if (bResult)
-	{
-		GraphBuilder.Execute();
-
-		// copy OCIO results back
-		if (DestResourceType != EDisplayClusterViewportResourceType::InputShaderResource)
-		{
-			ResolveResources_RenderThread(RHICmdList, DestResourceType, EDisplayClusterViewportResourceType::InputShaderResource);
-		}
-	}
-
-	return bResult;
-}
-
 void FDisplayClusterViewportProxy::UpdateDeferredResources(FRHICommandListImmediate& RHICmdList) const
 {
-	using namespace UE::DisplayCluster::ViewportProxy;
 	check(IsInRenderingThread());
 
 	if (RenderSettings.bFreezeRendering || RenderSettings.bSkipRendering)
@@ -304,49 +247,53 @@ void FDisplayClusterViewportProxy::UpdateDeferredResources(FRHICommandListImmedi
 		SrcResourceType = EDisplayClusterViewportResourceType::InputShaderResource;
 	}
 
-	// Pass 0 (OCIO): OCIO support on the first pass for an resolved RTT
-	if(!ApplyOCIO_RenderThread(RHICmdList, SourceViewportProxy, SrcResourceType))
-	{
-		// Pass 0 (default): Resolve from RTT region to separated viewport context resource:
-		if (SrcResourceType == EDisplayClusterViewportResourceType::InternalRenderTargetResource)
+	if (GetOpenColorIOMode() == EDisplayClusterViewportOpenColorIOMode::Resolved)
 		{
-			ImplResolveResources_RenderThread(RHICmdList, &SourceViewportProxy, SrcResourceType, EDisplayClusterViewportResourceType::InputShaderResource);
+		// Pass 0:  OCIO + Linear gamma
+		// At this point Resolver go to use RDG
+		OpenColorIO->AddPass_RenderThread(
+			FDisplayClusterShadersTextureUtilsSettings(),
+			GetShadersAPI().CreateTextureUtils_RenderThread(RHICmdList)
+				->SetInput(&SourceViewportProxy, SrcResourceType)
+				->SetOutput(this, EDisplayClusterViewportResourceType::InputShaderResource));
 		}
+	else
+	{
+		// Pass 0: Linear gamma
+		GetShadersAPI().CreateTextureUtils_RenderThread(RHICmdList)
+			->SetInput(&SourceViewportProxy, SrcResourceType)
+			->SetOutput(this, EDisplayClusterViewportResourceType::InputShaderResource)
+			->Resolve();
 	}
 
-	// Pass 1: Generate blur postprocess effect for render target texture rect for all contexts
+	// (Opt) Pass 1: Generate blur postprocess effect for render target texture rect for all contexts
 	if (PostRenderSettings.PostprocessBlur.IsEnabled())
 	{
-		TArray<FRHITexture*> InShaderResources;
-		TArray<FRHITexture*> OutTargetableResources;
-		if (GetResources_RenderThread(EDisplayClusterViewportResourceType::InputShaderResource, InShaderResources) && GetResources_RenderThread(EDisplayClusterViewportResourceType::AdditionalTargetableResource, OutTargetableResources))
-		{
-			// Render postprocess blur:
-			for (int32 ContextNum = 0; ContextNum < InShaderResources.Num(); ContextNum++)
-			{
-				GetShadersAPI().RenderPostprocess_Blur(RHICmdList, InShaderResources[ContextNum], OutTargetableResources[ContextNum], PostRenderSettings.PostprocessBlur);
-			}
-
-			// Copy result back to input
-			ResolveResources_RenderThread(RHICmdList, EDisplayClusterViewportResourceType::AdditionalTargetableResource, EDisplayClusterViewportResourceType::InputShaderResource);
-		}
+		GetShadersAPI().CreateTextureUtils_RenderThread(RHICmdList)
+			->SetOutput(this, EDisplayClusterViewportResourceType::InputShaderResource)
+			->ForEachContextByPredicate(EDisplayClusterShaderTextureUtilsFlags::UseOutputTextureAsInput,
+				[&](
+					const FDisplayClusterShadersTextureViewportContext& Input,
+					const FDisplayClusterShadersTextureViewportContext& Output)
+				{
+					GetShadersAPI().RenderPostprocess_Blur(RHICmdList, Input.TextureRHI, Output.TextureRHI, PostRenderSettings.PostprocessBlur);
+				});
 	}
 
 	// Pass 2: Create mips texture and generate mips from render target rect for all contexts
 	if (PostRenderSettings.GenerateMips.IsEnabled())
 	{
-		TArray<FRHITexture*> InOutMipsResources;
-		if (GetResources_RenderThread(EDisplayClusterViewportResourceType::MipsShaderResource, InOutMipsResources))
-		{
-			// Copy input image to layer0 on mips texture
-			ResolveResources_RenderThread(RHICmdList, EDisplayClusterViewportResourceType::InputShaderResource, EDisplayClusterViewportResourceType::MipsShaderResource);
-
-			// Generate mips
-			for (FRHITexture*& ResourceIt : InOutMipsResources)
+		GetShadersAPI().CreateTextureUtils_RenderThread(RHICmdList)
+			->SetInput(&SourceViewportProxy, EDisplayClusterViewportResourceType::InputShaderResource)
+			->SetOutput(this, EDisplayClusterViewportResourceType::MipsShaderResource)
+			->Resolve() // Copy `Input`->`Mips`
+			->ForEachContextByPredicate([&](
+				const FDisplayClusterShadersTextureViewportContext& Input,
+				const FDisplayClusterShadersTextureViewportContext& Output)
 			{
-				GetShadersAPI().GenerateMips(RHICmdList, ResourceIt, PostRenderSettings.GenerateMips);
-			}
-		}
+				// Generate mips
+				GetShadersAPI().GenerateMips(RHICmdList, Output.TextureRHI, PostRenderSettings.GenerateMips);
+			});
 	}
 }
 
@@ -376,7 +323,10 @@ void FDisplayClusterViewportProxy::OnResolvedSceneColor_RenderThread(FRDGBuilder
 		{
 			const FIntRect SrcRect = GetFinalContextRect(EDisplayClusterViewportResourceType::InternalRenderTargetResource, Contexts[InContextNum].RenderTargetRect);
 			// Copy alpha channel from 'SceneTextures.Color.Resolve' to 'InputShaderResource'
-			CopyResource_RenderThread(GraphBuilder, EDisplayClusterTextureCopyMode::Alpha, InContextNum, SceneTextures.Color.Resolve, SrcRect, EDisplayClusterViewportResourceType::InputShaderResource);
+			GetShadersAPI().CreateTextureUtils_RenderThread(GraphBuilder)
+				->SetInput({ SceneTextures.Color.Resolve, SrcRect }, InContextNum)
+				->SetOutput(this, EDisplayClusterViewportResourceType::InputShaderResource, InContextNum)
+				->Resolve(EColorWriteMask::CW_ALPHA);
 		}
 		break;
 
@@ -393,7 +343,10 @@ FScreenPassTexture FDisplayClusterViewportProxy::OnPostProcessPassAfterSSRInput_
 	{
 		// Copy alpha channel to 'InputShaderResource'
 		const FIntRect SrcRect = GetFinalContextRect(EDisplayClusterViewportResourceType::InternalRenderTargetResource, Contexts[ContextNum].RenderTargetRect);
-		CopyResource_RenderThread(GraphBuilder, EDisplayClusterTextureCopyMode::Alpha, ContextNum, OutScreenPassTexture.Texture, SrcRect, EDisplayClusterViewportResourceType::InputShaderResource);
+		GetShadersAPI().CreateTextureUtils_RenderThread(GraphBuilder)
+			->SetInput({ OutScreenPassTexture.Texture, SrcRect }, ContextNum)
+			->SetOutput(this, EDisplayClusterViewportResourceType::InputShaderResource, ContextNum)
+			->Resolve(EColorWriteMask::CW_ALPHA);
 	}
 
 	return OutScreenPassTexture;
@@ -407,7 +360,10 @@ FScreenPassTexture FDisplayClusterViewportProxy::OnPostProcessPassAfterFXAA_Rend
 		// Restore alpha channel after OCIO
 		// Copy alpha channel from 'InputShaderResource'
 		const FIntRect DestRect = GetFinalContextRect(EDisplayClusterViewportResourceType::InternalRenderTargetResource, Contexts[ContextNum].RenderTargetRect);
-		CopyResource_RenderThread(GraphBuilder, EDisplayClusterTextureCopyMode::Alpha, ContextNum, EDisplayClusterViewportResourceType::InputShaderResource, OutScreenPassTexture.Texture, DestRect);
+		GetShadersAPI().CreateTextureUtils_RenderThread(GraphBuilder)
+			->SetInput(this, EDisplayClusterViewportResourceType::InputShaderResource, ContextNum)
+			->SetOutput({ OutScreenPassTexture.Texture, DestRect }, ContextNum)
+			->Resolve(EColorWriteMask::CW_ALPHA);
 	}
 
 	return OutScreenPassTexture;
@@ -422,7 +378,7 @@ FScreenPassTexture FDisplayClusterViewportProxy::OnPostProcessPassAfterTonemap_R
 	if (GetOpenColorIOMode() == EDisplayClusterViewportOpenColorIOMode::PostProcess)
 	{
 		// Add OCIO pass
-		return OpenColorIO->PostProcessPassAfterTonemap_RenderThread(GraphBuilder, GetContexts_RenderThread()[ContextNum], View, Inputs);
+		return OpenColorIO->PostProcessPassAfterTonemap_RenderThread(GraphBuilder, GetResourceColorEncoding_RenderThread(EDisplayClusterViewportResourceType::InternalRenderTargetResource), View, Inputs);
 	}
 
 	return Inputs.ReturnUntouchedSceneColorForPostProcessing(GraphBuilder);
@@ -430,8 +386,6 @@ FScreenPassTexture FDisplayClusterViewportProxy::OnPostProcessPassAfterTonemap_R
 
 void FDisplayClusterViewportProxy::OnPostRenderViewFamily_RenderThread(FRDGBuilder& GraphBuilder, FSceneViewFamily& InViewFamily, const FSceneView& InSceneView, const FDisplayClusterViewportProxy_Context& InProxyContext)
 {
-	using namespace UE::DisplayCluster::ViewportProxy;
-
 	const uint32 InContextNum = InProxyContext.ContextNum;
 
 #if WITH_MGPU
@@ -462,56 +416,43 @@ void FDisplayClusterViewportProxy::OnPostRenderViewFamily_RenderThread(FRDGBuild
 
 	if (ShouldUseAlphaChannel_RenderThread())
 	{
-		switch (ConfigurationProxy->GetRenderFrameSettings().AlphaChannelCaptureMode)
+		const EDisplayClusterRenderFrameAlphaChannelCaptureMode AlphaChannelCaptureMode = ConfigurationProxy->GetRenderFrameSettings().AlphaChannelCaptureMode;
+		switch (AlphaChannelCaptureMode)
 		{
+		case EDisplayClusterRenderFrameAlphaChannelCaptureMode::Copy:
+		case EDisplayClusterRenderFrameAlphaChannelCaptureMode::CopyAA:
 		case EDisplayClusterRenderFrameAlphaChannelCaptureMode::FXAA:
 		{
-			// Restore alpha channed
+			// RenderPass 1: Copy Alpha channels back from 'InputShaderResource' to 'InternalRenderTargetResource'
+			TSharedRef<IDisplayClusterShadersTextureUtils> Resolver = GetShadersAPI().CreateTextureUtils_RenderThread(GraphBuilder)
+				->SetInput(this, EDisplayClusterViewportResourceType::InputShaderResource, InContextNum)
+				->SetOutput(this, EDisplayClusterViewportResourceType::InternalRenderTargetResource, InContextNum)
+				->Resolve(EColorWriteMask::CW_ALPHA);
+
 			EFXAAQuality FXAAQuality = EFXAAQuality::Q0;
-			if (ShouldApplyFXAA_RenderThread(FXAAQuality) && Resources[EDisplayClusterViewportResource::InputShaderResources].IsValidIndex(InContextNum))
+			if (AlphaChannelCaptureMode == EDisplayClusterRenderFrameAlphaChannelCaptureMode::FXAA &&
+				ShouldApplyFXAA_RenderThread(FXAAQuality))
 			{
-				if (FRHITexture* InputTextureRHI = Resources[EDisplayClusterViewportResource::InputShaderResources][InContextNum] ? Resources[EDisplayClusterViewportResource::InputShaderResources][InContextNum]->GetViewportResourceRHI_RenderThread() : nullptr)
-				{
-					// Apply FXAA for RGB only
-					// Note: Add AA for alpha channel
-
-					// Copy Alpha channels back from'InputShaderResource' to 'InternalRenderTargetResource'
-					CopyResource_RenderThread(GraphBuilder, EDisplayClusterTextureCopyMode::Alpha, InContextNum, EDisplayClusterViewportResourceType::InputShaderResource, EDisplayClusterViewportResourceType::InternalRenderTargetResource);
-
-					// 1. Copy RGB channels from 'InternalRenderTargetResource' to 'InputShaderResource'
-					CopyResource_RenderThread(GraphBuilder, EDisplayClusterTextureCopyMode::RGB, InContextNum, EDisplayClusterViewportResourceType::InternalRenderTargetResource, EDisplayClusterViewportResourceType::InputShaderResource);
-
-					// 2. FXAA render pass to 'InputShaderResource'
-					// Input is 'InputShaderResource'
-					FRDGTextureRef InputTexture = RegisterExternalTexture(GraphBuilder, InputTextureRHI, TEXT("DCViewportProxyFXAAResource"));
-					GraphBuilder.SetTextureAccessFinal(InputTexture, ERHIAccess::RTV);
-
-					FFXAAInputs PassInputs;
-					PassInputs.SceneColor = FScreenPassTexture(InputTexture);
-					PassInputs.Quality = FXAAQuality;
-
-					// 2.1. Do FXAA
-					FScreenPassTexture OutputColorTexture = AddFXAAPass(GraphBuilder, InSceneView, PassInputs);
-
-					// 2.2. Copy FXAA result from 'OutputTexture' to the 'InternalRenderTargetResource'
-					if (OutputColorTexture.Texture)
-					{
-						const FIntVector OutputTextureSize = OutputColorTexture.Texture->Desc.GetSize();
-						const FIntRect SrcRect(FIntPoint(0, 0), FIntPoint(OutputTextureSize.X, OutputTextureSize.Y));
-
-						CopyResource_RenderThread(GraphBuilder, EDisplayClusterTextureCopyMode::RGB, InContextNum, OutputColorTexture.Texture, SrcRect, EDisplayClusterViewportResourceType::InternalRenderTargetResource);
-					}
-				}
-				break;
+				Resolver
+					->ForEachContextByPredicate([&](
+						const FDisplayClusterShadersTextureViewportContext& Input,
+						const FDisplayClusterShadersTextureViewportContext& Output)
+						{
+							// RenderPass 2: Do FXAA with 'InternalRenderTargetResource' as input
+							FFXAAInputs PassInputs;
+							PassInputs.SceneColor = Output.ToScreenPassTexture();
+							PassInputs.Quality = FXAAQuality;
+							FScreenPassTexture FXAAColorTexture = AddFXAAPass(GraphBuilder, InSceneView, PassInputs);
+
+							// RenderPass 3: Copy FXAA result(RGB) back to the 'InternalRenderTargetResource'
+							Resolver->ResolveTextureContext(
+								EColorWriteMask::CW_RGB,
+								{ {FXAAColorTexture}, Output.ColorEncoding },
+								Output);
+						});
 			}
-			// don't break (FXAA not used, just copy alpha)
 		}
-
-		case EDisplayClusterRenderFrameAlphaChannelCaptureMode::Copy:
-		case EDisplayClusterRenderFrameAlphaChannelCaptureMode::CopyAA:
-			// Copy Alpha channels back from'InputShaderResource' to 'InternalRenderTargetResource'
-			CopyResource_RenderThread(GraphBuilder, EDisplayClusterTextureCopyMode::Alpha, InContextNum, EDisplayClusterViewportResourceType::InputShaderResource, EDisplayClusterViewportResourceType::InternalRenderTargetResource);
-			break;
+		break;
 
 		default:
 			break;
@@ -523,3 +464,74 @@ void FDisplayClusterViewportProxy::ReleaseTextures_RenderThread()
 {
 	Resources.ReleaseAllResources();
 }
+
+FDisplayClusterColorEncoding FDisplayClusterViewportProxy::GetResourceColorEncoding_RenderThread(const EDisplayClusterViewportResourceType InResourceType) const
+{
+	// Color modifiers:
+	EDisplayClusterColorPremultiply Premultiply = EDisplayClusterColorPremultiply::None;
+	if (EnumHasAnyFlags(RenderSettingsICVFX.RuntimeFlags, EDisplayClusterViewportRuntimeICVFXFlags::UVLightcard | EDisplayClusterViewportRuntimeICVFXFlags::Lightcard))
+	{
+		// All LightCard textures are premultiplied with inverted alpha.
+		Premultiply = EDisplayClusterColorPremultiply::InvertPremultiply;
+	}
+		
+	const EDisplayClusterViewportResourceType ResourceType = GetResourceType_RenderThread(InResourceType);
+	switch (ResourceType)
+	{
+	case EDisplayClusterViewportResourceType::InternalRenderTargetResource:
+		if (EnumHasAllFlags(RenderSettingsICVFX.RuntimeFlags, EDisplayClusterViewportRuntimeICVFXFlags::UVLightcard))
+		{
+			// The UVLightCard viewport is not actually rendered, so dont use context data
+			// This external texture will always be in linear color space.
+			return { EDisplayClusterColorEncoding::Linear, Premultiply };
+		}
+
+		if (!Contexts.IsEmpty())
+		{
+			// There is a special case where post processing and tonemapper are disabled. In this case tonemapper applies a static display Inverse of Gamma which defaults to 2.2.
+			const FEngineShowFlags& EngineShowFlags = Contexts[0].RenderThreadData.EngineShowFlags;
+			if (EngineShowFlags.Tonemapper == 0 || EngineShowFlags.PostProcessing == 0)
+			{
+				return { EDisplayClusterColorEncoding::Linear, Premultiply };
+			}
+
+			const float DefaultDisplayGamma = UTextureRenderTarget::GetDefaultDisplayGamma();
+			const float DisplayGamma = Contexts[0].RenderThreadData.EngineDisplayGamma;
+			if (DisplayGamma == DefaultDisplayGamma)
+			{
+				return { EDisplayClusterColorEncoding::Gamma,Premultiply };
+			}
+
+			// Custom gamma value is different from default
+			return { DisplayGamma, Premultiply };
+		}
+
+		break;
+
+	case EDisplayClusterViewportResourceType::InputShaderResource:
+	case EDisplayClusterViewportResourceType::MipsShaderResource:
+		if (EnumHasAnyFlags(RenderSettingsICVFX.RuntimeFlags, EDisplayClusterViewportRuntimeICVFXFlags::UVLightcard | EDisplayClusterViewportRuntimeICVFXFlags::Lightcard))
+		{
+			// ICVFX expects all lightcards in linear color space
+			return { EDisplayClusterColorEncoding::Linear, Premultiply };
+		}
+		break;
+
+		// Preview Output
+	case EDisplayClusterViewportResourceType::OutputPreviewTargetableResource:
+		if (ConfigurationProxy->GetRenderFrameSettings().ShouldUseHoldout())
+		{
+			// The HoldoutComposite plugin expects the input in linear gamma.
+			return { EDisplayClusterColorEncoding::Linear, Premultiply };
+		}
+
+		break;
+	}
+
+	if (ConfigurationProxy->GetRenderFrameSettings().IsPostProcessDisabled())
+	{
+		return { EDisplayClusterColorEncoding::Linear, Premultiply };
+	}
+
+	return { EDisplayClusterColorEncoding::Gamma, Premultiply };
+}
\ No newline at end of file
diff --git a/Engine/Plugins/Runtime/nDisplay/Source/DisplayCluster/Private/Render/Viewport/DisplayClusterViewportProxy.h b/Engine/Plugins/Runtime/nDisplay/Source/DisplayCluster/Private/Render/Viewport/DisplayClusterViewportProxy.h
--- a/Engine/Plugins/Runtime/nDisplay/Source/DisplayCluster/Private/Render/Viewport/DisplayClusterViewportProxy.h
+++ b/Engine/Plugins/Runtime/nDisplay/Source/DisplayCluster/Private/Render/Viewport/DisplayClusterViewportProxy.h
@@ -28,16 +28,6 @@ struct FPostProcessMaterialInputs;
 struct FScreenPassTexture;
 struct FSceneTextures;
 
-/**
- * Copy mode for texture channels.
- */
-enum class EDisplayClusterTextureCopyMode : uint8
-{
-	RGBA = 0,
-	RGB,
-	Alpha,
-};
-
 /**
  * OCIO is applied in different ways. It depends on the rendering workflow
  */
@@ -135,6 +125,8 @@ public:
 		Contexts.Append(InContexts);
 	}
 
+	virtual FDisplayClusterColorEncoding GetResourceColorEncoding_RenderThread(const EDisplayClusterViewportResourceType InResourceType) const override;
+
 	//  Return viewport scene proxy resources by type
 	virtual bool GetResources_RenderThread(const EDisplayClusterViewportResourceType InResourceType, TArray<FRHITexture*>& OutResources) const override;
 	virtual bool GetResourcesWithRects_RenderThread(const EDisplayClusterViewportResourceType InResourceType, TArray<FRHITexture*>& OutResources, TArray<FIntRect>& OutRects) const override;
@@ -301,6 +293,9 @@ public:
 	/** Returns the OCIO rendering type for the given viewport. */
 	EDisplayClusterViewportOpenColorIOMode GetOpenColorIOMode() const;
 
+	/** Return shaders API. */
+	static class IDisplayClusterShaders& GetShadersAPI();
+
 private:
 	bool ImplGetResourcesWithRects_RenderThread(const EDisplayClusterViewportResourceType InResourceType, TArray<FRHITexture*>& OutResources, TArray<FIntRect>& OutResourceRects, const int32 InRecursionDepth) const;
 	bool ImplGetResources_RenderThread(const EDisplayClusterViewportResourceType InResourceType, TArray<FRHITexture*>& OutResources, const int32 InRecursionDepth) const;
@@ -321,52 +316,12 @@ private:
 	/** When a resource by type can be overridden from another viewport, true is returned. */
 	bool ShouldOverrideViewportResource(const EDisplayClusterViewportResourceType InResourceType) const;
 
-	/** Copy pixel channels between resources.
-	 *
-	 * @param GraphBuilder       - RDG interface
-	 * @param InCopyMode         - channels copy mode
-	 * @param InContextNum       - viewport context
-	 * @param InSrcTextureRef    - Input resource
-	 * @param InSrcRect          - Input resource rect
-	 * @param InDestResourceType - Dest viewport resource type
-	 *
-	 * @return - true, if success
-	 */
-	bool CopyResource_RenderThread(FRDGBuilder& GraphBuilder, const EDisplayClusterTextureCopyMode InCopyMode, const int32 InContextNum, FRDGTextureRef InSrcTextureRef, const FIntRect& InSrcRect, const EDisplayClusterViewportResourceType InDestResourceType);
-
-	/** Copy pixel channels between resources.
-	 *
-	 * @param GraphBuilder      - RDG interface
-	 * @param InCopyMode        - channels copy mode
-	 * @param InContextNum      - viewport context
-	 * @param InDestTextureRef  - Dest resource
-	 * @param InDestRect        - Dest resource rect
-	 *
-	 * @return - true, if success
-	 */
-	bool CopyResource_RenderThread(FRDGBuilder& GraphBuilder, const EDisplayClusterTextureCopyMode InCopyMode, const int32 InContextNum, const EDisplayClusterViewportResourceType InSrcResourceType, FRDGTextureRef InDestTextureRef, const FIntRect& InDestRect);
-
-	/** Copy pixel channels between resources.
-	 *
-	 * @param GraphBuilder       - RDG interface
-	 * @param InCopyMode         - channels copy mode
-	 * @param InContextNum       - viewport context
-	 * @param InSrcResourceType  - Input viewport resource type
-	 * @param InDestResourceType - Dest viewport resource type
-	 *
-	 * @return - true, if success
-	 */
-	bool CopyResource_RenderThread(FRDGBuilder& GraphBuilder, const EDisplayClusterTextureCopyMode InCopyMode, const int32 InContextNum, const EDisplayClusterViewportResourceType InSrcResourceType, const EDisplayClusterViewportResourceType InDestResourceType);
-
 	/** Return viewport used to render RTT (support ViewportOverride)
 	 *
 	 * @return - viewport proxy with valid render resources
 	 */
 	const FDisplayClusterViewportProxy& GetRenderingViewportProxy() const;
 
-	/** Apply OCIO to the image from the 'InSrcResourceType' and put the result image to the EDisplayClusterViewportResourceType::InputShaderResource. */
-	bool ApplyOCIO_RenderThread(FRHICommandListImmediate& RHICmdList, const FDisplayClusterViewportProxy& InSrcViewportProxy, const EDisplayClusterViewportResourceType InSrcResourceType) const;
-
 	/** Check if there is an RTT source (internal or external) in this viewport proxy. */
 	bool IsInputRenderTargetResourceExists() const;
 
diff --git a/Engine/Plugins/Runtime/nDisplay/Source/DisplayCluster/Private/Render/Viewport/DisplayClusterViewportProxy_Resources.cpp b/Engine/Plugins/Runtime/nDisplay/Source/DisplayCluster/Private/Render/Viewport/DisplayClusterViewportProxy_Resources.cpp
--- a/Engine/Plugins/Runtime/nDisplay/Source/DisplayCluster/Private/Render/Viewport/DisplayClusterViewportProxy_Resources.cpp
+++ b/Engine/Plugins/Runtime/nDisplay/Source/DisplayCluster/Private/Render/Viewport/DisplayClusterViewportProxy_Resources.cpp
@@ -23,6 +23,7 @@
 #include "IDisplayCluster.h"
 #include "IDisplayClusterCallbacks.h"
 #include "IDisplayClusterShaders.h"
+#include "IDisplayClusterShadersTextureUtils.h"
 #include "TextureResource.h"
 
 #include "RHIStaticStates.h"
@@ -41,11 +42,6 @@
 #include "ScreenRendering.h"
 #include "RenderGraphUtils.h"
 #include "RenderGraphResources.h"
-#include "PostProcess/DrawRectangle.h"
-
-#include "ScreenPass.h"
-#include "SceneTextures.h"
-#include "PostProcess/PostProcessMaterialInputs.h"
 
 // Tile rect border width
 int32 GDisplayClusterRenderTileBorder = 0;
@@ -61,155 +57,6 @@ namespace UE::DisplayCluster::ViewportProxy
 {
 	// The viewport override has the maximum depth. This protects against a link cycle
 	static const int32 DisplayClusterViewportProxyResourcesOverrideRecursionDepthMax = 4;
-
-	template<class TScreenPixelShader>
-	void ResampleCopyTextureImpl_RenderThread(FRHICommandList& RHICmdList, FRHITexture* SrcTexture, FRHITexture* DstTexture, const FIntRect& SrcRect, const FIntRect& DstRect, const EDisplayClusterTextureCopyMode InCopyMode = EDisplayClusterTextureCopyMode::RGBA)
-	{
-		// Texture format mismatch, use a shader to do the copy.
-		// #todo-renderpasses there's no explicit resolve here? Do we need one?
-		FRHIRenderPassInfo RPInfo(DstTexture, ERenderTargetActions::Load_Store);
-		RHICmdList.Transition(FRHITransitionInfo(DstTexture, ERHIAccess::Unknown, ERHIAccess::RTV));
-
-		RHICmdList.BeginRenderPass(RPInfo, TEXT("DisplayCluster_ResampleCopyTexture"));
-		{
-			FIntVector SrcSizeXYZ = SrcTexture->GetSizeXYZ();
-			FIntVector DstSizeXYZ = DstTexture->GetSizeXYZ();
-
-			FIntPoint SrcSize(SrcSizeXYZ.X, SrcSizeXYZ.Y);
-			FIntPoint DstSize(DstSizeXYZ.X, DstSizeXYZ.Y);
-
-			RHICmdList.SetViewport(0.f, 0.f, 0.0f, DstSize.X, DstSize.Y, 1.0f);
-
-			FGraphicsPipelineStateInitializer GraphicsPSOInit;
-			RHICmdList.ApplyCachedRenderTargets(GraphicsPSOInit);
-
-			switch (InCopyMode)
-			{
-			case EDisplayClusterTextureCopyMode::Alpha:
-				// Copy alpha channel from source to dest
-				GraphicsPSOInit.BlendState = TStaticBlendState <CW_ALPHA, BO_Add, BF_One, BF_Zero, BO_Add, BF_One, BF_Zero>::GetRHI();
-				break;
-
-			case EDisplayClusterTextureCopyMode::RGB:
-				// Copy only RGB channels from source to dest
-				GraphicsPSOInit.BlendState = TStaticBlendState <CW_RGB, BO_Add, BF_One, BF_Zero, BO_Add, BF_One, BF_Zero>::GetRHI();
-				break;
-
-			case EDisplayClusterTextureCopyMode::RGBA:
-			default:
-				GraphicsPSOInit.BlendState = TStaticBlendState<>::GetRHI();
-				break;
-			}
-
-			GraphicsPSOInit.RasterizerState = TStaticRasterizerState<>::GetRHI();
-			GraphicsPSOInit.DepthStencilState = TStaticDepthStencilState<false, CF_Always>::GetRHI();
-
-			FGlobalShaderMap* ShaderMap = GetGlobalShaderMap(GMaxRHIFeatureLevel);
-			TShaderMapRef<FScreenVS> VertexShader(ShaderMap);
-			TShaderMapRef<TScreenPixelShader> PixelShader(ShaderMap);
-			if (!VertexShader.IsValid() || !PixelShader.IsValid())
-			{
-				// Always check if shaders are available on the current platform and hardware
-				return;
-			}
-
-			GraphicsPSOInit.BoundShaderState.VertexDeclarationRHI = GFilterVertexDeclaration.VertexDeclarationRHI;
-			GraphicsPSOInit.BoundShaderState.VertexShaderRHI = VertexShader.GetVertexShader();
-			GraphicsPSOInit.BoundShaderState.PixelShaderRHI = PixelShader.GetPixelShader();
-			GraphicsPSOInit.PrimitiveType = PT_TriangleList;
-
-			SetGraphicsPipelineState(RHICmdList, GraphicsPSOInit, 0);
-
-			const bool bSameSize = (SrcRect.Size() == DstRect.Size());
-			FRHISamplerState* PixelSampler = bSameSize ? TStaticSamplerState<SF_Point>::GetRHI() : TStaticSamplerState<SF_Bilinear>::GetRHI();
-
-			SetShaderParametersLegacyPS(RHICmdList, PixelShader, PixelSampler, SrcTexture);
-
-			// Set up vertex uniform parameters for scaling and biasing the rectangle.
-			// Note: Use DrawRectangle in the vertex shader to calculate the correct vertex position and uv.
-
-			UE::Renderer::PostProcess::DrawRectangle(
-				RHICmdList, VertexShader,
-				DstRect.Min.X, DstRect.Min.Y,
-				DstRect.Size().X, DstRect.Size().Y,
-				SrcRect.Min.X, SrcRect.Min.Y,
-				SrcRect.Size().X, SrcRect.Size().Y,
-				DstSize, SrcSize
-			);
-		}
-		RHICmdList.EndRenderPass();
-		RHICmdList.Transition(FRHITransitionInfo(DstTexture, ERHIAccess::Unknown, ERHIAccess::SRVMask));
-	}
-
-	static void ImplResolveResource(FRHICommandListImmediate& RHICmdList, FRHITexture* InputResource, const FIntRect& InputRect, FRHITexture* OutputResource, const FIntRect& OutputRect, const bool bOutputIsMipsResource, const bool bOutputIsPreviewResource)
-	{
-		check(InputResource);
-		check(OutputResource);
-
-		// Check if resources with the specified regions can be resolved.
-		FIntRect SrcRect(InputRect);
-		FIntRect DestRect(OutputRect);
-		if(!FDisplayClusterViewportHelpers::GetValidResourceRectsForResolve(InputResource, OutputResource, SrcRect, DestRect))
-		{
-			// The SrcRect or DestRect is invalid.
-			return;
-		}
-
-		if (bOutputIsPreviewResource)
-		{
-			// The preview texture should use only RGB colors and ignore the alpha channel. The alpha channel may or may not be inverted in third-party libraries.
-			ResampleCopyTextureImpl_RenderThread<FScreenPS>(RHICmdList, InputResource, OutputResource, SrcRect, DestRect, EDisplayClusterTextureCopyMode::RGB);
-		}
-		else if (SrcRect.Size() == DestRect.Size() && InputResource->GetFormat() == OutputResource->GetFormat())
-		{
-			FRHICopyTextureInfo CopyInfo;
-			CopyInfo.SourcePosition = FIntVector(SrcRect.Min.X, SrcRect.Min.Y, 0);
-			CopyInfo.DestPosition = FIntVector(DestRect.Min.X, DestRect.Min.Y, 0);
-
-			CopyInfo.Size = FIntVector(DestRect.Width(), DestRect.Height(), 0);
-
-			TransitionAndCopyTexture(RHICmdList, InputResource, OutputResource, CopyInfo);
-		}
-		else
-		{
-			ResampleCopyTextureImpl_RenderThread<FScreenPS>(RHICmdList, InputResource, OutputResource, SrcRect, DestRect);
-		}
-	}
-
-	struct FViewportResourceResolverData
-	{
-		FViewportResourceResolverData(const FDisplayClusterViewport_Context& InViewportContext, FRHITexture* InInputResource, const FIntRect& InInputRect, FRHITexture* InOutputResource, const FIntRect& InOutputRect, const bool InbOutputIsMipsResource, const bool InbOutputIsPreviewResource)
-			: InputResource(InInputResource)
-			, OutputResource(InOutputResource)
-			, InputRect(InInputRect)
-			, OutputRect(InOutputRect)
-			, bOutputIsMipsResource(InbOutputIsMipsResource)
-			, bOutputIsPreviewResource(InbOutputIsPreviewResource)
-			, ViewportContext(InViewportContext)
-		{ }
-
-		void Resolve_RenderThread(FRHICommandListImmediate& RHICmdList)
-		{
-			ImplResolveResource(RHICmdList, InputResource, InputRect, OutputResource, OutputRect, bOutputIsMipsResource, bOutputIsPreviewResource);
-		}
-
-		void AddFinalPass_RenderThread(FRDGBuilder& GraphBuilder, IDisplayClusterDisplayDeviceProxy& DisplayDevice)
-		{
-			DisplayDevice.AddFinalPass_RenderThread(GraphBuilder, ViewportContext, InputResource, InputRect, OutputResource, OutputRect);
-		}
-
-	private:
-		FRHITexture* InputResource;
-		FRHITexture* OutputResource;
-
-		const FIntRect InputRect;
-		const FIntRect OutputRect;
-
-		const bool bOutputIsMipsResource;
-		const bool bOutputIsPreviewResource;
-
-		const FDisplayClusterViewport_Context ViewportContext;
-	};
 };
 
 void FDisplayClusterViewportProxy::FillTextureWithColor_RenderThread(FRHICommandListImmediate& RHICmdList, FRHITexture* InRenderTargetTexture, const FLinearColor& InColor)
@@ -414,212 +261,107 @@ bool FDisplayClusterViewportProxy::ImplResolveResources_RenderThread(FRHICommand
 		return false;
 	}
 
-	const bool bOutputIsMipsResource = OutResourceType == EDisplayClusterViewportResourceType::MipsShaderResource;
-	const bool bOutputIsPreviewResource = OutResourceType == EDisplayClusterViewportResourceType::OutputPreviewTargetableResource;
-
-	TArray<FViewportResourceResolverData> ResourceResolverData;
-
-	TArray<FRHITexture*> SrcResources, DestResources;
-	TArray<FIntRect> SrcResourcesRect, DestResourcesRect;
-	if (SourceProxy->GetResourcesWithRects_RenderThread(InExtResourceType, SrcResources, SrcResourcesRect) && GetResourcesWithRects_RenderThread(OutExtResourceType, DestResources, DestResourcesRect))
+	FDisplayClusterShadersTextureUtilsSettings TextureUtilsSettings;
+	// The mode used to blend textures
+	if (OutResourceType == EDisplayClusterViewportResourceType::OutputPreviewTargetableResource)
 	{
-		const int32 SrcResourcesAmount = FMath::Min(SrcResources.Num(), DestResources.Num());
-		for (int32 SrcResourceContextIndex = 0; SrcResourceContextIndex < SrcResourcesAmount; SrcResourceContextIndex++)
-		{
-			// Get context num to copy
-			const int32 SrcContextNum = (InContextNum == INDEX_NONE) ? SrcResourceContextIndex : InContextNum;
-			const FIntRect SrcRect = SourceProxy->GetFinalContextRect(InExtResourceType, SrcResourcesRect[SrcContextNum]);
-
-			if ((SrcContextNum + 1) == SrcResourcesAmount)
-			{
-				// last input mono -> stereo outputs
-				for (int32 DestResourceContextIndex = SrcContextNum; DestResourceContextIndex < DestResources.Num(); DestResourceContextIndex++)
-				{
-					const FIntRect DestRect = GetFinalContextRect(OutResourceType, DestResourcesRect[DestResourceContextIndex]);
-					ResourceResolverData.Add(FViewportResourceResolverData(Contexts[SrcContextNum], SrcResources[SrcContextNum], SrcRect, DestResources[DestResourceContextIndex], DestRect, bOutputIsMipsResource, bOutputIsPreviewResource));
-				}
-				break;
-			}
-			else
-			{
-				const FIntRect DestRect = GetFinalContextRect(OutResourceType, DestResourcesRect[SrcContextNum]);
-				ResourceResolverData.Add(FViewportResourceResolverData(Contexts[SrcContextNum], SrcResources[SrcContextNum], SrcRect, DestResources[SrcContextNum], DestRect, bOutputIsMipsResource, bOutputIsPreviewResource));
-			}
-
-			if (InContextNum != INDEX_NONE)
-			{
-				// Copy only one texture
-				break;
-			}
-		}
-
-		if (InExtResourceType == EDisplayClusterViewportResourceType::AfterWarpBlendTargetableResource
-			&& OutExtResourceType == EDisplayClusterViewportResourceType::OutputTargetableResource
-			&& DisplayDeviceProxy.IsValid()
-			&& DisplayDeviceProxy->HasFinalPass_RenderThread())
-		{
-			// Custom resolve at external Display Device
-			FRDGBuilder GraphBuilder(RHICmdList);
-			for (FViewportResourceResolverData& ResourceResolverIt : ResourceResolverData)
-			{
-				ResourceResolverIt.AddFinalPass_RenderThread(GraphBuilder, *DisplayDeviceProxy.Get());
-			}
-			GraphBuilder.Execute();
-		}
-		else
-		{
-			// Standard resolve:
-			for (FViewportResourceResolverData& ResourceResolverIt : ResourceResolverData)
-			{
-				ResourceResolverIt.Resolve_RenderThread(RHICmdList);
-			}
-		}
-
-		return true;
+		// The preview texture should use only RGB colors and ignore the alpha channel.
+		// The alpha channel may or may not be inverted in third-party libraries.
+		TextureUtilsSettings.OverrideAlpha = EDisplayClusterShaderTextureUtilsOverrideAlpha::Set_Alpha_One;
 	}
 
-	return false;
-}
+	TSharedRef<IDisplayClusterShadersTextureUtils> TextureUtils =
+		GetShadersAPI().CreateTextureUtils_RenderThread(RHICmdList)
+			->SetInput(SourceProxy, InExtResourceType)
+			->SetOutput(this, OutExtResourceType);
 
-bool FDisplayClusterViewportProxy::CopyResource_RenderThread(FRDGBuilder& GraphBuilder, const EDisplayClusterTextureCopyMode InCopyMode, const int32 InContextNum, const EDisplayClusterViewportResourceType InExtSrcResourceType, const EDisplayClusterViewportResourceType InExtDestResourceType)
-{
-	TArray<FRHITexture*> SrcResources;
-	TArray<FIntRect> SrcResourceRects;
-	if (GetResourcesWithRects_RenderThread(InExtSrcResourceType, SrcResources, SrcResourceRects))
+	if (InExtResourceType == EDisplayClusterViewportResourceType::AfterWarpBlendTargetableResource
+		&& OutExtResourceType == EDisplayClusterViewportResourceType::OutputTargetableResource
+		&& DisplayDeviceProxy.IsValid()
+		&& DisplayDeviceProxy->HasFinalPass_RenderThread())
 	{
-		check(SrcResources.IsValidIndex(InContextNum));
-		check(SrcResourceRects.IsValidIndex(InContextNum));
-
-		const FIntRect SrcRect = SrcResourceRects[InContextNum];
-
-		FRDGTextureRef SrcTextureRef = RegisterExternalTexture(GraphBuilder, SrcResources[InContextNum]->GetTexture2D(), TEXT("DCViewportProxyCopySrcResource"));
-		GraphBuilder.SetTextureAccessFinal(SrcTextureRef, ERHIAccess::SRVGraphics);
-
-		return CopyResource_RenderThread(GraphBuilder, InCopyMode, InContextNum, SrcTextureRef, SrcRect, InExtDestResourceType);
+		// Custom resolve at external Display Device
+		DisplayDeviceProxy->AddFinalPass_RenderThread(TextureUtilsSettings, TextureUtils);
+	}
+	else
+	{
+		// Standard resolve:
+		TextureUtils->Resolve(TextureUtilsSettings);
 	}
 
-	return false;
+	return true;
 }
 
-BEGIN_SHADER_PARAMETER_STRUCT(FDisplayClusterCopyTextureParameters, )
-RDG_TEXTURE_ACCESS(Input, ERHIAccess::CopySrc)
-RDG_TEXTURE_ACCESS(Output, ERHIAccess::CopyDest)
-END_SHADER_PARAMETER_STRUCT()
+void FDisplayClusterViewportProxy::ImplResolveTileResource_RenderThread(FRHICommandListImmediate& RHICmdList, FDisplayClusterViewportProxy const* InDestViewportProxy) const
+	{
+	check(InDestViewportProxy);
 
-bool FDisplayClusterViewportProxy::CopyResource_RenderThread(FRDGBuilder& GraphBuilder, const EDisplayClusterTextureCopyMode InCopyMode, const int32 InContextNum, FRDGTextureRef InSrcTextureRef, const FIntRect& InSrcRect, const EDisplayClusterViewportResourceType InExtDestResourceType)
-{
 	using namespace UE::DisplayCluster::ViewportProxy;
 
-	TArray<FRHITexture*> DestResources;
-	TArray<FIntRect> DestResourceRects;
-	if (HasBeenProduced(InSrcTextureRef) && GetResourcesWithRects_RenderThread(InExtDestResourceType, DestResources, DestResourceRects))
-	{
-		check(DestResources.IsValidIndex(InContextNum));
-		check(DestResourceRects.IsValidIndex(InContextNum));
-
-		FRDGTextureRef DestTextureRef = RegisterExternalTexture(GraphBuilder, DestResources[InContextNum]->GetTexture2D(), TEXT("DCViewportProxyCopyDestResource"));
-		const FIntRect DestRect = DestResourceRects[InContextNum];
-
-		FDisplayClusterCopyTextureParameters* PassParameters = GraphBuilder.AllocParameters<FDisplayClusterCopyTextureParameters>();
-		PassParameters->Input = InSrcTextureRef;
-		PassParameters->Output = DestTextureRef;
-		GraphBuilder.SetTextureAccessFinal(DestTextureRef, ERHIAccess::RTV);
-
-		GraphBuilder.AddPass(
-			RDG_EVENT_NAME("DisplayClusterViewportProxy_CopyResource(%s)", *GetId()),
-			PassParameters,
-			ERDGPassFlags::Copy | ERDGPassFlags::NeverCull,
-			[InCopyMode, InContextNum, InSrcTextureRef, InSrcRect, DestTextureRef, DestRect](FRDGAsyncTask, FRHICommandList& RHICmdList)
-			{
-				ResampleCopyTextureImpl_RenderThread<FScreenPS>(RHICmdList, InSrcTextureRef->GetRHI(), DestTextureRef->GetRHI(), InSrcRect, DestRect, InCopyMode);
-			});
+	const TArray<FDisplayClusterViewport_Context>& DestContexts = InDestViewportProxy->GetContexts_RenderThread();
+	const TArray<FDisplayClusterViewport_Context>& SrcContexts = Contexts;
+
+	TSharedRef<IDisplayClusterShadersTextureUtils> TextureUtils = GetShadersAPI().CreateTextureUtils_RenderThread(RHICmdList)
+		->SetInput(this, EDisplayClusterViewportResourceType::InternalRenderTargetResource)
+		->SetOutput(InDestViewportProxy, EDisplayClusterViewportResourceType::InternalRenderTargetResource);
 
-		return true;
+	TextureUtils
+		->ForEachContextByPredicate(EDisplayClusterShaderTextureUtilsFlags::DisableUpdateResourcesRectsForResolve, [&](
+			const FDisplayClusterShadersTextureViewportContext& InputContext,
+			const FDisplayClusterShadersTextureViewportContext& OutputContext)
+			{
+				const uint32 ContextNum = InputContext.ContextNum;
+				if(!SrcContexts.IsValidIndex(ContextNum) || !DestContexts.IsValidIndex(ContextNum) || OutputContext.ContextNum != ContextNum)
+				{
+					return;
 	}
 
-	return false;
-}
+				FDisplayClusterShadersTextureViewport Input(InputContext);
+				FDisplayClusterShadersTextureViewport Output(OutputContext);
 
-bool FDisplayClusterViewportProxy::CopyResource_RenderThread(FRDGBuilder& GraphBuilder, const EDisplayClusterTextureCopyMode InCopyMode, const int32 InContextNum, const EDisplayClusterViewportResourceType InExtSrcResourceType, FRDGTextureRef InDestTextureRef, const FIntRect& InDestRect)
-{
-	using namespace UE::DisplayCluster::ViewportProxy;
+				Input.Rect  = GetFinalContextRect(EDisplayClusterViewportResourceType::InternalRenderTargetResource, Input.Rect);
+				Output.Rect = SrcContexts[ContextNum].TileDestRect;
 
-	TArray<FRHITexture*> SrcResources;
-	TArray<FIntRect> SrcResourceRects;
-	if (HasBeenProduced(InDestTextureRef) && GetResourcesWithRects_RenderThread(InExtSrcResourceType, SrcResources, SrcResourceRects))
+				if (GDisplayClusterRenderTileBorder > 0)
 	{
-		check(SrcResources.IsValidIndex(InContextNum));
-		check(SrcResourceRects.IsValidIndex(InContextNum));
-
-		FRDGTextureRef SrcTextureRef = RegisterExternalTexture(GraphBuilder, SrcResources[InContextNum]->GetTexture2D(), TEXT("DCViewportProxyCopySrcResource"));
-		const FIntRect SrcRect = SrcResourceRects[InContextNum];
-
-		FDisplayClusterCopyTextureParameters* PassParameters = GraphBuilder.AllocParameters<FDisplayClusterCopyTextureParameters>();
-		PassParameters->Input = SrcTextureRef;
-		PassParameters->Output = InDestTextureRef;
-		GraphBuilder.SetTextureAccessFinal(SrcTextureRef, ERHIAccess::RTV);
-
-		GraphBuilder.AddPass(
-			RDG_EVENT_NAME("DisplayClusterViewportProxy_CopyResource(%s)", *GetId()),
-			PassParameters,
-			ERDGPassFlags::Copy | ERDGPassFlags::NeverCull,
-			[InCopyMode, InContextNum, SrcTextureRef, SrcRect, InDestTextureRef, InDestRect](FRDGAsyncTask, FRHICommandList& RHICmdList)
-			{
-				ResampleCopyTextureImpl_RenderThread<FScreenPS>(RHICmdList, SrcTextureRef->GetRHI(), InDestTextureRef->GetRHI(), SrcRect, InDestRect, InCopyMode);
-			});
-
-		return true;
-	}
+					// The maximum border is 1/4 of the minimum side of the rectangle.
+					const int32 MaxBorderSize = FMath::Min(Input.Rect.Size().GetMin(), Output.Rect.Size().GetMin()) / 4;
+					const int32 TileBorderSize = FMath::Min(GDisplayClusterRenderTileBorder, MaxBorderSize);
 
-	return false;
+					// Shows borders for internal areas only
+					FIntRect SubRect(FIntPoint::ZeroValue, FIntPoint::ZeroValue);
+					if (RenderSettings.TileSettings.GetPos().X > 0)
+			{
+						SubRect.Min.X = TileBorderSize;
 }
-
-void FDisplayClusterViewportProxy::ImplResolveTileResource_RenderThread(FRHICommandListImmediate& RHICmdList, FDisplayClusterViewportProxy const* InDestViewportProxy) const
+					if (RenderSettings.TileSettings.GetPos().Y > 0)
 {
-	check(InDestViewportProxy);
-
-	using namespace UE::DisplayCluster::ViewportProxy;
-
-	const TArray<FDisplayClusterViewport_Context>& DestContexts = InDestViewportProxy->GetContexts_RenderThread();
-	const TArray<FDisplayClusterViewport_Context>& SrcContexts = Contexts;
-
-	TArray<FRHITexture*> SrcResources, DestResources;
-	if (GetResources_RenderThread(EDisplayClusterViewportResourceType::InternalRenderTargetResource, SrcResources)
-	&& InDestViewportProxy->GetResources_RenderThread(EDisplayClusterViewportResourceType::InternalRenderTargetResource, DestResources)
-	&& SrcResources.Num() == DestResources.Num()
-	&& SrcContexts.Num() == DestContexts.Num()
-	&& SrcResources.Num() == SrcContexts.Num()
-	&& !SrcResources.IsEmpty())
-	{
-		// Copy all contexts:
-		for (int32 ContextNum = 0; ContextNum < SrcContexts.Num(); ContextNum++)
+						SubRect.Min.Y = TileBorderSize;
+					}
+					if ((RenderSettings.TileSettings.GetPos().X + 1) < RenderSettings.TileSettings.GetSize().X)
 		{
-			// Get src rect inside overscan
-			FIntRect SrcRect = GetFinalContextRect(EDisplayClusterViewportResourceType::InternalRenderTargetResource, SrcContexts[ContextNum].RenderTargetRect);
-			FIntRect DestRect = SrcContexts[ContextNum].TileDestRect;
-
-			if (GDisplayClusterRenderTileBorder > 0)
+						SubRect.Max.X = TileBorderSize;
+					}
+					if ((RenderSettings.TileSettings.GetPos().Y + 1) < RenderSettings.TileSettings.GetSize().Y)
 			{
-				// The maximum border is 1/4 of the minimum side of the rectangle.
-				const int32 MinBorder = FMath::Min(SrcRect.Size().GetMin(), DestRect.Size().GetMin()) / 4;
-				const int32 TileBorder = FMath::Min(GDisplayClusterRenderTileBorder, MinBorder);
+						SubRect.Max.Y = TileBorderSize;
+					}
 
 				// Set rect smaller to show gaps between tiles:
-				SrcRect.Min.X += TileBorder;
-				SrcRect.Min.Y += TileBorder;
-				SrcRect.Max.X -= TileBorder * 2;
-				SrcRect.Max.Y -= TileBorder * 2;
-
-				DestRect.Min.X += TileBorder;
-				DestRect.Min.Y += TileBorder;
-				DestRect.Max.X -= TileBorder * 2;
-				DestRect.Max.Y -= TileBorder * 2;
+					Input.Rect.Min.X += SubRect.Min.X;
+					Input.Rect.Min.Y += SubRect.Min.Y;
+					Input.Rect.Max.X -= SubRect.Min.X + SubRect.Max.X;
+					Input.Rect.Max.Y -= SubRect.Min.Y + SubRect.Max.Y;
+
+					Output.Rect.Min.X += SubRect.Min.X;
+					Output.Rect.Min.Y += SubRect.Min.Y;
+					Output.Rect.Max.X -= SubRect.Min.X + SubRect.Max.X;
+					Output.Rect.Max.Y -= SubRect.Min.Y + SubRect.Max.Y;
 			}
 
-			// Copy tile to dest
-			ImplResolveResource(RHICmdList, SrcResources[ContextNum], SrcRect, DestResources[ContextNum], DestRect, false, false);
-		}
-	}
+				// Copy texture region
+				TextureUtils->ResolveTextureContext(EDisplayClusterShaderTextureUtilsFlags::DisableResampleShader, Input, Output);
+			});
 }
 
 void FDisplayClusterViewportProxy::CleanupResources_RenderThread(FRHICommandListImmediate& RHICmdList) const
diff --git a/Engine/Plugins/Runtime/nDisplay/Source/DisplayCluster/Private/Render/Viewport/DisplayClusterViewport_OpenColorIO.cpp b/Engine/Plugins/Runtime/nDisplay/Source/DisplayCluster/Private/Render/Viewport/DisplayClusterViewport_OpenColorIO.cpp
--- a/Engine/Plugins/Runtime/nDisplay/Source/DisplayCluster/Private/Render/Viewport/DisplayClusterViewport_OpenColorIO.cpp
+++ b/Engine/Plugins/Runtime/nDisplay/Source/DisplayCluster/Private/Render/Viewport/DisplayClusterViewport_OpenColorIO.cpp
@@ -4,6 +4,7 @@
 #include "Render/Viewport/DisplayClusterViewport.h"
 
 #include "Render/Viewport/Containers/DisplayClusterViewport_Context.h"
+#include "Misc/DisplayClusterColorEncoding.h"
 
 #include "OpenColorIOConfiguration.h"
 #include "OpenColorIOColorTransform.h"
@@ -11,6 +12,8 @@
 #include "OpenColorIODisplayExtension.h"
 #include "OpenColorIOShader.h"
 
+#include "IDisplayClusterShadersTextureUtils.h"
+
 #include "SceneView.h"
 #include "ScreenPass.h"
 #include "CommonRenderResources.h"
@@ -50,47 +53,99 @@ bool FDisplayClusterViewport_OpenColorIO::IsConversionSettingsEqual(const FOpenC
 	return ConversionSettings.ToString() == InConversionSettings.ToString();
 }
 
-float FDisplayClusterViewport_OpenColorIO::GetDisplayGamma(const FDisplayClusterViewport_Context& InViewportContext) const
+float FDisplayClusterViewport_OpenColorIO::GetGammaCorrection(const struct FDisplayClusterColorEncoding& InColorEncoding)
 {
-	const float DefaultDisplayGamma = UTextureRenderTarget::GetDefaultDisplayGamma();
-	const float EngineDisplayGamma = InViewportContext.RenderThreadData.EngineDisplayGamma;
+	// Convert input to the sRGB: default DisplayGamma is 2.2 which corresponds to SRGB
+	const EDisplayClusterColorEncoding EqualEncoding = InColorEncoding.GetEqualEncoding();
+	switch (EqualEncoding)
+	{
+	case EDisplayClusterColorEncoding::Linear: // Linear -> DefaultGamma/sRGB
+		return UTextureRenderTarget::GetDefaultDisplayGamma();
+
+	case EDisplayClusterColorEncoding::Gamma: // CustomGamma -> DefaultGamma/sRGB
+		return UTextureRenderTarget::GetDefaultDisplayGamma() / InColorEncoding.GammaValue;
 
-	// There is a special case where post processing and tonemapper are disabled. In this case tonemapper applies a static display Inverse of Gamma which defaults to 2.2.
-	// In the case when Both PostProcessing and ToneMapper are disabled we apply gamma manually. In every other case we apply inverse gamma before applying OCIO.
-	const float DisplayGamma = (InViewportContext.RenderThreadData.EngineShowFlags.Tonemapper == 0) || (InViewportContext.RenderThreadData.EngineShowFlags.PostProcessing == 0) ? DefaultDisplayGamma : DefaultDisplayGamma / EngineDisplayGamma;
+	case EDisplayClusterColorEncoding::sRGB: // DefaultGamma/sRGB
+		return 1.f;
 
-	return DisplayGamma;
+	default:
+		break;
 }
 
-bool FDisplayClusterViewport_OpenColorIO::AddPass_RenderThread(FRDGBuilder& GraphBuilder, const FDisplayClusterViewport_Context& InViewportContext,
-	FRHITexture* InputTextureRHI, const FIntRect& InputRect, FRHITexture* OutputTextureRHI, const FIntRect& OutputRect, bool bUnpremultiply, bool bInvertAlpha) const
-{
-	FRDGTextureRef InputTexture = GraphBuilder.RegisterExternalTexture(CreateRenderTarget(InputTextureRHI, TEXT("nDisplayOCIOInputTexture")));
-	FRDGTextureRef OutputTexture = GraphBuilder.RegisterExternalTexture(CreateRenderTarget(OutputTextureRHI, TEXT("nDisplayOCIORenderTargetTexture")));
+	// Unsupported input encoding, but we'll return 1 anyway, meaning we'll use the input as is.	
+	return 1.f;
+}
 
-	FScreenPassRenderTarget Output = FScreenPassRenderTarget(OutputTexture, OutputRect, ERenderTargetLoadAction::EClear);
-	EOpenColorIOTransformAlpha TransformAlpha = EOpenColorIOTransformAlpha::None;
-	if (bUnpremultiply)
+EOpenColorIOTransformAlpha FDisplayClusterViewport_OpenColorIO::GetTransformAlpha(const FDisplayClusterColorEncoding& InColorEncoding)
+{
+	// Unpremultiply while gamma conversion
+	switch (InColorEncoding.Premultiply)
 	{
-		TransformAlpha = bInvertAlpha ? EOpenColorIOTransformAlpha::InvertUnpremultiply : EOpenColorIOTransformAlpha::Unpremultiply;
+	case EDisplayClusterColorPremultiply::Premultiply:       return EOpenColorIOTransformAlpha::Unpremultiply;
+	case EDisplayClusterColorPremultiply::InvertPremultiply: return EOpenColorIOTransformAlpha::InvertUnpremultiply;
+	default:
+		break;
 	}
 
-	FOpenColorIORendering::AddPass_RenderThread(
-		GraphBuilder,
-		FScreenPassViewInfo(),
-		GMaxRHIFeatureLevel,
-		FScreenPassTexture(InputTexture),
-		Output,
-		CachedResourcesRenderThread,
-		GetDisplayGamma(InViewportContext),
-		TransformAlpha
-	);
+	return EOpenColorIOTransformAlpha::None;;
+}
+
+bool FDisplayClusterViewport_OpenColorIO::AddPass_RenderThread(
+	const FDisplayClusterShadersTextureUtilsSettings& InTextureUtilsSettings,
+	const TSharedRef<IDisplayClusterShadersTextureUtils>& InTextureUtils) const
+{
+	if (!CachedResourcesRenderThread.IsValid())
+{
+		// The OCIO shader is not ready at this point, use the default resolve method.
+		InTextureUtils->Resolve(InTextureUtilsSettings);
+
+		return true;
+	}
+
+	// When requesting RDG from the TextureUtils API, it switches from RHI to RDG.
+	FRDGBuilder& GraphBuilder = InTextureUtils->GetOrCreateRDGBuilder();
+
+	bool bFinalPassResult = false;
+	InTextureUtils->ForEachContextByPredicate([&](
+		const FDisplayClusterShadersTextureViewportContext& Input,
+		const FDisplayClusterShadersTextureViewportContext& Output)
+		{
+			// returns true if OCIO is applied.
+			bFinalPassResult = true;
+
+			// The OCIO output is in sRGB
+			FOpenColorIORendering::AddPass_RenderThread(
+				GraphBuilder,
+				FScreenPassViewInfo(),
+				GMaxRHIFeatureLevel,
+						Input.ToScreenPassTexture(),
+						FScreenPassRenderTarget(Output.ToScreenPassTexture(), ERenderTargetLoadAction::EClear),
+				CachedResourcesRenderThread,
+				GetGammaCorrection(Input.ColorEncoding),
+				GetTransformAlpha(Input.ColorEncoding)
+			);
+
+			// Convert OCIO output(sRGB) to output color space if necessary.
+			if (Output.ColorEncoding.GetEqualEncoding() != EDisplayClusterColorEncoding::sRGB)
+			{
+				// Convert OutputTexture color encoding from sRGB to Linear(Holdout):
+				InTextureUtils->ResolveTextureContext(
+					// Customize the settings by requesting a temporary input texture cloned from the output texture.
+					{ InTextureUtilsSettings , EDisplayClusterShaderTextureUtilsFlags::UseOutputTextureAsInput },
+					// The input texture is empty, a temporary texture created from the output will be used instead.
+					{ {}, {EDisplayClusterColorEncoding::Gamma, Input.ColorEncoding.Premultiply} },
+					// The output texture already contains the OCIO result and will be used as input in this pass.
+					Output
+				);
+			}
+		});
 
 	return true;
 }
 
 // This is a copy of FOpenColorIODisplayExtension::PostProcessPassAfterTonemap_RenderThread()
-FScreenPassTexture FDisplayClusterViewport_OpenColorIO::PostProcessPassAfterTonemap_RenderThread(FRDGBuilder& GraphBuilder, const FDisplayClusterViewport_Context& InViewportContext, const FSceneView& View, const FPostProcessMaterialInputs& InOutInputs)
+FScreenPassTexture FDisplayClusterViewport_OpenColorIO::PostProcessPassAfterTonemap_RenderThread(FRDGBuilder& GraphBuilder, const FDisplayClusterColorEncoding& InColorEncoding,
+	const FSceneView& View, const FPostProcessMaterialInputs& InOutInputs)
 {
 	const FScreenPassTexture& SceneColor = FScreenPassTexture::CopyFromSlice(GraphBuilder, InOutInputs.GetInput(EPostProcessMaterialInput::SceneColor));
 	check(SceneColor.IsValid());
@@ -107,7 +162,8 @@ FScreenPassTexture FDisplayClusterViewport_OpenColorIO::PostProcessPassAfterTone
 		View,
 		SceneColor,
 		Output,
-		CachedResourcesRenderThread
+		CachedResourcesRenderThread,
+		GetTransformAlpha(InColorEncoding)
 	);
 
 	return MoveTemp(Output);
diff --git a/Engine/Plugins/Runtime/nDisplay/Source/DisplayCluster/Private/Render/Viewport/DisplayClusterViewport_OpenColorIO.h b/Engine/Plugins/Runtime/nDisplay/Source/DisplayCluster/Private/Render/Viewport/DisplayClusterViewport_OpenColorIO.h
--- a/Engine/Plugins/Runtime/nDisplay/Source/DisplayCluster/Private/Render/Viewport/DisplayClusterViewport_OpenColorIO.h
+++ b/Engine/Plugins/Runtime/nDisplay/Source/DisplayCluster/Private/Render/Viewport/DisplayClusterViewport_OpenColorIO.h
@@ -17,6 +17,7 @@ class FSceneView;
 class FSceneViewFamily;
 struct FScreenPassTexture;
 struct FPostProcessMaterialInputs;
+struct FDisplayClusterColorEncoding;
 class FRDGBuilder;
 class FDisplayClusterViewport_Context;
 
@@ -43,20 +44,17 @@ public:
 
 	/** Add OCIO render pass.
 	 *
-	 * @param GraphBuilder      - RDG interface.
-	 * @param InViewportContext - DC viewport context.
-	 * @param InputTextureRHI   - Source texture.
-	 * @param InputRect         - Source rec.
-	 * @param OutputTextureRHI  - Destination texture.
-	 * @param OutputRect        - Destination rec.
+	 * @param TextureUtils   - resolving data
 	 *
 	 * @return - true if success.
 	 */
-	bool AddPass_RenderThread(FRDGBuilder& GraphBuilder, const FDisplayClusterViewport_Context& InViewportContext,
-		FRHITexture* InputTextureRHI, const FIntRect& InputRect, FRHITexture* OutputTextureRHI, const FIntRect& OutputRect, bool bUnpremultiply, bool InvertedAlpha) const;
+	bool AddPass_RenderThread(
+		const struct FDisplayClusterShadersTextureUtilsSettings& InTextureUtilsSettings,
+		const TSharedRef<class IDisplayClusterShadersTextureUtils>& TextureUtils) const;
 
 	/* This is a copy of FOpenColorIODisplayExtension::PostProcessPassAfterTonemap_RenderThread() */
-	FScreenPassTexture PostProcessPassAfterTonemap_RenderThread(FRDGBuilder& GraphBuilder, const FDisplayClusterViewport_Context& InViewportContext, const FSceneView& View, const FPostProcessMaterialInputs& InOutInputs);
+	FScreenPassTexture PostProcessPassAfterTonemap_RenderThread(FRDGBuilder& GraphBuilder, const FDisplayClusterColorEncoding& InColorEncoding,
+		const FSceneView& View, const FPostProcessMaterialInputs& InOutInputs);
 
 	/** Compare two OCIO configurations.
 	 *
@@ -78,9 +76,11 @@ public:
 		return CachedResourcesRenderThread.IsValid();
 	}
 
-private:
-	/* returns the DisplayGamma of OCIO for the viewport context.*/
-	float GetDisplayGamma(const FDisplayClusterViewport_Context& InViewportContext) const;
+	/** Gets the gamma correction for the OCIO shader from the InColorEncoding parameter. */
+	static float GetGammaCorrection(const FDisplayClusterColorEncoding& InColorEncoding);
+
+	/** Gets transform alpha value from the InColorEncoding parameter. */
+	static EOpenColorIOTransformAlpha GetTransformAlpha(const FDisplayClusterColorEncoding& InColorEncoding);
 
 private:
 	/** Cached pass resources required to apply conversion for render thread. */
diff --git a/Engine/Plugins/Runtime/nDisplay/Source/DisplayCluster/Private/Render/Viewport/DisplayClusterViewport_Rules.cpp b/Engine/Plugins/Runtime/nDisplay/Source/DisplayCluster/Private/Render/Viewport/DisplayClusterViewport_Rules.cpp
--- a/Engine/Plugins/Runtime/nDisplay/Source/DisplayCluster/Private/Render/Viewport/DisplayClusterViewport_Rules.cpp
+++ b/Engine/Plugins/Runtime/nDisplay/Source/DisplayCluster/Private/Render/Viewport/DisplayClusterViewport_Rules.cpp
@@ -367,12 +367,6 @@ bool FDisplayClusterViewport::IsResourceUsedImpl(const EDisplayClusterViewportRe
 	switch (InResourceType)
 	{
 	case EDisplayClusterViewportResource::AdditionalTargetableResources:
-		// PostRender Blur require additional RTT for shader
-		if (PostRenderSettings.PostprocessBlur.IsEnabled())
-		{
-			return true;
-		}
-
 		// Supoport projection policy additional resource
 		if (ProjectionPolicy.IsValid() && ProjectionPolicy->ShouldUseAdditionalTargetableResource(this))
 		{
diff --git a/Engine/Plugins/Runtime/nDisplay/Source/DisplayCluster/Private/Render/Viewport/LightCard/DisplayClusterViewportLightCardManager.cpp b/Engine/Plugins/Runtime/nDisplay/Source/DisplayCluster/Private/Render/Viewport/LightCard/DisplayClusterViewportLightCardManager.cpp
--- a/Engine/Plugins/Runtime/nDisplay/Source/DisplayCluster/Private/Render/Viewport/LightCard/DisplayClusterViewportLightCardManager.cpp
+++ b/Engine/Plugins/Runtime/nDisplay/Source/DisplayCluster/Private/Render/Viewport/LightCard/DisplayClusterViewportLightCardManager.cpp
@@ -300,8 +300,19 @@ void FDisplayClusterViewportLightCardManager::RenderUVLightCard(const EDisplayCl
 		{
 			FDisplayClusterShaderParameters_UVLightCards UVLightCardParameters;
 			UVLightCardParameters.ProjectionPlaneSize = ADisplayClusterLightCardActor::UVPlaneDefaultSize;
-			UVLightCardParameters.bRenderFinalColor = ViewportManager->ShouldRenderFinalColor();
-			UVLightCardParameters.LightCardGamma = 2.2;
+			if (!ViewportManager->ShouldRenderFinalColor())
+			{
+				// SceneCaptureSource = ESceneCaptureSource::SCS_SceneColorHDR
+				// Output gamma is linear
+				UVLightCardParameters.bRenderFinalColor = false;
+				UVLightCardParameters.LightCardGamma = 1;
+			}
+			else
+			{
+				// Output gamma is sRGB, have to convert to linear
+				UVLightCardParameters.bRenderFinalColor = true;
+				UVLightCardParameters.LightCardGamma = 2.2;
+			}
 
 			// Store any components that were invisible but forced to be visible so they can be set back to invisible after the render
 			TArray<UPrimitiveComponent*> ComponentsToUnload;
diff --git a/Engine/Plugins/Runtime/nDisplay/Source/DisplayCluster/Private/Render/Viewport/RenderFrame/DisplayClusterRenderFrameSettings.cpp b/Engine/Plugins/Runtime/nDisplay/Source/DisplayCluster/Private/Render/Viewport/RenderFrame/DisplayClusterRenderFrameSettings.cpp
--- a/Engine/Plugins/Runtime/nDisplay/Source/DisplayCluster/Private/Render/Viewport/RenderFrame/DisplayClusterRenderFrameSettings.cpp
+++ b/Engine/Plugins/Runtime/nDisplay/Source/DisplayCluster/Private/Render/Viewport/RenderFrame/DisplayClusterRenderFrameSettings.cpp
@@ -105,6 +105,16 @@ bool FDisplayClusterRenderFrameSettings::IsPostProcessDisabled() const
 	return false;
 }
 
+bool FDisplayClusterRenderFrameSettings::ShouldUseHoldout() const
+{
+	if (IsPreviewRendering())
+	{
+		return PreviewSettings.bPreviewEnableHoldoutComposite;
+	}
+
+	return false;
+}
+
 bool FDisplayClusterRenderFrameSettings::IsPreviewRendering() const
 {
 	if (RenderMode == EDisplayClusterRenderFrameMode::PreviewInScene
diff --git a/Engine/Plugins/Runtime/nDisplay/Source/DisplayCluster/Private/Render/Viewport/RenderFrame/DisplayClusterRenderFrameSettings.h b/Engine/Plugins/Runtime/nDisplay/Source/DisplayCluster/Private/Render/Viewport/RenderFrame/DisplayClusterRenderFrameSettings.h
--- a/Engine/Plugins/Runtime/nDisplay/Source/DisplayCluster/Private/Render/Viewport/RenderFrame/DisplayClusterRenderFrameSettings.h
+++ b/Engine/Plugins/Runtime/nDisplay/Source/DisplayCluster/Private/Render/Viewport/RenderFrame/DisplayClusterRenderFrameSettings.h
@@ -116,6 +116,9 @@ public:
 	/** Should use linear gamma. */
 	bool ShouldUseLinearGamma() const;
 
+	/** Should use Holdout. */
+	bool ShouldUseHoldout() const;
+
 	/** true, if output frame RTTs is used. */
 	bool ShouldUseOutputFrameTargetableResources() const;
 
diff --git a/Engine/Plugins/Runtime/nDisplay/Source/DisplayCluster/Public/Misc/DisplayClusterColorEncoding.h b/Engine/Plugins/Runtime/nDisplay/Source/DisplayCluster/Public/Misc/DisplayClusterColorEncoding.h
new file mode 100644
--- /dev/null
+++ b/Engine/Plugins/Runtime/nDisplay/Source/DisplayCluster/Public/Misc/DisplayClusterColorEncoding.h
@@ -0,0 +1,115 @@
+// Copyright Epic Games, Inc. All Rights Reserved.
+
+#pragma once
+
+#include "CoreMinimal.h"
+
+/** Color encoding type.  */
+enum class EDisplayClusterColorEncoding : uint8
+{
+	/** Linear color space */
+	Linear = 0,
+
+	/** Gamma encoding based on pow(In, Gamma). */
+	Gamma,
+
+	/** sRGB color space. */
+	sRGB,
+
+	/** ST2084 gamma with saturation (for special use case of nDisplay MediaIO). */
+	MediaPQ,
+};
+
+/** Color premultiply type.  */
+enum class EDisplayClusterColorPremultiply : uint8
+{
+	/** Color is not changed by alpha. */
+	None = 0,
+
+	/** Color is premultiplied by Alpha. */
+	Premultiply,
+
+	/** Color is premultiplied by (1-Alpha). */
+	InvertPremultiply
+};
+
+/** Color encoding data. */
+struct FDisplayClusterColorEncoding
+{
+	/** Constructors. */
+	FDisplayClusterColorEncoding() = default;
+	FDisplayClusterColorEncoding(
+		const EDisplayClusterColorEncoding InColorEncoding,
+		const EDisplayClusterColorPremultiply InPremultiply = EDisplayClusterColorPremultiply::None)
+		: Encoding(InColorEncoding), Premultiply(InPremultiply) { }
+	FDisplayClusterColorEncoding(
+		float InGamma,
+		const EDisplayClusterColorPremultiply InPremultiply = EDisplayClusterColorPremultiply::None)
+		: Encoding(EDisplayClusterColorEncoding::Gamma) , GammaValue(InGamma), Premultiply(InPremultiply) { }
+
+	/** Return equal encoding value. */
+	EDisplayClusterColorEncoding GetEqualEncoding() const
+	{
+		// Linear == Gamma(1.f)
+		if (Encoding == EDisplayClusterColorEncoding::Gamma && GammaValue == 1.f)
+		{
+			return EDisplayClusterColorEncoding::Linear;
+		}
+		// sRGB == Gamma(0.f) == Gamma(2.2f)
+		else if (Encoding == EDisplayClusterColorEncoding::Gamma && (GammaValue <= 0.f || GammaValue == 2.2f))
+		{
+			return EDisplayClusterColorEncoding::sRGB; // Default gamma means sRGB
+		}
+
+		return Encoding;
+	}
+
+	/** Returns true if the color encodings are the same. */
+	inline bool operator==(const FDisplayClusterColorEncoding& In) const
+	{
+		if (Premultiply != In.Premultiply)
+		{
+			return false;
+		}
+
+		return IsEqualsGammaEncoding(In);
+	}
+
+	/** Compare only gamma encodings. */
+	inline bool IsEqualsGammaEncoding(const FDisplayClusterColorEncoding& In) const
+	{
+		if (Encoding == In.Encoding)
+		{
+			// Additional rule for `Gamma` value
+			if (Encoding != EDisplayClusterColorEncoding::Gamma)
+			{
+				return true;
+			}
+
+			// EDisplayClusterColorEncoding::Gamma
+			if (GammaValue == In.GammaValue)
+			{
+				return true;
+			}
+		}
+
+		if (GetEqualEncoding() == In.GetEqualEncoding())
+		{
+			return true;
+		}
+
+		return false;
+	}
+
+public:
+	/** Color space encoding type.  */
+	EDisplayClusterColorEncoding Encoding = EDisplayClusterColorEncoding::Gamma;
+
+	/** The gamma value for the `EDisplayClusterColorEncoding::Gamma`.
+	* Zero value means default gamma.
+	*/
+	float GammaValue = 0.f;
+
+	/** Color premultiply modifiers. */
+	EDisplayClusterColorPremultiply Premultiply = EDisplayClusterColorPremultiply::None;
+};
diff --git a/Engine/Plugins/Runtime/nDisplay/Source/DisplayCluster/Public/Render/DisplayDevice/IDisplayClusterDisplayDeviceProxy.h b/Engine/Plugins/Runtime/nDisplay/Source/DisplayCluster/Public/Render/DisplayDevice/IDisplayClusterDisplayDeviceProxy.h
--- a/Engine/Plugins/Runtime/nDisplay/Source/DisplayCluster/Public/Render/DisplayDevice/IDisplayClusterDisplayDeviceProxy.h
+++ b/Engine/Plugins/Runtime/nDisplay/Source/DisplayCluster/Public/Render/DisplayDevice/IDisplayClusterDisplayDeviceProxy.h
@@ -13,6 +13,8 @@
 class FDisplayClusterViewport_Context;
 class FRDGBuilder;
 class FSceneView;
+class IDisplayClusterShadersTextureUtils;
+struct FDisplayClusterShadersTextureUtilsSettings;
 struct FPostProcessMaterialInputs;
 struct FScreenPassTexture;
 
@@ -35,18 +37,15 @@ public:
 
 	/** Add render pass at the end of nDisplay pipeline.
 	 *
-	 * @param GraphBuilder      - RDG interface.
-	 * @param InViewportContext - DC viewport context.
-	 * @param InputTextureRHI   - Source texture.
-	 * @param InputRect         - Source rec.
-	 * @param OutputTextureRHI  - Destination texture.
-	 * @param OutputRect        - Destination rec.
+	 * @param InTextureUtilsSettings - A settings for render
+	 * @param InTextureUtils         - A container of resources for processing.
 	 *
 	 * @return - true if success.
 	 */
-	virtual bool AddFinalPass_RenderThread(FRDGBuilder& GraphBuilder, const FDisplayClusterViewport_Context& InViewportContext,
-		FRHITexture* InputTextureRHI, const FIntRect& InputRect, FRHITexture* OutputTextureRHI, const FIntRect& OutputRect) const
-	{ return false; }
+	virtual bool AddFinalPass_RenderThread(const FDisplayClusterShadersTextureUtilsSettings& InTextureUtilsSettings, const TSharedRef<IDisplayClusterShadersTextureUtils, ESPMode::ThreadSafe>& InTextureUtils) const
+	{
+		return false;
+	}
 
 	/** Allow callback OnPostProcessPassAfterFXAA.  */
 	virtual bool ShouldUsePostProcessPassAfterFXAA() const
@@ -95,4 +94,22 @@ public:
 	 */
 	virtual FScreenPassTexture OnPostProcessPassAfterTonemap_RenderThread(FRDGBuilder& GraphBuilder, const FSceneView& View, const FPostProcessMaterialInputs& Inputs, const uint32 ContextNum)
 	{ return Inputs.ReturnUntouchedSceneColorForPostProcessing(GraphBuilder); }
+
+	/** Add render pass at the end of nDisplay pipeline.
+	 *
+	 * @param GraphBuilder      - RDG interface.
+	 * @param InViewportContext - DC viewport context.
+	 * @param InputTextureRHI   - Source texture.
+	 * @param InputRect         - Source rec.
+	 * @param OutputTextureRHI  - Destination texture.
+	 * @param OutputRect        - Destination rec.
+	 *
+	 * @return - true if success.
+	 */
+	UE_DEPRECATED(5.6, "This function has been deprecated. Please use 'AddFinalPass_RenderThread()'.")
+	virtual bool AddFinalPass_RenderThread(FRDGBuilder& GraphBuilder, const FDisplayClusterViewport_Context& InViewportContext,
+		FRHITexture* InputTextureRHI, const FIntRect& InputRect, FRHITexture* OutputTextureRHI, const FIntRect& OutputRect) const
+	{
+		return false;
+	}
 };
diff --git a/Engine/Plugins/Runtime/nDisplay/Source/DisplayCluster/Public/Render/Viewport/Containers/DisplayClusterViewport_PreviewSettings.h b/Engine/Plugins/Runtime/nDisplay/Source/DisplayCluster/Public/Render/Viewport/Containers/DisplayClusterViewport_PreviewSettings.h
--- a/Engine/Plugins/Runtime/nDisplay/Source/DisplayCluster/Public/Render/Viewport/Containers/DisplayClusterViewport_PreviewSettings.h
+++ b/Engine/Plugins/Runtime/nDisplay/Source/DisplayCluster/Public/Render/Viewport/Containers/DisplayClusterViewport_PreviewSettings.h
@@ -22,6 +22,9 @@ struct FDisplayClusterViewport_PreviewSettings
 	bool bPreviewInGameEnable = false;
 	bool bPreviewInGameRenderFrustum = false;
 
+	/** Enables the use of 'HoldoutComposite' plugin for the preview rendering. */
+	bool bPreviewEnableHoldoutComposite = true;
+
 	// Preview uses techvis mode for rendering
 	bool bEnablePreviewTechvis = false;
 
diff --git a/Engine/Plugins/Runtime/nDisplay/Source/DisplayCluster/Public/Render/Viewport/IDisplayClusterViewportProxy.h b/Engine/Plugins/Runtime/nDisplay/Source/DisplayCluster/Public/Render/Viewport/IDisplayClusterViewportProxy.h
--- a/Engine/Plugins/Runtime/nDisplay/Source/DisplayCluster/Public/Render/Viewport/IDisplayClusterViewportProxy.h
+++ b/Engine/Plugins/Runtime/nDisplay/Source/DisplayCluster/Public/Render/Viewport/IDisplayClusterViewportProxy.h
@@ -13,6 +13,8 @@
 #include "Render/Viewport/RenderFrame/DisplayClusterRenderFrameEnums.h"
 #include "Render/Viewport/IDisplayClusterViewportConfigurationProxy.h"
 
+#include "Misc/DisplayClusterColorEncoding.h"
+
 #include "RHI.h"
 #include "RHICommandList.h"
 #include "RHIResources.h"
@@ -54,6 +56,13 @@ public:
 	/** Return viewport proxy contexts data. */
 	virtual const TArray<FDisplayClusterViewport_Context>& GetContexts_RenderThread() const = 0;
 
+	/**
+	* Return color-encoding information for the resource
+	* 
+	* @param InResourceType   - (in) resource type (RTT, Shader, MIPS, etc)
+	*/
+	virtual FDisplayClusterColorEncoding GetResourceColorEncoding_RenderThread(const EDisplayClusterViewportResourceType InResourceType) const = 0;
+
 	/** Get viewport resources by type
 	 *
 	 * @param InResourceType - resource type (RTT, Shader, MIPS, etc)
diff --git a/Engine/Plugins/Runtime/nDisplay/Source/DisplayClusterProjection/Private/Policy/MPCDI/DisplayClusterProjectionMPCDIPolicy.cpp b/Engine/Plugins/Runtime/nDisplay/Source/DisplayClusterProjection/Private/Policy/MPCDI/DisplayClusterProjectionMPCDIPolicy.cpp
--- a/Engine/Plugins/Runtime/nDisplay/Source/DisplayClusterProjection/Private/Policy/MPCDI/DisplayClusterProjectionMPCDIPolicy.cpp
+++ b/Engine/Plugins/Runtime/nDisplay/Source/DisplayClusterProjection/Private/Policy/MPCDI/DisplayClusterProjectionMPCDIPolicy.cpp
@@ -419,6 +419,17 @@ void FDisplayClusterProjectionMPCDIPolicy::ApplyWarpBlend_RenderThread(FRHIComma
 			// cleanup camera list for render
 			ShaderICVFX.CleanupCamerasForRender();
 
+			// Convert Lightcards gamma from linear to the Outer viewport gamma
+			FDisplayClusterColorEncoding InputEncoding = InViewportProxy->GetResourceColorEncoding_RenderThread(EDisplayClusterViewportResourceType::InputShaderResource);
+			if (InputEncoding == EDisplayClusterColorEncoding::Linear)
+			{
+				ShaderICVFX.LightCardGamma = 1.0f;
+			}
+			else
+			{
+				ShaderICVFX.LightCardGamma = 2.2f;
+			}
+
 			// Initialize shader input data
 			FDisplayClusterShaderParameters_WarpBlend WarpBlendParameters;
 
diff --git a/Engine/Plugins/Runtime/nDisplay/Source/DisplayClusterShaders/Private/DisplayClusterShadersModule.cpp b/Engine/Plugins/Runtime/nDisplay/Source/DisplayClusterShaders/Private/DisplayClusterShadersModule.cpp
--- a/Engine/Plugins/Runtime/nDisplay/Source/DisplayClusterShaders/Private/DisplayClusterShadersModule.cpp
+++ b/Engine/Plugins/Runtime/nDisplay/Source/DisplayClusterShaders/Private/DisplayClusterShadersModule.cpp
@@ -14,6 +14,7 @@
 #include "Shaders/DisplayClusterShadersPostprocess_OutputRemap.h"
 #include "Shaders/DisplayClusterShadersWarpblend_ICVFX.h"
 #include "Shaders/DisplayClusterShadersWarpblend_MPCDI.h"
+#include "DisplayClusterShadersTextureUtils.h"
 
 //////////////////////////////////////////////////////////////////////////////////////////////
 // IModuleInterface
@@ -74,5 +75,14 @@ void FDisplayClusterShadersModule::AddPQToLinearPass(FRDGBuilder& GraphBuilder,
 	FDisplayClusterShadersMedia::AddPQToLinearPass(GraphBuilder, Parameters);
 }
 
-IMPLEMENT_MODULE(FDisplayClusterShadersModule, DisplayClusterShaders);
+TSharedRef<IDisplayClusterShadersTextureUtils> FDisplayClusterShadersModule::CreateTextureUtils_RenderThread(FRHICommandListImmediate& RHICmdList) const
+{
+	return FDisplayClusterShadersTextureUtils::CreateTextureUtils_RenderThread(RHICmdList);
+}
 
+TSharedRef<IDisplayClusterShadersTextureUtils> FDisplayClusterShadersModule::CreateTextureUtils_RenderThread(FRDGBuilder& GraphBuilder) const
+{
+	return FDisplayClusterShadersTextureUtils::CreateTextureUtils_RenderThread(GraphBuilder);
+}
+
+IMPLEMENT_MODULE(FDisplayClusterShadersModule, DisplayClusterShaders);
diff --git a/Engine/Plugins/Runtime/nDisplay/Source/DisplayClusterShaders/Private/DisplayClusterShadersModule.h b/Engine/Plugins/Runtime/nDisplay/Source/DisplayClusterShaders/Private/DisplayClusterShadersModule.h
--- a/Engine/Plugins/Runtime/nDisplay/Source/DisplayClusterShaders/Private/DisplayClusterShadersModule.h
+++ b/Engine/Plugins/Runtime/nDisplay/Source/DisplayClusterShaders/Private/DisplayClusterShadersModule.h
@@ -15,6 +15,7 @@ public:
 	virtual void ShutdownModule() override;
 
 public:
+	//~Begin IDisplayClusterShaders
 	virtual bool RenderWarpBlend_MPCDI(FRHICommandListImmediate& RHICmdList, const FDisplayClusterShaderParameters_WarpBlend& InWarpBlendParameters) const override;
 	virtual bool RenderWarpBlend_ICVFX(FRHICommandListImmediate& RHICmdList, const FDisplayClusterShaderParameters_WarpBlend& InWarpBlendParameters, const FDisplayClusterShaderParameters_ICVFX& InICVFXParameters) const override;
 	virtual bool RenderPreprocess_UVLightCards(FRHICommandListImmediate& RHICmdList, FSceneInterface* InScene, FRenderTarget* InRenderTarget, float ProjectionPlaneSize, bool bRenderFinalColor) const override { return false; }
@@ -24,4 +25,8 @@ public:
 	virtual bool GenerateMips(FRHICommandListImmediate& RHICmdList, FRHITexture* InOutMipsTexture, const FDisplayClusterShaderParameters_GenerateMips& InSettings) const override;
 	virtual void AddLinearToPQPass(FRDGBuilder& GraphBuilder, const FDisplayClusterShaderParameters_MediaPQ& Parameters) const override;
 	virtual void AddPQToLinearPass(FRDGBuilder& GraphBuilder, const FDisplayClusterShaderParameters_MediaPQ& Parameters) const override;
+	
+	virtual TSharedRef<IDisplayClusterShadersTextureUtils> CreateTextureUtils_RenderThread(FRHICommandListImmediate& RHICmdList) const override;
+	virtual TSharedRef<IDisplayClusterShadersTextureUtils> CreateTextureUtils_RenderThread(FRDGBuilder& GraphBuilder) const override;
+	//~~End IDisplayClusterShaders
 };
diff --git a/Engine/Plugins/Runtime/nDisplay/Source/DisplayClusterShaders/Private/DisplayClusterShadersTextureUtils.cpp b/Engine/Plugins/Runtime/nDisplay/Source/DisplayClusterShaders/Private/DisplayClusterShadersTextureUtils.cpp
new file mode 100644
--- /dev/null
+++ b/Engine/Plugins/Runtime/nDisplay/Source/DisplayClusterShaders/Private/DisplayClusterShadersTextureUtils.cpp
@@ -0,0 +1,746 @@
+// Copyright Epic Games, Inc. All Rights Reserved.
+
+#include "DisplayClusterShadersTextureUtils.h"
+#include "Render/Viewport/IDisplayClusterViewportProxy.h"
+#include "Shaders/DisplayClusterShadersCopyTexture.h"
+
+#include "RenderResource.h"
+#include "RenderGraphUtils.h"
+#include "RenderGraphResources.h"
+
+#include "RenderTargetPool.h"
+
+namespace UE::DisplayClusterShaders::Private
+{
+	/** Return resource name. */
+	static constexpr auto GetDisplayClusterViewportResourceTypeName(const EDisplayClusterViewportResourceType InResourceType)
+	{
+		switch (InResourceType)
+		{
+		case EDisplayClusterViewportResourceType::InternalRenderTargetResource:     return TEXT("nDisplay.InternalRenderTargetResource");
+
+		case EDisplayClusterViewportResourceType::InputShaderResource:              return TEXT("nDisplay.InputShaderResource");
+		case EDisplayClusterViewportResourceType::MipsShaderResource:               return TEXT("nDisplay.MipsShaderResource");
+		case EDisplayClusterViewportResourceType::AdditionalTargetableResource:     return TEXT("nDisplay.AdditionalTargetableResource");
+
+		case EDisplayClusterViewportResourceType::BeforeWarpBlendTargetableResource:return TEXT("nDisplay.BeforeWarpBlendTargetableResource");
+		case EDisplayClusterViewportResourceType::AfterWarpBlendTargetableResource: return TEXT("nDisplay.AfterWarpBlendTargetableResource");
+
+		case EDisplayClusterViewportResourceType::OutputTargetableResource:         return TEXT("nDisplay.OutputTargetableResource");
+		case EDisplayClusterViewportResourceType::OutputPreviewTargetableResource:  return TEXT("nDisplay.OutputPreviewTargetableResource");
+
+		case EDisplayClusterViewportResourceType::OutputFrameTargetableResource:      return TEXT("nDisplay.OutputFrameTargetableResource");
+		case EDisplayClusterViewportResourceType::AdditionalFrameTargetableResource:  return TEXT("nDisplay.AdditionalFrameTargetableResource");
+		default:
+			break;
+		}
+
+		return TEXT("");
+	}
+
+	/** Return pixel format. */
+	static inline EPixelFormat GetPixelFormat(const FDisplayClusterShadersTextureViewport& In)
+	{
+		if (In.TextureRHI)
+		{
+			return In.TextureRHI->GetDesc().Format;
+		}
+		else if (In.TextureRDG)
+		{
+			return In.TextureRDG->Desc.Format;
+		}
+
+		return EPixelFormat::PF_Unknown;
+	}
+
+	/** Return size of the referenced texture. */
+	static inline FIntPoint GetTextureSize(const FDisplayClusterShadersTextureViewport& In)
+	{
+		if (In.TextureRHI)
+		{
+			return In.TextureRHI->GetDesc().Extent;
+		}
+		else if (In.TextureRDG)
+		{
+			return In.TextureRDG->Desc.Extent;
+		}
+
+		return FIntPoint::ZeroValue;
+	}
+
+	/**
+	 * Create RHI texture viewport resource 
+	 * 
+	 * @param RHICmdList               - API
+ 	 * @param OutTextureViewport       - (out) The new texture
+	 * @param InTextureViewport        - (in) the source texture to be clone (size and format)
+	 * @param InOutPooledRenderTargets - (in, out) save all temporary textures here so they exist for the lifetime of this external variable
+	 * @param InDebugName              - (in) the unique debug name for the new resource
+	 */
+	static inline bool CloneTextureViewportResource(
+		FRHICommandListImmediate& RHICmdList,
+		FDisplayClusterShadersTextureViewport& OutTextureViewport,
+		const FDisplayClusterShadersTextureViewport& InTextureViewport,
+		TArray<TRefCountPtr<IPooledRenderTarget>>& InOutPooledRenderTargets,
+		const TCHAR* InDebugName)
+	{
+		const FIntPoint Size = GetTextureSize(InTextureViewport);
+		const EPixelFormat Format = GetPixelFormat(InTextureViewport);
+
+		if (Size.GetMin() > 0 && Format != EPixelFormat::PF_Unknown)
+		{
+			// Create a temporary pool texture
+			const FPooledRenderTargetDesc NewResourceDesc = FPooledRenderTargetDesc::Create2DDesc(
+				Size, Format, FClearValueBinding::None,
+				TexCreate_None, TexCreate_RenderTargetable | TexCreate_ShaderResource, false);
+
+			TRefCountPtr<IPooledRenderTarget> RenderTargetPoolResource;
+			GRenderTargetPool.FindFreeElement(RHICmdList, NewResourceDesc, RenderTargetPoolResource, InDebugName);
+
+			if (RenderTargetPoolResource.IsValid())
+			{
+				if (FRHITexture* RHITexture = RenderTargetPoolResource->GetRHI())
+				{
+					// Maintains an internal link to this resource.
+					// It will be released late, after the TS has completed all operations on that resource.
+					InOutPooledRenderTargets.Add(RenderTargetPoolResource);
+
+					OutTextureViewport = FDisplayClusterShadersTextureViewport(RHITexture, FIntRect(FIntPoint::ZeroValue, Size));
+
+					return true;
+				}
+			}
+		}
+
+		return false;
+	}
+
+	/**
+	 * Initialize texture viewport for RHI
+	 *
+	 * @param RHICmdList         - API
+	 * @param InTextureViewport  - Texture for cloning
+	 */
+	static inline bool InitializeTextureViewport(
+		FRHICommandListImmediate& RHICmdList,
+		FDisplayClusterShadersTextureViewport& InOutTextureViewport)
+	{
+		if (InOutTextureViewport.TextureRDG && HasBeenProduced(InOutTextureViewport.TextureRDG))
+		{
+			if (FRHITexture* TextureRHI = InOutTextureViewport.TextureRDG->GetRHI())
+			{
+				InOutTextureViewport.TextureRHI = TextureRHI;
+				InOutTextureViewport.TextureRDG = nullptr;
+
+				return true;
+			}
+		}
+		else if (InOutTextureViewport.TextureRHI)
+		{
+			return true;
+		}
+
+		return false;
+	}
+
+	/**
+	 * Create RDG texture viewport resource
+	 *
+	 * @param GraphBuilder       - RDG API
+	 * @param OutTextureViewport - (out) The new texture
+	 * @param InTextureViewport  - (in) the source texture to be clone (size and format)
+	 * @param InDebugName        - (in) the unique debug name for the new resource
+	 */
+	static inline bool CloneTextureViewportResource(
+		FRDGBuilder& GraphBuilder,
+		FDisplayClusterShadersTextureViewport& OutTextureViewport,
+		const FDisplayClusterShadersTextureViewport& InTextureViewport,
+		const TCHAR* InDebugName)
+	{
+		const FIntPoint Size = GetTextureSize(InTextureViewport);
+		const EPixelFormat Format = GetPixelFormat(InTextureViewport);
+
+		if (Size.GetMin() > 0 && Format != EPixelFormat::PF_Unknown)
+		{
+			// Use temporary RTT
+			const FRDGTextureDesc TemporaryTextureDesc = FRDGTextureDesc::Create2D(
+				Size, Format, FClearValueBinding::None,
+				TexCreate_RenderTargetable | TexCreate_ShaderResource);
+
+			if (FRDGTextureRef RDGTexture = GraphBuilder.CreateTexture(TemporaryTextureDesc, InDebugName))
+			{
+				OutTextureViewport = FDisplayClusterShadersTextureViewport(RDGTexture, FIntRect(FIntPoint::ZeroValue, Size));
+
+				return true;
+			}
+		}
+
+		return false;
+	}
+
+	/**
+	 * Initialize texture viewport for RDG
+	 *
+	 * @param GraphBuilder       - API
+	 * @param InTextureViewport  - Texture for cloning
+	 * @param InDebugName        - the unique debug name for the new resource
+	 */
+	static inline bool InitializeTextureViewport(
+		FRDGBuilder& GraphBuilder,
+		FDisplayClusterShadersTextureViewport& InOutTextureViewport,
+		const TCHAR* InDebugName)
+	{
+		if (InOutTextureViewport.TextureRHI)
+		{
+			const TCHAR* DebugName = (InOutTextureViewport.DebugName && InOutTextureViewport.DebugName[0])
+				? InOutTextureViewport.DebugName : InDebugName;
+
+			InOutTextureViewport.bExternalTextureRDG = true;
+
+			InOutTextureViewport.TextureRDG = RegisterExternalTexture(GraphBuilder, InOutTextureViewport.TextureRHI->GetTexture2D(), DebugName);
+			InOutTextureViewport.TextureRHI = nullptr;
+
+			return true;
+		}
+		else if (InOutTextureViewport.bExternalTextureRDG || HasBeenProduced(InOutTextureViewport.TextureRDG))
+		{
+			return true;
+		}
+
+		return false;
+	}
+
+	/**
+	* Check if resources with the specified regions can be resolved.
+	* If any rect exceeds the texture size, RHI will crash.
+	* This function adjusts the rects to the size of the textures.
+	*
+	* @param InOutSource      - (in, out) Input texture with rect
+	* @param InOutDestination - (inm out) Output texture with rect
+	* @param Settings         - (in) current settings
+	*
+	* @return false, if it isn't possible
+	*/
+	static inline bool UpdateResourcesRectsForResolve(
+		FDisplayClusterShadersTextureViewport& InOutSource,
+		FDisplayClusterShadersTextureViewport& InOutDestination,
+		const FDisplayClusterShadersTextureUtilsSettings& Settings)
+	{
+		using namespace UE::DisplayClusterShaders::Private;
+
+		const FIntPoint InputTextureSize = GetTextureSize(InOutSource);
+		const FIntPoint OutputTextureSize = GetTextureSize(InOutDestination);
+
+		// One of the texture dimensions is zero
+		if (InputTextureSize.GetMin() <= 0 || OutputTextureSize.GetMin() <= 0)
+		{
+			return false;
+		}
+
+		if (InOutSource.Rect.IsEmpty())
+		{
+			InOutSource.Rect = FIntRect(FIntPoint(0, 0), InputTextureSize);
+		}
+		if (InOutDestination.Rect.IsEmpty())
+		{
+			InOutDestination.Rect = FIntRect(FIntPoint(0, 0), OutputTextureSize);
+		}
+
+		if (Settings.HasAnyFlags(EDisplayClusterShaderTextureUtilsFlags::DisableUpdateResourcesRectsForResolve))
+		{
+			return true;
+		}
+
+		FIntRect InputRect(InOutSource.Rect);
+		FIntRect OutputRect(InOutDestination.Rect);
+
+		// if InputRect.Min<0, also adjust the OutputRect.Min
+		OutputRect.Min += FIntPoint(
+			FMath::Max(0, -InOutSource.Rect.Min.X),
+			FMath::Max(0, -InOutSource.Rect.Min.Y));
+
+		// if OutputRect.Min<0, also adjust the InputRect.Min
+		InputRect.Min += FIntPoint(
+			FMath::Max(0, -InOutDestination.Rect.Min.X),
+			FMath::Max(0, -InOutDestination.Rect.Min.Y));
+
+		// If InputRect or OutputRect exceeds the texture size, RHI will crash. Let's adjust it to the texture size.
+		{
+			InputRect.Min.X = FMath::Clamp(InputRect.Min.X, 0, InputTextureSize.X);
+			InputRect.Min.Y = FMath::Clamp(InputRect.Min.Y, 0, InputTextureSize.Y);
+
+			InputRect.Max.X = FMath::Clamp(InputRect.Max.X, 0, InputTextureSize.X);
+			InputRect.Max.Y = FMath::Clamp(InputRect.Max.Y, 0, InputTextureSize.Y);
+
+			OutputRect.Min.X = FMath::Clamp(OutputRect.Min.X, 0, OutputTextureSize.X);
+			OutputRect.Min.Y = FMath::Clamp(OutputRect.Min.Y, 0, OutputTextureSize.Y);
+
+			OutputRect.Max.X = FMath::Clamp(OutputRect.Max.X, 0, OutputTextureSize.X);
+			OutputRect.Max.Y = FMath::Clamp(OutputRect.Max.Y, 0, OutputTextureSize.Y);
+		}
+
+		// InputRect.Min and OutputRect.Min always > 0
+
+		// Check the SrcRect and DestRect
+		if (InputRect.Size().GetMin() <= 0
+			|| OutputRect.Size().GetMin() <= 0)
+		{
+			// The SrcRect or DestRect is invalid.
+			return false;
+		}
+
+		// Input and output rects will be cropped to match each other.
+		if (EnumHasAnyFlags(Settings.Flags, EDisplayClusterShaderTextureUtilsFlags::DisableResize))
+		{
+			const FIntPoint MinSize(
+				FMath::Min(InputRect.Size().X, OutputRect.Size().X),
+				FMath::Min(InputRect.Size().Y, OutputRect.Size().Y));
+
+			// Crop the rects to the smallest size
+			InputRect.Max = InputRect.Min + MinSize;
+			OutputRect.Max = OutputRect.Min + MinSize;
+		}
+
+		// Can be resolved
+		InOutSource.Rect = InputRect;
+		InOutDestination.Rect = OutputRect;
+
+		return true;
+	}
+};
+
+/**
+* FDisplayClusterShadersTextureUtils
+*/
+TSharedRef<IDisplayClusterShadersTextureUtils> FDisplayClusterShadersTextureUtils::CreateTextureUtils_RenderThread(FRHICommandListImmediate& RHICmdList)
+{
+	check(IsInRenderingThread());
+
+	return MakeShared<FDisplayClusterShadersRHITextureUtils>(RHICmdList);
+}
+
+TSharedRef<IDisplayClusterShadersTextureUtils> FDisplayClusterShadersTextureUtils::CreateTextureUtils_RenderThread(FRDGBuilder& GraphBuilder)
+{
+	check(IsInRenderingThread());
+
+	return MakeShared<FDisplayClusterShadersRDGTextureUtils>(GraphBuilder);
+}
+
+TSharedRef<IDisplayClusterShadersTextureUtils> FDisplayClusterShadersTextureUtils::ResolveTextureContext(const FDisplayClusterShadersTextureViewportContext& Input, const FDisplayClusterShadersTextureViewportContext& Output)
+{
+	return ResolveTextureContext(FDisplayClusterShadersTextureUtilsSettings(), Input, Output);
+}
+
+TSharedRef<IDisplayClusterShadersTextureUtils> FDisplayClusterShadersTextureUtils::Resolve()
+{
+	return Resolve(FDisplayClusterShadersTextureUtilsSettings());
+}
+
+bool FDisplayClusterShadersTextureUtils::ImplementTextureContextResolve(
+	const FDisplayClusterShadersTextureViewportContext& Input,
+	const FDisplayClusterShadersTextureViewportContext& Output,
+	const FDisplayClusterShadersTextureUtilsSettings& Settings)
+{
+	// Don't use a shader if possible.
+	if (!ShouldUseResampleShader(Input, Output, Settings))
+	{
+		return TransitionAndCopyTexture(Input, Output, Settings);
+	}
+	else if (Settings.HasAnyFlags(EDisplayClusterShaderTextureUtilsFlags::DisableResampleShader))
+	{
+		// A resampling shader needs to be used, but it is disabled by the user.
+		return false;
+	}
+
+	// Custom implementations must perform the copying from the output texture to the input texture themselves.
+	if (Settings.HasAnyFlags(EDisplayClusterShaderTextureUtilsFlags::UseOutputTextureAsInput))
+	{
+		if (!TransitionAndCopyTexture(Output, Input, Settings))
+		{
+			return false;
+		}
+	}
+
+	return ResampleColorEncodingCopyRect(Input, Output, Settings);
+}
+
+TSharedRef<IDisplayClusterShadersTextureUtils> FDisplayClusterShadersTextureUtils::Resolve(const FDisplayClusterShadersTextureUtilsSettings& Settings)
+{
+	return ForEachContextByPredicate([&](
+		const FDisplayClusterShadersTextureViewportContext& Input,
+		const FDisplayClusterShadersTextureViewportContext& Output)
+		{
+			ImplementTextureContextResolve(Input, Output, Settings);
+		});
+}
+
+TSharedRef<IDisplayClusterShadersTextureUtils> FDisplayClusterShadersTextureUtils::ResolveTextureContext(const FDisplayClusterShadersTextureUtilsSettings& Settings, const FDisplayClusterShadersTextureViewportContext& InputContext, const FDisplayClusterShadersTextureViewportContext& OutputContext)
+{
+	using namespace UE::DisplayClusterShaders::Private;
+
+	FDisplayClusterShadersTextureViewportContext Input(InputContext);
+	FDisplayClusterShadersTextureViewportContext Output(OutputContext);
+
+	if (InitializeTextureViewportForRenderPass(Output, TEXT("nDisplay.Output")))
+	{
+		const bool bInputInitialized = Settings.HasAnyFlags(EDisplayClusterShaderTextureUtilsFlags::UseOutputTextureAsInput)
+			? CloneTextureViewportResourceForRenderPass(Input, Output, TEXT("nDisplay.OutputClone"))
+			: InitializeTextureViewportForRenderPass(Input, TEXT("nDisplay.Input"));
+		if (bInputInitialized && UpdateResourcesRectsForResolve(Input, Output, Settings))
+		{
+			ImplementTextureContextResolve(Input, Output, Settings);
+		}
+	}
+
+	return SharedThis(this);
+}
+
+FDisplayClusterShadersTextureParameters FDisplayClusterShadersTextureUtils::GetTextureParametersFromViewport(
+	const IDisplayClusterViewportProxy* InViewportProxy,
+	const EDisplayClusterViewportResourceType InResourceType)
+{
+	using namespace UE::DisplayClusterShaders::Private;
+
+	FDisplayClusterShadersTextureParameters OutTextureParameters;
+
+	TArray<FRHITexture*> Textures;
+	TArray<FIntRect> TextureRects;
+	if (InViewportProxy && InViewportProxy->GetResourcesWithRects_RenderThread(InResourceType, Textures, TextureRects)
+		&& Textures.Num() == TextureRects.Num())
+	{
+		// Get resource color encoding
+		OutTextureParameters.ColorEncoding = InViewportProxy->GetResourceColorEncoding_RenderThread(InResourceType);
+
+		// Get all contexts
+		for (int32 ContextNum = 0; ContextNum < Textures.Num(); ContextNum++)
+		{
+			if (Textures[ContextNum])
+			{
+				OutTextureParameters.TextureViewports.Emplace(ContextNum,
+					FDisplayClusterShadersTextureViewport(
+						Textures[ContextNum],
+						TextureRects[ContextNum],
+						GetDisplayClusterViewportResourceTypeName(InResourceType)));
+			}
+		}
+	}
+
+	return OutTextureParameters;
+}
+
+TSharedRef<IDisplayClusterShadersTextureUtils> FDisplayClusterShadersTextureUtils::SetInputEncoding(const FDisplayClusterColorEncoding& InColorEncoding)
+{
+	InputTextureParameters.ColorEncoding = InColorEncoding;
+
+	return SharedThis(this);
+}
+TSharedRef<IDisplayClusterShadersTextureUtils> FDisplayClusterShadersTextureUtils::SetOutputEncoding(const FDisplayClusterColorEncoding& InColorEncoding)
+{
+	OutputTextureParameters.ColorEncoding = InColorEncoding;
+
+	return SharedThis(this);
+}
+
+
+TSharedRef<IDisplayClusterShadersTextureUtils> FDisplayClusterShadersTextureUtils::SetInput(
+	const IDisplayClusterViewportProxy* InViewportProxy,
+	const EDisplayClusterViewportResourceType InResourceType,
+	const int32 InContextNum)
+{
+	const FDisplayClusterShadersTextureParameters NewTextureParameters = GetTextureParametersFromViewport(InViewportProxy, InResourceType);
+
+	if (InContextNum == INDEX_NONE)
+	{
+		// Override entire input parameters.
+		SetInputEncoding(NewTextureParameters.ColorEncoding);
+		for (const TPair<uint32, FDisplayClusterShadersTextureViewport>& TextureIt : NewTextureParameters.TextureViewports)
+		{
+			SetInput(TextureIt.Value, TextureIt.Key);
+		}
+	}
+	else if(NewTextureParameters.TextureViewports.Contains(InContextNum))
+	{
+		SetInputEncoding(NewTextureParameters.ColorEncoding);
+		SetInput(NewTextureParameters.TextureViewports[InContextNum], InContextNum);
+	}
+
+
+	return SharedThis(this);
+}
+
+TSharedRef<IDisplayClusterShadersTextureUtils> FDisplayClusterShadersTextureUtils::SetOutput(
+	const IDisplayClusterViewportProxy* InViewportProxy,
+	const EDisplayClusterViewportResourceType InResourceType,
+	const int32 InContextNum)
+{
+	const FDisplayClusterShadersTextureParameters NewTextureParameters = GetTextureParametersFromViewport(InViewportProxy, InResourceType);
+
+	if (InContextNum == INDEX_NONE)
+	{
+		// Override entire input parameters.
+		SetOutputEncoding(NewTextureParameters.ColorEncoding);
+		for (const TPair<uint32, FDisplayClusterShadersTextureViewport>& TextureIt : NewTextureParameters.TextureViewports)
+		{
+			SetOutput(TextureIt.Value, TextureIt.Key);
+		}
+	}
+	else if (NewTextureParameters.TextureViewports.Contains(InContextNum))
+	{
+		SetOutputEncoding(NewTextureParameters.ColorEncoding);
+		SetOutput(NewTextureParameters.TextureViewports[InContextNum], InContextNum);
+	}
+
+	return SharedThis(this);
+}
+
+TSharedRef<IDisplayClusterShadersTextureUtils> FDisplayClusterShadersTextureUtils::SetInput(const FDisplayClusterShadersTextureViewport& InTextureViewport, const int32 InContextNum)
+{
+	const uint32 ContextNum = InContextNum > 0 ? InContextNum : 0;
+	InputTextureParameters.TextureViewports.Emplace(ContextNum, InTextureViewport);
+
+	return SharedThis(this);
+}
+
+TSharedRef<IDisplayClusterShadersTextureUtils> FDisplayClusterShadersTextureUtils::SetOutput(const FDisplayClusterShadersTextureViewport& InTextureViewport, const int32 InContextNum)
+{
+	const uint32 ContextNum = InContextNum > 0 ? InContextNum : 0;
+	OutputTextureParameters.TextureViewports.Emplace(ContextNum, InTextureViewport);
+
+	return SharedThis(this);
+}
+
+/** Iterate through contexts matching these settings. */
+void FDisplayClusterShadersTextureUtils::ImplForEachContextByPredicate(
+	const FDisplayClusterShadersTextureUtilsSettings& Settings,
+	TFunctionDisplayClusterShaders_TextureContextIterator&& TextureContextIteratorFunc)
+{
+	using namespace UE::DisplayClusterShaders::Private;
+	// The input is ignored, and temporary textures is used for every output texture
+	if (Settings.HasAnyFlags(EDisplayClusterShaderTextureUtilsFlags::UseOutputTextureAsInput))
+	{
+		for (const TPair<uint32, FDisplayClusterShadersTextureViewport>& OutputTextureIt : OutputTextureParameters.TextureViewports)
+		{
+			const uint32 OutputContextNum = OutputTextureIt.Key;
+			FDisplayClusterShadersTextureViewportContext Output(OutputTextureIt.Value, OutputTextureParameters.ColorEncoding, OutputContextNum);
+			if (InitializeTextureViewportForRenderPass(Output, TEXT("nDisplay.Output")))
+			{
+				FDisplayClusterShadersTextureViewportContext Input;
+				if(CloneTextureViewportResourceForRenderPass(Input, Output, TEXT("nDisplay.OutputClone")))
+				{
+					// Use input color encoding
+					Input.ColorEncoding = InputTextureParameters.ColorEncoding;
+					Input.ContextNum = OutputContextNum;
+
+					if (InitializeTextureViewportForRenderPass(Input, TEXT("nDisplay.Input")))
+					{
+						if (UpdateResourcesRectsForResolve(Input, Output, Settings))
+						{
+							// Copy Output to Input (temp)
+							TransitionAndCopyTexture(Output, Input, Settings);
+
+							// Continue the process using the external functor.
+							if (Settings.HasAnyFlags(EDisplayClusterShaderTextureUtilsFlags::InvertDirection))
+							{
+								TextureContextIteratorFunc(Output, Input);
+							}
+							else
+							{
+								TextureContextIteratorFunc(Input, Output);
+							}
+						}
+					}
+				}
+			}
+		}
+		return;
+	}
+
+	// Map input to output:
+	for (const TPair<uint32, FDisplayClusterShadersTextureViewport>& OutputTextureIt : OutputTextureParameters.TextureViewports)
+	{
+		// Monoscopic input can be copied to stereoscopic output
+		const uint32 OutputContextNum = OutputTextureIt.Key;
+		const uint32 InputContextNum = InputTextureParameters.TextureViewports.Contains(OutputContextNum) ? OutputContextNum : 0;
+
+		if (const FDisplayClusterShadersTextureViewport* InputTexturePtr = InputTextureParameters.TextureViewports.Find(InputContextNum))
+		{
+			FDisplayClusterShadersTextureViewportContext Output(OutputTextureIt.Value, OutputTextureParameters.ColorEncoding, OutputContextNum);
+			FDisplayClusterShadersTextureViewportContext Input(*InputTexturePtr,       InputTextureParameters.ColorEncoding,  InputContextNum);
+
+			if (UpdateResourcesRectsForResolve(Input, Output, Settings)
+				&& InitializeTextureViewportForRenderPass(Input, TEXT("nDisplay.Input"))
+				&& InitializeTextureViewportForRenderPass(Output, TEXT("nDisplay.Output")))
+			{
+				// Continue the process using the external functor.
+				if (Settings.HasAnyFlags(EDisplayClusterShaderTextureUtilsFlags::InvertDirection))
+				{
+					TextureContextIteratorFunc(Output, Input);
+				}
+				else
+				{
+					TextureContextIteratorFunc(Input, Output);
+				}
+			}
+		}
+	}
+}
+
+TSharedRef<IDisplayClusterShadersTextureUtils> FDisplayClusterShadersTextureUtils::ForEachContextByPredicate(TFunctionDisplayClusterShaders_TextureContextIterator&& InFunction)
+{
+	const FDisplayClusterShadersTextureUtilsSettings DefaultSettings;
+
+	ImplForEachContextByPredicate(DefaultSettings, [&](
+		const FDisplayClusterShadersTextureViewportContext& Input,
+		const FDisplayClusterShadersTextureViewportContext& Output)
+		{
+			InFunction(Input, Output);
+		});
+
+	return SharedThis(this);
+}
+
+TSharedRef<IDisplayClusterShadersTextureUtils> FDisplayClusterShadersTextureUtils::ForEachContextByPredicate(
+	const FDisplayClusterShadersTextureUtilsSettings& Settings,
+	TFunctionDisplayClusterShaders_TextureContextIterator&& InFunction)
+{
+	ImplForEachContextByPredicate(Settings, [&](
+		const FDisplayClusterShadersTextureViewportContext& Input,
+		const FDisplayClusterShadersTextureViewportContext& Output)
+		{
+			InFunction(Input, Output);
+		});
+
+	return SharedThis(this);
+}
+
+bool FDisplayClusterShadersTextureUtils::ShouldUseResampleShader(
+	const FDisplayClusterShadersTextureViewportContext& InputTextureContext,
+	const FDisplayClusterShadersTextureViewportContext& OutputTextureContext,
+	const FDisplayClusterShadersTextureUtilsSettings& Settings) const
+{
+	using namespace UE::DisplayClusterShaders::Private;
+
+	if (Settings.ColorMask != EColorWriteMask::CW_RGBA)
+	{
+		// Color mask requires a shader
+		return true;
+	}
+
+	if (Settings.OverrideAlpha != EDisplayClusterShaderTextureUtilsOverrideAlpha::None)
+	{
+		// Use shader to overide alpha channel value
+		return true;
+	}
+
+	if (InputTextureContext.Rect.Size() != OutputTextureContext.Rect.Size())
+	{
+		// Resizing should be done using a resampling shader.
+		return true;
+	}
+
+	if (GetPixelFormat(InputTextureContext) != GetPixelFormat(OutputTextureContext))
+	{
+		// changing format requires a shader
+		return true;
+	}
+
+	if (InputTextureContext.ColorEncoding != OutputTextureContext.ColorEncoding)
+	{
+		// Encode color in pixel shader
+		return true;
+	}
+
+	return false;
+}
+
+/**
+* FDisplayClusterShadersRHITextureUtils
+*/
+bool FDisplayClusterShadersRHITextureUtils::TransitionAndCopyTexture(const FDisplayClusterShadersTextureViewport& Input, const FDisplayClusterShadersTextureViewport& Output, const FDisplayClusterShadersTextureUtilsSettings& Settings)
+{
+	if (Input.TextureRDG && Output.TextureRDG)
+	{
+		return FDisplayClusterShadersCopyTexture::AddPassTransitionAndCopyTexture_RenderThread(GetOrCreateRDGBuilder(), Input, Output, Settings);
+	}
+	else if (Input.TextureRHI && Output.TextureRHI)
+	{
+		return FDisplayClusterShadersCopyTexture::TransitionAndCopyTexture_RenderThread(RHICmdList, Input, Output, Settings);
+	}
+
+	return false;
+}
+
+bool FDisplayClusterShadersRHITextureUtils::ResampleColorEncodingCopyRect(const FDisplayClusterShadersTextureViewportContext& Input, const FDisplayClusterShadersTextureViewportContext& Output, const FDisplayClusterShadersTextureUtilsSettings& Settings)
+{
+	if (Input.TextureRDG && Output.TextureRDG)
+	{
+		return FDisplayClusterShadersCopyTexture::AddPassColorEncodingCopyRect_RenderThread(GetOrCreateRDGBuilder(), Input, Output, Settings);
+	}
+	else if (Input.TextureRHI && Output.TextureRHI)
+	{
+		return FDisplayClusterShadersCopyTexture::ColorEncodingCopyRect_RenderThread(RHICmdList, Input, Output, Settings);
+	}
+
+	return false;
+}
+
+bool FDisplayClusterShadersRHITextureUtils::CloneTextureViewportResourceForRenderPass(
+	FDisplayClusterShadersTextureViewport& OutTextureViewport,
+	const FDisplayClusterShadersTextureViewport& InTextureViewport,
+	const TCHAR* InDebugName)
+{
+	using namespace UE::DisplayClusterShaders::Private;
+
+	if (ShouldUseRDG() || InTextureViewport.TextureRDG)
+	{
+		return CloneTextureViewportResource(GetOrCreateRDGBuilder(), OutTextureViewport, InTextureViewport, InDebugName);
+	}
+
+	return CloneTextureViewportResource(RHICmdList, OutTextureViewport, InTextureViewport, PooledRenderTargets, InDebugName);
+}
+
+bool FDisplayClusterShadersRHITextureUtils::InitializeTextureViewportForRenderPass(
+	FDisplayClusterShadersTextureViewport& InOutTextureViewport,
+	const TCHAR* InDebugName)
+{
+	using namespace UE::DisplayClusterShaders::Private;
+
+	if (ShouldUseRDG() || InOutTextureViewport.TextureRDG)
+	{
+		return InitializeTextureViewport(GetOrCreateRDGBuilder(), InOutTextureViewport, InDebugName);
+	}
+
+	return InitializeTextureViewport(RHICmdList, InOutTextureViewport);
+}
+
+/**
+* FDisplayClusterShadersRDGTextureUtils
+*/
+bool FDisplayClusterShadersRDGTextureUtils::CloneTextureViewportResourceForRenderPass(
+	FDisplayClusterShadersTextureViewport& OutTextureViewport,
+	const FDisplayClusterShadersTextureViewport& InTextureViewport,
+	const TCHAR* InDebugName)
+{
+	using namespace UE::DisplayClusterShaders::Private;
+
+	return CloneTextureViewportResource(GraphBuilder, OutTextureViewport, InTextureViewport, InDebugName);
+}
+
+bool FDisplayClusterShadersRDGTextureUtils::InitializeTextureViewportForRenderPass(
+	FDisplayClusterShadersTextureViewport& InOutTextureViewport,
+	const TCHAR* InDebugName)
+{
+	using namespace UE::DisplayClusterShaders::Private;
+
+	return InitializeTextureViewport(GraphBuilder, InOutTextureViewport, InDebugName);
+}
+
+bool FDisplayClusterShadersRDGTextureUtils::TransitionAndCopyTexture(const FDisplayClusterShadersTextureViewport& Input, const FDisplayClusterShadersTextureViewport& Output, const FDisplayClusterShadersTextureUtilsSettings& Settings)
+{
+	return FDisplayClusterShadersCopyTexture::AddPassTransitionAndCopyTexture_RenderThread(GraphBuilder, Input, Output, Settings);
+}
+
+bool FDisplayClusterShadersRDGTextureUtils::ResampleColorEncodingCopyRect(const FDisplayClusterShadersTextureViewportContext& Input, const FDisplayClusterShadersTextureViewportContext& Output, const FDisplayClusterShadersTextureUtilsSettings& Settings)
+{
+	return FDisplayClusterShadersCopyTexture::AddPassColorEncodingCopyRect_RenderThread(GraphBuilder, Input, Output, Settings);
+}
diff --git a/Engine/Plugins/Runtime/nDisplay/Source/DisplayClusterShaders/Private/DisplayClusterShadersTextureUtils.h b/Engine/Plugins/Runtime/nDisplay/Source/DisplayClusterShaders/Private/DisplayClusterShadersTextureUtils.h
new file mode 100644
--- /dev/null
+++ b/Engine/Plugins/Runtime/nDisplay/Source/DisplayClusterShaders/Private/DisplayClusterShadersTextureUtils.h
@@ -0,0 +1,281 @@
+// Copyright Epic Games, Inc. All Rights Reserved.
+
+#pragma once
+#include "IDisplayClusterShadersTextureUtils.h"
+#include "Templates/SharedPointer.h"
+
+#include "RHICommandList.h"
+
+/**
+ *  Texture utils class for nDisplay.
+ */
+class FDisplayClusterShadersTextureUtils
+	: public IDisplayClusterShadersTextureUtils
+	, public TSharedFromThis<FDisplayClusterShadersTextureUtils>
+{
+public:
+	virtual ~FDisplayClusterShadersTextureUtils() = default;
+
+	/**
+	* Return new instance of the resource utils.
+	*
+	* @param RHICmdList   - RHI API
+	*/
+	static TSharedRef<IDisplayClusterShadersTextureUtils> CreateTextureUtils_RenderThread(FRHICommandListImmediate& RHICmdList);
+
+	/**
+	* Return new instance of the resource utils.
+	* @param GraphBuilder - RDG builder
+	*/
+	static TSharedRef<IDisplayClusterShadersTextureUtils> CreateTextureUtils_RenderThread(FRDGBuilder& GraphBuilder);
+
+
+public:
+	//~Begin IDisplayClusterShadersTextureUtils
+	virtual TSharedRef<IDisplayClusterShadersTextureUtils> ResolveTextureContext(const FDisplayClusterShadersTextureViewportContext& Input, const FDisplayClusterShadersTextureViewportContext& Output) override;
+	virtual TSharedRef<IDisplayClusterShadersTextureUtils> ResolveTextureContext(const FDisplayClusterShadersTextureUtilsSettings& InSettings, const FDisplayClusterShadersTextureViewportContext& Input, const FDisplayClusterShadersTextureViewportContext& Output) override;
+
+	virtual TSharedRef<IDisplayClusterShadersTextureUtils> Resolve() override;
+	virtual TSharedRef<IDisplayClusterShadersTextureUtils> Resolve(const FDisplayClusterShadersTextureUtilsSettings& Settings) override;
+
+	virtual TSharedRef<IDisplayClusterShadersTextureUtils> ForEachContextByPredicate(TFunctionDisplayClusterShaders_TextureContextIterator&& InFunction) override;
+	virtual TSharedRef<IDisplayClusterShadersTextureUtils> ForEachContextByPredicate(const FDisplayClusterShadersTextureUtilsSettings& Settings, TFunctionDisplayClusterShaders_TextureContextIterator&& InFunction) override;
+
+	virtual TSharedRef<IDisplayClusterShadersTextureUtils> SetInput(const FDisplayClusterShadersTextureViewport& InTextureViewport, const int32 InContextNum = INDEX_NONE) override;
+	virtual TSharedRef<IDisplayClusterShadersTextureUtils> SetOutput(const FDisplayClusterShadersTextureViewport& InTextureViewport, const int32 InContextNum = INDEX_NONE) override;
+
+	virtual TSharedRef<IDisplayClusterShadersTextureUtils> SetInput(const IDisplayClusterViewportProxy* InViewportProxy, const EDisplayClusterViewportResourceType InResourceType, const int32 InContextNum = INDEX_NONE) override;
+	virtual TSharedRef<IDisplayClusterShadersTextureUtils> SetOutput(const IDisplayClusterViewportProxy* InViewportProxy, const EDisplayClusterViewportResourceType InResourceType,const int32 InContextNum = INDEX_NONE) override;
+
+
+	virtual TSharedRef<IDisplayClusterShadersTextureUtils> SetInputEncoding(const FDisplayClusterColorEncoding& InColorEncoding) override;
+	virtual TSharedRef<IDisplayClusterShadersTextureUtils> SetOutputEncoding(const FDisplayClusterColorEncoding& InColorEncoding) override;
+
+	virtual const FDisplayClusterShadersTextureParameters& GetInputTextureParameters() const override
+	{
+		return InputTextureParameters;
+	}
+	virtual const FDisplayClusterShadersTextureParameters& GetOutputTextureParameters() const override
+	{
+		return OutputTextureParameters;
+	}
+
+	//~~ End IDisplayClusterShadersTextureUtils
+
+public:
+	/** Return true if RDG is required. */
+	virtual bool ShouldUseRDG() const
+	{
+		return false;
+	}
+
+	/**
+	* Implements Transition and copy
+	*
+	* @param Input    - Input texture with rect
+	* @param Output   - Output texture with rect
+	* @param Settings - current settings
+	*
+	* @return false, if it isn't possible
+	*/
+	virtual bool TransitionAndCopyTexture(
+		const FDisplayClusterShadersTextureViewport& Input,
+		const FDisplayClusterShadersTextureViewport& Output,
+		const FDisplayClusterShadersTextureUtilsSettings& Settings)
+	{
+		return false;
+	}
+
+	/**
+	* Implements resample shader
+	*
+	* @param Input    - Input texture with rect
+	* @param Output   - Output texture with rect
+	* @param Settings - current settings
+	*
+	* @return false, if it isn't possible
+	*/
+	virtual bool ResampleColorEncodingCopyRect(
+		const FDisplayClusterShadersTextureViewportContext& Input,
+		const FDisplayClusterShadersTextureViewportContext& Output,
+		const FDisplayClusterShadersTextureUtilsSettings& Settings)
+	{
+		return false;
+	}
+
+	/** Create a new input texture from the output and copy the image as well.
+	* 
+	* @param OutTextureViewport - (out) The new texture
+	* @param InTextureViewport  - (in) the source texture to be clone (size and format)
+	* @param InDebugName - (opt) texture debug name
+	* 
+	* @return true on success.
+	*/
+	virtual bool CloneTextureViewportResourceForRenderPass(
+		FDisplayClusterShadersTextureViewport& OutTextureViewport,
+		const FDisplayClusterShadersTextureViewport& InTextureViewport,
+		const TCHAR* InDebugName)
+	{
+		return false;
+	};
+
+	/** Update texture viewport data for render pass
+	* return false if error
+	*/
+	virtual bool InitializeTextureViewportForRenderPass(
+		FDisplayClusterShadersTextureViewport& InOutTextureViewport,
+		const TCHAR* InDebugName)
+	{
+		return true;
+	};
+
+	/**
+	* Implements texture context resolving
+	*
+	* @param Input    - Input texture with rect
+	* @param Output   - Output texture with rect
+	* @param Settings - current settings
+	*
+	* @return false, if it isn't possible
+	*/
+	virtual bool ImplementTextureContextResolve(
+		const FDisplayClusterShadersTextureViewportContext& Input,
+		const FDisplayClusterShadersTextureViewportContext& Output,
+		const FDisplayClusterShadersTextureUtilsSettings& Settings);
+
+public:
+	/** Get texture parameter from the nDisplay viewport resourece
+	*
+	* @param InViewportProxy  - viewport proxy object with resources to retrieve
+	* @param InResourceType   - Type of the viewport resource
+	*
+	* @return themself
+	*/
+	static FDisplayClusterShadersTextureParameters GetTextureParametersFromViewport(
+		const IDisplayClusterViewportProxy* InViewportProxy,
+		const EDisplayClusterViewportResourceType InResourceType);
+
+	/** Returns true if a resampling shader should be used to copy this texture.
+	*
+	* @param InputTextureContext  - input texture context
+	* @param OutputTextureContext - output texture context
+	* @param Settings             - current settings
+	*/
+	bool ShouldUseResampleShader(
+		const FDisplayClusterShadersTextureViewportContext& InputTextureContext,
+		const FDisplayClusterShadersTextureViewportContext& OutputTextureContext,
+		const FDisplayClusterShadersTextureUtilsSettings& Settings) const;
+
+
+	/** Iterate through contexts matching these settings. */
+	void ImplForEachContextByPredicate(
+		const FDisplayClusterShadersTextureUtilsSettings& Settings,
+		TFunctionDisplayClusterShaders_TextureContextIterator&& TextureContextIteratorFunc);
+
+protected:
+	/** Input textures parameters. */
+	FDisplayClusterShadersTextureParameters InputTextureParameters;
+
+	/** Output textures parameters. */
+	FDisplayClusterShadersTextureParameters OutputTextureParameters;
+};
+
+/**
+ *  RDI:Texture utils class for nDisplay.
+ */
+class FDisplayClusterShadersRHITextureUtils
+	: public FDisplayClusterShadersTextureUtils
+{
+public:
+	virtual ~FDisplayClusterShadersRHITextureUtils()
+	{
+		if (GraphBuilderUniquePtr)
+		{
+			// If using RDG, call execute at the end
+			GraphBuilderUniquePtr->Execute();
+			GraphBuilderUniquePtr.Reset();
+		}
+
+		// And finally release pooled RTT
+		PooledRenderTargets.Reset();
+	}
+
+	FDisplayClusterShadersRHITextureUtils(FRHICommandListImmediate& InRHICmdList)
+			: FDisplayClusterShadersTextureUtils(), RHICmdList(InRHICmdList)
+	{ }
+
+public:
+	//~Begin IDisplayClusterShadersTextureUtils
+	virtual FRDGBuilder& GetOrCreateRDGBuilder() override
+	{
+		if (!GraphBuilderUniquePtr.IsValid())
+		{
+			GraphBuilderUniquePtr = MakeUnique<FRDGBuilder>(RHICmdList);
+		}
+
+		return *GraphBuilderUniquePtr;
+	}
+	//~~End IDisplayClusterShadersTextureUtils
+
+public:
+	//~Begin FDisplayClusterShadersTextureUtils
+	virtual bool TransitionAndCopyTexture(const FDisplayClusterShadersTextureViewport& Input, const FDisplayClusterShadersTextureViewport& Output, const FDisplayClusterShadersTextureUtilsSettings& Settings) override;
+	virtual bool ResampleColorEncodingCopyRect(const FDisplayClusterShadersTextureViewportContext& Input, const FDisplayClusterShadersTextureViewportContext& Output, const FDisplayClusterShadersTextureUtilsSettings& Settings) override;
+	virtual bool CloneTextureViewportResourceForRenderPass(FDisplayClusterShadersTextureViewport& OutTextureViewport, const FDisplayClusterShadersTextureViewport& InTextureViewport, const TCHAR* InDebugName) override;
+	virtual bool InitializeTextureViewportForRenderPass(FDisplayClusterShadersTextureViewport& InOutTextureViewport, const TCHAR* InDebugName) override;
+
+	virtual bool ShouldUseRDG() const override
+	{
+		return GraphBuilderUniquePtr.IsValid();
+	}
+
+	//~~End FDisplayClusterShadersTextureUtils
+
+private:
+	// Stored RHI api ref
+	FRHICommandListImmediate& RHICmdList;
+
+	/** Temporary RTTs*/
+	TArray<TRefCountPtr<IPooledRenderTarget>> PooledRenderTargets;
+
+	// A GraphBuilder that can be created on request.
+	TUniquePtr<FRDGBuilder> GraphBuilderUniquePtr;
+};
+
+/**
+ *  RDG:Texture utils class for nDisplay.
+ */
+class FDisplayClusterShadersRDGTextureUtils
+	: public FDisplayClusterShadersTextureUtils
+{
+public:
+	virtual ~FDisplayClusterShadersRDGTextureUtils() = default;	
+	FDisplayClusterShadersRDGTextureUtils(FRDGBuilder& InGraphBuilder)
+		: FDisplayClusterShadersTextureUtils(), GraphBuilder(InGraphBuilder)
+	{ }
+
+public:
+	//~Begin IDisplayClusterShadersTextureUtils
+	virtual FRDGBuilder& GetOrCreateRDGBuilder() override
+	{
+		return GraphBuilder;
+	}
+	//~~End IDisplayClusterShadersTextureUtils
+
+public:
+	//~Begin FDisplayClusterShadersTextureUtils
+	virtual bool TransitionAndCopyTexture(const FDisplayClusterShadersTextureViewport& Input, const FDisplayClusterShadersTextureViewport& Output, const FDisplayClusterShadersTextureUtilsSettings& Settings) override;
+	virtual bool ResampleColorEncodingCopyRect(const FDisplayClusterShadersTextureViewportContext& Input, const FDisplayClusterShadersTextureViewportContext& Output, const FDisplayClusterShadersTextureUtilsSettings& Settings) override;
+	virtual bool CloneTextureViewportResourceForRenderPass(FDisplayClusterShadersTextureViewport& OutTextureViewport, const FDisplayClusterShadersTextureViewport& InTextureViewport, const TCHAR* InDebugName) override;
+	virtual bool InitializeTextureViewportForRenderPass(FDisplayClusterShadersTextureViewport& InOutTextureViewport, const TCHAR* InDebugName) override;
+	virtual bool ShouldUseRDG() const override
+	{
+		return true;
+	}
+	//~~End FDisplayClusterShadersTextureUtils
+
+
+private:
+	// Stored RDG api ref
+	FRDGBuilder& GraphBuilder;
+};
\ No newline at end of file
diff --git a/Engine/Plugins/Runtime/nDisplay/Source/DisplayClusterShaders/Private/Shaders/DisplayClusterShadersCopyTexture.cpp b/Engine/Plugins/Runtime/nDisplay/Source/DisplayClusterShaders/Private/Shaders/DisplayClusterShadersCopyTexture.cpp
new file mode 100644
--- /dev/null
+++ b/Engine/Plugins/Runtime/nDisplay/Source/DisplayClusterShaders/Private/Shaders/DisplayClusterShadersCopyTexture.cpp
@@ -0,0 +1,511 @@
+// Copyright Epic Games, Inc. All Rights Reserved.
+
+#include "Shaders/DisplayClusterShadersCopyTexture.h"
+#include "DisplayClusterShadersLog.h"
+
+#include "GlobalShader.h"
+#include "ShaderParameters.h"
+#include "ShaderParameterUtils.h"
+#include "TextureResource.h"
+
+#include "RHIStaticStates.h"
+
+#include "RenderResource.h"
+#include "RenderingThread.h"
+#include "CommonRenderResources.h"
+#include "PixelShaderUtils.h"
+
+#include "ClearQuad.h"
+
+#include "GlobalShader.h"
+#include "ShaderParameters.h"
+#include "ShaderParameterUtils.h"
+
+#include "ScreenRendering.h"
+#include "RenderGraphUtils.h"
+#include "RenderGraphResources.h"
+#include "PostProcess/DrawRectangle.h"
+#include "PostProcess/PostProcessMaterialInputs.h"
+
+#include "SceneView.h"
+#include "ScreenPass.h"
+
+#include "Engine/TextureRenderTarget.h"
+
+namespace UE::DisplayClusterShaders::Private
+{
+	/** Returns the RHI blend state for the requested DC blend mode. */
+	FRHIBlendState* GetBlendStateRHI(const EColorWriteMask InColorMask)
+	{
+		if (InColorMask == EColorWriteMask::CW_ALPHA)
+		{
+			// Copy alpha channel from source to dest
+			return TStaticBlendState <CW_ALPHA, BO_Add, BF_One, BF_Zero, BO_Add, BF_One, BF_Zero>::GetRHI();
+
+		}
+		else if (InColorMask == EColorWriteMask::CW_RGB)
+		{
+			// Copy only RGB channels from source to dest
+			return TStaticBlendState <CW_RGB, BO_Add, BF_One, BF_Zero, BO_Add, BF_One, BF_Zero>::GetRHI();
+		}
+
+		return TStaticBlendState<>::GetRHI();
+	}
+
+	/** Enumerations that cover definitions in the shader. */
+	enum class ESourceEncoding : uint32
+	{
+		Linear  = 0, // #define ESourceEncoding_Linear  0
+		Gamma   = 1, // #define ESourceEncoding_Gamma   1
+		sRGB    = 2, // #define ESourceEncoding_sRGB    2
+		MediaPQ = 3, // #define ESourceEncoding_MediaPQ 3
+	};
+
+	enum class EColorPremultiply : uint32
+	{
+		None          = 0, // #define EColorPremultiply_None         0
+		Alpha         = 1, // #define EColorPremultiply_Alpha        1
+		InvertedAlpha = 2, // #define EColorPremultiply_InvertedAlpha 2
+	};
+
+	/** Enumerations that cover definitions in the shader. */
+	enum class EOverrideAlpha : uint32
+	{
+		None   = 0, // #define EOverrideAlpha_None   0
+		Invert = 1, // #define EOverrideAlpha_Invert 1
+		One    = 2, // #define EOverrideAlpha_One    2
+		Zero   = 3, // #define EOverrideAlpha_Zero   3
+	};
+
+	namespace FColorEncodingCopyRectPSPermutation
+	{
+		// Shared permutation for picp warp
+		class FPermutationEncodeInput   : SHADER_PERMUTATION_BOOL("ENCODE_INPUT");
+		class FPermutationEncodeOutput  : SHADER_PERMUTATION_BOOL("ENCODE_OUTPUT");
+		class FPermutationOverrideAlpha : SHADER_PERMUTATION_BOOL("OVERRIDE_ALPHA");
+		class FPermutationColorPremultiply : SHADER_PERMUTATION_BOOL("COLOR_PREMULTIPLY");
+
+		using FCommonPSDomain = TShaderPermutationDomain<
+			FPermutationEncodeInput,
+			FPermutationEncodeOutput,
+			FPermutationOverrideAlpha,
+			FPermutationColorPremultiply
+		>;
+
+		bool ShouldCompileCommonPSPermutation(const FCommonPSDomain& PermutationVector)
+		{
+			return true;
+		}
+	};
+
+	/** RDG copy texture parameters. */
+	BEGIN_SHADER_PARAMETER_STRUCT(FDisplayClusterCopyTextureParameters, )
+		RDG_TEXTURE_ACCESS(Input, ERHIAccess::CopySrc)
+		RDG_TEXTURE_ACCESS(Output, ERHIAccess::CopyDest)
+	END_SHADER_PARAMETER_STRUCT()
+
+	/**
+	* Pixel shaders parameters for 'FColorEncodingCopyRectPS'
+	*/
+	BEGIN_SHADER_PARAMETER_STRUCT(FColorEncodingCopyRectPSParameters, )
+		SHADER_PARAMETER_TEXTURE(Texture2D, InputTexture)
+		SHADER_PARAMETER_SAMPLER(SamplerState,  InputTextureSampler)
+		SHADER_PARAMETER(FUint32Vector, ColorPremultiply)
+		SHADER_PARAMETER(FUint32Vector, Encodings)
+		SHADER_PARAMETER(FVector3f, DisplayGamma)
+	END_SHADER_PARAMETER_STRUCT()
+
+	/**
+	 * A pixel shader for TextureShare resource resampling
+	 */
+	class FColorEncodingCopyRectPS : public FGlobalShader
+	{
+	public:
+		DECLARE_GLOBAL_SHADER(FColorEncodingCopyRectPS);
+		SHADER_USE_PARAMETER_STRUCT(FColorEncodingCopyRectPS, FGlobalShader);
+
+		using FPermutationDomain = FColorEncodingCopyRectPSPermutation::FCommonPSDomain;
+		using FParameters = FColorEncodingCopyRectPSParameters;
+
+		static bool ShouldCompilePermutation(const FGlobalShaderPermutationParameters& Parameters)
+		{
+			return FColorEncodingCopyRectPSPermutation::ShouldCompileCommonPSPermutation(FPermutationDomain(Parameters.PermutationId));
+		}
+
+		/** Initialize shader parameters. */
+		static bool InitializeShaderParameters(
+			const FDisplayClusterShadersTextureViewportContext& Input,
+			const FDisplayClusterShadersTextureViewportContext& Output,
+			const FDisplayClusterShadersTextureUtilsSettings& Settings,
+			FParameters& OutParameters,
+			FColorEncodingCopyRectPSPermutation::FCommonPSDomain& OutPermutationVector)
+		{
+			OutParameters.InputTextureSampler = (Input.Rect.Size() == Output.Rect.Size())
+				? TStaticSamplerState<SF_Point>::GetRHI() // Use the 'Point' sampler if the input and output texture sizes are equal
+				: TStaticSamplerState<SF_Bilinear>::GetRHI();
+
+			// Setup color encoding.
+			// The default `0` is ESceneColorSourceEncoding::Linear.
+			// Also linear encoding is used for HoldoutComposite input
+			OutParameters.Encodings = FUint32Vector::ZeroValue;
+			OutParameters.DisplayGamma = FVector3f(1.f, 1.f, 1.f);
+
+			const bool UseGammaEncoding = (Settings.ColorMask & EColorWriteMask::CW_RGB)
+				&& !Input.ColorEncoding.IsEqualsGammaEncoding(Output.ColorEncoding);
+
+			// Only apply color transform if input and output encodings are not equal and color channels are used.
+			if (UseGammaEncoding)
+			{
+				// Gets source and dest gamma values:
+				const float DefaultDisplayGamma = UTextureRenderTarget::GetDefaultDisplayGamma();
+				const float SrcGamma  = (Input.ColorEncoding.GammaValue > 0)  ? Input.ColorEncoding.GammaValue : DefaultDisplayGamma;
+				const float DestGamma = (Output.ColorEncoding.GammaValue > 0) ? Output.ColorEncoding.GammaValue : DefaultDisplayGamma;
+
+				if (Input.ColorEncoding.Encoding == Output.ColorEncoding.Encoding
+					&& Input.ColorEncoding.Encoding == EDisplayClusterColorEncoding::Gamma)
+				{
+					// Convert Gamma->Gamma in one pow()
+					OutParameters.Encodings.X = static_cast<uint32>(ESourceEncoding::Gamma);
+					OutParameters.DisplayGamma.X = SrcGamma / DestGamma;
+					OutPermutationVector.Set<FColorEncodingCopyRectPSPermutation::FPermutationEncodeInput>(true);
+				}
+				else // Transformation via linear color space:
+				{
+					// Convert input encoding to the linear
+					switch (Input.ColorEncoding.Encoding)
+					{
+					case EDisplayClusterColorEncoding::Gamma: // Gamma -> Linear
+						OutParameters.Encodings.X = static_cast<uint32>(ESourceEncoding::Gamma);
+						OutParameters.DisplayGamma.X = SrcGamma;
+						OutPermutationVector.Set<FColorEncodingCopyRectPSPermutation::FPermutationEncodeInput>(true);
+						break;
+
+					case EDisplayClusterColorEncoding::sRGB: // sRGB -> Linear
+						OutParameters.Encodings.X = static_cast<uint32>(ESourceEncoding::sRGB);
+						OutPermutationVector.Set<FColorEncodingCopyRectPSPermutation::FPermutationEncodeInput>(true);
+						break;
+
+					case EDisplayClusterColorEncoding::MediaPQ: // MediaPQ -> Linear
+						OutParameters.Encodings.X = static_cast<uint32>(ESourceEncoding::MediaPQ);
+						OutPermutationVector.Set<FColorEncodingCopyRectPSPermutation::FPermutationEncodeInput>(true);
+						break;
+
+					default:
+						break;
+					}
+
+					// Convert linear to the output encoding
+					switch (Output.ColorEncoding.Encoding)
+					{
+					case EDisplayClusterColorEncoding::Gamma: // Linear -> Gamma
+						OutParameters.Encodings.Y = static_cast<uint32>(ESourceEncoding::Gamma);
+						OutParameters.DisplayGamma.Y = 1.0f / DestGamma;
+						OutPermutationVector.Set<FColorEncodingCopyRectPSPermutation::FPermutationEncodeOutput>(true);
+						break;
+
+					case EDisplayClusterColorEncoding::sRGB: // Linear -> sRGB
+						OutParameters.Encodings.Y = static_cast<uint32>(ESourceEncoding::sRGB);
+						OutPermutationVector.Set<FColorEncodingCopyRectPSPermutation::FPermutationEncodeOutput>(true);
+						break;
+
+					case EDisplayClusterColorEncoding::MediaPQ: // Linear -> MediaPQ
+						OutParameters.Encodings.Y = static_cast<uint32>(ESourceEncoding::MediaPQ);
+						OutPermutationVector.Set<FColorEncodingCopyRectPSPermutation::FPermutationEncodeOutput>(true);
+						break;
+
+					default:
+						break;
+					}
+				}
+			}
+
+			// Setup Color Premultiply
+			OutParameters.ColorPremultiply = FUint32Vector::ZeroValue;
+
+			if (Input.ColorEncoding.Premultiply != Output.ColorEncoding.Premultiply
+			|| (UseGammaEncoding && Input.ColorEncoding.Premultiply != EDisplayClusterColorPremultiply::None))
+			{
+				switch (Input.ColorEncoding.Premultiply)
+				{
+				case EDisplayClusterColorPremultiply::Premultiply:
+					OutParameters.ColorPremultiply.X = static_cast<uint32>(EColorPremultiply::Alpha);
+					OutPermutationVector.Set<FColorEncodingCopyRectPSPermutation::FPermutationColorPremultiply>(true);
+					break;
+				case EDisplayClusterColorPremultiply::InvertPremultiply:
+					OutParameters.ColorPremultiply.X = static_cast<uint32>(EColorPremultiply::InvertedAlpha);
+					OutPermutationVector.Set<FColorEncodingCopyRectPSPermutation::FPermutationColorPremultiply>(true);
+					break;
+				default:
+					break;
+				}
+				switch (Output.ColorEncoding.Premultiply)
+				{
+				case EDisplayClusterColorPremultiply::Premultiply:
+					OutParameters.ColorPremultiply.Y = static_cast<uint32>(EColorPremultiply::Alpha);
+					OutPermutationVector.Set<FColorEncodingCopyRectPSPermutation::FPermutationColorPremultiply>(true);
+					break;
+				case EDisplayClusterColorPremultiply::InvertPremultiply:
+					OutParameters.ColorPremultiply.Y = static_cast<uint32>(EColorPremultiply::InvertedAlpha);
+					OutPermutationVector.Set<FColorEncodingCopyRectPSPermutation::FPermutationColorPremultiply>(true);
+					break;
+				default:
+					break;
+				}
+			}
+
+			// Override alpha
+			switch(Settings.OverrideAlpha)
+			{
+			case EDisplayClusterShaderTextureUtilsOverrideAlpha::Invert_Alpha:
+				OutParameters.Encodings.Z = static_cast<uint32>(EOverrideAlpha::Invert);
+				OutPermutationVector.Set<FColorEncodingCopyRectPSPermutation::FPermutationOverrideAlpha>(true);
+				break;
+
+			case EDisplayClusterShaderTextureUtilsOverrideAlpha::Set_Alpha_One:
+				OutParameters.Encodings.Z = static_cast<uint32>(EOverrideAlpha::One);
+				OutPermutationVector.Set<FColorEncodingCopyRectPSPermutation::FPermutationOverrideAlpha>(true);
+				break;
+
+			case EDisplayClusterShaderTextureUtilsOverrideAlpha::Set_Alpha_Zero:
+				OutParameters.Encodings.Z = static_cast<uint32>(EOverrideAlpha::Zero);
+				OutPermutationVector.Set<FColorEncodingCopyRectPSPermutation::FPermutationOverrideAlpha>(true);
+				break;
+			}	
+
+			// Check the permutation vectors. This should prevent a crash when no shader permutation is found.
+			if (!FColorEncodingCopyRectPSPermutation::ShouldCompileCommonPSPermutation(OutPermutationVector))
+			{
+				UE_LOG(LogDisplayClusterShaders, Warning, TEXT("Invalid permutation vector %d for shader `FColorEncodingCopyRectPS`"), OutPermutationVector.ToDimensionValueId());
+
+				return false;
+			}
+
+			return true;
+		}
+
+		/** Render input to output using this shader. */
+		static bool RenderPass(
+			FRHICommandListImmediate& RHICmdList,
+			FRHITexture* SrcTexture,
+			FRHITexture* DestTexture,
+			const FDisplayClusterShadersTextureViewportContext& InputContext,
+			const FDisplayClusterShadersTextureViewportContext& OutputContext,
+			const FDisplayClusterShadersTextureUtilsSettings& Settings)
+		{
+			if (!SrcTexture || !DestTexture)
+			{
+				return false;
+			}
+
+			// Initialize shader parameters and permutation vectors
+			FParameters PixelShaderParameters;
+			FColorEncodingCopyRectPSPermutation::FCommonPSDomain PermutationVector;
+			if (!InitializeShaderParameters(InputContext, OutputContext, Settings, PixelShaderParameters, PermutationVector))
+			{
+				return false;
+			}
+			PixelShaderParameters.InputTexture = SrcTexture;
+
+			FGlobalShaderMap* ShaderMap = GetGlobalShaderMap(GMaxRHIFeatureLevel);
+			TShaderMapRef<FScreenVS> VertexShader(ShaderMap);
+			TShaderMapRef<FColorEncodingCopyRectPS> PixelShader(ShaderMap, PermutationVector);
+			if (!VertexShader.IsValid() || !PixelShader.IsValid())
+			{
+				// Always check if shaders are available on the current platform and hardware
+				return false;
+			}
+
+			const FIntPoint& SrcTextureSize = SrcTexture->GetDesc().Extent;
+			const FIntPoint& DestTextureSize = DestTexture->GetDesc().Extent;
+			const FIntRect& SrcRect = InputContext.Rect;
+			const FIntRect& DestRect = OutputContext.Rect;
+
+			RHICmdList.Transition(FRHITransitionInfo(SrcTexture, ERHIAccess::Unknown, ERHIAccess::SRVMask));
+			RHICmdList.Transition(FRHITransitionInfo(DestTexture, ERHIAccess::Unknown, ERHIAccess::RTV));
+
+			FRHIRenderPassInfo RPInfo(DestTexture, ERenderTargetActions::Load_Store);
+			RHICmdList.BeginRenderPass(RPInfo, TEXT("nDisplay.Shaders.ColorEncodingCopyRect"));
+			
+			{
+				RHICmdList.SetViewport(0.0f, 0.0f, 0.0f, DestTextureSize.X, DestTextureSize.Y, 1.0f);
+
+				FGraphicsPipelineStateInitializer GraphicsPSOInit;
+				RHICmdList.ApplyCachedRenderTargets(GraphicsPSOInit);
+
+				GraphicsPSOInit.BlendState = GetBlendStateRHI(Settings.ColorMask);
+				GraphicsPSOInit.RasterizerState = TStaticRasterizerState<>::GetRHI();
+				GraphicsPSOInit.DepthStencilState = TStaticDepthStencilState<false, CF_Always>::GetRHI();
+
+				GraphicsPSOInit.BoundShaderState.VertexDeclarationRHI = GFilterVertexDeclaration.VertexDeclarationRHI;
+				GraphicsPSOInit.BoundShaderState.VertexShaderRHI = VertexShader.GetVertexShader();
+				GraphicsPSOInit.BoundShaderState.PixelShaderRHI = PixelShader.GetPixelShader();
+				GraphicsPSOInit.PrimitiveType = PT_TriangleList;
+
+				SetGraphicsPipelineState(RHICmdList, GraphicsPSOInit, 0);
+
+				SetShaderParameters(RHICmdList, PixelShader, PixelShader.GetPixelShader(), PixelShaderParameters);
+
+				UE::Renderer::PostProcess::DrawRectangle(
+					RHICmdList, VertexShader,
+					DestRect.Min.X, DestRect.Min.Y,
+					DestRect.Size().X, DestRect.Size().Y,
+					SrcRect.Min.X, SrcRect.Min.Y,
+					SrcRect.Size().X, SrcRect.Size().Y,
+					DestTextureSize, SrcTextureSize
+				);
+			}
+
+			RHICmdList.EndRenderPass();
+			RHICmdList.Transition(FRHITransitionInfo(DestTexture, ERHIAccess::Unknown, ERHIAccess::SRVMask));
+
+			return true;
+		}
+	};
+
+	IMPLEMENT_GLOBAL_SHADER(FColorEncodingCopyRectPS, "/Plugin/nDisplay/Private/ResourceUtils.usf", "Main", SF_Pixel);
+
+	/** Return CopyTextureInfo structure. */
+	static inline FRHICopyTextureInfo GetRHICopyTextureInfo(
+		const FDisplayClusterShadersTextureViewport& Input,
+		const FDisplayClusterShadersTextureViewport& Output,
+		const FDisplayClusterShadersTextureUtilsSettings& InSettings)
+	{
+		FRHICopyTextureInfo OutCopyInfo;
+
+		OutCopyInfo.SourceSliceIndex = InSettings.SourceSliceIndex;
+		OutCopyInfo.DestSliceIndex = InSettings.DestSliceIndex;
+
+		OutCopyInfo.SourcePosition = FIntVector(Input.Rect.Min.X, Input.Rect.Min.Y, 0);
+		OutCopyInfo.DestPosition = FIntVector(Output.Rect.Min.X, Output.Rect.Min.Y, 0);
+
+		OutCopyInfo.Size = FIntVector(Output.Rect.Width(), Output.Rect.Height(), 0);
+
+		return OutCopyInfo;
+	}
+};
+
+bool FDisplayClusterShadersCopyTexture::ColorEncodingCopyRect_RenderThread(
+	FRHICommandListImmediate& RHICmdList,
+	const FDisplayClusterShadersTextureViewportContext& Input,
+	const FDisplayClusterShadersTextureViewportContext& Output,
+	const FDisplayClusterShadersTextureUtilsSettings& Settings)
+{
+	using namespace UE::DisplayClusterShaders::Private;
+
+	return FColorEncodingCopyRectPS::RenderPass(RHICmdList, Input.TextureRHI, Output.TextureRHI, Input, Output, Settings);
+}
+
+/**
+* Copy RDG textures via shader.
+*/
+bool FDisplayClusterShadersCopyTexture::AddPassColorEncodingCopyRect_RenderThread(
+	FRDGBuilder& GraphBuilder,
+	const FDisplayClusterShadersTextureViewportContext& Input,
+	const FDisplayClusterShadersTextureViewportContext& Output,
+	const FDisplayClusterShadersTextureUtilsSettings& Settings)
+{
+	using namespace UE::DisplayClusterShaders::Private;
+
+	if (!Input.TextureRDG || !Output.TextureRDG)
+	{
+		return false;
+	}
+
+	// Initialize render pass parameters.
+	FDisplayClusterCopyTextureParameters* Parameters = GraphBuilder.AllocParameters<FDisplayClusterCopyTextureParameters>();
+	Parameters->Input  = Input.TextureRDG;
+	Parameters->Output = Output.TextureRDG;
+
+	if (Input.bExternalTextureRDG)
+	{
+		GraphBuilder.SetTextureAccessFinal(Input.TextureRDG, ERHIAccess::SRVGraphics);
+	}
+	if (Output.bExternalTextureRDG)
+	{
+		GraphBuilder.SetTextureAccessFinal(Output.TextureRDG, ERHIAccess::RTV);
+	}
+
+	GraphBuilder.AddPass(
+		RDG_EVENT_NAME("nDisplayShaders.ResampleTexture"),
+		Parameters,
+		ERDGPassFlags::Copy | ERDGPassFlags::NeverCull,
+		[Parameters, Input, Output, Settings](FRHICommandListImmediate& RHICmdList)
+	{
+		FColorEncodingCopyRectPS::RenderPass(
+			RHICmdList,
+			Parameters->Input->GetRHI(),
+			Parameters->Output->GetRHI(),
+			Input, Output, Settings);
+	});
+
+	return true;
+}
+
+/**
+* Copy RDG textures without shader.
+*/
+bool FDisplayClusterShadersCopyTexture::AddPassTransitionAndCopyTexture_RenderThread(
+	FRDGBuilder& GraphBuilder,
+	const FDisplayClusterShadersTextureViewport& Input,
+	const FDisplayClusterShadersTextureViewport& Output,
+	const FDisplayClusterShadersTextureUtilsSettings& Settings)
+{
+	using namespace UE::DisplayClusterShaders::Private;
+
+	if (!Input.TextureRDG || !Output.TextureRDG)
+	{
+		return false;
+	}
+
+	if (Input.Rect.Size() != Output.Rect.Size())
+	{
+		return false;
+	}
+
+	// Initialize render pass parameters.
+	FDisplayClusterCopyTextureParameters* Parameters = GraphBuilder.AllocParameters<FDisplayClusterCopyTextureParameters>();
+	Parameters->Input = Input.TextureRDG;
+	Parameters->Output = Output.TextureRDG;
+
+	GraphBuilder.AddPass(
+		RDG_EVENT_NAME("nDisplayShaders.CopyTexture"),
+		Parameters,
+		ERDGPassFlags::Copy | ERDGPassFlags::NeverCull,
+		[ Input, Output, Settings](FRHICommandListImmediate& RHICmdList)
+		{
+			FRHITexture* Src = Input.TextureRDG->GetRHI();
+			FRHITexture* Dest = Output.TextureRDG->GetRHI();
+			if (Src && Dest)
+			{
+				TransitionAndCopyTexture(RHICmdList, Src, Dest, GetRHICopyTextureInfo(Input, Output, Settings));
+			}
+		});
+
+	return true;
+}
+
+bool FDisplayClusterShadersCopyTexture::TransitionAndCopyTexture_RenderThread(
+	FRHICommandListImmediate& RHICmdList,
+	const FDisplayClusterShadersTextureViewport& Input,
+	const FDisplayClusterShadersTextureViewport& Output,
+	const FDisplayClusterShadersTextureUtilsSettings& InSettings)
+{
+	using namespace UE::DisplayClusterShaders::Private;
+
+	if (!Input.TextureRHI || !Output.TextureRHI)
+	{
+		return false;
+	}
+
+	if (Input.Rect.Size() != Output.Rect.Size())
+	{
+		return false;
+	}
+
+	TransitionAndCopyTexture(RHICmdList,
+		Input.TextureRHI,
+		Output.TextureRHI,
+		GetRHICopyTextureInfo(Input, Output, InSettings));
+
+	return true;
+}
diff --git a/Engine/Plugins/Runtime/nDisplay/Source/DisplayClusterShaders/Private/Shaders/DisplayClusterShadersCopyTexture.h b/Engine/Plugins/Runtime/nDisplay/Source/DisplayClusterShaders/Private/Shaders/DisplayClusterShadersCopyTexture.h
new file mode 100644
--- /dev/null
+++ b/Engine/Plugins/Runtime/nDisplay/Source/DisplayClusterShaders/Private/Shaders/DisplayClusterShadersCopyTexture.h
@@ -0,0 +1,76 @@
+// Copyright Epic Games, Inc. All Rights Reserved.
+
+#pragma once
+#include "DisplayClusterShadersTextureUtils.h"
+
+#include "RHI.h"
+#include "RHICommandList.h"
+
+class FRDGBuilder;
+
+/** Implementation of texture copying. */
+class FDisplayClusterShadersCopyTexture
+{
+public:
+	/** Implements texture copy via pixel shader.
+	*
+	* @param RHICmdList   - RHI API
+	* @param Input        - input texture context
+	* @param Output       - output texture context
+	* @param InSettings
+	*
+	* @return - true if success
+	*/
+	static bool ColorEncodingCopyRect_RenderThread(
+		FRHICommandListImmediate& RHICmdList,
+		const FDisplayClusterShadersTextureViewportContext& Input,
+		const FDisplayClusterShadersTextureViewportContext& Output,
+		const FDisplayClusterShadersTextureUtilsSettings& InSettings);
+
+	/** Implements texture copy via pixel shader.
+	*
+	* @param RHICmdList   - RHI API
+	* @param Input        - input texture context
+	* @param Output       - output texture context
+	* @param InSettings
+	*
+	* @return - true if success
+	*/
+	static bool AddPassColorEncodingCopyRect_RenderThread(
+		FRDGBuilder& GraphBuilder,
+		const FDisplayClusterShadersTextureViewportContext& Input,
+		const FDisplayClusterShadersTextureViewportContext& Output,
+		const FDisplayClusterShadersTextureUtilsSettings& InSettings);
+
+	/**
+	* Call TransitionAndCopyTexture()
+	*
+	* @param RHICmdList   - RHI API
+	* @param Input        - input texture context
+	* @param Output       - output texture context
+	* @param InSettings
+	*
+	* @return - true if success
+	*/
+	static bool TransitionAndCopyTexture_RenderThread(
+		FRHICommandListImmediate& RHICmdList,
+		const FDisplayClusterShadersTextureViewport& Input,
+		const FDisplayClusterShadersTextureViewport& Output,
+		const FDisplayClusterShadersTextureUtilsSettings& InSettings);
+
+	/**
+	* Call TransitionAndCopyTexture()
+	*
+	* @param GraphBuilder - RDG builder
+	* @param Input        - input texture context
+	* @param Output       - output texture context
+	* @param InSettings
+	*
+	* @return - true if success
+	*/
+	static bool AddPassTransitionAndCopyTexture_RenderThread(
+		FRDGBuilder& GraphBuilder,
+		const FDisplayClusterShadersTextureViewport& Input,
+		const FDisplayClusterShadersTextureViewport& Output,
+		const FDisplayClusterShadersTextureUtilsSettings& InSettings);
+};
diff --git a/Engine/Plugins/Runtime/nDisplay/Source/DisplayClusterShaders/Public/Containers/DisplayClusterShaderContainers_TextureUtils.h b/Engine/Plugins/Runtime/nDisplay/Source/DisplayClusterShaders/Public/Containers/DisplayClusterShaderContainers_TextureUtils.h
new file mode 100644
--- /dev/null
+++ b/Engine/Plugins/Runtime/nDisplay/Source/DisplayClusterShaders/Public/Containers/DisplayClusterShaderContainers_TextureUtils.h
@@ -0,0 +1,239 @@
+// Copyright Epic Games, Inc. All Rights Reserved.
+
+#pragma once
+
+#include "CoreMinimal.h"
+#include "Misc/DisplayClusterColorEncoding.h"
+
+#include "RHIResources.h"
+#include "RenderGraphFwd.h"
+#include "ScreenPass.h"
+
+/** Texture utiles flags. */
+enum class EDisplayClusterShaderTextureUtilsFlags : uint8
+{
+	None = 0,
+
+	/**
+	* Use only output texture as the source and destination.
+	* A temporary texture will be used as the output texture, which is then copied back to the output texture.
+	* The input texture is not defined.
+	* Custom implementations must perform the copying from the output texture to the input texture themselves.
+	*/
+	UseOutputTextureAsInput = 1 << 0,
+
+	/**
+	* Invert direction: From `Output` to the `Input`
+	*/
+	InvertDirection = 1 << 1,
+
+	/**
+	* Trim the input and output rect sizes to the same value.
+	* (The input rect size will be equal to the output rect size.)
+	*/
+	DisableResize = 1 << 4,
+
+	/** Disable the use of the shader for resizing, color coding, etc. */
+	DisableResampleShader = 1 << 5,
+
+	/** Dont update resource rects, they are expected to be user defined. */
+	DisableUpdateResourcesRectsForResolve = 1 << 6
+};
+ENUM_CLASS_FLAGS(EDisplayClusterShaderTextureUtilsFlags);
+
+/** Alpha channel rendering mode. */
+enum class EDisplayClusterShaderTextureUtilsOverrideAlpha : uint8
+{
+	// Alpha does not change
+	None = 0,
+
+	// Invert alpha channel
+	Invert_Alpha,
+
+	// set alpha to One(1)
+	Set_Alpha_One,
+
+	// Set alpha to Zero(0)
+	Set_Alpha_Zero,
+};
+
+/*
+ * A container with a texture on which a rectangular area is located.
+ */
+struct FDisplayClusterShadersTextureViewport
+{
+	FDisplayClusterShadersTextureViewport() = default;
+
+	FDisplayClusterShadersTextureViewport(FRHITexture* InTextureRHI, const TCHAR* InDebugName = nullptr)
+		: TextureRHI(InTextureRHI), TextureRDG(nullptr), DebugName(InDebugName) { }
+
+	FDisplayClusterShadersTextureViewport(FRDGTextureRef InTextureRDG)
+		: TextureRHI(nullptr), TextureRDG(InTextureRDG) { }
+
+	FDisplayClusterShadersTextureViewport(FRHITexture* InTextureRHI, const FIntRect& InTextureRect, const TCHAR* InDebugName = nullptr)
+		: TextureRHI(InTextureRHI), TextureRDG(nullptr), Rect(InTextureRect), DebugName(InDebugName) { }
+
+	FDisplayClusterShadersTextureViewport(FRDGTextureRef InTextureRDG, const FIntRect& InTextureRect)
+		: TextureRHI(nullptr), TextureRDG(InTextureRDG), Rect(InTextureRect) { }
+
+	FDisplayClusterShadersTextureViewport(const FScreenPassTexture& InScreenPass)
+		: TextureRHI(nullptr), TextureRDG(InScreenPass.Texture), Rect(InScreenPass.ViewRect) { }
+
+	/** Returns true if these parameters are valid. */
+	inline bool IsValid() const
+	{
+		return (TextureRHI || TextureRDG) && !Rect.IsEmpty();
+	}
+
+	/** Convert to the ScreenPass texture type. */
+	inline FScreenPassTexture ToScreenPassTexture() const
+	{
+		return FScreenPassTexture(TextureRDG, Rect);
+	}
+
+	/** Convert to the ScreenPass viewport type. */
+	inline FScreenPassTextureViewport ToScreenPassTextureViewport() const
+	{
+		return FScreenPassTextureViewport(TextureRDG, Rect);
+	}
+
+public:
+	/** RHI Texture. If this parameter is nullptr, the TextureRDG parameter is used. */
+	FRHITexture* TextureRHI = nullptr;
+
+	/** RDG texture.  */
+	FRDGTextureRef TextureRDG = nullptr;
+
+	/** true if RDG texture external. */
+	bool bExternalTextureRDG = false;
+
+	/** Texture sub-region */
+	FIntRect Rect = { FIntPoint::ZeroValue, FIntPoint::ZeroValue };
+
+	/** Optional debug name for this resource, used to name the external RHI resource for RDG. */
+	const TCHAR* DebugName = nullptr;
+};
+
+/*
+ * A container with textures and their configurations to be used as input and output data.
+ */
+struct FDisplayClusterShadersTextureParameters
+{
+	/** Returns true if these parameters are valid. */
+	bool IsValid() const
+	{
+		// All defined textures must be valid.
+		for (const TPair<uint32, FDisplayClusterShadersTextureViewport>& TextureIt : TextureViewports)
+		{
+			if(!TextureIt.Value.IsValid())
+			{
+				return false;
+			}
+		}
+
+		// At least one texture must be defined.
+		return !TextureViewports.IsEmpty();
+	}
+
+public:
+	/** Textures of all viewport contexts. The default is context 0. */
+	TMap<uint32, FDisplayClusterShadersTextureViewport> TextureViewports;
+
+	/** texture color space. */
+	FDisplayClusterColorEncoding ColorEncoding;
+};
+
+/*
+ * Container with the texture utils settings
+ */
+struct FDisplayClusterShadersTextureUtilsSettings
+{
+	FDisplayClusterShadersTextureUtilsSettings() = default;
+
+	FDisplayClusterShadersTextureUtilsSettings(
+		const FDisplayClusterShadersTextureUtilsSettings& InSettigns,
+		const EDisplayClusterShaderTextureUtilsFlags InFlags,
+		const uint32 InDestSliceIndex = 0)
+		: FDisplayClusterShadersTextureUtilsSettings(InSettigns)
+	{
+		Flags = InFlags;
+		DestSliceIndex = InDestSliceIndex;
+	}
+
+	FDisplayClusterShadersTextureUtilsSettings(const EColorWriteMask InColorMask)
+		: ColorMask(InColorMask) { }
+
+	FDisplayClusterShadersTextureUtilsSettings(
+		const FDisplayClusterShadersTextureUtilsSettings& InSettigns,
+		const EColorWriteMask InColorMask)
+		: FDisplayClusterShadersTextureUtilsSettings(InSettigns)
+	{
+		ColorMask = InColorMask;
+	}
+
+	FDisplayClusterShadersTextureUtilsSettings(const EDisplayClusterShaderTextureUtilsOverrideAlpha InOverrideAlpha)
+		: OverrideAlpha(InOverrideAlpha) {  }
+
+	FDisplayClusterShadersTextureUtilsSettings(
+		const EColorWriteMask InColorMask,
+		const EDisplayClusterShaderTextureUtilsFlags InFlags,
+		const uint32 InDestSliceIndex = 0)
+		: ColorMask(InColorMask), Flags(InFlags), DestSliceIndex(InDestSliceIndex) { }
+
+	FDisplayClusterShadersTextureUtilsSettings(
+		const EDisplayClusterShaderTextureUtilsFlags InFlags,
+		const uint32 InDestSliceIndex = 0)
+		: Flags(InFlags), DestSliceIndex(InDestSliceIndex) { }
+
+	inline bool HasAnyFlags(const EDisplayClusterShaderTextureUtilsFlags InFlags) const
+	{
+		return EnumHasAnyFlags(Flags, InFlags);
+	}
+
+	inline void AddFlags(const EDisplayClusterShaderTextureUtilsFlags InFlags)
+	{
+		return EnumAddFlags(Flags, InFlags);
+	}
+
+public:
+	/** Color mask*/
+	EColorWriteMask ColorMask = EColorWriteMask::CW_RGBA;
+
+	/** Additional flags. */
+	EDisplayClusterShaderTextureUtilsFlags Flags = EDisplayClusterShaderTextureUtilsFlags::None;
+
+	/** Override alpha channel. */
+	EDisplayClusterShaderTextureUtilsOverrideAlpha OverrideAlpha = EDisplayClusterShaderTextureUtilsOverrideAlpha::None;
+
+	/** Additional parameters for the FRHICopyTextureInfo. */
+	uint32 SourceSliceIndex = 0;
+	uint32 DestSliceIndex = 0;
+};
+
+/** This container is used for IDisplayClusterShadersTextureUtils::ForEachContextByPredicate() */
+struct FDisplayClusterShadersTextureViewportContext : public FDisplayClusterShadersTextureViewport
+{
+	FDisplayClusterShadersTextureViewportContext() = default;
+
+	FDisplayClusterShadersTextureViewportContext(
+		const FDisplayClusterShadersTextureViewport& InTextureViewport,
+		const uint32 InContextNum = 0)
+		: FDisplayClusterShadersTextureViewport(InTextureViewport), ContextNum(InContextNum) { }
+	
+	FDisplayClusterShadersTextureViewportContext(
+		const FDisplayClusterShadersTextureViewport& InTextureViewport,
+		const FDisplayClusterColorEncoding& InColorEncoding,
+		const uint32 InContextNum = 0)
+		: FDisplayClusterShadersTextureViewport(InTextureViewport), ContextNum(InContextNum), ColorEncoding(InColorEncoding){ }
+
+	/** Texture context num. */
+	uint32 ContextNum = 0;
+
+	/** texture color space. */
+	FDisplayClusterColorEncoding ColorEncoding;
+};
+
+class IDisplayClusterShadersTextureUtils;
+using TFunctionDisplayClusterShaders_TextureContextIterator = TFunction<void(
+	const FDisplayClusterShadersTextureViewportContext& Input,
+	const FDisplayClusterShadersTextureViewportContext& Output)>;
diff --git a/Engine/Plugins/Runtime/nDisplay/Source/DisplayClusterShaders/Public/IDisplayClusterShaders.h b/Engine/Plugins/Runtime/nDisplay/Source/DisplayClusterShaders/Public/IDisplayClusterShaders.h
--- a/Engine/Plugins/Runtime/nDisplay/Source/DisplayClusterShaders/Public/IDisplayClusterShaders.h
+++ b/Engine/Plugins/Runtime/nDisplay/Source/DisplayClusterShaders/Public/IDisplayClusterShaders.h
@@ -12,6 +12,7 @@
 class IDisplayClusterRender_MeshComponentProxy;
 class FDisplayClusterShaderParameters_ICVFX;
 class FDisplayClusterShaderParameters_Override;
+class IDisplayClusterShadersTextureUtils;
 
 class FSceneInterface;
 class FRDGBuilder;
@@ -151,4 +152,17 @@ public:
 	* @param Parameters   - Conversion parameters
 	*/
 	virtual void AddPQToLinearPass(FRDGBuilder& GraphBuilder, const FDisplayClusterShaderParameters_MediaPQ& Parameters) const = 0;
+
+	/**
+	* Return new instance of the resource utils.
+	* 
+	* @param RHICmdList   - RHI API
+	*/
+	virtual TSharedRef<IDisplayClusterShadersTextureUtils> CreateTextureUtils_RenderThread(FRHICommandListImmediate& RHICmdList) const = 0;
+
+	/**
+	* Return new instance of the resource utils.
+	* @param GraphBuilder - RDG builder
+	*/
+	virtual TSharedRef<IDisplayClusterShadersTextureUtils> CreateTextureUtils_RenderThread(FRDGBuilder& GraphBuilder) const = 0;
 };
diff --git a/Engine/Plugins/Runtime/nDisplay/Source/DisplayClusterShaders/Public/IDisplayClusterShadersTextureUtils.h b/Engine/Plugins/Runtime/nDisplay/Source/DisplayClusterShaders/Public/IDisplayClusterShadersTextureUtils.h
new file mode 100644
--- /dev/null
+++ b/Engine/Plugins/Runtime/nDisplay/Source/DisplayClusterShaders/Public/IDisplayClusterShadersTextureUtils.h
@@ -0,0 +1,167 @@
+// Copyright Epic Games, Inc. All Rights Reserved.
+
+#pragma once
+
+#include "CoreMinimal.h"
+#include "Containers/DisplayClusterShaderContainers_TextureUtils.h"
+#include "Render/Viewport/Containers/DisplayClusterViewport_Enums.h"
+
+#include "RHICommandList.h"
+#include "RenderGraphBuilder.h"
+
+class IDisplayClusterViewportProxy;
+
+/**
+ * Auxiliary Texture utils class for nDisplay.
+ */
+class IDisplayClusterShadersTextureUtils
+{
+public:
+	/** Default destructor. */
+	virtual ~IDisplayClusterShadersTextureUtils() = default;
+
+	/**
+	*  Set the input texture viewport for the specified context.
+	*
+	* @parama InTextureViewport - texture viewport
+	* @param InContextNum       - (opt) context(eye)  number, otherwise all contexts will be collected.
+	*
+	* @return themself
+	*/
+	virtual TSharedRef<IDisplayClusterShadersTextureUtils> SetInput(const FDisplayClusterShadersTextureViewport& InTextureViewport, const int32 InContextNum = INDEX_NONE) = 0;
+
+	/**
+	*  Set the output texture viewport for the specified context.
+	*
+	* @parama InTextureViewport - texture viewport
+	* @param InContextNum       - (opt) context(eye)  number, otherwise all contexts will be collected.
+	*
+	* @return themself
+	*/
+	virtual TSharedRef<IDisplayClusterShadersTextureUtils> SetOutput(const FDisplayClusterShadersTextureViewport& InTextureViewport, const int32 InContextNum = INDEX_NONE) = 0;
+
+	/** Set input texture from a viewport proxy object.
+	* 
+	* @param InViewportProxy  - viewport proxy object with resources to retrieve
+	* @param InResourceType   - Type of the viewport resource
+	* @param InContextNum       - (opt) context(eye)  number, otherwise all contexts will be collected.
+	*
+	* @return themself
+	*/
+	virtual TSharedRef<IDisplayClusterShadersTextureUtils> SetInput(
+		const IDisplayClusterViewportProxy* InViewportProxy,
+		const EDisplayClusterViewportResourceType InResourceType,
+		const int32 InContextNum = INDEX_NONE) = 0;
+
+	/** Set output texture from a viewport proxy object.
+	*
+	* @param InViewportProxy  - viewport proxy object with resources to retrieve
+	* @param InResourceType   - Type of the viewport resource
+	* @param InContextNum       - (opt) context(eye)  number, otherwise all contexts will be collected.
+	*
+	* @return themself
+	*/
+	virtual TSharedRef<IDisplayClusterShadersTextureUtils> SetOutput(
+		const IDisplayClusterViewportProxy* InViewportProxy,
+		const EDisplayClusterViewportResourceType InResourceType,
+		const int32 InContextNum = INDEX_NONE) = 0;
+
+	/** Set input color encodig for texture
+	*
+	* @param InColorEncoding  - color encoding
+	*
+	* @return themself
+	*/
+	virtual TSharedRef<IDisplayClusterShadersTextureUtils> SetInputEncoding(const FDisplayClusterColorEncoding& InColorEncoding) = 0;
+
+	/** Set output color encodig for texture
+	*
+	* @param InColorEncoding  - color encoding
+	*
+	* @return themself
+	*/
+	virtual TSharedRef<IDisplayClusterShadersTextureUtils> SetOutputEncoding(const FDisplayClusterColorEncoding& InColorEncoding) = 0;
+
+
+	/** Return input texture parameters. */
+	virtual const FDisplayClusterShadersTextureParameters& GetInputTextureParameters() const = 0;
+
+	/** Return input texture parameters. */
+	virtual const FDisplayClusterShadersTextureParameters& GetOutputTextureParameters() const = 0;
+
+	/**
+	* Resolve all contexts
+	* 
+	* @return themself
+	*/
+	virtual TSharedRef<IDisplayClusterShadersTextureUtils> Resolve() = 0;
+
+	/**
+	* Resolve all contexts
+	*
+	* @param ColorMask - Color Recording Mask. Used to copy only selected channels.
+	* @param InSettings - (in) render settings
+	*
+	* @return themself
+	*/
+	virtual TSharedRef<IDisplayClusterShadersTextureUtils> Resolve(const FDisplayClusterShadersTextureUtilsSettings& InSettings) = 0;
+
+	/**
+	* Resolve one context to the other
+	*
+	* @param Input  - input texture context
+	* @param Output - output texture context
+	*
+	* @return - themself
+	*/
+	virtual TSharedRef<IDisplayClusterShadersTextureUtils> ResolveTextureContext(
+		const FDisplayClusterShadersTextureViewportContext& Input,
+		const FDisplayClusterShadersTextureViewportContext& Output) = 0;
+
+	/**
+	* Resolve one context to the other
+	*
+	* @param InSettings - (in) render settings
+	* @param Input      - input texture context
+	* @param Output     - output texture context
+	*
+	* @return - themself
+	*/
+	virtual TSharedRef<IDisplayClusterShadersTextureUtils> ResolveTextureContext(
+		const FDisplayClusterShadersTextureUtilsSettings& InSettings,
+		const FDisplayClusterShadersTextureViewportContext& Input,
+		const FDisplayClusterShadersTextureViewportContext& Output) = 0;
+
+	/**
+	* Custom implementation of an action between input and output contexts.
+	* Iterate over all contexts on input and output textures.
+	* Note: This function creates unique TextureContexts instances for the iterator arguments.
+	*       Multiple calls to ForEachContextByPredicate() will create unique contexts for each use.
+	* 
+	* @param TextureContextIteratorFunc - predicate function that processes pairs of input and output textures (via resolve or any other way).
+	* 
+	* @ returns themself
+	*/
+	virtual TSharedRef<IDisplayClusterShadersTextureUtils> ForEachContextByPredicate(TFunctionDisplayClusterShaders_TextureContextIterator&& TextureContextIteratorFunc) = 0;
+
+	/**
+	* Custom implementation of an action between input and output contexts.
+	* Iterate over all contexts on input and output textures.
+	* Note: This function creates unique TextureContexts instances for the iterator arguments.
+	*       Multiple calls to ForEachContextByPredicate() will create unique contexts for each use.
+	*
+	* @param InSettings - (in) render settings
+	* @param TextureContextIteratorFunc - predicate function that processes pairs of input and output textures (via resolve or any other way).
+	*
+	* @ returns themself
+	*/
+	virtual TSharedRef<IDisplayClusterShadersTextureUtils> ForEachContextByPredicate(
+		const FDisplayClusterShadersTextureUtilsSettings& InSettings,
+		TFunctionDisplayClusterShaders_TextureContextIterator&& TextureContextIteratorFunc) = 0;
+
+	/**
+	* Get or create RDG builder for this texture utils instance
+	* Instances created for RHI after this call will only use RDG.
+	*/
+	virtual FRDGBuilder& GetOrCreateRDGBuilder() = 0;
+};
\ No newline at end of file
