[nDisplay] HoldoutComposite integration 3/3: Integration nDisplay preview with HoldoutComposite plugin.

-Added dependency on the HoldoutComposite plugin
- Added new property to the DCRA `bPreviewEnableHoldoutComposite`
- Implemented HoldoutComposite integration for preview in the scene.

diff --git a/Engine/Plugins/Runtime/nDisplay/Source/DisplayCluster/DisplayCluster.Build.cs b/Engine/Plugins/Runtime/nDisplay/Source/DisplayCluster/DisplayCluster.Build.cs
--- a/Engine/Plugins/Runtime/nDisplay/Source/DisplayCluster/DisplayCluster.Build.cs
+++ b/Engine/Plugins/Runtime/nDisplay/Source/DisplayCluster/DisplayCluster.Build.cs
@@ -27,6 +27,7 @@ public class DisplayCluster : ModuleRules
 		PrivateDependencyModuleNames.AddRange(
 			new string[] {
 				"HeadMountedDisplay",
+				"HoldoutComposite",
 				"InputCore",
 				"Json",
 				"JsonUtilities",
diff --git a/Engine/Plugins/Runtime/nDisplay/Source/DisplayCluster/Private/Game/EngineClasses/Scene/DisplayClusterRootActor.cpp b/Engine/Plugins/Runtime/nDisplay/Source/DisplayCluster/Private/Game/EngineClasses/Scene/DisplayClusterRootActor.cpp
--- a/Engine/Plugins/Runtime/nDisplay/Source/DisplayCluster/Private/Game/EngineClasses/Scene/DisplayClusterRootActor.cpp
+++ b/Engine/Plugins/Runtime/nDisplay/Source/DisplayCluster/Private/Game/EngineClasses/Scene/DisplayClusterRootActor.cpp
@@ -52,6 +52,8 @@
 #include "Components/DisplayClusterStageGeometryComponent.h"
 #include "Components/DisplayClusterStageIsosphereComponent.h"
 #include "Components/LineBatchComponent.h"
+
+#include "HoldoutCompositeSubsystem.h"
 #include "UObject/Package.h"
 
 namespace UE::DisplayCluster::RootActor
@@ -879,6 +881,17 @@ void ADisplayClusterRootActor::SetPreviewEnablePostProcess(const bool bNewPrevie
 #endif // WITH_EDITOR
 }
 
+void ADisplayClusterRootActor::SetPreviewEnableHoldoutComposite(const bool bNewPreviewEnableHoldoutComposite)
+{
+#if WITH_EDITOR
+	if (bNewPreviewEnableHoldoutComposite != bPreviewEnableHoldoutComposite)
+	{
+		bPreviewEnableHoldoutComposite = bNewPreviewEnableHoldoutComposite;
+
+		ResetEntireClusterPreviewRendering();
+	}
+#endif // WITH_EDITOR
+}
 
 bool ADisplayClusterRootActor::BuildHierarchy()
 {
@@ -1246,6 +1259,22 @@ void ADisplayClusterRootActor::TickPreviewRenderer()
 		NewPreviewSettings.bPreviewEnable = false;
 	}
 
+	// To use the `HoldoutComposite` plugin, a post-process is required
+	if (!NewPreviewSettings.bPreviewEnablePostProcess)
+	{
+		NewPreviewSettings.bPreviewEnableHoldoutComposite = false;
+	}
+
+	// The HoldoutComposite plugin requires special project settings.
+	if (NewPreviewSettings.bPreviewEnableHoldoutComposite
+	&& !UHoldoutCompositeSubsystem::IsProjectSettingsValid())
+	{
+		NewPreviewSettings.bPreviewEnableHoldoutComposite = false;
+
+		// In this particular case, also disable the postprocess.
+		NewPreviewSettings.bPreviewEnablePostProcess = false;
+	}
+
 #if WITH_EDITOR
 	// The StageIsosphereComponent can only be seen if DCRA preview is enabled.
 	const bool bShowStageIsosphere = NewPreviewSettings.bPreviewEnable && bPreviewStageGeometryMesh;
@@ -1321,6 +1350,8 @@ FDisplayClusterViewport_PreviewSettings ADisplayClusterRootActor::GetPreviewSett
 	OutPreviewSettings.bPreviewInGameEnable        = bPreviewInGameEnable;
 	OutPreviewSettings.bPreviewInGameRenderFrustum = bPreviewInGameRenderFrustum;
 
+	OutPreviewSettings.bPreviewEnableHoldoutComposite = bPreviewEnableHoldoutComposite;
+
 	return OutPreviewSettings;
 }
 
@@ -1596,6 +1627,7 @@ void ADisplayClusterRootActor::OverrideRootActorProperties(const FDisplayCluster
 	OVERRIDE_PROPERTY(PreviewRenderTargetRatioMult);
 	OVERRIDE_PROPERTY(bPreviewEnablePostProcess);
 	OVERRIDE_PROPERTY(bPreviewEnableOCIO);
+	OVERRIDE_PROPERTY(bPreviewEnableHoldoutComposite);
 	OVERRIDE_PROPERTY(bPreviewEnableOverlayMaterial);
 	OVERRIDE_PROPERTY(bEnablePreviewTechvis);
 	OVERRIDE_PROPERTY(bEnablePreviewMesh);
diff --git a/Engine/Plugins/Runtime/nDisplay/Source/DisplayCluster/Private/Game/EngineClasses/Scene/DisplayClusterRootActorEditor.cpp b/Engine/Plugins/Runtime/nDisplay/Source/DisplayCluster/Private/Game/EngineClasses/Scene/DisplayClusterRootActorEditor.cpp
--- a/Engine/Plugins/Runtime/nDisplay/Source/DisplayCluster/Private/Game/EngineClasses/Scene/DisplayClusterRootActorEditor.cpp
+++ b/Engine/Plugins/Runtime/nDisplay/Source/DisplayCluster/Private/Game/EngineClasses/Scene/DisplayClusterRootActorEditor.cpp
@@ -710,6 +710,7 @@ void ADisplayClusterRootActor::PostEditChangeProperty(FPropertyChangedEvent& Pro
 		ResetInnerFrustumPriority();
 	}
 	else if (PropertyName == GET_MEMBER_NAME_CHECKED(ADisplayClusterRootActor, bPreviewEnablePostProcess)
+		|| PropertyName == GET_MEMBER_NAME_CHECKED(ADisplayClusterRootActor, bPreviewEnableHoldoutComposite)
 		|| PropertyName == GET_MEMBER_NAME_CHECKED(ADisplayClusterRootActor, DefaultDisplayDeviceName))
 	{
 		ResetEntireClusterPreviewRendering();
diff --git a/Engine/Plugins/Runtime/nDisplay/Source/DisplayCluster/Private/Render/Viewport/Preview/DisplayClusterViewportPreviewMesh.cpp b/Engine/Plugins/Runtime/nDisplay/Source/DisplayCluster/Private/Render/Viewport/Preview/DisplayClusterViewportPreviewMesh.cpp
--- a/Engine/Plugins/Runtime/nDisplay/Source/DisplayCluster/Private/Render/Viewport/Preview/DisplayClusterViewportPreviewMesh.cpp
+++ b/Engine/Plugins/Runtime/nDisplay/Source/DisplayCluster/Private/Render/Viewport/Preview/DisplayClusterViewportPreviewMesh.cpp
@@ -9,6 +9,8 @@
 #include "Components/DisplayClusterCameraComponent.h"
 #include "Components/StaticMeshComponent.h"
 
+#include "HoldoutCompositeSubsystem.h"
+
 namespace UE::DisplayCluster::ViewportPreviewMesh
 {
 	/** Reset the mesh material to default values from its archetype. */
@@ -48,6 +50,49 @@ namespace UE::DisplayCluster::ViewportPreviewMesh
 		// Ignore deleted materials
 		return OutMaterial;
 	}
+
+	/** Configure Holdout for the preview mesh
+	* 
+	* @param Configuration   - current configuration
+	* @param InMeshComponent - Mesh for to be configured for Holdout
+	* @param InMeshType      - Mesh type
+	* @param InMaterialType  - Type of material on mesh
+	*/
+	static void UpdateHoldoutForMeshComponent(
+		FDisplayClusterViewportConfiguration& Configuration,
+		UMeshComponent* InMeshComponent,
+		const EDisplayClusterDisplayDeviceMeshType InMeshType,
+		const EDisplayClusterDisplayDeviceMaterialType InMaterialType = EDisplayClusterDisplayDeviceMaterialType::PreviewMeshMaterial)
+	{
+		if (!IsValid(InMeshComponent) || !UHoldoutCompositeSubsystem::IsProjectSettingsValid())
+		{
+			return;
+		}
+
+		UHoldoutCompositeSubsystem* HoldoutComposite = UWorld::GetSubsystem<UHoldoutCompositeSubsystem>(Configuration.GetCurrentWorld());
+		if (!IsValid(HoldoutComposite))
+		{
+			return;
+		}
+
+		//By default, all materials must support Holdout.
+		bool bUseHoldout = true;
+		if (InMeshType == EDisplayClusterDisplayDeviceMeshType::DefaultMesh)
+		{
+			bUseHoldout = false;
+		}
+
+		if (bUseHoldout && Configuration.GetRenderFrameSettings().ShouldUseHoldout())
+		{
+			// Render this mesh using Holdout
+			HoldoutComposite->RegisterPrimitive(InMeshComponent);
+		}
+		else
+		{
+			// Use normal rendering
+			HoldoutComposite->UnregisterPrimitive(InMeshComponent);
+		}
+	}
 };
 
 ////////////////////////////////////////////////////////////////////////////////////////
@@ -126,6 +171,9 @@ void FDisplayClusterViewportPreviewMesh::Update(FDisplayClusterViewport* InViewp
 
 		// The material must be assigned to the mesh at each tick, without any conditions in case it can be changed externally.
 		MeshComponent->SetMaterial(0, MaterialInstancePtr.Get());
+
+		// Apply a holdout to this mesh component if possible.
+		UpdateHoldoutForMeshComponent(*Configuration, MeshComponent, MeshType, GetCurrentMaterialType());
 	}
 
 	// Handling the material overlay logic for the preview mesh:
@@ -198,6 +246,8 @@ void FDisplayClusterViewportPreviewMesh::ReleaseMeshComponent(FDisplayClusterVie
 
 	if (MeshComponent)
 	{
+		UpdateHoldoutForMeshComponent(*Configuration, MeshComponent, EDisplayClusterDisplayDeviceMeshType::DefaultMesh);
+
 		EnumAddFlags(RuntimeFlags, EDisplayClusterViewportPreviewMeshFlags::HasDeletedMeshComponent);
 
 		UMaterial* const DefaultMaterial = DefaultMaterialPtr.Get();
diff --git a/Engine/Plugins/Runtime/nDisplay/Source/DisplayCluster/Public/DisplayClusterRootActor.h b/Engine/Plugins/Runtime/nDisplay/Source/DisplayCluster/Public/DisplayClusterRootActor.h
--- a/Engine/Plugins/Runtime/nDisplay/Source/DisplayCluster/Public/DisplayClusterRootActor.h
+++ b/Engine/Plugins/Runtime/nDisplay/Source/DisplayCluster/Public/DisplayClusterRootActor.h
@@ -252,6 +252,14 @@ public:
 	UFUNCTION(BlueprintSetter)
 	void SetPreviewEnablePostProcess(const bool bNewPreviewEnablePostProcess);
 
+	/**
+	* Blueprint setter for the PreviewEnableHoldoutComposite property. Makes sure preview pipeline is updated properly.
+	*
+	* @param bNewbPreviewEnableHoldoutComposite - Desired new property value.
+	*/
+	UFUNCTION(BlueprintSetter)
+	void SetPreviewEnableHoldoutComposite(const bool bNewPreviewEnableHoldoutComposite);
+
 public:
 	/** Get ViewportManager API. */
 	IDisplayClusterViewportManager* GetViewportManager() const;
@@ -462,6 +470,18 @@ public:
 	UPROPERTY(EditAnywhere, BlueprintReadWrite, Category = "Preview", meta = (DisplayName = "Enable Post Process"), BlueprintSetter = SetPreviewEnablePostProcess)
 	bool bPreviewEnablePostProcess = false;
 
+	/**
+	* Render a preview with postprocessing in the scene using the HoldoutComposite plugin.
+	* 
+	* The `Preview`>`Enable Post Process` option must be set
+	* 
+	* These project settings also need to be set:
+	* 1. "Alpha Output"
+	* 2. "Support Primitive Alpha Holdout (Deferred)"
+	*/
+	UPROPERTY(EditAnywhere, BlueprintReadWrite, Category = "Preview", meta = (DisplayName = "Enable Preview Holdout", EditCondition = "bPreviewEnablePostProcess"), BlueprintSetter = SetPreviewEnableHoldoutComposite)
+	bool bPreviewEnableHoldoutComposite = true;
+
 	/** Enable OCIO on preview. */
 	UPROPERTY(EditAnywhere, BlueprintReadWrite, Category = "Preview", meta = (DisplayName = "Enable OCIO on Preview"))
 	bool bPreviewEnableOCIO = false;
diff --git a/Engine/Plugins/Runtime/nDisplay/Source/DisplayCluster/Public/DisplayClusterRootActorContainers.h b/Engine/Plugins/Runtime/nDisplay/Source/DisplayCluster/Public/DisplayClusterRootActorContainers.h
--- a/Engine/Plugins/Runtime/nDisplay/Source/DisplayCluster/Public/DisplayClusterRootActorContainers.h
+++ b/Engine/Plugins/Runtime/nDisplay/Source/DisplayCluster/Public/DisplayClusterRootActorContainers.h
@@ -27,6 +27,8 @@ struct FDisplayClusterRootActorPropertyOverrides
 
 	/** Enable OCIO on preview. */
 	TOptional<bool> bPreviewEnableOCIO;
+	/** Enables the use of 'HoldoutComposite' plugin for the preview rendering. */
+	TOptional<bool> bPreviewEnableHoldoutComposite;
 
 	/** Show overlay material on the preview mesh when preview rendering is enabled (UMeshComponent::OverlayMaterial). */
 	TOptional<bool> bPreviewEnableOverlayMaterial;
diff --git a/Engine/Plugins/Runtime/nDisplay/Source/DisplayClusterLightCardEditor/Private/Viewport/DisplayClusterLightCardEditorViewportClient.cpp b/Engine/Plugins/Runtime/nDisplay/Source/DisplayClusterLightCardEditor/Private/Viewport/DisplayClusterLightCardEditorViewportClient.cpp
--- a/Engine/Plugins/Runtime/nDisplay/Source/DisplayClusterLightCardEditor/Private/Viewport/DisplayClusterLightCardEditorViewportClient.cpp
+++ b/Engine/Plugins/Runtime/nDisplay/Source/DisplayClusterLightCardEditor/Private/Viewport/DisplayClusterLightCardEditorViewportClient.cpp
@@ -1674,6 +1674,9 @@ void FDisplayClusterLightCardEditorViewportClient::UpdatePreviewActor(ADisplayCl
 				PropertyOverrides.bPreviewEnablePostProcess = true;
 				PropertyOverrides.bPreviewEnable = true;
 
+				// Can render a DCRA preview without the HoldoutComposite plugin.
+				PropertyOverrides.bPreviewEnableHoldoutComposite = false;
+
 				PropertyOverrides.PreviewSetttingsSource = EDisplayClusterConfigurationRootActorPreviewSettingsSource::RootActor;
 			}
 
diff --git a/Engine/Plugins/Runtime/nDisplay/Source/DisplayClusterScenePreview/Private/DisplayClusterLightCardEditorHelper.cpp b/Engine/Plugins/Runtime/nDisplay/Source/DisplayClusterScenePreview/Private/DisplayClusterLightCardEditorHelper.cpp
--- a/Engine/Plugins/Runtime/nDisplay/Source/DisplayClusterScenePreview/Private/DisplayClusterLightCardEditorHelper.cpp
+++ b/Engine/Plugins/Runtime/nDisplay/Source/DisplayClusterScenePreview/Private/DisplayClusterLightCardEditorHelper.cpp
@@ -242,6 +242,9 @@ void FDisplayClusterLightCardEditorHelper::SetRootActor(ADisplayClusterRootActor
 		PropertyOverrides.bPreviewEnablePostProcess = true;
 		PropertyOverrides.bPreviewEnable = true;
 
+		// Can render a DCRA preview without the HoldoutComposite plugin.
+		PropertyOverrides.bPreviewEnableHoldoutComposite = false;
+
 		PropertyOverrides.PreviewSetttingsSource = EDisplayClusterConfigurationRootActorPreviewSettingsSource::RootActor;
 	}
 
diff --git a/Engine/Plugins/Runtime/nDisplay/nDisplay.uplugin b/Engine/Plugins/Runtime/nDisplay/nDisplay.uplugin
--- a/Engine/Plugins/Runtime/nDisplay/nDisplay.uplugin
+++ b/Engine/Plugins/Runtime/nDisplay/nDisplay.uplugin
@@ -303,6 +303,10 @@
 		{
 			"Name": "OnlineSubsystemUtils",
 			"Enabled": true
+		},
+		{
+			"Name": "HoldoutComposite",
+			"Enabled": true
 		}
 	]
 }
diff --git a/Engine/Plugins/VirtualProduction/EpicStageApp/Source/EpicStageApp/Private/StageAppRouteHandler.cpp b/Engine/Plugins/VirtualProduction/EpicStageApp/Source/EpicStageApp/Private/StageAppRouteHandler.cpp
--- a/Engine/Plugins/VirtualProduction/EpicStageApp/Source/EpicStageApp/Private/StageAppRouteHandler.cpp
+++ b/Engine/Plugins/VirtualProduction/EpicStageApp/Source/EpicStageApp/Private/StageAppRouteHandler.cpp
@@ -417,6 +417,9 @@ void FStageAppRouteHandler::SetRendererRootActorPath(const int32 RendererId, con
 		PropertyOverrides.bPreviewEnablePostProcess = true;
 		PropertyOverrides.bPreviewEnable = true;
 
+		// Can render a DCRA preview without the HoldoutComposite plugin.
+		PropertyOverrides.bPreviewEnableHoldoutComposite = false;
+
 		PropertyOverrides.PreviewSetttingsSource = EDisplayClusterConfigurationRootActorPreviewSettingsSource::RootActor;
 	}
 
