nDisplay: Fix automation tests not cleaning up assets properly and causing crashes.

diff --git a/Engine/Plugins/Runtime/nDisplay/Source/DisplayClusterTests/Private/DisplayCluster.spec.cpp b/Engine/Plugins/Runtime/nDisplay/Source/DisplayClusterTests/Private/DisplayCluster.spec.cpp
--- a/Engine/Plugins/Runtime/nDisplay/Source/DisplayClusterTests/Private/DisplayCluster.spec.cpp
+++ b/Engine/Plugins/Runtime/nDisplay/Source/DisplayClusterTests/Private/DisplayCluster.spec.cpp
@@ -119,8 +119,8 @@ BEGIN_DEFINE_SPEC(FDisplayClusterSpec, "DisplayCluster", EAutomationTestFlags_Ap
 	UDisplayClusterConfigurationViewport* InstancedViewport;
 	UDisplayClusterXformComponent* InstancedComponent;
 
-	UOpenColorIOConfiguration* OCIOConfig;
-	UOpenColorIOConfiguration* AltOCIOConfig;
+	TWeakObjectPtr<UOpenColorIOConfiguration> OCIOConfig;
+	TWeakObjectPtr<UOpenColorIOConfiguration> AltOCIOConfig;
 
 	AActor* AltActor;
 
@@ -1000,17 +1000,17 @@ void FDisplayClusterSpec::DoOCIOConfigurationPropagationTests(TFunction<OwnerTyp
 	{
 		OCIOConfig = NewObject<UOpenColorIOConfiguration>(ClusterAsset->GetPackage(), FName("TestOCIOConfig"), RF_Transient | RF_Public);
 		OCIOConfig->AddToRoot();
-		FAssetRegistryModule::AssetCreated(OCIOConfig);
+		FAssetRegistryModule::AssetCreated(OCIOConfig.Get());
 		
 		AltOCIOConfig = NewObject<UOpenColorIOConfiguration>(ClusterAsset->GetPackage(), FName("TestOCIOConfig"), RF_Transient | RF_Public);
 		AltOCIOConfig->AddToRoot();
-		FAssetRegistryModule::AssetCreated(AltOCIOConfig);
+		FAssetRegistryModule::AssetCreated(AltOCIOConfig.Get());
 	});
 	
 	const FString GroupName = FString::Printf(TEXT("Color grading settings (%s)"), *FieldNames.Last().ToString()); 
 	Describe(*GroupName, [this, FieldNames, BeforeTestAction, AssetOwnerAccessor, InstancedOwnerAccessor, OwnerName, bDoInstanceOverrideTest]
 	{
-		OCIO_CONFIG_TEST_INTERNAL(static_cast<UObject*>(OCIOConfig), static_cast<UObject*>(AltOCIOConfig),
+		OCIO_CONFIG_TEST_INTERNAL(static_cast<UObject*>(OCIOConfig.Get()), static_cast<UObject*>(AltOCIOConfig.Get()),
 			TArray({ GET_MEMBER_NAME_CHECKED(FOpenColorIODisplayConfiguration, ColorConfiguration), GET_MEMBER_NAME_CHECKED(FOpenColorIOColorConversionSettings, ConfigurationSource) }));
 		
 		OCIO_CONFIG_TEST_INTERNAL(FString("TestName"), FString("OverrideName"),
@@ -1024,11 +1024,11 @@ void FDisplayClusterSpec::DoOCIOConfigurationPropagationTests(TFunction<OwnerTyp
 
 	AfterEach([this]()
 	{
-		DisplayClusterTestUtils::CleanUpAsset(OCIOConfig);
-		OCIOConfig = nullptr;
+		DisplayClusterTestUtils::CleanUpAsset(OCIOConfig.Get());
+		OCIOConfig.Reset();
 		
-		DisplayClusterTestUtils::CleanUpAsset(AltOCIOConfig);
-		AltOCIOConfig = nullptr;
+		DisplayClusterTestUtils::CleanUpAsset(AltOCIOConfig.Get());
+		AltOCIOConfig.Reset();
 	});
 }
 
diff --git a/Engine/Plugins/Runtime/nDisplay/Source/DisplayClusterTests/Private/DisplayClusterTestUtils.cpp b/Engine/Plugins/Runtime/nDisplay/Source/DisplayClusterTests/Private/DisplayClusterTestUtils.cpp
--- a/Engine/Plugins/Runtime/nDisplay/Source/DisplayClusterTests/Private/DisplayClusterTestUtils.cpp
+++ b/Engine/Plugins/Runtime/nDisplay/Source/DisplayClusterTests/Private/DisplayClusterTestUtils.cpp
@@ -7,11 +7,14 @@
 #include "ClusterConfiguration/DisplayClusterConfiguratorClusterUtils.h"
 #include "DisplayClusterConfiguratorFactory.h"
 #include "DisplayClusterTestsModule.h"
+#include "Editor/UnrealEdEngine.h"
 #include "Engine/Engine.h"
 #include "Engine/World.h"
 #include "ISinglePropertyView.h"
 #include "ObjectTools.h"
 #include "PackageTools.h"
+#include "Tests/AutomationEditorCommon.h"
+#include "UnrealEdGlobals.h"
 
 #if WITH_DEV_AUTOMATION_TESTS
 
@@ -135,9 +138,15 @@ namespace DisplayClusterTestUtils
 
 		Package->RemoveFromRoot();
 	
-		TArray<UPackage*> PackagesToUnload;
-		PackagesToUnload.Add(Package);
-		UPackageTools::UnloadPackages(PackagesToUnload);
+		TArray<UPackage*> PackagesToDelete;
+		PackagesToDelete.Add(Package);
+		// Let the package auto-saver know that it needs to ignore the deleted packages
+		GUnrealEd->GetPackageAutoSaver().OnPackagesDeleted(PackagesToDelete);
+
+		Package->SetDirtyFlag(false);
+
+		// Unload the packages and collect garbage.
+		UPackageTools::UnloadPackages(PackagesToDelete);
 	}
 
 	void CleanUpAsset(UObject* Asset)
@@ -150,6 +159,8 @@ namespace DisplayClusterTestUtils
 		// Rather than calling ObjectTools::ForceDeleteObjects, we just mark the asset and any of its Blueprint instances
 		// as garbage and trigger collection. This runs much faster (important since we need to do it after each test), but
 		// assumes that we haven't created any other references to the asset in the process of running the test.
+		
+		FAutomationEditorCommonUtils::NullReferencesToObject(Asset);
 		Asset->RemoveFromRoot();
 
 		UBlueprint* BlueprintObject = Cast<UBlueprint>(Asset);
@@ -161,6 +172,7 @@ namespace DisplayClusterTestUtils
 			for (TArray<UObject*>::TConstIterator InstanceItr( InstancesToDelete ); InstanceItr; ++InstanceItr)
 			{
 				UObject* CurrentInstance = *InstanceItr;
+				FAutomationEditorCommonUtils::NullReferencesToObject(CurrentInstance);
 				CurrentInstance->MarkAsGarbage();
 			}
 		}
