[nDisplay] Fixed typos.

diff --git a/Engine/Plugins/Runtime/nDisplay/Source/DisplayClusterProjection/Private/Policy/EasyBlend/DisplayClusterProjectionEasyBlendPolicyConfiguration.cpp b/Engine/Plugins/Runtime/nDisplay/Source/DisplayClusterProjection/Private/Policy/EasyBlend/DisplayClusterProjectionEasyBlendPolicyConfiguration.cpp
--- a/Engine/Plugins/Runtime/nDisplay/Source/DisplayClusterProjection/Private/Policy/EasyBlend/DisplayClusterProjectionEasyBlendPolicyConfiguration.cpp
+++ b/Engine/Plugins/Runtime/nDisplay/Source/DisplayClusterProjection/Private/Policy/EasyBlend/DisplayClusterProjectionEasyBlendPolicyConfiguration.cpp
@@ -32,13 +32,13 @@ bool FDisplayClusterProjectionEasyBlendPolicyConfiguration::Initialize(const TMa
 	// Origin node (optional)
 	if (DisplayClusterHelpers::map::template ExtractValueFromString(InParameters, DisplayClusterProjectionStrings::cfg::easyblend::Origin, OriginCompId))
 	{
-		UE_LOG(LogDisplayClusterProjectionVIOSO, Verbose, TEXT("Viewport <%s>: Projection parameter '%s' - '%s'"), *InViewportId, DisplayClusterProjectionStrings::cfg::easyblend::Origin, *OriginCompId);
+		UE_LOG(LogDisplayClusterProjectionEasyBlend, Verbose, TEXT("Viewport <%s>: Projection parameter '%s' - '%s'"), *InViewportId, DisplayClusterProjectionStrings::cfg::easyblend::Origin, *OriginCompId);
 	}
 	else
 	{
 		if (bEnableExtraLog)
 		{
-			UE_LOG(LogDisplayClusterProjectionVIOSO, Error, TEXT("Viewport <%s>: Projection parameter '%s' not found"), *InViewportId, DisplayClusterProjectionStrings::cfg::easyblend::Origin);
+			UE_LOG(LogDisplayClusterProjectionEasyBlend, Error, TEXT("Viewport <%s>: Projection parameter '%s' not found"), *InViewportId, DisplayClusterProjectionStrings::cfg::easyblend::Origin);
 		}
 
 		return false;
@@ -46,7 +46,7 @@ bool FDisplayClusterProjectionEasyBlendPolicyConfiguration::Initialize(const TMa
 
 	if (DisplayClusterHelpers::map::template ExtractValueFromString(InParameters, DisplayClusterProjectionStrings::cfg::mpcdi::EnablePreview, bIsPreviewMeshEnabled))
 	{
-		UE_LOG(LogDisplayClusterProjectionVIOSO, Verbose, TEXT("Viewport <%s>: Projection parameter '%s' - '%s'"), *InViewportId, DisplayClusterProjectionStrings::cfg::mpcdi::EnablePreview, *DisplayClusterHelpers::str::BoolToStr(bIsPreviewMeshEnabled));
+		UE_LOG(LogDisplayClusterProjectionEasyBlend, Verbose, TEXT("Viewport <%s>: Projection parameter '%s' - '%s'"), *InViewportId, DisplayClusterProjectionStrings::cfg::mpcdi::EnablePreview, *DisplayClusterHelpers::str::BoolToStr(bIsPreviewMeshEnabled));
 	}
 
 
diff --git a/Engine/Plugins/Runtime/nDisplay/Source/DisplayClusterProjection/Private/Policy/EasyBlend/Windows/DX11/DisplayClusterProjectionEasyBlendLibraryDX11.h b/Engine/Plugins/Runtime/nDisplay/Source/DisplayClusterProjection/Private/Policy/EasyBlend/Windows/DX11/DisplayClusterProjectionEasyBlendLibraryDX11.h
--- a/Engine/Plugins/Runtime/nDisplay/Source/DisplayClusterProjection/Private/Policy/EasyBlend/Windows/DX11/DisplayClusterProjectionEasyBlendLibraryDX11.h
+++ b/Engine/Plugins/Runtime/nDisplay/Source/DisplayClusterProjection/Private/Policy/EasyBlend/Windows/DX11/DisplayClusterProjectionEasyBlendLibraryDX11.h
@@ -42,7 +42,7 @@ public:
 	DeclareEasyBlendFunc(SetOutputTexture2D);
 
 private:
-	/** Imprort all functionf from VIOSO DLL. Return true on success. */
+	/** Imprort all functionf from EasyBlend DLL. Return true on success. */
 	bool InitializeDLL();
 	void ReleaseDLL();
 
@@ -51,7 +51,7 @@ private:
 	// Initializetion state
 	bool bInitialized = false;
 
-	// Saved handle to VIOSO DLL
+	// Saved handle to EasyBlend DLL
 	void* EasyBlend_DLL_Handler = nullptr;
 };
 
diff --git a/Engine/Plugins/Runtime/nDisplay/Source/DisplayClusterProjection/Private/Policy/EasyBlend/Windows/DX12/DisplayClusterProjectionEasyBlendLibraryDX12.h b/Engine/Plugins/Runtime/nDisplay/Source/DisplayClusterProjection/Private/Policy/EasyBlend/Windows/DX12/DisplayClusterProjectionEasyBlendLibraryDX12.h
--- a/Engine/Plugins/Runtime/nDisplay/Source/DisplayClusterProjection/Private/Policy/EasyBlend/Windows/DX12/DisplayClusterProjectionEasyBlendLibraryDX12.h
+++ b/Engine/Plugins/Runtime/nDisplay/Source/DisplayClusterProjection/Private/Policy/EasyBlend/Windows/DX12/DisplayClusterProjectionEasyBlendLibraryDX12.h
@@ -39,7 +39,7 @@ public:
 	DeclareEasyBlendFunc(EasyBlendSDK_TransformInputToOutputDX12_CommandQueue);
 
 private:
-	/** Imprort all functionf from VIOSO DLL. Return true on success. */
+	/** Imprort all functionf from EasyBlend DLL. Return true on success. */
 	bool InitializeDLL();
 	void ReleaseDLL();
 
@@ -47,7 +47,7 @@ private:
 	// Initializetion state
 	bool bInitialized = false;
 
-	// Saved handle to VIOSO DLL
+	// Saved handle to EasyBlend DLL
 	void* EasyBlend_DLL_Handler = nullptr;
 };
 
