[nDisplay] Fixed typo in the ICVFX pixel shader.

A typo was found in the ICVFX pixel shader that causes a gamma issue.

diff --git a/Engine/Plugins/Runtime/nDisplay/Shaders/Private/MPCDIOverlayShaders.usf b/Engine/Plugins/Runtime/nDisplay/Shaders/Private/MPCDIOverlayShaders.usf
--- a/Engine/Plugins/Runtime/nDisplay/Shaders/Private/MPCDIOverlayShaders.usf
+++ b/Engine/Plugins/Runtime/nDisplay/Shaders/Private/MPCDIOverlayShaders.usf
@@ -47,7 +47,7 @@ float4 GetFinalLightcardColor(float4 InColor)
 	OutColor.rgb = DisplayGammaCorrection(OutColor.rgb, LightCardGamma);
 
 	// And finally premultiply it back
-	OutColor = InvertedPremultiply(InColor);
+	OutColor = InvertedPremultiply(OutColor);
 	
 	return OutColor;
 }
