nDisplay: Fix issue where recording with TakeRecorder and MU could cause a crash when rejoining a session.

diff --git a/Engine/Plugins/Runtime/nDisplay/Source/DisplayClusterConfiguration/Private/DisplayClusterConfigurationTypes.cpp b/Engine/Plugins/Runtime/nDisplay/Source/DisplayClusterConfiguration/Private/DisplayClusterConfigurationTypes.cpp
--- a/Engine/Plugins/Runtime/nDisplay/Source/DisplayClusterConfiguration/Private/DisplayClusterConfigurationTypes.cpp
+++ b/Engine/Plugins/Runtime/nDisplay/Source/DisplayClusterConfiguration/Private/DisplayClusterConfigurationTypes.cpp
@@ -389,7 +389,7 @@ void UDisplayClusterConfigurationCluster::GetObjectsToExport(TArray<UObject*>& O
 void UDisplayClusterConfigurationData_Base::Serialize(FArchive& Ar)
 {
 #if WITH_EDITOR
-	if (Ar.IsSaving())
+	if (Ar.IsSaving() && !Ar.IsTransacting())
 	{
 		/*
 		 * We need to set everything to public so it can be saved & referenced properly with the object.
