nDisplay: Avoid constant media restart if the camera is disabled on start. Media restarts flush rendering commands which introduces hitches.

diff --git a/Engine/Plugins/Runtime/nDisplay/Source/DisplayClusterMedia/Private/Capture/DisplayClusterMediaCaptureBase.cpp b/Engine/Plugins/Runtime/nDisplay/Source/DisplayClusterMedia/Private/Capture/DisplayClusterMediaCaptureBase.cpp
--- a/Engine/Plugins/Runtime/nDisplay/Source/DisplayClusterMedia/Private/Capture/DisplayClusterMediaCaptureBase.cpp
+++ b/Engine/Plugins/Runtime/nDisplay/Source/DisplayClusterMedia/Private/Capture/DisplayClusterMediaCaptureBase.cpp
@@ -236,7 +236,14 @@ void FDisplayClusterMediaCaptureBase::OnPostClusterTick()
 			const FIntPoint LastSrcRegionIntPoint = LastSrcRegionSize.load().ToIntPoint();
 			const FIntPoint DesiredSize = MediaCapture->GetDesiredSize();
 
-			if (DesiredSize != LastSrcRegionIntPoint)
+			// We don't restart if we haven't exported any textures yet (indicated by zero size LastSrcRegion)
+			// to avoid constant media restarts since in such case media is set to use GetCaptureSize() != (0.0).
+			// Once an export happens, any media restart will use LastSrcRegion which should not trigger a restart
+			// when texture exports are suspended, since that does not cause a mismatch. 
+			// This is preferred to setting LastSrcRegion to GetCaptureSize() because that would not reflect the actual
+			// last captured size.
+
+			if ((DesiredSize != LastSrcRegionIntPoint) && (LastSrcRegionIntPoint != FIntPoint(0,0)))
 			{
 				UE_LOG(LogDisplayClusterMedia, Log, TEXT("Stopping MediaCapture '%s' because its DesiredSize (%d, %d) doesn't match the captured texture size (%d, %d)"), 
 					*GetMediaId(), DesiredSize.X, DesiredSize.Y, LastSrcRegionIntPoint.X, LastSrcRegionIntPoint.Y);
@@ -289,13 +296,16 @@ bool FDisplayClusterMediaCaptureBase::StartMediaCapture()
 	MediaCaptureOptions.OverrunAction = EMediaCaptureOverrunAction::Flush;
 
 	const bool bCaptureStarted = MediaCapture->CaptureRHITexture(Descriptor, MediaCaptureOptions);
+
 	if (bCaptureStarted)
 	{
-		UE_LOG(LogDisplayClusterMedia, Log, TEXT("Started media capture: '%s'"), *GetMediaId());
+		UE_LOG(LogDisplayClusterMedia, Log, TEXT("Started media capture: '%s' (%d x %d)"),
+			*GetMediaId(), Descriptor.ResourceSize.X, Descriptor.ResourceSize.Y);
 	}
 	else
 	{
-		UE_LOG(LogDisplayClusterMedia, Warning, TEXT("Couldn't start media capture '%s'"), *GetMediaId());
+		UE_LOG(LogDisplayClusterMedia, Warning, TEXT("Couldn't start media capture '%s' (%d x %d)"), 
+			*GetMediaId(), Descriptor.ResourceSize.X, Descriptor.ResourceSize.Y);
 	}
 
 	return bCaptureStarted;
