nDisplay: Fix a crash when icvfx camera resolution is changed and SharedMemoryMedia is used to for sharing GPU textures.

Game thread was releasing resources after async tasks were done and flushing rendering commands, but that flush only includes render thread and not rhi thread, which could still access those resources and launch an async task when building the rdg graph. The fix is to include the graph builder lambda as an async task to wait for.

diff --git a/Engine/Plugins/Runtime/nDisplay/Source/SharedMemoryMedia/Private/SharedMemoryMediaCapture.cpp b/Engine/Plugins/Runtime/nDisplay/Source/SharedMemoryMedia/Private/SharedMemoryMediaCapture.cpp
--- a/Engine/Plugins/Runtime/nDisplay/Source/SharedMemoryMedia/Private/SharedMemoryMediaCapture.cpp
+++ b/Engine/Plugins/Runtime/nDisplay/Source/SharedMemoryMedia/Private/SharedMemoryMediaCapture.cpp
@@ -421,6 +421,11 @@ void USharedMemoryMediaCapture::AddCopyToSharedGpuTexturePass(FRDGBuilder& Graph
 	PassParameters->DstTexture = GraphBuilder.RegisterExternalTexture(
 		CreateRenderTarget(SharedCrossGpuTextures[SharedTextureIdx], *FString::Printf(TEXT("SharedCrossGpuTextures_%d"), SharedTextureIdx)));
 
+	// We increment RunningTasksCount here because the RDG build will run on an RHI related thread,
+	// so for the purposes of waiting for resources to be freed, we must include the AddPass lambda execution,
+	// which contains the async task that will decrement RunningTasksCount.
+	RunningTasksCount++;
+
 	GraphBuilder.AddPass(
 		RDG_EVENT_NAME("Copy%sToSharedGpuTexture", InSourceTexture->Name),
 		PassParameters,
@@ -453,7 +458,6 @@ void USharedMemoryMediaCapture::AddCopyToSharedGpuTexturePass(FRDGBuilder& Graph
 
 			// Spawn a thread that via shared ram will notify receiver that data is ready
 			// It will also verify that the data has been consumed (with a timeout).
-			RunningTasksCount++;
 			UE::Tasks::Launch(UE_SOURCE_LOCATION, [FrameNumber = GFrameCounterRenderThread, SharedTextureIdx, this]()
 			{
 				// Decrement RunningTasksCount when the task exits
