[nDisplay] Fixed the gamma parameter value of the ICVFX shader for LightCards.

To transform a linear function into a gamma function, we should use the inverted gamma value for the pow function.

diff --git a/Engine/Plugins/Runtime/nDisplay/Source/DisplayClusterShaders/Private/Shaders/DisplayClusterShadersWarpblend_ICVFX.cpp b/Engine/Plugins/Runtime/nDisplay/Source/DisplayClusterShaders/Private/Shaders/DisplayClusterShadersWarpblend_ICVFX.cpp
--- a/Engine/Plugins/Runtime/nDisplay/Source/DisplayClusterShaders/Private/Shaders/DisplayClusterShadersWarpblend_ICVFX.cpp
+++ b/Engine/Plugins/Runtime/nDisplay/Source/DisplayClusterShaders/Private/Shaders/DisplayClusterShadersWarpblend_ICVFX.cpp
@@ -671,8 +671,6 @@ public:
 
 		if (EnumHasAnyLightCardModeFlags(RenderPassData, EShaderParameterLightCardMode::LIGHTCARD_UNDER | EShaderParameterLightCardMode::LIGHTCARD_OVER))
 		{
-			RenderPassData.PSParameters.LightCardGamma = ICVFXParameters.LightCardGamma;
-
 			RenderPassData.PSParameters.LightCardUnderTexture = ICVFXParameters.LightCardUnder.Texture ? ICVFXParameters.LightCardUnder.Texture : ICVFXParameters.LightCardOver.Texture;
 			RenderPassData.PSParameters.LightCardOverTexture  = ICVFXParameters.LightCardOver.Texture  ? ICVFXParameters.LightCardOver.Texture  : ICVFXParameters.LightCardUnder.Texture;
 
@@ -724,8 +722,6 @@ public:
 
 		if (EnumHasAnyLightCardModeFlags(RenderPassData, EShaderParameterLightCardMode::UVLIGHTCARD_UNDER | EShaderParameterLightCardMode::UVLIGHTCARD_OVER))
 		{
-			RenderPassData.PSParameters.LightCardGamma = ICVFXParameters.LightCardGamma;
-
 			RenderPassData.PSParameters.UVLightCardUnderTexture = ICVFXParameters.UVLightCardUnder.Texture ? ICVFXParameters.UVLightCardUnder.Texture : ICVFXParameters.UVLightCardOver.Texture;
 			RenderPassData.PSParameters.UVLightCardOverTexture  = ICVFXParameters.UVLightCardOver.Texture  ? ICVFXParameters.UVLightCardOver.Texture  : ICVFXParameters.UVLightCardUnder.Texture;
 
@@ -991,6 +987,20 @@ public:
 
 			// Configure Lightcards:
 			{
+				// Both LightCards textures use linear gamma, so they have the same gamma correction at the end.
+				{
+					// The LightCard texture is in linear color space and should be converted to the target gamma.
+					// Math: GammaDst = Pow( LinearSrc, (1/Gamma) ).
+					if (ICVFXParameters.LightCardGamma > KINDA_SMALL_NUMBER)
+					{
+						RenderPassData.PSParameters.LightCardGamma = 1.0f / ICVFXParameters.LightCardGamma;
+					}
+					else
+					{
+						RenderPassData.PSParameters.LightCardGamma = 1.0f;
+					}
+				}
+
 				InitializeLightCardParameters(RenderPassData);
 				GetLightCardParameters(RenderPassData);
 				GetUVLightCardParameters(RenderPassData);
