[nDisplay] Fixed issues related to overscan

- Added a new texture resource enumeration type called `InternalRenderTargetEntireRectResource`. This type is used for the entire RTT rendering area.

- Changed the use of `InternalRenderTargetResource` resource type: it is used for internal area on texture (after overscan, etc).
*Note: MediaIO will now copy smaller textures (after overscan), which should improve performance.

- `InternalRenderTargetEntireRectResource` and `InternalRenderTargetResource` refer to the same texture resource, but with different rects.

diff --git a/Engine/Plugins/Runtime/nDisplay/Source/DisplayCluster/Private/Render/Viewport/DisplayClusterViewportProxy.cpp b/Engine/Plugins/Runtime/nDisplay/Source/DisplayCluster/Private/Render/Viewport/DisplayClusterViewportProxy.cpp
--- a/Engine/Plugins/Runtime/nDisplay/Source/DisplayCluster/Private/Render/Viewport/DisplayClusterViewportProxy.cpp
+++ b/Engine/Plugins/Runtime/nDisplay/Source/DisplayCluster/Private/Render/Viewport/DisplayClusterViewportProxy.cpp
@@ -321,7 +321,7 @@ void FDisplayClusterViewportProxy::OnResolvedSceneColor_RenderThread(FRDGBuilder
 		case EDisplayClusterRenderFrameAlphaChannelCaptureMode::Copy:
 		case EDisplayClusterRenderFrameAlphaChannelCaptureMode::CopyAA:
 		{
-			const FIntRect SrcRect = GetFinalContextRect(EDisplayClusterViewportResourceType::InternalRenderTargetResource, Contexts[InContextNum].RenderTargetRect);
+			const FIntRect SrcRect = GetResourceRect_RenderThread(EDisplayClusterViewportResourceType::InternalRenderTargetResource, InContextNum);
 			// Copy alpha channel from 'SceneTextures.Color.Resolve' to 'InputShaderResource'
 			GetShadersAPI().CreateTextureUtils_RenderThread(GraphBuilder)
 				->SetInput({ SceneTextures.Color.Resolve, SrcRect }, InContextNum)
@@ -342,7 +342,7 @@ FScreenPassTexture FDisplayClusterViewportProxy::OnPostProcessPassAfterSSRInput_
 	if (OutScreenPassTexture.IsValid())
 	{
 		// Copy alpha channel to 'InputShaderResource'
-		const FIntRect SrcRect = GetFinalContextRect(EDisplayClusterViewportResourceType::InternalRenderTargetResource, Contexts[ContextNum].RenderTargetRect);
+		const FIntRect SrcRect = GetResourceRect_RenderThread(EDisplayClusterViewportResourceType::InternalRenderTargetResource, ContextNum);
 		GetShadersAPI().CreateTextureUtils_RenderThread(GraphBuilder)
 			->SetInput({ OutScreenPassTexture.Texture, SrcRect }, ContextNum)
 			->SetOutput(this, EDisplayClusterViewportResourceType::InputShaderResource, ContextNum)
@@ -359,7 +359,7 @@ FScreenPassTexture FDisplayClusterViewportProxy::OnPostProcessPassAfterFXAA_Rend
 	{
 		// Restore alpha channel after OCIO
 		// Copy alpha channel from 'InputShaderResource'
-		const FIntRect DestRect = GetFinalContextRect(EDisplayClusterViewportResourceType::InternalRenderTargetResource, Contexts[ContextNum].RenderTargetRect);
+		const FIntRect DestRect = GetResourceRect_RenderThread(EDisplayClusterViewportResourceType::InternalRenderTargetResource, ContextNum);
 		GetShadersAPI().CreateTextureUtils_RenderThread(GraphBuilder)
 			->SetInput(this, EDisplayClusterViewportResourceType::InputShaderResource, ContextNum)
 			->SetOutput({ OutScreenPassTexture.Texture, DestRect }, ContextNum)
@@ -478,6 +478,7 @@ FDisplayClusterColorEncoding FDisplayClusterViewportProxy::GetResourceColorEncod
 	const EDisplayClusterViewportResourceType ResourceType = GetResourceType_RenderThread(InResourceType);
 	switch (ResourceType)
 	{
+	case EDisplayClusterViewportResourceType::InternalRenderTargetEntireRectResource:
 	case EDisplayClusterViewportResourceType::InternalRenderTargetResource:
 		if (EnumHasAllFlags(RenderSettingsICVFX.RuntimeFlags, EDisplayClusterViewportRuntimeICVFXFlags::UVLightcard))
 		{
diff --git a/Engine/Plugins/Runtime/nDisplay/Source/DisplayCluster/Private/Render/Viewport/DisplayClusterViewportProxy.h b/Engine/Plugins/Runtime/nDisplay/Source/DisplayCluster/Private/Render/Viewport/DisplayClusterViewportProxy.h
--- a/Engine/Plugins/Runtime/nDisplay/Source/DisplayCluster/Private/Render/Viewport/DisplayClusterViewportProxy.h
+++ b/Engine/Plugins/Runtime/nDisplay/Source/DisplayCluster/Private/Render/Viewport/DisplayClusterViewportProxy.h
@@ -157,11 +157,11 @@ public:
 	/** Get actual region for viewport context
 	 * 
 	 * @param InResourceType - the requested resource type from the entire namespace
-	 * @param InRect - viewport context region
+	 * @param InContextNum - context num (eye)
 	 * 
 	 * @return - InRect with applied overscan for RTT
 	 */
-	FIntRect GetFinalContextRect(const EDisplayClusterViewportResourceType InResourceType, const FIntRect& InRect) const;
+	FIntRect GetResourceRect_RenderThread(const EDisplayClusterViewportResourceType InResourceType, const uint32 InContextNum) const;
 
 	/** Resolve viewport RTT: render OCIO, PP, generate MIPS, etc.
 	 *
diff --git a/Engine/Plugins/Runtime/nDisplay/Source/DisplayCluster/Private/Render/Viewport/DisplayClusterViewportProxy_Resources.cpp b/Engine/Plugins/Runtime/nDisplay/Source/DisplayCluster/Private/Render/Viewport/DisplayClusterViewportProxy_Resources.cpp
--- a/Engine/Plugins/Runtime/nDisplay/Source/DisplayCluster/Private/Render/Viewport/DisplayClusterViewportProxy_Resources.cpp
+++ b/Engine/Plugins/Runtime/nDisplay/Source/DisplayCluster/Private/Render/Viewport/DisplayClusterViewportProxy_Resources.cpp
@@ -98,6 +98,7 @@ bool FDisplayClusterViewportProxy::ImplGetResources_RenderThread(const EDisplayC
 
 	switch (InResourceType)
 	{
+	case EDisplayClusterViewportResourceType::InternalRenderTargetEntireRectResource:
 	case EDisplayClusterViewportResourceType::InternalRenderTargetResource:
 	{
 		bool bResult = false;
@@ -203,43 +204,23 @@ bool FDisplayClusterViewportProxy::ImplGetResourcesWithRects_RenderThread(const
 	}
 
 	const EDisplayClusterViewportResourceType InResourceType = GetResourceType_RenderThread(InExtResourceType);
-
 	if (!GetResources_RenderThread(InResourceType, OutResources))
 	{
 		return false;
 	}
 
 	// Collect all resource rects:
-	OutResourceRects.AddDefaulted(OutResources.Num());
-
-	switch (InResourceType)
+	for (int32 ContextIt = 0; ContextIt < OutResources.Num(); ContextIt++)
 	{
-	case EDisplayClusterViewportResourceType::InternalRenderTargetResource:
-		for (int32 ContextIt = 0; ContextIt < OutResourceRects.Num(); ContextIt++)
-		{
-			if (PostRenderSettings.Replace.IsEnabled())
-			{
-				// Get image from Override
-				OutResourceRects[ContextIt] = PostRenderSettings.Replace.Rect;
-			}
-			else
-			{
-				OutResourceRects[ContextIt] = Contexts[ContextIt].RenderTargetRect;
-			}
-		}
-		break;
-	case EDisplayClusterViewportResourceType::OutputFrameTargetableResource:
-	case EDisplayClusterViewportResourceType::AdditionalFrameTargetableResource:
-		for (int32 ContextIt = 0; ContextIt < OutResourceRects.Num(); ContextIt++)
-		{
-			OutResourceRects[ContextIt] = Contexts[ContextIt].FrameTargetRect;
-		}
-		break;
-	default:
-		for (int32 ContextIt = 0; ContextIt < OutResourceRects.Num(); ContextIt++)
+		FIntRect ResourceRect = GetResourceRect_RenderThread(InResourceType, ContextIt);
+
+		// Rect({0,0}, {0,0} means we are using the entire texture.
+		if (ResourceRect.IsEmpty() && OutResources[ContextIt])
 		{
-			OutResourceRects[ContextIt] = FIntRect(FIntPoint(0, 0), OutResources[ContextIt]->GetSizeXY());
+			ResourceRect.Max = OutResources[ContextIt]->GetDesc().Extent;
 		}
+
+		OutResourceRects.Add(ResourceRect);
 	}
 
 	return true;
@@ -300,13 +281,13 @@ void FDisplayClusterViewportProxy::ImplResolveTileResource_RenderThread(FRHIComm
 
 	const TArray<FDisplayClusterViewport_Context>& DestContexts = InDestViewportProxy->GetContexts_RenderThread();
 	const TArray<FDisplayClusterViewport_Context>& SrcContexts = Contexts;
-
-	TSharedRef<IDisplayClusterShadersTextureUtils> TextureUtils = GetShadersAPI().CreateTextureUtils_RenderThread(RHICmdList)
-		->SetInput(this, EDisplayClusterViewportResourceType::InternalRenderTargetResource)
-		->SetOutput(InDestViewportProxy, EDisplayClusterViewportResourceType::InternalRenderTargetResource);
+	TSharedRef<IDisplayClusterShadersTextureUtils> TextureUtils = GetShadersAPI().CreateTextureUtils_RenderThread(RHICmdList);
 
 	TextureUtils
-		->ForEachContextByPredicate(EDisplayClusterShaderTextureUtilsFlags::DisableUpdateResourcesRectsForResolve, [&](
+		->SetInput(this, EDisplayClusterViewportResourceType::InternalRenderTargetResource)
+		->SetOutput(InDestViewportProxy, EDisplayClusterViewportResourceType::InternalRenderTargetResource)
+		->ForEachContextByPredicate(
+			EDisplayClusterShaderTextureUtilsFlags::DisableUpdateResourcesRectsForResolve, [&](
 			const FDisplayClusterShadersTextureViewportContext& InputContext,
 			const FDisplayClusterShadersTextureViewportContext& OutputContext)
 			{
@@ -314,16 +295,15 @@ void FDisplayClusterViewportProxy::ImplResolveTileResource_RenderThread(FRHIComm
 				if(!SrcContexts.IsValidIndex(ContextNum) || !DestContexts.IsValidIndex(ContextNum) || OutputContext.ContextNum != ContextNum)
 				{
 					return;
-	}
+				}
 
 				FDisplayClusterShadersTextureViewport Input(InputContext);
 				FDisplayClusterShadersTextureViewport Output(OutputContext);
 
-				Input.Rect  = GetFinalContextRect(EDisplayClusterViewportResourceType::InternalRenderTargetResource, Input.Rect);
 				Output.Rect = SrcContexts[ContextNum].TileDestRect;
 
 				if (GDisplayClusterRenderTileBorder > 0)
-	{
+				{
 					// The maximum border is 1/4 of the minimum side of the rectangle.
 					const int32 MaxBorderSize = FMath::Min(Input.Rect.Size().GetMin(), Output.Rect.Size().GetMin()) / 4;
 					const int32 TileBorderSize = FMath::Min(GDisplayClusterRenderTileBorder, MaxBorderSize);
@@ -331,23 +311,23 @@ void FDisplayClusterViewportProxy::ImplResolveTileResource_RenderThread(FRHIComm
 					// Shows borders for internal areas only
 					FIntRect SubRect(FIntPoint::ZeroValue, FIntPoint::ZeroValue);
 					if (RenderSettings.TileSettings.GetPos().X > 0)
-			{
+					{
 						SubRect.Min.X = TileBorderSize;
-}
+					}
 					if (RenderSettings.TileSettings.GetPos().Y > 0)
-{
+					{
 						SubRect.Min.Y = TileBorderSize;
 					}
 					if ((RenderSettings.TileSettings.GetPos().X + 1) < RenderSettings.TileSettings.GetSize().X)
-		{
+					{
 						SubRect.Max.X = TileBorderSize;
 					}
 					if ((RenderSettings.TileSettings.GetPos().Y + 1) < RenderSettings.TileSettings.GetSize().Y)
-			{
+					{
 						SubRect.Max.Y = TileBorderSize;
 					}
 
-				// Set rect smaller to show gaps between tiles:
+					// Set rect smaller to show gaps between tiles:
 					Input.Rect.Min.X += SubRect.Min.X;
 					Input.Rect.Min.Y += SubRect.Min.Y;
 					Input.Rect.Max.X -= SubRect.Min.X + SubRect.Max.X;
@@ -357,10 +337,12 @@ void FDisplayClusterViewportProxy::ImplResolveTileResource_RenderThread(FRHIComm
 					Output.Rect.Min.Y += SubRect.Min.Y;
 					Output.Rect.Max.X -= SubRect.Min.X + SubRect.Max.X;
 					Output.Rect.Max.Y -= SubRect.Min.Y + SubRect.Max.Y;
-			}
+				}
 
 				// Copy texture region
-				TextureUtils->ResolveTextureContext(EDisplayClusterShaderTextureUtilsFlags::DisableResampleShader, Input, Output);
+				TextureUtils->ResolveTextureContext(
+					EDisplayClusterShaderTextureUtilsFlags::DisableResampleShader,
+					Input, Output);
 			});
 }
 
@@ -370,7 +352,7 @@ void FDisplayClusterViewportProxy::CleanupResources_RenderThread(FRHICommandList
 	if (GDisplayClusterRenderTileBorder > 0 && RenderSettings.TileSettings.GetType() == EDisplayClusterViewportTileType::Source)
 	{
 		TArray<FRHITexture*> RenderTargets;
-		if (GetResources_RenderThread(EDisplayClusterViewportResourceType::InternalRenderTargetResource, RenderTargets))
+		if (GetResources_RenderThread(EDisplayClusterViewportResourceType::InternalRenderTargetEntireRectResource, RenderTargets))
 		{
 			for (FRHITexture* TextureIt : RenderTargets)
 			{
diff --git a/Engine/Plugins/Runtime/nDisplay/Source/DisplayCluster/Private/Render/Viewport/DisplayClusterViewportProxy_Rules.cpp b/Engine/Plugins/Runtime/nDisplay/Source/DisplayCluster/Private/Render/Viewport/DisplayClusterViewportProxy_Rules.cpp
--- a/Engine/Plugins/Runtime/nDisplay/Source/DisplayCluster/Private/Render/Viewport/DisplayClusterViewportProxy_Rules.cpp
+++ b/Engine/Plugins/Runtime/nDisplay/Source/DisplayCluster/Private/Render/Viewport/DisplayClusterViewportProxy_Rules.cpp
@@ -185,6 +185,7 @@ bool FDisplayClusterViewportProxy::ShouldOverrideViewportResource(const EDisplay
 		{
 			switch (GetResourceType_RenderThread(InExtResourceType))
 			{
+			case EDisplayClusterViewportResourceType::InternalRenderTargetEntireRectResource:
 			case EDisplayClusterViewportResourceType::InternalRenderTargetResource:
 			case EDisplayClusterViewportResourceType::InputShaderResource:
 			case EDisplayClusterViewportResourceType::MipsShaderResource:
@@ -201,6 +202,7 @@ bool FDisplayClusterViewportProxy::ShouldOverrideViewportResource(const EDisplay
 		{
 			switch (GetResourceType_RenderThread(InExtResourceType))
 			{
+			case EDisplayClusterViewportResourceType::InternalRenderTargetEntireRectResource:
 			case EDisplayClusterViewportResourceType::InternalRenderTargetResource:
 			case EDisplayClusterViewportResourceType::InputShaderResource:
 			case EDisplayClusterViewportResourceType::MipsShaderResource:
@@ -216,6 +218,7 @@ bool FDisplayClusterViewportProxy::ShouldOverrideViewportResource(const EDisplay
 		{
 			switch (GetResourceType_RenderThread(InExtResourceType))
 			{
+			case EDisplayClusterViewportResourceType::InternalRenderTargetEntireRectResource:
 			case EDisplayClusterViewportResourceType::InternalRenderTargetResource:
 				return true;
 
@@ -296,27 +299,6 @@ bool FDisplayClusterViewportProxy::IsInputRenderTargetResourceExists() const
 	return !Resources[EDisplayClusterViewportResource::RenderTargets].IsEmpty();
 }
 
-FIntRect FDisplayClusterViewportProxy::GetFinalContextRect(const EDisplayClusterViewportResourceType InExtResourceType, const FIntRect& InRect) const
-{
-	const EDisplayClusterViewportResourceType InResourceType = GetResourceType_RenderThread(InExtResourceType);
-
-	// When resolving without warp, apply overscan
-	switch (InResourceType)
-	{
-	case EDisplayClusterViewportResourceType::InternalRenderTargetResource:
-		if (OverscanRuntimeSettings.bIsEnabled && CVarDisplayClusterRenderOverscanResolve.GetValueOnRenderThread() != 0)
-		{
-			// Support overscan crop
-			return OverscanRuntimeSettings.OverscanPixels.GetInnerRect(InRect);
-		}
-		break;
-	default:
-		break;
-	}
-
-	return InRect;
-}
-
 /**
 * The viewport proxy priority values
 * A lower priority value for a viewport means that this viewport will be the first in the list of viewports.
@@ -377,3 +359,52 @@ uint8 FDisplayClusterViewportProxy::GetPriority_RenderThread() const
 
 	return OutOrder;
 }
+
+FIntRect FDisplayClusterViewportProxy::GetResourceRect_RenderThread(const EDisplayClusterViewportResourceType InResourceType, const uint32 InContextNum) const
+{
+	// Rect({0,0}, {0,0} means we are using the entire texture.
+	FIntRect OutResourceRect = FIntRect(FIntPoint::ZeroValue, FIntPoint::ZeroValue);
+
+	if (Contexts.IsValidIndex(InContextNum))
+	{
+		switch (InResourceType)
+		{
+		case EDisplayClusterViewportResourceType::InternalRenderTargetEntireRectResource:
+		case EDisplayClusterViewportResourceType::InternalRenderTargetResource:
+			if (PostRenderSettings.Replace.IsEnabled())
+			{
+				// Get external texture rect.
+				OutResourceRect = PostRenderSettings.Replace.Rect;
+			}
+			else
+			{
+				// Get RTT rect
+				OutResourceRect = Contexts[InContextNum].RenderTargetRect;
+
+				// Get final rect
+				if (InResourceType == EDisplayClusterViewportResourceType::InternalRenderTargetResource)
+				{
+					// Get the inner region from the overscan texture.
+					if (OverscanRuntimeSettings.bIsEnabled && CVarDisplayClusterRenderOverscanResolve.GetValueOnRenderThread() != 0)
+					{
+						// Apply overscan crop
+						OutResourceRect = OverscanRuntimeSettings.OverscanPixels.GetInnerRect(OutResourceRect);
+					}
+				}
+			}
+			break;
+
+		case EDisplayClusterViewportResourceType::OutputFrameTargetableResource:
+		case EDisplayClusterViewportResourceType::AdditionalFrameTargetableResource:
+			OutResourceRect = Contexts[InContextNum].FrameTargetRect;
+			break;
+
+		default:
+			// All other resources use the same context size
+			OutResourceRect.Max = Contexts[InContextNum].ContextSize;
+			break;
+		}
+	}
+
+	return OutResourceRect;
+}
\ No newline at end of file
diff --git a/Engine/Plugins/Runtime/nDisplay/Source/DisplayCluster/Private/Render/Viewport/DisplayClusterViewportResources.h b/Engine/Plugins/Runtime/nDisplay/Source/DisplayCluster/Private/Render/Viewport/DisplayClusterViewportResources.h
--- a/Engine/Plugins/Runtime/nDisplay/Source/DisplayCluster/Private/Render/Viewport/DisplayClusterViewportResources.h
+++ b/Engine/Plugins/Runtime/nDisplay/Source/DisplayCluster/Private/Render/Viewport/DisplayClusterViewportResources.h
@@ -17,7 +17,9 @@ class FDisplayClusterViewportResource;
 enum class EDisplayClusterViewportResource : uint8
 {
 	// View family render to this resources
-	// This is the resource that can be pointed to using the enum 'EDisplayClusterViewportResourceType::InternalRenderTargetResource' value.
+	// This is the resource that can be pointed to using the enum values:
+	// - 'EDisplayClusterViewportResourceType::InternalRenderTargetEntireRectResource'
+	// - 'EDisplayClusterViewportResourceType::InternalRenderTargetResource'
 	RenderTargets,
 
 	// unique viewport resources
diff --git a/Engine/Plugins/Runtime/nDisplay/Source/DisplayCluster/Public/Render/Viewport/Containers/DisplayClusterViewport_Enums.h b/Engine/Plugins/Runtime/nDisplay/Source/DisplayCluster/Public/Render/Viewport/Containers/DisplayClusterViewport_Enums.h
--- a/Engine/Plugins/Runtime/nDisplay/Source/DisplayCluster/Public/Render/Viewport/Containers/DisplayClusterViewport_Enums.h
+++ b/Engine/Plugins/Runtime/nDisplay/Source/DisplayCluster/Public/Render/Viewport/Containers/DisplayClusterViewport_Enums.h
@@ -14,6 +14,12 @@ enum class EDisplayClusterViewportResourceType : uint8
 
 	// This RTT is used to render the scene for the viewport eye context to the specified area on this texture.
 	// The mGPU rendering with cross-GPU transfer can be used for this texture.
+	// This resource has an area that includes the overscan extension.
+	InternalRenderTargetEntireRectResource, /** Just for Internal use */
+
+	/**
+	* Same RTT texture resource, but with the final value of the inner area on the texture (after overscan, etc.).
+	*/
 	InternalRenderTargetResource, /** Just for Internal use */
 
 
diff --git a/Engine/Plugins/Runtime/nDisplay/Source/DisplayClusterShaders/Private/DisplayClusterShadersTextureUtils.cpp b/Engine/Plugins/Runtime/nDisplay/Source/DisplayClusterShaders/Private/DisplayClusterShadersTextureUtils.cpp
--- a/Engine/Plugins/Runtime/nDisplay/Source/DisplayClusterShaders/Private/DisplayClusterShadersTextureUtils.cpp
+++ b/Engine/Plugins/Runtime/nDisplay/Source/DisplayClusterShaders/Private/DisplayClusterShadersTextureUtils.cpp
@@ -17,7 +17,8 @@ namespace UE::DisplayClusterShaders::Private
 	{
 		switch (InResourceType)
 		{
-		case EDisplayClusterViewportResourceType::InternalRenderTargetResource:     return TEXT("nDisplay.InternalRenderTargetResource");
+		case EDisplayClusterViewportResourceType::InternalRenderTargetEntireRectResource : return TEXT("nDisplay.InternalRenderTargetEntireRectResource");
+		case EDisplayClusterViewportResourceType::InternalRenderTargetResource           : return TEXT("nDisplay.InternalRenderTargetResource");
 
 		case EDisplayClusterViewportResourceType::InputShaderResource:              return TEXT("nDisplay.InputShaderResource");
 		case EDisplayClusterViewportResourceType::MipsShaderResource:               return TEXT("nDisplay.MipsShaderResource");
