[nDisplay] The rules for limiting the viewport size have been reworked.

- Added new CVar "nDisplay.render.MaxTextureAreaRoot"
- Code cleanup
- Added new rule for viewports `FDisplayClusterViewport::ShouldApplyMaxTextureConstraints()`

diff --git a/Engine/Plugins/Runtime/nDisplay/Source/DisplayCluster/Private/Render/Viewport/DisplayClusterViewport.cpp b/Engine/Plugins/Runtime/nDisplay/Source/DisplayCluster/Private/Render/Viewport/DisplayClusterViewport.cpp
--- a/Engine/Plugins/Runtime/nDisplay/Source/DisplayCluster/Private/Render/Viewport/DisplayClusterViewport.cpp
+++ b/Engine/Plugins/Runtime/nDisplay/Source/DisplayCluster/Private/Render/Viewport/DisplayClusterViewport.cpp
@@ -452,11 +452,7 @@ FIntPoint FDisplayClusterViewport::GetDesiredContextSize(const FIntPoint& InCont
 		? InSize // Use original RenderTarget size.
 		: FDisplayClusterViewportHelpers::ScaleTextureSize(InSize, FinalRenderTargetMult);
 
-	const int32 MaxTextureSize = FDisplayClusterViewportHelpers::GetMaxTextureDimension();
-	DesiredContextSize.X = FMath::Min(DesiredContextSize.X, MaxTextureSize);
-	DesiredContextSize.Y = FMath::Min(DesiredContextSize.Y, MaxTextureSize);
-
-	return DesiredContextSize;
+	return InFrameSettings.ApplyViewportSizeConstraint(*this, DesiredContextSize);
 }
 
 float FDisplayClusterViewport::GetCustomBufferRatio(const FDisplayClusterRenderFrameSettings& InFrameSettings) const
@@ -627,13 +623,6 @@ bool FDisplayClusterViewport::UpdateFrameContexts(const uint32 InStereoViewIndex
 		}
 	}
 
-	// Apply restrictions on the maximum size of the viewport texture.
-	const int32 ViewportTextureMaxSize = InFrameSettings.GetViewportTextureMaxSize();
-	if (ViewportTextureMaxSize > 0)
-	{
-		DesiredContextSize = FDisplayClusterViewportHelpers::GetTextureSizeLessThanMax(DesiredContextSize, ViewportTextureMaxSize);
-	}
-
 	// Exclude zero-size viewports from render
 	if (DesiredContextSize.GetMin() <= 0)
 	{
diff --git a/Engine/Plugins/Runtime/nDisplay/Source/DisplayCluster/Private/Render/Viewport/DisplayClusterViewport.h b/Engine/Plugins/Runtime/nDisplay/Source/DisplayCluster/Private/Render/Viewport/DisplayClusterViewport.h
--- a/Engine/Plugins/Runtime/nDisplay/Source/DisplayCluster/Private/Render/Viewport/DisplayClusterViewport.h
+++ b/Engine/Plugins/Runtime/nDisplay/Source/DisplayCluster/Private/Render/Viewport/DisplayClusterViewport.h
@@ -336,6 +336,11 @@ public:
 	/** Returns true if per-view Lumen scene is enabled. */
 	bool ShouldUseLumenPerView() const;
 
+	/** Returns true if textures in this viewport with an area larger than the MaxTextureArea limit
+	*  must be reduced to stay within that limit with the same aspect ratio.
+	*/
+	bool ShouldApplyMaxTextureConstraints() const;
+
 	/** Returns true if this viewport should use this type of resource. */
 	bool IsResourceUsed(const EDisplayClusterViewportResource InResourceType) const;
 
diff --git a/Engine/Plugins/Runtime/nDisplay/Source/DisplayCluster/Private/Render/Viewport/DisplayClusterViewportHelpers.cpp b/Engine/Plugins/Runtime/nDisplay/Source/DisplayCluster/Private/Render/Viewport/DisplayClusterViewportHelpers.cpp
--- a/Engine/Plugins/Runtime/nDisplay/Source/DisplayCluster/Private/Render/Viewport/DisplayClusterViewportHelpers.cpp
+++ b/Engine/Plugins/Runtime/nDisplay/Source/DisplayCluster/Private/Render/Viewport/DisplayClusterViewportHelpers.cpp
@@ -9,32 +9,19 @@
 #include "Engine/TextureRenderTarget2D.h"
 
 ////////////////////////////////////////////////////////////////////////////////
-int32 GDisplayClusterOverrideMaxTextureDimension = 8192;
-static FAutoConsoleVariableRef CVarDisplayClusterOverrideMaxTextureDimension(
-	TEXT("DC.OverrideMaxTextureDimension"),
-	GDisplayClusterOverrideMaxTextureDimension,
-	TEXT("Override max texture dimension for nDisplay rendering (-1 == disabled, positive values == override)"),
+int32 GDisplayClusterMaxTextureAreaRoot = 8192;
+static FAutoConsoleVariableRef CVarDisplayClusterMaxTextureAreaRoot(
+	TEXT("nDisplay.render.MaxTextureAreaRoot"),
+	GDisplayClusterMaxTextureAreaRoot,
+	TEXT("The root value of the maximum texture area. (The default value is 8192, which means Area=8192*8192.\n")
+	TEXT("Textures with an area larger than this limit will be reduced to stay within this limit with the same aspect ratio.\n")
+	TEXT("Set a negative or zero value to disable this constraint.\n"),
 	ECVF_RenderThreadSafe
 );
 
-int32 GDisplayClusterOverrideMinTextureDimension = -1;
-static FAutoConsoleVariableRef CVarDisplayClusterOverrideMinTextureDimension(
-	TEXT("DC.OverrideMinTextureDimension"),
-	GDisplayClusterOverrideMinTextureDimension,
-	TEXT("Override min texture dimension for nDisplay rendering (-1 == disabled, positive values == override)"),
-	ECVF_RenderThreadSafe
-);
-
-int32 GDisplayClusterMaxNumMips = -1;
-static FAutoConsoleVariableRef CVarDisplayClusterMaxNumMips(
-	TEXT("DC.TextureMaxNumMips"),
-	GDisplayClusterMaxNumMips,
-	TEXT("Maximum number of mips for viewport texture.\n")
-	TEXT(" '0' - disable mips generation.\n")
-	TEXT("'-1' - disable this limit.\n")
-	TEXT("positive value - set the limit.\n"),
-	ECVF_RenderThreadSafe
-);
+static FAutoConsoleVariableDeprecated CVarDisplayClusterOverrideMaxTextureDimension_Dep(TEXT("DC.OverrideMaxTextureDimension"), TEXT(""), TEXT("5.5"));
+static FAutoConsoleVariableDeprecated CVarDisplayClusterOverrideMinTextureDimension_Dep(TEXT("DC.OverrideMinTextureDimension"), TEXT(""), TEXT("5.5"));
+static FAutoConsoleVariableDeprecated CVarDisplayClusterMaxNumMips_Dep(TEXT("DC.TextureMaxNumMips"), TEXT(""), TEXT("5.5"));
 
 // Default pixel format for preview rendering
 int32 GDisplayClusterPreviewDefaultPixelFormat = 1;
@@ -59,7 +46,6 @@ namespace UE::DisplayCluster::ViewportHelpers
 		return SPixelFormat[ValidIndex];
 	}
 };
-using namespace UE::DisplayCluster::ViewportHelpers;
 
 ////////////////////////////////////////////////////////////////////////////////
 // FDisplayClusterViewportHelpers
@@ -76,47 +62,36 @@ int32 FDisplayClusterViewportHelpers::GetMaxTextureNumMips(const FDisplayCluster
 		NumMips = 0;
 	}
 
-	if (GDisplayClusterMaxNumMips >= 0)
-	{
-		return FMath::Min(GDisplayClusterMaxNumMips, NumMips);
-	}
-
 	return NumMips;
 }
 
-int32 FDisplayClusterViewportHelpers::GetMaxTextureDimension()
+int32 FDisplayClusterViewportHelpers::GetMaxTextureArea()
 {
-	// The target always needs be within GMaxTextureDimensions, larger dimensions are not supported by the engine
-	static const int32 MaxTextureDimension = 1 << (GMaxTextureMipCount - 1);
-
-	// Use overridden value
-	if (GDisplayClusterOverrideMaxTextureDimension > 0)
+	if (GDisplayClusterMaxTextureAreaRoot > 0)
 	{
-		// Protect vs wrong CVar values
-		if (GDisplayClusterOverrideMinTextureDimension < 0 || GDisplayClusterOverrideMaxTextureDimension > GDisplayClusterOverrideMinTextureDimension)
-		{
-			return GDisplayClusterOverrideMaxTextureDimension;
-		}
+		// make sure this value is within hardware limits.
+		const int32 MaxTextureAreaRoot = FMath::Clamp(
+			GDisplayClusterMaxTextureAreaRoot,
+			GetMinTextureDimension(),
+			GetMaxTextureDimension());
+
+		// Returns the area calculated from the root.
+		return MaxTextureAreaRoot * MaxTextureAreaRoot;
 	}
 
-	return MaxTextureDimension;
+	return -1;
 }
 
-int32 FDisplayClusterViewportHelpers::GetMinTextureDimension()
+int32 FDisplayClusterViewportHelpers::GetMaxTextureDimension()
 {
-	static const int32 MinTextureDimension = 16;
+	static int32 NativeMaxTextureDimension = 1 << (GMaxTextureMipCount - 1);
 
-	// Use overridden value
-	if (GDisplayClusterOverrideMinTextureDimension > 0)
-	{
-		// Protect vs wrong CVar values
-		if (GDisplayClusterOverrideMaxTextureDimension < 0 || GDisplayClusterOverrideMaxTextureDimension > GDisplayClusterOverrideMinTextureDimension)
-		{
-			return GDisplayClusterOverrideMinTextureDimension;
-		}
-	}
+	return NativeMaxTextureDimension;
+}
 
-	return MinTextureDimension;
+int32 FDisplayClusterViewportHelpers::GetMinTextureDimension()
+{
+	return 1;
 }
 
 FIntRect FDisplayClusterViewportHelpers::GetValidViewportRect(const FIntRect& InRect, const FString& InViewportId, const TCHAR* InResourceName)
@@ -189,19 +164,6 @@ bool FDisplayClusterViewportHelpers::IsValidTextureSize(const FIntPoint& InSize)
 	return InSize.GetMin() >= GetMinTextureDimension() && InSize.GetMax() <= GetMaxTextureDimension();
 }
 
-// Return size less than max texture dimension. Aspect ratio not changed
-FIntPoint FDisplayClusterViewportHelpers::GetTextureSizeLessThanMax(const FIntPoint& InSize, const int32 InMaxTextureDimension)
-{
-	if (InSize.GetMax() > InMaxTextureDimension)
-	{
-		const float DownscaleMult =  double(InMaxTextureDimension) / double(InSize.GetMax());
-
-		return ScaleTextureSize(InSize, DownscaleMult);
-	}
-
-	return InSize;
-}
-
 FIntPoint FDisplayClusterViewportHelpers::ScaleTextureSize(const FIntPoint& InSize, float InMult)
 {
 	const int32 ScaledX = FMath::CeilToInt(InSize.X * InMult);
@@ -251,6 +213,8 @@ float FDisplayClusterViewportHelpers::GetValidSizeMultiplier(const FIntPoint& In
 
 EPixelFormat FDisplayClusterViewportHelpers::GetPreviewDefaultPixelFormat()
 {
+	using namespace UE::DisplayCluster::ViewportHelpers;
+
 	const ETextureRenderTargetFormat RenderTargetFormat = ImplGetRenderTargetFormatFromInt(GDisplayClusterPreviewDefaultPixelFormat);
 
 	return GetPixelFormatFromRenderTargetFormat(RenderTargetFormat);
diff --git a/Engine/Plugins/Runtime/nDisplay/Source/DisplayCluster/Private/Render/Viewport/DisplayClusterViewportHelpers.h b/Engine/Plugins/Runtime/nDisplay/Source/DisplayCluster/Private/Render/Viewport/DisplayClusterViewportHelpers.h
--- a/Engine/Plugins/Runtime/nDisplay/Source/DisplayCluster/Private/Render/Viewport/DisplayClusterViewportHelpers.h
+++ b/Engine/Plugins/Runtime/nDisplay/Source/DisplayCluster/Private/Render/Viewport/DisplayClusterViewportHelpers.h
@@ -20,6 +20,9 @@ public:
 	/** Get the minimum allowable texture size used for the nDisplay viewport. */
 	static int32 GetMaxTextureDimension();
 
+	/** Get the maximum allowable texture area that can be used for the nDisplay viewport. */
+	static int32 GetMaxTextureArea();
+
 	/** Get the valid viewport size.
 	* 
 	* @param InRect         - source region
@@ -47,15 +50,6 @@ public:
 	*/
 	static FIntPoint ScaleTextureSize(const FIntPoint& InSize, float InMult);
 
-	/** The returned size is smaller than the maximum texture size. The aspect ratio does not change
-	* 
-	* @param InSize - input texture size
-	* @param InMaxTextureDimension - max alloable dimension
-	* 
-	* @return - the same as the input size, or reduced to the InMaxTextureDimension
-	 */
-	static FIntPoint GetTextureSizeLessThanMax(const FIntPoint& InSize, const int32 InMaxTextureDimension);
-
 	/** Find an acceptable multiplier for the texture size
 	* 
 	* @param InSize - in texture size
diff --git a/Engine/Plugins/Runtime/nDisplay/Source/DisplayCluster/Private/Render/Viewport/DisplayClusterViewport_Rules.cpp b/Engine/Plugins/Runtime/nDisplay/Source/DisplayCluster/Private/Render/Viewport/DisplayClusterViewport_Rules.cpp
--- a/Engine/Plugins/Runtime/nDisplay/Source/DisplayCluster/Private/Render/Viewport/DisplayClusterViewport_Rules.cpp
+++ b/Engine/Plugins/Runtime/nDisplay/Source/DisplayCluster/Private/Render/Viewport/DisplayClusterViewport_Rules.cpp
@@ -204,6 +204,17 @@ bool FDisplayClusterViewport::IsUsedByMedia() const
 	return RenderSettings.HasAnyMediaStates(EDisplayClusterViewportMediaState::Input | EDisplayClusterViewportMediaState::Capture);
 }
 
+bool FDisplayClusterViewport::ShouldApplyMaxTextureConstraints() const
+{
+	if (IsUsedByMedia())
+	{
+		// Do not apply this restriction to a viewport that uses media.
+		return false;
+	}
+
+	return true;
+}
+
 bool FDisplayClusterViewport::IsRenderEnabledByMedia() const
 {
 	// Use this function to disable rendering of this viewport.
diff --git a/Engine/Plugins/Runtime/nDisplay/Source/DisplayCluster/Private/Render/Viewport/RenderFrame/DisplayClusterRenderFrameSettings.cpp b/Engine/Plugins/Runtime/nDisplay/Source/DisplayCluster/Private/Render/Viewport/RenderFrame/DisplayClusterRenderFrameSettings.cpp
--- a/Engine/Plugins/Runtime/nDisplay/Source/DisplayCluster/Private/Render/Viewport/RenderFrame/DisplayClusterRenderFrameSettings.cpp
+++ b/Engine/Plugins/Runtime/nDisplay/Source/DisplayCluster/Private/Render/Viewport/RenderFrame/DisplayClusterRenderFrameSettings.cpp
@@ -1,9 +1,10 @@
 // Copyright Epic Games, Inc. All Rights Reserved.
 
 #include "Render/Viewport/RenderFrame/DisplayClusterRenderFrameSettings.h"
+#include "Render/Viewport/DisplayClusterViewport.h"
+#include "Render/Viewport/DisplayClusterViewportHelpers.h"
 #include "HAL/IConsoleManager.h"
 
-
 ////////////////////////////////////////////////////////////////////////////////
 // Experimental feature: to be approved after testing
 int32 GDisplayClusterPreviewEnableReuseViewportInCluster = 1;
@@ -62,15 +63,52 @@ bool FDisplayClusterRenderFrameSettings::CanReuseViewportWithinClusterNodes() co
 	return false;
 }
 
-int32 FDisplayClusterRenderFrameSettings::GetViewportTextureMaxSize() const
+FIntPoint FDisplayClusterRenderFrameSettings::ApplyViewportSizeConstraint(const FDisplayClusterViewport& InViewport, const FIntPoint& InSize) const
 {
-	if (IsPreviewRendering())
+	// Adjust the size of the viewport size depending on the constraints
+	FIntPoint OutSize = InSize;
+
+	if (InViewport.ShouldApplyMaxTextureConstraints())
+	{
+		// Apply restrictions on the maximum texture area
+		const int32 MaxTextureArea = FDisplayClusterViewportHelpers::GetMaxTextureArea();
+		if (MaxTextureArea > 0)
+		{
+			const int32 OutArea = OutSize.X * OutSize.Y;
+			if (OutArea > MaxTextureArea)
+			{
+				// downsize to the max area
+				const double DownsizeMult = FMath::Sqrt(double(MaxTextureArea) / OutArea);
+				OutSize = FDisplayClusterViewportHelpers::ScaleTextureSize(OutSize, DownsizeMult);
+			}
+		}
+	}
+
+	const int32 MinTextureDimension = FDisplayClusterViewportHelpers::GetMinTextureDimension();
+	int32 MaxTextureDimension = FDisplayClusterViewportHelpers::GetMaxTextureDimension();
+
+	// Preview rendering has its own maximum texture size, which overrides the default one
+	if (IsPreviewRendering() && PreviewSettings.PreviewMaxTextureDimension > 0)
 	{
-		// Special constraints for texture when rendering previews
-		return PreviewSettings.PreviewMaxTextureDimension;
+		const int32 MaxPreviewTextureDimension = FMath::Clamp(PreviewSettings.PreviewMaxTextureDimension, MinTextureDimension, MaxTextureDimension);
+
+		// Preview consume less GPU memory
+		MaxTextureDimension = MaxPreviewTextureDimension;
 	}
 
-	return -1;
+	// Downsize the size to the limits
+	if (MaxTextureDimension < OutSize.GetMax())
+	{
+		// downsize to the max dimension
+		const double DownsizeMult = double(MaxTextureDimension) / OutSize.GetMax();
+		OutSize = FDisplayClusterViewportHelpers::ScaleTextureSize(OutSize, DownsizeMult);
+	}
+
+	// Final crop
+	OutSize.X = FMath::Clamp(OutSize.X, MinTextureDimension, MaxTextureDimension);
+	OutSize.Y = FMath::Clamp(OutSize.Y, MinTextureDimension, MaxTextureDimension);
+
+	return OutSize;
 }
 
 bool FDisplayClusterRenderFrameSettings::ShouldUseLinearGamma() const
diff --git a/Engine/Plugins/Runtime/nDisplay/Source/DisplayCluster/Private/Render/Viewport/RenderFrame/DisplayClusterRenderFrameSettings.h b/Engine/Plugins/Runtime/nDisplay/Source/DisplayCluster/Private/Render/Viewport/RenderFrame/DisplayClusterRenderFrameSettings.h
--- a/Engine/Plugins/Runtime/nDisplay/Source/DisplayCluster/Private/Render/Viewport/RenderFrame/DisplayClusterRenderFrameSettings.h
+++ b/Engine/Plugins/Runtime/nDisplay/Source/DisplayCluster/Private/Render/Viewport/RenderFrame/DisplayClusterRenderFrameSettings.h
@@ -132,8 +132,15 @@ public:
 	FVector2D GetDesiredRTTSize(const FVector2D& InSize) const;
 	FIntPoint GetDesiredRTTSize(const FIntPoint& InSize) const;
 
-	/** Get the maximum texture size that is allowed to be used in the viewport. */
-	int32 GetViewportTextureMaxSize() const;
+	/** 
+	* Limit viewport size within:
+	* 1. MaxArea - for main viewports only, controlled by ShouldApplyMaxTextureConstraints().
+	* 2. Maximum texture size based on hardware limitations	* The output size will have the same aspect ratio as the input size.
+	* 
+	* @param InViewport - the nDisplay viewport for which the size is used.
+	* @param InSize     - The base size of the viewport.
+	*/
+	FIntPoint ApplyViewportSizeConstraint(const class FDisplayClusterViewport& InViewport, const FIntPoint& InSize) const;
 
 	/** Should use Cross-GPU transfers. */
 	inline bool ShouldUseCrossGPUTransfers() const
