[nDisplay] Added media passthrough support + some minor changes

- Added another media capture step for passthrough pipeline
- Added another media input step for custom JITR based solutions
- A bunch of minor fixes/updates

[RN] Fixed media passthrough for nDisplay viewports

diff --git a/Engine/Plugins/Runtime/nDisplay/Source/DisplayCluster/Private/DisplayClusterCallbacks.h b/Engine/Plugins/Runtime/nDisplay/Source/DisplayCluster/Private/DisplayClusterCallbacks.h
--- a/Engine/Plugins/Runtime/nDisplay/Source/DisplayCluster/Private/DisplayClusterCallbacks.h
+++ b/Engine/Plugins/Runtime/nDisplay/Source/DisplayCluster/Private/DisplayClusterCallbacks.h
@@ -122,6 +122,16 @@ public:
 		return DisplayClusterPostCrossGpuTransferEvent;
 	}
 
+	virtual FDisplayClusterPassthroughMediaCapture_RenderThread& OnDisplayClusterPassthroughMediaCapture_RenderThread() override
+	{
+		return DisplayClusterPassthroughMediaCaptureEvent;
+	}
+
+	virtual FDisplayClusterPassthroughMediaInput_RenderThread& OnDisplayClusterPassthroughMediaInput_RenderThread() override
+	{
+		return DisplayClusterPassthroughMediaInputEvent;
+	}
+
 	virtual FDisplayClusterProcessLatency_RenderThread& OnDisplayClusterProcessLatency_RenderThread() override
 	{
 		return FDisplayClusterProcessLatencyEvent;
@@ -178,6 +188,8 @@ private:
 	FDisplayClusterPostWarp_RenderThread             DisplayClusterPostWarpEvent;
 	FDisplayClusterPostWarpViewport_RenderThread     DisplayClusterPostWarpViewportEvent;
 	FDisplayClusterPostCrossGpuTransfer_RenderThread DisplayClusterPostCrossGpuTransferEvent;
+	FDisplayClusterPassthroughMediaCapture_RenderThread DisplayClusterPassthroughMediaCaptureEvent;
+	FDisplayClusterPassthroughMediaInput_RenderThread   DisplayClusterPassthroughMediaInputEvent;
 	FDisplayClusterProcessLatency_RenderThread       FDisplayClusterProcessLatencyEvent;
 	FDisplayClusterPostFrameRender_RenderThread      DisplayClusterPostFrameRenderEvent;
 	FDisplayClusterPostBackbufferUpdate_RenderThread DisplayClusterPostBackbufferUpdateEvent;
diff --git a/Engine/Plugins/Runtime/nDisplay/Source/DisplayCluster/Private/Render/Viewport/DisplayClusterViewportManagerProxy.cpp b/Engine/Plugins/Runtime/nDisplay/Source/DisplayCluster/Private/Render/Viewport/DisplayClusterViewportManagerProxy.cpp
--- a/Engine/Plugins/Runtime/nDisplay/Source/DisplayCluster/Private/Render/Viewport/DisplayClusterViewportManagerProxy.cpp
+++ b/Engine/Plugins/Runtime/nDisplay/Source/DisplayCluster/Private/Render/Viewport/DisplayClusterViewportManagerProxy.cpp
@@ -221,6 +221,11 @@ void FDisplayClusterViewportManagerProxy::ImplRenderFrame_GameThread(FViewport*
 
 		// PostCrossGpuTransfer notification
 		IDisplayCluster::Get().GetCallbacks().OnDisplayClusterPostCrossGpuTransfer_RenderThread().Broadcast(RHICmdList, ViewportManagerProxy, OutputViewport);
+
+		// Another capture+playback round for media passthrough
+		IDisplayCluster::Get().GetCallbacks().OnDisplayClusterPassthroughMediaCapture_RenderThread().Broadcast(RHICmdList, ViewportManagerProxy);
+		IDisplayCluster::Get().GetCallbacks().OnDisplayClusterPassthroughMediaInput_RenderThread().Broadcast(RHICmdList, ViewportManagerProxy);
+
 		// Latency processing
 		IDisplayCluster::Get().GetCallbacks().OnDisplayClusterProcessLatency_RenderThread().Broadcast(RHICmdList, ViewportManagerProxy, OutputViewport);
 	});
diff --git a/Engine/Plugins/Runtime/nDisplay/Source/DisplayCluster/Public/IDisplayClusterCallbacks.h b/Engine/Plugins/Runtime/nDisplay/Source/DisplayCluster/Public/IDisplayClusterCallbacks.h
--- a/Engine/Plugins/Runtime/nDisplay/Source/DisplayCluster/Public/IDisplayClusterCallbacks.h
+++ b/Engine/Plugins/Runtime/nDisplay/Source/DisplayCluster/Public/IDisplayClusterCallbacks.h
@@ -119,6 +119,14 @@ public:
 	DECLARE_EVENT_ThreeParams(IDisplayClusterCallbacks, FDisplayClusterPostCrossGpuTransfer_RenderThread, FRHICommandListImmediate&, const IDisplayClusterViewportManagerProxy*, FViewport*);
 	virtual FDisplayClusterPostCrossGpuTransfer_RenderThread& OnDisplayClusterPostCrossGpuTransfer_RenderThread() = 0;
 
+	/** Called to process passthrough media capture **/
+	DECLARE_EVENT_TwoParams(IDisplayClusterCallbacks, FDisplayClusterPassthroughMediaCapture_RenderThread, FRHICommandListImmediate&, const IDisplayClusterViewportManagerProxy*);
+	virtual FDisplayClusterPassthroughMediaCapture_RenderThread& OnDisplayClusterPassthroughMediaCapture_RenderThread() = 0;
+
+	/** Called to process passthrough media input **/
+	DECLARE_EVENT_TwoParams(IDisplayClusterCallbacks, FDisplayClusterPassthroughMediaInput_RenderThread, FRHICommandListImmediate&, const IDisplayClusterViewportManagerProxy*);
+	virtual FDisplayClusterPassthroughMediaInput_RenderThread& OnDisplayClusterPassthroughMediaInput_RenderThread() = 0;
+
 	/** Called to let the artificial latency subsystem do its job **/
 	DECLARE_EVENT_ThreeParams(IDisplayClusterCallbacks, FDisplayClusterProcessLatency_RenderThread, FRHICommandListImmediate&, const IDisplayClusterViewportManagerProxy*, FViewport*);
 	virtual FDisplayClusterProcessLatency_RenderThread& OnDisplayClusterProcessLatency_RenderThread() = 0;
diff --git a/Engine/Plugins/Runtime/nDisplay/Source/DisplayClusterMedia/Private/Capture/DisplayClusterMediaCaptureCameraBase.cpp b/Engine/Plugins/Runtime/nDisplay/Source/DisplayClusterMedia/Private/Capture/DisplayClusterMediaCaptureCameraBase.cpp
--- a/Engine/Plugins/Runtime/nDisplay/Source/DisplayClusterMedia/Private/Capture/DisplayClusterMediaCaptureCameraBase.cpp
+++ b/Engine/Plugins/Runtime/nDisplay/Source/DisplayClusterMedia/Private/Capture/DisplayClusterMediaCaptureCameraBase.cpp
@@ -17,7 +17,7 @@ FDisplayClusterMediaCaptureCameraBase::FDisplayClusterMediaCaptureCameraBase(
 }
 
 
-void FDisplayClusterMediaCaptureCameraBase::UpdateLateOCIOState(IDisplayClusterViewport* Viewport)
+void FDisplayClusterMediaCaptureCameraBase::UpdateLateOCIOState(const IDisplayClusterViewport* Viewport)
 {
 	FLateOCIOData NewLateOCIOConfiguration;
 
@@ -27,3 +27,8 @@ void FDisplayClusterMediaCaptureCameraBase::UpdateLateOCIOState(IDisplayClusterV
 	// And store/update for current frame
 	SetLateOCIO(NewLateOCIOConfiguration);
 }
+
+void FDisplayClusterMediaCaptureCameraBase::UpdateMediaPassthrough(const IDisplayClusterViewport* Viewport)
+{
+	// No implementation for camera so far. It's unlikely we'd ever need media passthrough for ICVFX cameras.
+}
diff --git a/Engine/Plugins/Runtime/nDisplay/Source/DisplayClusterMedia/Private/Capture/DisplayClusterMediaCaptureCameraBase.h b/Engine/Plugins/Runtime/nDisplay/Source/DisplayClusterMedia/Private/Capture/DisplayClusterMediaCaptureCameraBase.h
--- a/Engine/Plugins/Runtime/nDisplay/Source/DisplayClusterMedia/Private/Capture/DisplayClusterMediaCaptureCameraBase.h
+++ b/Engine/Plugins/Runtime/nDisplay/Source/DisplayClusterMedia/Private/Capture/DisplayClusterMediaCaptureCameraBase.h
@@ -27,5 +27,8 @@ public:
 protected:
 
 	/** Updates late OCIO state */
-	virtual void UpdateLateOCIOState(IDisplayClusterViewport* Viewport) override;
+	virtual void UpdateLateOCIOState(const IDisplayClusterViewport* Viewport) override;
+
+	/** Checks whether current frame should use media passthrough */
+	virtual void UpdateMediaPassthrough(const IDisplayClusterViewport* Viewport) override;
 };
diff --git a/Engine/Plugins/Runtime/nDisplay/Source/DisplayClusterMedia/Private/Capture/DisplayClusterMediaCaptureViewportBase.cpp b/Engine/Plugins/Runtime/nDisplay/Source/DisplayClusterMedia/Private/Capture/DisplayClusterMediaCaptureViewportBase.cpp
--- a/Engine/Plugins/Runtime/nDisplay/Source/DisplayClusterMedia/Private/Capture/DisplayClusterMediaCaptureViewportBase.cpp
+++ b/Engine/Plugins/Runtime/nDisplay/Source/DisplayClusterMedia/Private/Capture/DisplayClusterMediaCaptureViewportBase.cpp
@@ -2,20 +2,24 @@
 
 #include "Capture/DisplayClusterMediaCaptureViewportBase.h"
 
-#include "IDisplayCluster.h"
-#include "IDisplayClusterCallbacks.h"
-
-#include "DisplayClusterMediaLog.h"
-
+#include "Game/IDisplayClusterGameManager.h"
 #include "PostProcess/PostProcessMaterialInputs.h"
 
 #include "Render/IDisplayClusterRenderManager.h"
 #include "Render/Viewport/IDisplayClusterViewport.h"
 #include "Render/Viewport/IDisplayClusterViewportManager.h"
+#include "Render/Viewport/IDisplayClusterViewportManagerProxy.h"
 #include "Render/Viewport/IDisplayClusterViewportProxy.h"
 #include "Render/Viewport/Containers/DisplayClusterViewport_Context.h"
 #include "Render/Viewport/Containers/DisplayClusterViewport_RenderSettings.h"
 
+#include "DisplayClusterConfigurationTypes.h"
+#include "DisplayClusterMediaLog.h"
+#include "DisplayClusterRootActor.h"
+
+#include "IDisplayCluster.h"
+#include "IDisplayClusterCallbacks.h"
+
 #include "RHICommandList.h"
 #include "RHIResources.h"
 
@@ -35,7 +39,7 @@ FDisplayClusterMediaCaptureViewportBase::FDisplayClusterMediaCaptureViewportBase
 
 bool FDisplayClusterMediaCaptureViewportBase::StartCapture()
 {
-	// Subscribe for events
+	// Media state update callback
 	IDisplayCluster::Get().GetCallbacks().OnDisplayClusterUpdateViewportMediaState().AddRaw(this, &FDisplayClusterMediaCaptureViewportBase::OnUpdateViewportMediaState);
 
 	// PostTonemap callback for the late OCIO path
@@ -44,6 +48,9 @@ bool FDisplayClusterMediaCaptureViewportBase::StartCapture()
 	// PostRenderViewFamily for a regular capture path
 	IDisplayCluster::Get().GetCallbacks().OnDisplayClusterPostRenderViewFamily_RenderThread().AddRaw(this, &FDisplayClusterMediaCaptureViewportBase::OnPostRenderViewFamily_RenderThread);
 
+	// Passthrough media capture
+	IDisplayCluster::Get().GetCallbacks().OnDisplayClusterPassthroughMediaCapture_RenderThread().AddRaw(this, &FDisplayClusterMediaCaptureViewportBase::OnPassthroughMediaCapture_RenderThread);
+
 	// Start capture
 	const bool bStarted = FDisplayClusterMediaCaptureBase::StartCapture();
 
@@ -56,6 +63,7 @@ void FDisplayClusterMediaCaptureViewportBase::StopCapture()
 	IDisplayCluster::Get().GetCallbacks().OnDisplayClusterPostRenderViewFamily_RenderThread().RemoveAll(this);
 	IDisplayCluster::Get().GetCallbacks().OnDisplayClusterPostTonemapPass_RenderThread().RemoveAll(this);
 	IDisplayCluster::Get().GetCallbacks().OnDisplayClusterUpdateViewportMediaState().RemoveAll(this);
+	IDisplayCluster::Get().GetCallbacks().OnDisplayClusterPassthroughMediaCapture_RenderThread().RemoveAll(this);
 
 	// Stop capturing
 	FDisplayClusterMediaCaptureBase::StopCapture();
@@ -77,6 +85,9 @@ void FDisplayClusterMediaCaptureViewportBase::OnUpdateViewportMediaState(IDispla
 		{
 			InOutMediaState |= EDisplayClusterViewportMediaState::CaptureLateOCIO;
 		}
+
+		// Update media passthrough for this frame
+		UpdateMediaPassthrough(InViewport);
 	}
 }
 
@@ -113,7 +124,7 @@ bool FDisplayClusterMediaCaptureViewportBase::GetCaptureSizeFromGameProxy(FIntPo
 	{
 		if (const IDisplayClusterViewportManager* const ViewportMgr = RenderMgr->GetViewportManager())
 		{
-			if (const IDisplayClusterViewport* const Viewport = ViewportMgr->FindViewport(ReferencedViewportId))
+			if (const IDisplayClusterViewport* const Viewport = ViewportMgr->FindViewport(GetViewportId()))
 			{
 				const TArray<FDisplayClusterViewport_Context>& Contexts = Viewport->GetContexts();
 				if (Contexts.Num() > 0)
@@ -128,7 +139,7 @@ bool FDisplayClusterMediaCaptureViewportBase::GetCaptureSizeFromGameProxy(FIntPo
 	return false;
 }
 
-void FDisplayClusterMediaCaptureViewportBase::UpdateLateOCIOState(IDisplayClusterViewport* Viewport)
+void FDisplayClusterMediaCaptureViewportBase::UpdateLateOCIOState(const IDisplayClusterViewport* Viewport)
 {
 	// So far, regular viewports dont't support late OCIO feature.
 }
@@ -140,6 +151,34 @@ void FDisplayClusterMediaCaptureViewportBase::HandleLateOCIOChanged(const FLateO
 	StartCapture();
 }
 
+void FDisplayClusterMediaCaptureViewportBase::UpdateMediaPassthrough(const IDisplayClusterViewport* Viewport)
+{
+	bool bNewPassthroughState = false;
+
+	// Media passthough is used when the same viewport has both media input and output configured. Let's see if this is the case.
+	if (const ADisplayClusterRootActor* const RootActor = IDisplayCluster::Get().GetGameMgr()->GetRootActor())
+	{
+		if (const UDisplayClusterConfigurationData* const ConfigData = RootActor->GetConfigData())
+		{
+			if (const UDisplayClusterConfigurationViewport* const ViewportCfg = ConfigData->GetViewport(GetClusterNodeId(), GetViewportId()))
+			{
+				const FDisplayClusterConfigurationMediaViewport& MediaSettings = ViewportCfg->RenderSettings.Media;
+				if (MediaSettings.bEnable)
+				{
+					bNewPassthroughState = MediaSettings.IsMediaInputAssigned() && MediaSettings.IsMediaOutputAssigned();
+				}
+			}
+		}
+	}
+
+	// Pass to the render thread
+	ENQUEUE_RENDER_COMMAND(DCMediaCaptureUpdatePassthrough)(
+		[This = AsShared(), bSetVal = bNewPassthroughState](FRHICommandListImmediate& RHICmdList)
+		{
+			StaticCastSharedRef<FDisplayClusterMediaCaptureViewportBase>(This)->bUseMediaPassthrough_RT = bSetVal;
+		});
+}
+
 void FDisplayClusterMediaCaptureViewportBase::OnPostTonemapPass_RenderThread(FRDGBuilder& GraphBuilder, const IDisplayClusterViewportProxy* ViewportProxy, const FSceneView& View, const FPostProcessMaterialInputs& Inputs, const uint32 ContextNum)
 {
 	checkSlow(ViewportProxy);
@@ -209,3 +248,40 @@ void FDisplayClusterMediaCaptureViewportBase::OnPostRenderViewFamily_RenderThrea
 		}
 	}
 }
+
+void FDisplayClusterMediaCaptureViewportBase::OnPassthroughMediaCapture_RenderThread(FRHICommandListImmediate& RHICmdList, const IDisplayClusterViewportManagerProxy* ViewportManagerProxy)
+{
+	checkSlow(ViewportManagerProxy);
+
+	// Nothing to do if media passthrough is not used
+	if (!bUseMediaPassthrough_RT)
+	{
+		return;
+	}
+
+	// Otherwise, find our viewport and export its texture
+	if (ViewportManagerProxy)
+	{
+		if (IDisplayClusterViewportProxy* ViewportProxy = ViewportManagerProxy->FindViewport_RenderThread(ReferencedViewportId))
+		{
+			TArray<FRHITexture*> Textures;
+			TArray<FIntRect>     Regions;
+
+			// Get RHI texture and pass it to the media capture pipeline
+			if (ViewportProxy->GetResourcesWithRects_RenderThread(EDisplayClusterViewportResourceType::InternalRenderTargetResource, Textures, Regions))
+			{
+				if (Textures.Num() > 0 && Regions.Num() > 0 && Textures[0])
+				{
+					FRDGBuilder GraphBuilder(RHICmdList);
+
+					FRDGTextureRef SrcTextureRef = RegisterExternalTexture(GraphBuilder, Textures[0], TEXT("DCMediaOutViewportTex"));
+
+					FMediaOutputTextureInfo TextureInfo{ SrcTextureRef, Regions[0] };
+					ExportMediaData_RenderThread(GraphBuilder, TextureInfo);
+
+					GraphBuilder.Execute();
+				}
+			}
+		}
+	}
+}
diff --git a/Engine/Plugins/Runtime/nDisplay/Source/DisplayClusterMedia/Private/Capture/DisplayClusterMediaCaptureViewportBase.h b/Engine/Plugins/Runtime/nDisplay/Source/DisplayClusterMedia/Private/Capture/DisplayClusterMediaCaptureViewportBase.h
--- a/Engine/Plugins/Runtime/nDisplay/Source/DisplayClusterMedia/Private/Capture/DisplayClusterMediaCaptureViewportBase.h
+++ b/Engine/Plugins/Runtime/nDisplay/Source/DisplayClusterMedia/Private/Capture/DisplayClusterMediaCaptureViewportBase.h
@@ -10,6 +10,7 @@ class FRDGBuilder;
 class FSceneView;
 class FSceneViewFamily;
 class IDisplayClusterViewport;
+class IDisplayClusterViewportManagerProxy;
 class IDisplayClusterViewportProxy;
 struct FPostProcessMaterialInputs;
 
@@ -56,11 +57,14 @@ public:
 protected:
 
 	/** Updates late OCIO state */
-	virtual void UpdateLateOCIOState(IDisplayClusterViewport* Viewport);
+	virtual void UpdateLateOCIOState(const IDisplayClusterViewport* Viewport);
 
 	/** Process late OCIO state updates */
 	virtual void HandleLateOCIOChanged(const FLateOCIOData& NewLateOCIOConfiguration) override;
 
+	/** Checks whether current frame should use media passthrough */
+	virtual void UpdateMediaPassthrough(const IDisplayClusterViewport* Viewport);
+
 private:
 
 	/** UpdateViewportMediaState callback to configure media state for a viewoprt */
@@ -72,6 +76,9 @@ private:
 	/** PostTonemapPass callback handler (late OCIO) */
 	void OnPostTonemapPass_RenderThread(FRDGBuilder& GraphBuilder, const IDisplayClusterViewportProxy* ViewportProxy, const FSceneView& View, const FPostProcessMaterialInputs& Inputs, const uint32 ContextNum);
 
+	/** PassthroughMediaCapture callback handler is used to capture in the passthrough scenarios */
+	void OnPassthroughMediaCapture_RenderThread(FRHICommandListImmediate& RHICmdList, const IDisplayClusterViewportManagerProxy* ViewportManagerProxy);
+
 private:
 
 	/** Returns size of the viewport bound to this media */
@@ -81,4 +88,7 @@ private:
 
 	/** Viewport ID assigned to capture */
 	const FString ReferencedViewportId;
+
+	/** Whether passthrough capture is needed */
+	bool bUseMediaPassthrough_RT = false;
 };
diff --git a/Engine/Plugins/Runtime/nDisplay/Source/DisplayClusterMedia/Private/Common/DisplayClusterMediaBase.cpp b/Engine/Plugins/Runtime/nDisplay/Source/DisplayClusterMedia/Private/Common/DisplayClusterMediaBase.cpp
--- a/Engine/Plugins/Runtime/nDisplay/Source/DisplayClusterMedia/Private/Common/DisplayClusterMediaBase.cpp
+++ b/Engine/Plugins/Runtime/nDisplay/Source/DisplayClusterMedia/Private/Common/DisplayClusterMediaBase.cpp
@@ -51,9 +51,9 @@ void FDisplayClusterMediaBase::SetLateOCIO(const FLateOCIOData& NewLateOCIOConfi
 
 		// And pass to the render thread
 		ENQUEUE_RENDER_COMMAND(DCMediaUpdateOCIOState)(
-			[this, SetNewLateOCIO = LateOCIOConfiguration](FRHICommandListImmediate& RHICmdList)
+			[This = AsShared(), SetNewLateOCIO = LateOCIOConfiguration](FRHICommandListImmediate& RHICmdList)
 			{
-				LateOCIOConfiguration_RT = SetNewLateOCIO;
+				This->LateOCIOConfiguration_RT = SetNewLateOCIO;
 			});
 	}
 }
diff --git a/Engine/Plugins/Runtime/nDisplay/Source/DisplayClusterMedia/Private/Common/DisplayClusterMediaBase.h b/Engine/Plugins/Runtime/nDisplay/Source/DisplayClusterMedia/Private/Common/DisplayClusterMediaBase.h
--- a/Engine/Plugins/Runtime/nDisplay/Source/DisplayClusterMedia/Private/Common/DisplayClusterMediaBase.h
+++ b/Engine/Plugins/Runtime/nDisplay/Source/DisplayClusterMedia/Private/Common/DisplayClusterMediaBase.h
@@ -10,6 +10,7 @@
  * Base media adapter class
  */
 class FDisplayClusterMediaBase
+	: public TSharedFromThis<FDisplayClusterMediaBase>
 {
 public:
 
@@ -70,6 +71,7 @@ protected:
 	/** Set late OCIO configuration by children */
 	void SetLateOCIO(const FLateOCIOData& NewLateOCIOConfiguration);
 
+	/** Let children handle the OCIO changes */
 	virtual void HandleLateOCIOChanged(const FLateOCIOData& NewLateOCIOConfiguration)
 	{ }
 
diff --git a/Engine/Plugins/Runtime/nDisplay/Source/DisplayClusterMedia/Private/DisplayClusterMediaModule.cpp b/Engine/Plugins/Runtime/nDisplay/Source/DisplayClusterMedia/Private/DisplayClusterMediaModule.cpp
--- a/Engine/Plugins/Runtime/nDisplay/Source/DisplayClusterMedia/Private/DisplayClusterMediaModule.cpp
+++ b/Engine/Plugins/Runtime/nDisplay/Source/DisplayClusterMedia/Private/DisplayClusterMediaModule.cpp
@@ -32,6 +32,11 @@
 #include "Misc/CoreDelegates.h"
 
 
+FDisplayClusterMediaModule::FDisplayClusterMediaModule()
+	: FrameQueue(MakeShared<FDisplayClusterFrameQueue>())
+{
+}
+
 void FDisplayClusterMediaModule::StartupModule()
 {
 	UE_LOG(LogDisplayClusterMedia, Log, TEXT("Starting module 'DisplayClusterMedia'..."));
@@ -83,7 +88,7 @@ void FDisplayClusterMediaModule::InitializeMedia()
 	}
 
 	// Instantiate latency queue
-	FrameQueue.Init();
+	FrameQueue->Init();
 
 	// Parse DCRA configuration and initialize media
 	if (const ADisplayClusterRootActor* const RootActor = IDisplayCluster::Get().GetGameMgr()->GetRootActor())
@@ -142,7 +147,7 @@ void FDisplayClusterMediaModule::ReleaseMedia()
 	CaptureDevices.Reset();
 	InputDevices.Reset();
 
-	FrameQueue.Release();
+	FrameQueue->Release();
 }
 
 void FDisplayClusterMediaModule::StartCapture()
diff --git a/Engine/Plugins/Runtime/nDisplay/Source/DisplayClusterMedia/Private/DisplayClusterMediaModule.h b/Engine/Plugins/Runtime/nDisplay/Source/DisplayClusterMedia/Private/DisplayClusterMediaModule.h
--- a/Engine/Plugins/Runtime/nDisplay/Source/DisplayClusterMedia/Private/DisplayClusterMediaModule.h
+++ b/Engine/Plugins/Runtime/nDisplay/Source/DisplayClusterMedia/Private/DisplayClusterMediaModule.h
@@ -21,6 +21,11 @@ class UDisplayClusterICVFXCameraComponent;
 class FDisplayClusterMediaModule :
 	public IModuleInterface
 {
+public:
+
+	FDisplayClusterMediaModule();
+	~FDisplayClusterMediaModule() = default;
+
 public:
 
 	//~ Begin IModuleInterface Implementation
@@ -93,5 +98,5 @@ private:
 private:
 
 	/** Latency queue */
-	FDisplayClusterFrameQueue FrameQueue;
+	TSharedRef<FDisplayClusterFrameQueue> FrameQueue;
 };
diff --git a/Engine/Plugins/Runtime/nDisplay/Source/DisplayClusterMedia/Private/Input/DisplayClusterMediaInputCameraBase.cpp b/Engine/Plugins/Runtime/nDisplay/Source/DisplayClusterMedia/Private/Input/DisplayClusterMediaInputCameraBase.cpp
--- a/Engine/Plugins/Runtime/nDisplay/Source/DisplayClusterMedia/Private/Input/DisplayClusterMediaInputCameraBase.cpp
+++ b/Engine/Plugins/Runtime/nDisplay/Source/DisplayClusterMedia/Private/Input/DisplayClusterMediaInputCameraBase.cpp
@@ -24,7 +24,7 @@ FDisplayClusterMediaInputCameraBase::FDisplayClusterMediaInputCameraBase(
 {
 }
 
-void FDisplayClusterMediaInputCameraBase::UpdateLateOCIOState(IDisplayClusterViewport* Viewport)
+void FDisplayClusterMediaInputCameraBase::UpdateLateOCIOState(const IDisplayClusterViewport* Viewport)
 {
 	FLateOCIOData NewLateOCIOConfiguration;
 
diff --git a/Engine/Plugins/Runtime/nDisplay/Source/DisplayClusterMedia/Private/Input/DisplayClusterMediaInputCameraBase.h b/Engine/Plugins/Runtime/nDisplay/Source/DisplayClusterMedia/Private/Input/DisplayClusterMediaInputCameraBase.h
--- a/Engine/Plugins/Runtime/nDisplay/Source/DisplayClusterMedia/Private/Input/DisplayClusterMediaInputCameraBase.h
+++ b/Engine/Plugins/Runtime/nDisplay/Source/DisplayClusterMedia/Private/Input/DisplayClusterMediaInputCameraBase.h
@@ -28,5 +28,5 @@ public:
 protected:
 
 	/** Updates late OCIO state */
-	virtual void UpdateLateOCIOState(IDisplayClusterViewport* Viewport) override;
+	virtual void UpdateLateOCIOState(const IDisplayClusterViewport* Viewport) override;
 };
diff --git a/Engine/Plugins/Runtime/nDisplay/Source/DisplayClusterMedia/Private/Input/DisplayClusterMediaInputViewportBase.cpp b/Engine/Plugins/Runtime/nDisplay/Source/DisplayClusterMedia/Private/Input/DisplayClusterMediaInputViewportBase.cpp
--- a/Engine/Plugins/Runtime/nDisplay/Source/DisplayClusterMedia/Private/Input/DisplayClusterMediaInputViewportBase.cpp
+++ b/Engine/Plugins/Runtime/nDisplay/Source/DisplayClusterMedia/Private/Input/DisplayClusterMediaInputViewportBase.cpp
@@ -59,7 +59,7 @@ void FDisplayClusterMediaInputViewportBase::Stop()
 	FDisplayClusterMediaInputBase::Stop();
 }
 
-void FDisplayClusterMediaInputViewportBase::UpdateLateOCIOState(IDisplayClusterViewport* Viewport)
+void FDisplayClusterMediaInputViewportBase::UpdateLateOCIOState(const IDisplayClusterViewport* Viewport)
 {
 	// So far, regular viewports dont't support late OCIO feature. But it can be overrided down the hierarchy.
 }
@@ -93,9 +93,9 @@ void FDisplayClusterMediaInputViewportBase::OnPreSubmitViewFamilies(TArray<FScen
 
 			// And push it to the rendering thread
 			ENQUEUE_RENDER_COMMAND(DCMediaInputUpdateOCIOResources)(
-				[this, InOCIOPassResources = MoveTemp(OCIOPassResources)](FRHICommandListImmediate& RHICmdList)
+				[This = AsShared(), InOCIOPassResources = MoveTemp(OCIOPassResources)](FRHICommandListImmediate& RHICmdList)
 				{
-					OCIOPassResources_RT = InOCIOPassResources;
+					StaticCastSharedRef<FDisplayClusterMediaInputViewportBase>(This)->OCIOPassResources_RT = InOCIOPassResources;
 				}
 			);
 		}
diff --git a/Engine/Plugins/Runtime/nDisplay/Source/DisplayClusterMedia/Private/Input/DisplayClusterMediaInputViewportBase.h b/Engine/Plugins/Runtime/nDisplay/Source/DisplayClusterMedia/Private/Input/DisplayClusterMediaInputViewportBase.h
--- a/Engine/Plugins/Runtime/nDisplay/Source/DisplayClusterMedia/Private/Input/DisplayClusterMediaInputViewportBase.h
+++ b/Engine/Plugins/Runtime/nDisplay/Source/DisplayClusterMedia/Private/Input/DisplayClusterMediaInputViewportBase.h
@@ -44,7 +44,7 @@ public:
 protected:
 
 	/** Updates late OCIO state */
-	virtual void UpdateLateOCIOState(IDisplayClusterViewport* Viewport);
+	virtual void UpdateLateOCIOState(const IDisplayClusterViewport* Viewport);
 
 	/** Process late OCIO state updates */
 	virtual void HandleLateOCIOChanged(const FLateOCIOData& NewLateOCIOConfiguration) override;
diff --git a/Engine/Plugins/Runtime/nDisplay/Source/DisplayClusterMedia/Private/Synchronization/LatencyQueue/DisplayClusterFrameQueue.cpp b/Engine/Plugins/Runtime/nDisplay/Source/DisplayClusterMedia/Private/Synchronization/LatencyQueue/DisplayClusterFrameQueue.cpp
--- a/Engine/Plugins/Runtime/nDisplay/Source/DisplayClusterMedia/Private/Synchronization/LatencyQueue/DisplayClusterFrameQueue.cpp
+++ b/Engine/Plugins/Runtime/nDisplay/Source/DisplayClusterMedia/Private/Synchronization/LatencyQueue/DisplayClusterFrameQueue.cpp
@@ -53,9 +53,9 @@ void FDisplayClusterFrameQueue::Release()
 
 	// We need to properly release render thread data
 	ENQUEUE_RENDER_COMMAND(FDisplayClusterFrameQueue_Release)(
-		[this](FRHICommandListImmediate& RHICmdList)
+		[This = AsShared()](FRHICommandListImmediate& RHICmdList)
 		{
-			CleanQueue_RenderThread();
+			This->CleanQueue_RenderThread();
 		});
 }
 
@@ -72,18 +72,18 @@ void FDisplayClusterFrameQueue::HandleBeginDraw()
 
 			// Update latency data on RT
 			ENQUEUE_RENDER_COMMAND(FDisplayClusterFrameQueue_SetLatency)(
-				[this, NewLatency](FRHICommandListImmediate& RHICmdList)
+				[This = AsShared(), NewLatency](FRHICommandListImmediate& RHICmdList)
 				{
-					SetLatency_RenderThread(NewLatency);
+					This->SetLatency_RenderThread(NewLatency);
 				});
 		}
 	}
 
 	// Process pending frames adding
 	ENQUEUE_RENDER_COMMAND(FDisplayClusterFrameQueue_ProcessAddingPendingFrames)(
-		[this](FRHICommandListImmediate& RHICmdList)
+		[This = AsShared()](FRHICommandListImmediate& RHICmdList)
 		{
-			ProcessAddingPendingFrames_RenderThread();
+			This->ProcessAddingPendingFrames_RenderThread();
 		});
 }
 
@@ -91,9 +91,9 @@ void FDisplayClusterFrameQueue::HandleEndDraw()
 {
 	// We need to update head and tail indices before rendering next frame
 	ENQUEUE_RENDER_COMMAND(FDisplayClusterFrameQueue_StepQueueIndices)(
-		[this](FRHICommandListImmediate& RHICmdList)
+		[This = AsShared()](FRHICommandListImmediate& RHICmdList)
 		{
-			StepQueueIndices_RenderThread();
+			This->StepQueueIndices_RenderThread();
 		});
 }
 
diff --git a/Engine/Plugins/Runtime/nDisplay/Source/DisplayClusterMedia/Private/Synchronization/LatencyQueue/DisplayClusterFrameQueue.h b/Engine/Plugins/Runtime/nDisplay/Source/DisplayClusterMedia/Private/Synchronization/LatencyQueue/DisplayClusterFrameQueue.h
--- a/Engine/Plugins/Runtime/nDisplay/Source/DisplayClusterMedia/Private/Synchronization/LatencyQueue/DisplayClusterFrameQueue.h
+++ b/Engine/Plugins/Runtime/nDisplay/Source/DisplayClusterMedia/Private/Synchronization/LatencyQueue/DisplayClusterFrameQueue.h
@@ -26,6 +26,7 @@ struct FDisplayClusterShaderParameters_WarpBlend;
  * at some point during refactoring.
  */
 class FDisplayClusterFrameQueue
+	: public TSharedFromThis<FDisplayClusterFrameQueue>
 {
 public:
 	// Initializes latency feature
