nDisplay: Default sync policy in windows should be Nvidia.

This is configured when the nDisplay blueprint is created.

#rn nDisplay default sync policy in Windows is now Nvidia sync policy.

diff --git a/Engine/Plugins/Runtime/nDisplay/Source/DisplayClusterConfigurator/Private/DisplayClusterConfiguratorFactory.cpp b/Engine/Plugins/Runtime/nDisplay/Source/DisplayClusterConfigurator/Private/DisplayClusterConfiguratorFactory.cpp
--- a/Engine/Plugins/Runtime/nDisplay/Source/DisplayClusterConfigurator/Private/DisplayClusterConfiguratorFactory.cpp
+++ b/Engine/Plugins/Runtime/nDisplay/Source/DisplayClusterConfigurator/Private/DisplayClusterConfiguratorFactory.cpp
@@ -164,7 +164,18 @@ void UDisplayClusterConfiguratorFactory::SetupNewBlueprint(UDisplayClusterBluepr
 	check(NewBlueprint);
 	
 	FDisplayClusterConfiguratorVersionUtils::SetToLatestVersion(NewBlueprint);
-	NewBlueprint->SetConfigData(UDisplayClusterConfigurationData::CreateNewConfigData());
+
+	UDisplayClusterConfigurationData* NewConfigData = UDisplayClusterConfigurationData::CreateNewConfigData();
+
+	// In Windows, the default sync policy should be Nvidia. Linux/other can be Ethernet by default, since Nvidia sync might not be supported.
+#if PLATFORM_WINDOWS
+	if (ensure(NewConfigData && NewConfigData->Cluster))
+	{
+		NewConfigData->Cluster->Sync.RenderSyncPolicy.Type = DisplayClusterConfigurationStrings::config::cluster::render_sync::Nvidia;
+	}
+#endif // PLATFORM_WINDOWS
+
+	NewBlueprint->SetConfigData(NewConfigData);
 	
 	check(NewBlueprint->GeneratedClass);
 	UClass* SuperClass = NewBlueprint->GeneratedClass->GetSuperClass();
