nDisplay: Fixed issue where ICVFX components weren't properly being reselected in the ICVFX drawer when a property was changed

diff --git a/Engine/Plugins/Runtime/nDisplay/Source/DisplayClusterDetails/Private/Drawer/SDisplayClusterDetailsDrawer.cpp b/Engine/Plugins/Runtime/nDisplay/Source/DisplayClusterDetails/Private/Drawer/SDisplayClusterDetailsDrawer.cpp
--- a/Engine/Plugins/Runtime/nDisplay/Source/DisplayClusterDetails/Private/Drawer/SDisplayClusterDetailsDrawer.cpp
+++ b/Engine/Plugins/Runtime/nDisplay/Source/DisplayClusterDetails/Private/Drawer/SDisplayClusterDetailsDrawer.cpp
@@ -294,6 +294,7 @@ void SDisplayClusterDetailsDrawer::SetDrawerState(const FDisplayClusterDetailsDr
 		}
 	}
 
+	ObjectListView->SetSelectedItems(ItemsToSelect);
 	DetailsDataModel->SetDrawerState(InDrawerState);
 
 	if (DetailsPanel.IsValid())
