[nDisplay] Base failover implementation (v1.5, P-node replacement)

About:
- Supports any node loss (depends on the configuration)
- It has everything we need for the next step where we could use backup nodes for complete (rendering + media) replacement of the failed ones

[RN] Improved the failover subsystem in order to support primary node loss

diff --git a/Engine/Plugins/Runtime/nDisplay/Source/DisplayCluster/Internal/IPDisplayCluster.h b/Engine/Plugins/Runtime/nDisplay/Source/DisplayCluster/Internal/IPDisplayCluster.h
--- a/Engine/Plugins/Runtime/nDisplay/Source/DisplayCluster/Internal/IPDisplayCluster.h
+++ b/Engine/Plugins/Runtime/nDisplay/Source/DisplayCluster/Internal/IPDisplayCluster.h
@@ -9,7 +9,6 @@ class IPDisplayClusterRenderManager;
 class IPDisplayClusterClusterManager;
 class IPDisplayClusterConfigManager;
 class IPDisplayClusterGameManager;
-class ADisplayClusterSettings;
 
 
 /**
@@ -20,11 +19,20 @@ class IPDisplayCluster
 	, public IPDisplayClusterManager
 {
 public:
+
 	virtual ~IPDisplayCluster() = default;
 
 public:
-	virtual IPDisplayClusterRenderManager*    GetPrivateRenderMgr() const = 0;
-	virtual IPDisplayClusterClusterManager*   GetPrivateClusterMgr() const = 0;
-	virtual IPDisplayClusterConfigManager*    GetPrivateConfigMgr() const = 0;
-	virtual IPDisplayClusterGameManager*      GetPrivateGameMgr() const = 0;
+
+	/** Returns private render manager interface */
+	virtual IPDisplayClusterRenderManager* GetPrivateRenderMgr() const = 0;
+
+	/** Returns private cluster manager interface */
+	virtual IPDisplayClusterClusterManager* GetPrivateClusterMgr() const = 0;
+
+	/** Returns private config manager interface */
+	virtual IPDisplayClusterConfigManager* GetPrivateConfigMgr() const = 0;
+
+	/** Returns private game manager interface */
+	virtual IPDisplayClusterGameManager* GetPrivateGameMgr() const = 0;
 };
diff --git a/Engine/Plugins/Runtime/nDisplay/Source/DisplayCluster/Private/Blueprints/DisplayClusterBlueprintLib.cpp b/Engine/Plugins/Runtime/nDisplay/Source/DisplayCluster/Private/Blueprints/DisplayClusterBlueprintLib.cpp
--- a/Engine/Plugins/Runtime/nDisplay/Source/DisplayCluster/Private/Blueprints/DisplayClusterBlueprintLib.cpp
+++ b/Engine/Plugins/Runtime/nDisplay/Source/DisplayCluster/Private/Blueprints/DisplayClusterBlueprintLib.cpp
@@ -71,6 +71,11 @@ bool UDisplayClusterBlueprintLib::IsBackup()
 	return IDisplayCluster::Get().GetClusterMgr()->IsBackup();
 }
 
+bool UDisplayClusterBlueprintLib::HasClusterRole(EDisplayClusterNodeRole Role)
+{
+	return IDisplayCluster::Get().GetClusterMgr()->HasClusterRole(Role);
+}
+
 EDisplayClusterNodeRole UDisplayClusterBlueprintLib::GetClusterRole()
 {
 	return IDisplayCluster::Get().GetClusterMgr()->GetClusterRole();
diff --git a/Engine/Plugins/Runtime/nDisplay/Source/DisplayCluster/Private/Cluster/Controller/DisplayClusterClusterNodeCtrlBase.cpp b/Engine/Plugins/Runtime/nDisplay/Source/DisplayCluster/Private/Cluster/Controller/DisplayClusterClusterNodeCtrlBase.cpp
--- a/Engine/Plugins/Runtime/nDisplay/Source/DisplayCluster/Private/Cluster/Controller/DisplayClusterClusterNodeCtrlBase.cpp
+++ b/Engine/Plugins/Runtime/nDisplay/Source/DisplayCluster/Private/Cluster/Controller/DisplayClusterClusterNodeCtrlBase.cpp
@@ -1,24 +1,22 @@
 // Copyright Epic Games, Inc. All Rights Reserved.
 
 #include "Cluster/Controller/DisplayClusterClusterNodeCtrlBase.h"
+#include "Cluster/IPDisplayClusterClusterManager.h"
 
-#include "IDisplayCluster.h"
+#include "Misc/DisplayClusterGlobals.h"
+#include "Misc/DisplayClusterLog.h"
 
 #include "Network/IDisplayClusterServer.h"
 #include "Network/IDisplayClusterClient.h"
+#include "Network/Configuration/DisplayClusterNetworkConfiguration.h"
 #include "Network/Listener/DisplayClusterTcpListener.h"
 #include "Network/Service/ClusterEventsJson/DisplayClusterClusterEventsJsonClient.h"
 #include "Network/Service/ClusterEventsBinary/DisplayClusterClusterEventsBinaryClient.h"
 
-#include "Misc/DisplayClusterGlobals.h"
-#include "Misc/DisplayClusterLog.h"
-
 
-FDisplayClusterClusterNodeCtrlBase::FDisplayClusterClusterNodeCtrlBase(const FString& CtrlName, const FString& NodeName)
-	: NodeName(NodeName)
-	, ControllerName(CtrlName)
-	, ExternalEventsClientJson(new FDisplayClusterClusterEventsJsonClient(), FDisplayClusterClientDeleter())
-	, ExternalEventsClientBinary(new FDisplayClusterClusterEventsBinaryClient(), FDisplayClusterClientDeleter())
+FDisplayClusterClusterNodeCtrlBase::FDisplayClusterClusterNodeCtrlBase(const FString& InControllerName, const FString& InClusterNodeId)
+	: ClusterNodeId(InClusterNodeId)
+	, ControllerName(InControllerName)
 {
 }
 
@@ -26,76 +24,66 @@ FDisplayClusterClusterNodeCtrlBase::~FDisplayClusterClusterNodeCtrlBase()
 {
 }
 
-
-//////////////////////////////////////////////////////////////////////////////////////////////
-// IDisplayClusterClusterNodeController
-//////////////////////////////////////////////////////////////////////////////////////////////
-bool FDisplayClusterClusterNodeCtrlBase::Initialize()
+TWeakPtr<FDisplayClusterService> FDisplayClusterClusterNodeCtrlBase::GetService(const FName& ServiceName) const
 {
-	if (!InitializeServers())
-	{
-		UE_LOG(LogDisplayClusterCluster, Error, TEXT("Servers initialization failed"));
-		return false;
-	}
-
-	if (!InitializeClients())
-	{
-		UE_LOG(LogDisplayClusterCluster, Error, TEXT("Clients initialization failed"));
-		return false;
-	}
+	FScopeLock Lock(&RegisteredServicesCS);
 
-	if (!StartServers())
+	const TSharedPtr<FDisplayClusterService>* const ServiceFound = RegisteredServices.Find(ServiceName);
+	if (!ServiceFound)
 	{
-		UE_LOG(LogDisplayClusterCluster, Error, TEXT("An error occurred during servers start"));
-		return false;
+		return nullptr;
 	}
 
-	if (!StartClients())
-	{
-		UE_LOG(LogDisplayClusterCluster, Error, TEXT("An error occurred during clients start"));
-		return false;
-	}
-
-	return true;
-}
-
-void FDisplayClusterClusterNodeCtrlBase::Shutdown()
-{
-	StopClients();
-	StopServers();
+	return ServiceFound->ToWeakPtr();
 }
 
 void FDisplayClusterClusterNodeCtrlBase::SendClusterEventTo(const FString& Address, const uint16 Port, const FDisplayClusterClusterEventJson& Event, bool bPrimaryOnly)
 {
-	// We should synchronize access to the client
-	FScopeLock Lock(&ExternEventsClientJsonGuard);
+	const bool bIsPrimary  = GDisplayCluster->GetClusterMgr()->IsPrimary();
+	const bool bShouldSend = (bPrimaryOnly ? bIsPrimary : true);
 
-	// One-shot connection
-	ExternalEventsClientJson->Connect(Address, Port, 1, 0.f);
-	ExternalEventsClientJson->EmitClusterEventJson(Event);
-	ExternalEventsClientJson->Disconnect();
+	if (bShouldSend)
+	{
+		if (TUniquePtr<FDisplayClusterClusterEventsJsonClient> Client = MakeUnique<FDisplayClusterClusterEventsJsonClient>(UE::nDisplay::Network::Configuration::JsonEventsExternalClientName, false))
+		{
+			if (Client->Connect(Address, Port, 1, 0.f))
+			{
+				Client->EmitClusterEventJson(Event);
+				Client->Disconnect();
+			}
+		}
+	}
 }
 
 void FDisplayClusterClusterNodeCtrlBase::SendClusterEventTo(const FString& Address, const uint16 Port, const FDisplayClusterClusterEventBinary& Event, bool bPrimaryOnly)
 {
-	// We should synchronize access to the client
-	FScopeLock Lock(&ExternEventsClientBinaryGuard);
+	const bool bIsPrimary  = GDisplayCluster->GetClusterMgr()->IsPrimary();
+	const bool bShouldSend = (bPrimaryOnly ? bIsPrimary : true);
 
-	// One-shot connection
-	ExternalEventsClientBinary->Connect(Address, Port, 1, 0.f);
-	ExternalEventsClientBinary->EmitClusterEventBinary(Event);
-	ExternalEventsClientBinary->Disconnect();
+	if (bShouldSend)
+	{
+		if (TUniquePtr<FDisplayClusterClusterEventsBinaryClient> Client = MakeUnique<FDisplayClusterClusterEventsBinaryClient>(UE::nDisplay::Network::Configuration::BinaryEventsExternalClientName, false))
+		{
+			if (Client->Connect(Address, Port, 1, 0.f))
+			{
+				Client->EmitClusterEventBinary(Event);
+				Client->Disconnect();
+			}
+		}
+	}
 }
 
+TMap<FName, TSharedPtr<FDisplayClusterService>> FDisplayClusterClusterNodeCtrlBase::GetRegisteredServices()
+{
+	FScopeLock Lock(&RegisteredServicesCS);
+	return RegisteredServices;
+}
 
-//////////////////////////////////////////////////////////////////////////////////////////////
-// FDisplayClusterNodeCtrlBase
-//////////////////////////////////////////////////////////////////////////////////////////////
-bool FDisplayClusterClusterNodeCtrlBase::StartServerWithLogs(IDisplayClusterServer* Server, const FString& Address, const uint16 Port) const
+bool FDisplayClusterClusterNodeCtrlBase::StartServerWithLogs(const TSharedPtr<IDisplayClusterServer>& Server, const FString& Address, const uint16 Port) const
 {
 	if (!Server)
 	{
-		UE_LOG(LogDisplayClusterCluster, Warning, TEXT("Invalid server instance (nullptr)"));
+		UE_LOG(LogDisplayClusterCluster, Error, TEXT("%s - Invalid server instance (nullptr)"), *GetControllerName());
 		return false;
 	}
 
@@ -103,21 +91,21 @@ bool FDisplayClusterClusterNodeCtrlBase::StartServerWithLogs(IDisplayClusterServ
 
 	if (bResult)
 	{
-		UE_LOG(LogDisplayClusterCluster, Log, TEXT("Server %s has started"), *Server->GetName());
+		UE_LOG(LogDisplayClusterCluster, Log, TEXT("%s - Server '%s' started at [%s:%u]"), *GetControllerName(), *Server->GetName(), *Address, Port);
 	}
 	else
 	{
-		UE_LOG(LogDisplayClusterCluster, Error, TEXT("Server %s failed to start"), *Server->GetName());
+		UE_LOG(LogDisplayClusterCluster, Error, TEXT("%s - Server '%s' failed to start at [%s:%u]"), *GetControllerName(), *Server->GetName(), *Address, Port);
 	}
 
 	return bResult;
 }
 
-bool FDisplayClusterClusterNodeCtrlBase::StartServerWithLogs(IDisplayClusterServer* const Server, TSharedPtr<FDisplayClusterTcpListener>& TcpListener)
+bool FDisplayClusterClusterNodeCtrlBase::StartServerWithLogs(const TSharedPtr<IDisplayClusterServer>& Server, TSharedRef<FDisplayClusterTcpListener>& TcpListener)
 {
 	if (!Server)
 	{
-		UE_LOG(LogDisplayClusterCluster, Warning, TEXT("Invalid server instance (nullptr)"));
+		UE_LOG(LogDisplayClusterCluster, Error, TEXT("%s - Invalid server instance (nullptr)"), *GetControllerName());
 		return false;
 	}
 
@@ -125,21 +113,28 @@ bool FDisplayClusterClusterNodeCtrlBase::StartServerWithLogs(IDisplayClusterServ
 
 	if (bResult)
 	{
-		UE_LOG(LogDisplayClusterCluster, Log, TEXT("Server %s has started"), *Server->GetName());
+		if (TcpListener->IsListening())
+		{
+			UE_LOG(LogDisplayClusterCluster, Log, TEXT("%s - Server '%s' started at [%s:%u]"), *GetControllerName(), *Server->GetName(), *TcpListener->GetListeningHost(), TcpListener->GetListeningPort());
+		}
+		else
+		{
+			UE_LOG(LogDisplayClusterCluster, Log, TEXT("%s - Server '%s' is awaiting for the listener to start"), *GetControllerName(), *Server->GetName());
+		}
 	}
 	else
 	{
-		UE_LOG(LogDisplayClusterCluster, Error, TEXT("Server %s failed to start"), *Server->GetName());
+		UE_LOG(LogDisplayClusterCluster, Error, TEXT("%s - Server '%s' failed to start at [%s:%u]"), *GetControllerName(), *Server->GetName(), *TcpListener->GetListeningHost(), TcpListener->GetListeningPort());
 	}
 
 	return bResult;
 }
 
-bool FDisplayClusterClusterNodeCtrlBase::StartClientWithLogs(IDisplayClusterClient* const Client, const FString& Address, const uint16 Port, const uint32 ClientConnTriesAmount, const uint32 ClientConnRetryDelay) const
+bool FDisplayClusterClusterNodeCtrlBase::StartClientWithLogs(IDisplayClusterClient* Client, const FString& Address, const uint16 Port, const uint32 ClientConnTriesAmount, const uint32 ClientConnRetryDelay)
 {
 	if (!Client)
 	{
-		UE_LOG(LogDisplayClusterCluster, Warning, TEXT("Invalid client instance (nullptr)"));
+		UE_LOG(LogDisplayClusterCluster, Error, TEXT("%s - Invalid client instance (nullptr)"), *GetControllerName());
 		return false;
 	}
 
@@ -147,12 +142,20 @@ bool FDisplayClusterClusterNodeCtrlBase::StartClientWithLogs(IDisplayClusterClie
 
 	if (bResult)
 	{
-		UE_LOG(LogDisplayClusterCluster, Log, TEXT("%s connected to the server %s:%d"), *Client->GetName(), *Address, Port);
+		UE_LOG(LogDisplayClusterCluster, Log, TEXT("%s - client '%s' connected to [%s:%u]"), *GetControllerName(), *Client->GetName(), *Address, Port);
 	}
 	else
 	{
-		UE_LOG(LogDisplayClusterCluster, Error, TEXT("%s couldn't connect to the server %s:%d"), *Client->GetName(), *Address, Port);
+		UE_LOG(LogDisplayClusterCluster, Error, TEXT("%s - client '%s' couldn't connect to [%s:%u]"), *GetControllerName(), *Client->GetName(), *Address, Port);
 	}
 
 	return bResult;
 }
+
+void FDisplayClusterClusterNodeCtrlBase::RegisterLocalService(const FName& ServiceName, const TSharedRef<FDisplayClusterService>& Service)
+{
+	FScopeLock Lock(&RegisteredServicesCS);
+
+	checkSlow(!RegisteredServices.Contains(ServiceName));
+	RegisteredServices.Emplace(ServiceName, Service);
+}
diff --git a/Engine/Plugins/Runtime/nDisplay/Source/DisplayCluster/Private/Cluster/Controller/DisplayClusterClusterNodeCtrlBase.h b/Engine/Plugins/Runtime/nDisplay/Source/DisplayCluster/Private/Cluster/Controller/DisplayClusterClusterNodeCtrlBase.h
--- a/Engine/Plugins/Runtime/nDisplay/Source/DisplayCluster/Private/Cluster/Controller/DisplayClusterClusterNodeCtrlBase.h
+++ b/Engine/Plugins/Runtime/nDisplay/Source/DisplayCluster/Private/Cluster/Controller/DisplayClusterClusterNodeCtrlBase.h
@@ -9,28 +9,30 @@ class IDisplayClusterServer;
 class FDisplayClusterTcpListener;
 class FDisplayClusterClusterEventsBinaryClient;
 class FDisplayClusterClusterEventsJsonClient;
+struct FDisplayClusterClusterEventBinary;
+struct FDisplayClusterClusterEventJson;
 
 
 /**
- * Abstract cluster node controller (cluster mode).
+ * Base node controller
+ *
+ * Encapsulates some common controller logic and data
  */
 class FDisplayClusterClusterNodeCtrlBase
 	: public IDisplayClusterClusterNodeController
 {
 public:
-	FDisplayClusterClusterNodeCtrlBase(const FString& CtrlName, const FString& NodeName);
+
+	FDisplayClusterClusterNodeCtrlBase(const FString& InControllerName, const FString& InClusterNodeId);
+
 	virtual ~FDisplayClusterClusterNodeCtrlBase();
 
 public:
-	//////////////////////////////////////////////////////////////////////////////////////////////
-	// IDisplayClusterClusterNodeController
-	//////////////////////////////////////////////////////////////////////////////////////////////
-	virtual bool Initialize() override final;
-	virtual void Shutdown()   override final;
 
+	//~ Begin IDisplayClusterClusterNodeController
 	virtual FString GetNodeId() const override final
 	{
-		return NodeName;
+		return ClusterNodeId;
 	}
 
 	virtual FString GetControllerName() const override final
@@ -38,118 +40,42 @@ public:
 		return ControllerName;
 	}
 
-	virtual void SendClusterEventTo(const FString& Address, const uint16 Port, const FDisplayClusterClusterEventJson&   Event, bool bPrimaryOnly) override;
-	virtual void SendClusterEventTo(const FString& Address, const uint16 Port, const FDisplayClusterClusterEventBinary& Event, bool bPrimaryOnly) override;
-
-public:
-	//////////////////////////////////////////////////////////////////////////////////////////////
-	// IDisplayClusterProtocolEventsJson - default overrides
-	//////////////////////////////////////////////////////////////////////////////////////////////
-	virtual EDisplayClusterCommResult EmitClusterEventJson(const FDisplayClusterClusterEventJson& Event) override
-	{
-		return EDisplayClusterCommResult::NotImplemented;
-	}
-
-public:
-	//////////////////////////////////////////////////////////////////////////////////////////////
-	// IDisplayClusterProtocolEventsBinary - default overrides
-	//////////////////////////////////////////////////////////////////////////////////////////////
-	virtual EDisplayClusterCommResult EmitClusterEventBinary(const FDisplayClusterClusterEventBinary& Event) override
-	{
-		return EDisplayClusterCommResult::NotImplemented;
-	}
-
-public:
-	//////////////////////////////////////////////////////////////////////////////////////////////
-	// IDisplayClusterProtocolClusterSync - default overrides
-	//////////////////////////////////////////////////////////////////////////////////////////////
-	virtual EDisplayClusterCommResult WaitForGameStart() override
-	{
-		return EDisplayClusterCommResult::NotImplemented;
-	}
-
-	virtual EDisplayClusterCommResult WaitForFrameStart() override
-	{
-		return EDisplayClusterCommResult::NotImplemented;
-	}
+	virtual TWeakPtr<FDisplayClusterService> GetService(const FName& ServiceName) const override;
 
-	virtual EDisplayClusterCommResult WaitForFrameEnd() override
-	{
-		return EDisplayClusterCommResult::NotImplemented;
-	}
-
-	virtual EDisplayClusterCommResult GetTimeData(double& OutDeltaTime, double& OutGameTime, TOptional<FQualifiedFrameTime>& OutFrameTime) override
-	{
-		return EDisplayClusterCommResult::NotImplemented;
-	}
-
-	virtual EDisplayClusterCommResult GetObjectsData(const EDisplayClusterSyncGroup InSyncGroup, TMap<FString, FString>& OutObjectsData) override
-	{
-		return EDisplayClusterCommResult::NotImplemented;
-	}
+	virtual void SendClusterEventTo(const FString& Address, const uint16 Port, const FDisplayClusterClusterEventJson& Event, bool bPrimaryOnly) override;
+	virtual void SendClusterEventTo(const FString& Address, const uint16 Port, const FDisplayClusterClusterEventBinary& Event, bool bPrimaryOnly) override;
+	//~ End IDisplayClusterClusterNodeController
 
-	virtual EDisplayClusterCommResult GetEventsData(TArray<TSharedPtr<FDisplayClusterClusterEventJson, ESPMode::ThreadSafe>>& OutJsonEvents, TArray<TSharedPtr<FDisplayClusterClusterEventBinary, ESPMode::ThreadSafe>>& OutBinaryEvents) override
-	{
-		return EDisplayClusterCommResult::NotImplemented;
-	}
+protected:
 
-	virtual EDisplayClusterCommResult GetNativeInputData(TMap<FString, FString>& OutNativeInputData) override
-	{
-		return EDisplayClusterCommResult::NotImplemented;
-	}
+	/** Provides accesss to the service register for children */
+	TMap<FName, TSharedPtr<FDisplayClusterService>> GetRegisteredServices();
 
-public:
-	//////////////////////////////////////////////////////////////////////////////////////////////
-	// IDisplayClusterProtocolRenderSync - default overrides
-	//////////////////////////////////////////////////////////////////////////////////////////////
-	virtual EDisplayClusterCommResult SyncOnBarrier() override
-	{
-		return EDisplayClusterCommResult::NotImplemented;
-	}
+	/** A helper function to start a server with its own connection listener */
+	bool StartServerWithLogs(const TSharedPtr<IDisplayClusterServer>& Server, const FString& Address, const uint16 Port) const;
 
-protected:
-	virtual bool InitializeServers()
-	{
-		return true;
-	}
+	/** A helper function to start a connection listener provided */
+	bool StartServerWithLogs(const TSharedPtr<IDisplayClusterServer>& Server, TSharedRef<FDisplayClusterTcpListener>& TcpListener);
 
-	virtual bool StartServers()
-	{
-		return true;
-	}
-
-	virtual void StopServers()
-	{
-		return;
-	}
+	/** A helper function to initialize (connect) a specific client to a specific address */
+	bool StartClientWithLogs(IDisplayClusterClient* Client, const FString& Address, const uint16 Port, const uint32 ClientConnTriesAmount, const uint32 ClientConnRetryDelay);
 
-	virtual bool InitializeClients()
-	{
-		return true;
-	}
+	/** Register local server so it can be reached via API */
+	void RegisterLocalService(const FName& ServiceName, const TSharedRef<FDisplayClusterService>& Service);
 
-	virtual bool StartClients()
-	{
-		return true;
-	}
+private:
 
-	virtual void StopClients()
-	{ }
+	/** Cluster node ID */
+	const FString ClusterNodeId;
 
-protected:
-	bool StartServerWithLogs(IDisplayClusterServer* const Server, const FString& Address, const uint16 Port) const;
-	bool StartServerWithLogs(IDisplayClusterServer* const Server, TSharedPtr<FDisplayClusterTcpListener>& TcpListener);
-	bool StartClientWithLogs(IDisplayClusterClient* const Client, const FString& Address, const uint16 Port, const uint32 ClientConnTriesAmount, const uint32 ClientConnRetryDelay) const;
+	/** Controller name/ID */
+	const FString ControllerName;
 
 private:
-	const FString NodeName;
-	const FString ControllerName;
 
-	// JSON client for sending events outside of the cluster
-	FCriticalSection ExternEventsClientJsonGuard;
-	TUniquePtr<FDisplayClusterClusterEventsJsonClient, FDisplayClusterClientDeleter> ExternalEventsClientJson;
+	/** Services that have been registered, and accessible from the outside via API */
+	TMap<FName, TSharedPtr<FDisplayClusterService>> RegisteredServices;
 
-	// Binary client for sending events outside of the cluster
-	FCriticalSection ExternEventsClientBinaryGuard;
-	TUniquePtr<FDisplayClusterClusterEventsBinaryClient, FDisplayClusterClientDeleter> ExternalEventsClientBinary;
+	/** A critical section to control access to the container of registered services */
+	mutable FCriticalSection RegisteredServicesCS;
 };
diff --git a/Engine/Plugins/Runtime/nDisplay/Source/DisplayCluster/Private/Cluster/Controller/DisplayClusterClusterNodeCtrlEditor.cpp b/Engine/Plugins/Runtime/nDisplay/Source/DisplayCluster/Private/Cluster/Controller/DisplayClusterClusterNodeCtrlEditor.cpp
--- a/Engine/Plugins/Runtime/nDisplay/Source/DisplayCluster/Private/Cluster/Controller/DisplayClusterClusterNodeCtrlEditor.cpp
+++ b/Engine/Plugins/Runtime/nDisplay/Source/DisplayCluster/Private/Cluster/Controller/DisplayClusterClusterNodeCtrlEditor.cpp
@@ -5,32 +5,108 @@
 #include "Cluster/IPDisplayClusterClusterManager.h"
 #include "Config/IPDisplayClusterConfigManager.h"
 
+#include "Misc/DisplayClusterGlobals.h"
+#include "Misc/DisplayClusterLog.h"
+
+#include "Network/Configuration/DisplayClusterNetworkConfiguration.h"
+#include "Network/Service/ClusterEventsJson/DisplayClusterClusterEventsJsonClient.h"
 #include "Network/Service/ClusterEventsBinary/DisplayClusterClusterEventsBinaryService.h"
+#include "Network/Service/ClusterEventsBinary/DisplayClusterClusterEventsBinaryClient.h"
 #include "Network/Service/ClusterEventsJson/DisplayClusterClusterEventsJsonService.h"
 
 #include "DisplayClusterConfigurationTypes.h"
 
-#include "Misc/DisplayClusterGlobals.h"
-#include "Misc/DisplayClusterLog.h"
-
 
-FDisplayClusterClusterNodeCtrlEditor::FDisplayClusterClusterNodeCtrlEditor(const FString& CtrlName, const FString& NodeName)
-	: FDisplayClusterClusterNodeCtrlBase(CtrlName, NodeName)
+FDisplayClusterClusterNodeCtrlEditor::FDisplayClusterClusterNodeCtrlEditor()
+	: FDisplayClusterClusterNodeCtrlBase(TEXT("CTRL_ED"), TEXT("Node_Editor"))
 {
 }
 
 FDisplayClusterClusterNodeCtrlEditor::~FDisplayClusterClusterNodeCtrlEditor()
 {
+	// In case Shutdown() has not been called before deleting this controller, we have to stop
+	// all the clients and servers this controller owns. We do it safely (non-virtual shutdown).
+	ShutdownImpl();
+}
+
+bool FDisplayClusterClusterNodeCtrlEditor::Initialize()
+{
+	const bool bBaseInitialized = FDisplayClusterClusterNodeCtrlBase::Initialize();
+	if (!bBaseInitialized)
+	{
+		UE_LOG(LogDisplayClusterCluster, Warning, TEXT("%s - Couldn't initialize base controller internals"), *GetControllerName());
+		// Don't return false as it's not critical for PIE
+	}
+
+	if (!InitializeServers())
+	{
+		UE_LOG(LogDisplayClusterCluster, Warning, TEXT("%s - Couldn't initialize internal servers"), *GetControllerName());
+		// Don't return false as it's not critical for PIE
+	}
+
+	if (!StartServers())
+	{
+		UE_LOG(LogDisplayClusterCluster, Warning, TEXT("%s - An error occurred while starting internal servers"), *GetControllerName());
+		// Don't return false as it's not critical for PIE
+	}
+
+	return true;
+}
+
+void FDisplayClusterClusterNodeCtrlEditor::Shutdown()
+{
+	ShutdownImpl();
+	FDisplayClusterClusterNodeCtrlBase::Shutdown();
 }
 
 
 //////////////////////////////////////////////////////////////////////////////////////////////
 // IDisplayClusterProtocolClusterSync
 //////////////////////////////////////////////////////////////////////////////////////////////
-EDisplayClusterCommResult FDisplayClusterClusterNodeCtrlEditor::GetEventsData(TArray<TSharedPtr<FDisplayClusterClusterEventJson, ESPMode::ThreadSafe>>& OutJsonEvents, TArray<TSharedPtr<FDisplayClusterClusterEventBinary, ESPMode::ThreadSafe>>& OutBinaryEvents)
+EDisplayClusterCommResult FDisplayClusterClusterNodeCtrlEditor::WaitForGameStart()
+{
+	// PIE stub
+	return EDisplayClusterCommResult::Ok;
+}
+
+EDisplayClusterCommResult FDisplayClusterClusterNodeCtrlEditor::WaitForFrameStart()
+{
+	// PIE stub
+	return EDisplayClusterCommResult::Ok;
+}
+
+EDisplayClusterCommResult FDisplayClusterClusterNodeCtrlEditor::WaitForFrameEnd()
+{
+	// PIE stub
+	return EDisplayClusterCommResult::Ok;
+}
+
+EDisplayClusterCommResult FDisplayClusterClusterNodeCtrlEditor::GetTimeData(double& OutDeltaTime, double& OutGameTime, TOptional<FQualifiedFrameTime>& OutFrameTime)
+{
+	IPDisplayClusterClusterManager* const ClusterMgr = GDisplayCluster->GetPrivateClusterMgr();
+	if (!ClusterMgr)
+	{
+		UE_LOG(LogDisplayClusterCluster, Warning, TEXT("%s - wrong cluster manager object"), *GetControllerName());
+		return EDisplayClusterCommResult::InternalError;
+	}
+
+	ClusterMgr->CacheTimeData();
+	ClusterMgr->ExportTimeData(OutDeltaTime, OutGameTime, OutFrameTime);
+
+	return EDisplayClusterCommResult::Ok;
+}
+
+EDisplayClusterCommResult FDisplayClusterClusterNodeCtrlEditor::GetObjectsData(const EDisplayClusterSyncGroup InSyncGroup, TMap<FString, FString>& OutObjectsData)
+{
+	// PIE stub
+	return EDisplayClusterCommResult::Ok;
+}
+
+EDisplayClusterCommResult FDisplayClusterClusterNodeCtrlEditor::GetEventsData(TArray<TSharedPtr<FDisplayClusterClusterEventJson>>& OutJsonEvents, TArray<TSharedPtr<FDisplayClusterClusterEventBinary>>& OutBinaryEvents)
 {
 	if (IPDisplayClusterClusterManager* const ClusterMgr = GDisplayCluster->GetPrivateClusterMgr())
 	{
+		ClusterMgr->CacheEvents();
 		ClusterMgr->ExportEventsData(OutJsonEvents, OutBinaryEvents);
 		return EDisplayClusterCommResult::Ok;
 	}
@@ -38,30 +114,135 @@ EDisplayClusterCommResult FDisplayClusterClusterNodeCtrlEditor::GetEventsData(TA
 	return EDisplayClusterCommResult::InternalError;
 }
 
+EDisplayClusterCommResult FDisplayClusterClusterNodeCtrlEditor::GetNativeInputData(TMap<FString, FString>& OutNativeInputData)
+{
+	// PIE stub
+	return EDisplayClusterCommResult::Ok;
+}
+
 
 //////////////////////////////////////////////////////////////////////////////////////////////
-// FDisplayClusterNodeCtrlBase
+// IDisplayClusterProtocolRenderSync
+//////////////////////////////////////////////////////////////////////////////////////////////
+EDisplayClusterCommResult FDisplayClusterClusterNodeCtrlEditor::SynchronizeOnBarrier()
+{
+	// PIE stub
+	return EDisplayClusterCommResult::Ok;
+}
+
+
+//////////////////////////////////////////////////////////////////////////////////////////////
+// IDisplayClusterProtocolEventsJson
+//////////////////////////////////////////////////////////////////////////////////////////////
+EDisplayClusterCommResult FDisplayClusterClusterNodeCtrlEditor::EmitClusterEventJson(const FDisplayClusterClusterEventJson& Event)
+{
+	// PIE stub
+	return EDisplayClusterCommResult::Ok;
+}
+
+
+//////////////////////////////////////////////////////////////////////////////////////////////
+// IDisplayClusterProtocolEventsBinary
+//////////////////////////////////////////////////////////////////////////////////////////////
+EDisplayClusterCommResult FDisplayClusterClusterNodeCtrlEditor::EmitClusterEventBinary(const FDisplayClusterClusterEventBinary& Event)
+{
+	// PIE stub
+	return EDisplayClusterCommResult::Ok;
+}
+
+
+//////////////////////////////////////////////////////////////////////////////////////////////
+// IDisplayClusterProtocolGenericBarrier
+//////////////////////////////////////////////////////////////////////////////////////////////
+EDisplayClusterCommResult FDisplayClusterClusterNodeCtrlEditor::CreateBarrier(const FString& BarrierId, const TMap<FString, TSet<FString>>& NodeToSyncCallers, uint32 Timeout, IDisplayClusterProtocolGenericBarrier::EBarrierControlResult& Result)
+{
+	// PIE stub
+	return EDisplayClusterCommResult::Ok;
+}
+
+EDisplayClusterCommResult FDisplayClusterClusterNodeCtrlEditor::WaitUntilBarrierIsCreated(const FString& BarrierId, IDisplayClusterProtocolGenericBarrier::EBarrierControlResult& Result)
+{
+	// PIE stub
+	return EDisplayClusterCommResult::Ok;
+}
+
+EDisplayClusterCommResult FDisplayClusterClusterNodeCtrlEditor::IsBarrierAvailable(const FString& BarrierId, IDisplayClusterProtocolGenericBarrier::EBarrierControlResult& Result)
+{
+	// PIE stub
+	return EDisplayClusterCommResult::Ok;
+}
+
+EDisplayClusterCommResult FDisplayClusterClusterNodeCtrlEditor::ReleaseBarrier(const FString& BarrierId, IDisplayClusterProtocolGenericBarrier::EBarrierControlResult& Result)
+{
+	// PIE stub
+	return EDisplayClusterCommResult::Ok;
+}
+
+EDisplayClusterCommResult FDisplayClusterClusterNodeCtrlEditor::SyncOnBarrier(const FString& BarrierId, const FString& CallerId, IDisplayClusterProtocolGenericBarrier::EBarrierControlResult& Result)
+{
+	// PIE stub
+	return EDisplayClusterCommResult::Ok;
+}
+
+EDisplayClusterCommResult FDisplayClusterClusterNodeCtrlEditor::SyncOnBarrierWithData(const FString& BarrierId, const FString& CallerId, const TArray<uint8>& RequestData, TArray<uint8>& OutResponseData, IDisplayClusterProtocolGenericBarrier::EBarrierControlResult& Result)
+{
+	// PIE stub
+	return EDisplayClusterCommResult::Ok;
+}
+
+//////////////////////////////////////////////////////////////////////////////////////////////
+// IDisplayClusterProtocolInternalComm
+//////////////////////////////////////////////////////////////////////////////////////////////
+EDisplayClusterCommResult FDisplayClusterClusterNodeCtrlEditor::GatherServicesHostingInfo(const FNodeServicesHostingInfo& ThisNodeInfo, FClusterServicesHostingInfo& OutHostingInfo)
+{
+	// PIE stub
+	return EDisplayClusterCommResult::Ok;
+}
+
+EDisplayClusterCommResult FDisplayClusterClusterNodeCtrlEditor::PostFailureNegotiate(TArray<uint8>& InOutRecoveryData)
+{
+	// PIE stub
+	return EDisplayClusterCommResult::Ok;
+}
+
+EDisplayClusterCommResult FDisplayClusterClusterNodeCtrlEditor::RequestNodeDrop(const FString& NodeId, uint8 DropReason)
+{
+	// PIE stub
+	return EDisplayClusterCommResult::Ok;
+}
+
+
+//////////////////////////////////////////////////////////////////////////////////////////////
+// FDisplayClusterClusterNodeCtrlEditor
 //////////////////////////////////////////////////////////////////////////////////////////////
 bool FDisplayClusterClusterNodeCtrlEditor::InitializeServers()
 {
-	UE_LOG(LogDisplayClusterCluster, Log, TEXT("%s - initializing servers..."), *GetControllerName());
+	UE_LOG(LogDisplayClusterCluster, Log, TEXT("%s - Initializing servers..."), *GetControllerName());
 
-	// Instantiate node servers
-	ClusterEventsJsonServer   = MakeUnique<FDisplayClusterClusterEventsJsonService>();
-	ClusterEventsBinaryServer = MakeUnique<FDisplayClusterClusterEventsBinaryService>();
+	// Instantiate public external servers                                                             
+	ClusterEventsJsonServer   = MakeShared<FDisplayClusterClusterEventsJsonService>  (UE::nDisplay::Network::Configuration::JsonEventsExternalServerName);
+	ClusterEventsBinaryServer = MakeShared<FDisplayClusterClusterEventsBinaryService>(UE::nDisplay::Network::Configuration::BinaryEventsExternalServerName);
 
-	return ClusterEventsJsonServer && ClusterEventsBinaryServer;
+	return true;
 }
 
 bool FDisplayClusterClusterNodeCtrlEditor::StartServers()
 {
-	UE_LOG(LogDisplayClusterCluster, Log, TEXT("%s - starting servers..."), *GetControllerName());
+	UE_LOG(LogDisplayClusterCluster, Log, TEXT("%s - Starting servers..."), *GetControllerName());
 
 	// Get config data
 	const UDisplayClusterConfigurationData* ConfigData = GDisplayCluster->GetPrivateConfigMgr()->GetConfig();
 	if (!ConfigData)
 	{
-		UE_LOG(LogDisplayClusterCluster, Error, TEXT("Couldn't get configuration data"));
+		UE_LOG(LogDisplayClusterCluster, Warning, TEXT("%s - Couldn't get configuration data"), *GetControllerName());
+		return false;
+	}
+
+	// Primary node configuration
+	const UDisplayClusterConfigurationClusterNode* const PrimaryNodeCfg = ConfigData->GetPrimaryNode();
+	if (!PrimaryNodeCfg)
+	{
+		UE_LOG(LogDisplayClusterCluster, Warning, TEXT("%s - No P-node configuration was found"), *GetControllerName());
 		return false;
 	}
 
@@ -69,17 +250,25 @@ bool FDisplayClusterClusterNodeCtrlEditor::StartServers()
 	const FString HostForPie("127.0.0.1");
 	const FDisplayClusterConfigurationPrimaryNodePorts& Ports = ConfigData->Cluster->PrimaryNode.Ports;
 
-	UE_LOG(LogDisplayClusterCluster, Log, TEXT("Servers: addr %s, port_cs %d, port_ce %d, port_ceb %d"), *HostForPie, Ports.ClusterSync, Ports.ClusterEventsJson, Ports.ClusterEventsBinary);
+	UE_LOG(LogDisplayClusterCluster, Log, TEXT("%s - Listening at `%s` to port_cej=%u, port_ceb=%u"),
+		*GetControllerName(), *HostForPie, Ports.ClusterEventsJson, Ports.ClusterEventsBinary);
 
 	// Start the servers
-	return StartServerWithLogs(ClusterEventsJsonServer.Get(),   HostForPie, Ports.ClusterEventsJson)
-		&& StartServerWithLogs(ClusterEventsBinaryServer.Get(), HostForPie, Ports.ClusterEventsBinary);
+	const bool bCEJServerStarted = StartServerWithLogs(ClusterEventsJsonServer,   HostForPie, Ports.ClusterEventsJson);
+	const bool bCEBServerStarted = StartServerWithLogs(ClusterEventsBinaryServer, HostForPie, Ports.ClusterEventsBinary);
+
+	return bCEJServerStarted && bCEBServerStarted;
 }
 
 void FDisplayClusterClusterNodeCtrlEditor::StopServers()
 {
-	UE_LOG(LogDisplayClusterCluster, Log, TEXT("%s - shutting down..."), *GetControllerName());
-
 	ClusterEventsJsonServer->Shutdown();
 	ClusterEventsBinaryServer->Shutdown();
 }
+
+void FDisplayClusterClusterNodeCtrlEditor::ShutdownImpl()
+{
+	UE_LOG(LogDisplayClusterCluster, Log, TEXT("%s - Shutting down..."), *GetControllerName());
+
+	StopServers();
+}
diff --git a/Engine/Plugins/Runtime/nDisplay/Source/DisplayCluster/Private/Cluster/Controller/DisplayClusterClusterNodeCtrlEditor.h b/Engine/Plugins/Runtime/nDisplay/Source/DisplayCluster/Private/Cluster/Controller/DisplayClusterClusterNodeCtrlEditor.h
--- a/Engine/Plugins/Runtime/nDisplay/Source/DisplayCluster/Private/Cluster/Controller/DisplayClusterClusterNodeCtrlEditor.h
+++ b/Engine/Plugins/Runtime/nDisplay/Source/DisplayCluster/Private/Cluster/Controller/DisplayClusterClusterNodeCtrlEditor.h
@@ -9,40 +9,96 @@ class FDisplayClusterClusterEventsBinaryService;
 
 
 /**
- * Editor node controller implementation.
+ * Node controller for 'Editor' operation mode.
+ *
+ * This controller is used in PIE only, therefore it has a very limited functionality.
+ * So far, the following features are supported:
+ *  - JSON and binary events processsing. This is useful for debugging event based logic in PIE.
+ *  - JSON and binary events sending outside (inheriated from base controller).
  */
 class FDisplayClusterClusterNodeCtrlEditor
 	: public FDisplayClusterClusterNodeCtrlBase
 {
 public:
-	FDisplayClusterClusterNodeCtrlEditor(const FString& CtrlName, const FString& NodeName);
+
+	FDisplayClusterClusterNodeCtrlEditor();
 	virtual ~FDisplayClusterClusterNodeCtrlEditor();
 
 public:
-	//////////////////////////////////////////////////////////////////////////////////////////////
-	// IDisplayClusterClusterNodeController
-	//////////////////////////////////////////////////////////////////////////////////////////////
-	virtual EDisplayClusterNodeRole GetClusterRole() const override
-	{
-		return EDisplayClusterNodeRole::Primary;
-	}
+
+	//~ Begin IDisplayClusterClusterNodeController
+	virtual bool Initialize() override;
+	virtual void Shutdown() override;
+	//~ End IDisplayClusterClusterNodeController
+
+public:
+
+	//~ Begin IDisplayClusterProtocolClusterSync
+	virtual EDisplayClusterCommResult WaitForGameStart() override;
+	virtual EDisplayClusterCommResult WaitForFrameStart() override;
+	virtual EDisplayClusterCommResult WaitForFrameEnd() override;
+	virtual EDisplayClusterCommResult GetTimeData(double& OutDeltaTime, double& OutGameTime, TOptional<FQualifiedFrameTime>& OutFrameTime) override;
+	virtual EDisplayClusterCommResult GetObjectsData(const EDisplayClusterSyncGroup InSyncGroup, TMap<FString, FString>& OutObjectsData) override;
+	virtual EDisplayClusterCommResult GetEventsData(TArray<TSharedPtr<FDisplayClusterClusterEventJson>>& OutJsonEvents, TArray<TSharedPtr<FDisplayClusterClusterEventBinary>>& OutBinaryEvents) override;
+	virtual EDisplayClusterCommResult GetNativeInputData(TMap<FString, FString>& OutNativeInputData) override;
+	//~ End IDisplayClusterProtocolClusterSync
+
+public:
+
+	//~ Begin IDisplayClusterProtocolRenderSync
+	virtual EDisplayClusterCommResult SynchronizeOnBarrier() override;
+	//~ End IDisplayClusterProtocolRenderSync
+
+public:
+
+	//~ Begin IDisplayClusterProtocolEventsJson
+	virtual EDisplayClusterCommResult EmitClusterEventJson(const FDisplayClusterClusterEventJson& Event) override;
+	//~ End IDisplayClusterProtocolEventsJson
+
+public:
+
+	//~ Begin IDisplayClusterProtocolEventsBinary
+	virtual EDisplayClusterCommResult EmitClusterEventBinary(const FDisplayClusterClusterEventBinary& Event) override;
+	//~ End IDisplayClusterProtocolEventsBinary
+
+public:
+
+	//~ Begin IDisplayClusterProtocolGenericBarrier
+	virtual EDisplayClusterCommResult CreateBarrier(const FString& BarrierId, const TMap<FString, TSet<FString>>& NodeToSyncCallers, uint32 Timeout, IDisplayClusterProtocolGenericBarrier::EBarrierControlResult& Result) override;
+	virtual EDisplayClusterCommResult WaitUntilBarrierIsCreated(const FString& BarrierId, IDisplayClusterProtocolGenericBarrier::EBarrierControlResult& Result) override;
+	virtual EDisplayClusterCommResult IsBarrierAvailable(const FString& BarrierId, IDisplayClusterProtocolGenericBarrier::EBarrierControlResult& Result) override;
+	virtual EDisplayClusterCommResult ReleaseBarrier(const FString& BarrierId, IDisplayClusterProtocolGenericBarrier::EBarrierControlResult& Result) override;
+	virtual EDisplayClusterCommResult SyncOnBarrier(const FString& BarrierId, const FString& CallerId, IDisplayClusterProtocolGenericBarrier::EBarrierControlResult& Result) override;
+	virtual EDisplayClusterCommResult SyncOnBarrierWithData(const FString& BarrierId, const FString& CallerId, const TArray<uint8>& RequestData, TArray<uint8>& OutResponseData, IDisplayClusterProtocolGenericBarrier::EBarrierControlResult& Result) override;
+	//~ End IDisplayClusterProtocolGenericBarrier
 
 public:
-	//////////////////////////////////////////////////////////////////////////////////////////////
-	// IDisplayClusterProtocolClusterSync
-	//////////////////////////////////////////////////////////////////////////////////////////////
-	virtual EDisplayClusterCommResult GetEventsData(TArray<TSharedPtr<FDisplayClusterClusterEventJson, ESPMode::ThreadSafe>>& OutJsonEvents, TArray<TSharedPtr<FDisplayClusterClusterEventBinary, ESPMode::ThreadSafe>>& OutBinaryEvents) override;
 
-protected:
-	//////////////////////////////////////////////////////////////////////////////////////////////
-	// FDisplayClusterNodeCtrlBase
-	//////////////////////////////////////////////////////////////////////////////////////////////
-	virtual bool InitializeServers() override;
-	virtual bool StartServers()      override;
-	virtual void StopServers()       override;
+	//~ Begin IDisplayClusterProtocolInternalComm
+	virtual EDisplayClusterCommResult GatherServicesHostingInfo(const FNodeServicesHostingInfo& ThisNodeInfo, FClusterServicesHostingInfo& OutHostingInfo) override;
+	virtual EDisplayClusterCommResult PostFailureNegotiate(TArray<uint8>& InOutRecoveryData) override;
+	virtual EDisplayClusterCommResult RequestNodeDrop(const FString& NodeId, uint8 DropReason) override;
+	//~ End IDisplayClusterProtocolInternalComm
 
 private:
-	// Node servers
-	TUniquePtr<FDisplayClusterClusterEventsJsonService>   ClusterEventsJsonServer;
-	TUniquePtr<FDisplayClusterClusterEventsBinaryService> ClusterEventsBinaryServer;
+
+	/** Auxiliary function that initializes all the internal servers */
+	bool InitializeServers();
+
+	/** Auxiliary function that starts all the internal servers */
+	bool StartServers();
+
+	/** Auxiliary function that stops all the internal servers */
+	void StopServers();
+
+	/** Non-virtual implementation of 'Shutdown' to prevent any potential issues when called from destrcutor */
+	void ShutdownImpl();
+
+private: // servers
+
+	/** JSON events server */
+	TSharedPtr<FDisplayClusterClusterEventsJsonService> ClusterEventsJsonServer;
+
+	/** Binary events server */
+	TSharedPtr<FDisplayClusterClusterEventsBinaryService> ClusterEventsBinaryServer;
 };
diff --git a/Engine/Plugins/Runtime/nDisplay/Source/DisplayCluster/Private/Cluster/Controller/IDisplayClusterClusterNodeController.h b/Engine/Plugins/Runtime/nDisplay/Source/DisplayCluster/Private/Cluster/Controller/IDisplayClusterClusterNodeController.h
--- a/Engine/Plugins/Runtime/nDisplay/Source/DisplayCluster/Private/Cluster/Controller/IDisplayClusterClusterNodeController.h
+++ b/Engine/Plugins/Runtime/nDisplay/Source/DisplayCluster/Private/Cluster/Controller/IDisplayClusterClusterNodeController.h
@@ -2,11 +2,11 @@
 
 #pragma once
 
-#include "CoreMinimal.h"
-
-#include "Network/Protocol/IDisplayClusterProtocolEventsJson.h"
-#include "Network/Protocol/IDisplayClusterProtocolEventsBinary.h"
 #include "Network/Protocol/IDisplayClusterProtocolClusterSync.h"
+#include "Network/Protocol/IDisplayClusterProtocolEventsBinary.h"
+#include "Network/Protocol/IDisplayClusterProtocolEventsJson.h"
+#include "Network/Protocol/IDisplayClusterProtocolGenericBarrier.h"
+#include "Network/Protocol/IDisplayClusterProtocolInternalComm.h"
 #include "Network/Protocol/IDisplayClusterProtocolRenderSync.h"
 
 #include "DisplayClusterEnums.h"
@@ -18,23 +18,29 @@ class FDisplayClusterService;
  * Node controller interface
  */
 class IDisplayClusterClusterNodeController
-	: public IDisplayClusterProtocolEventsJson
+	: public IDisplayClusterProtocolClusterSync
 	, public IDisplayClusterProtocolEventsBinary
-	, public IDisplayClusterProtocolClusterSync
+	, public IDisplayClusterProtocolEventsJson
+	, public IDisplayClusterProtocolGenericBarrier
+	, public IDisplayClusterProtocolInternalComm
 	, public IDisplayClusterProtocolRenderSync
 {
 public:
+
 	virtual ~IDisplayClusterClusterNodeController() = default;
 
 public:
 	/** Initialize controller instance */
-	virtual bool Initialize() = 0;
+	virtual bool Initialize()
+	{
+		return true;
+	}
 
 	/** Stop  clients/servers/etc */
-	virtual void Shutdown() = 0;
+	virtual void Shutdown()
+	{ }
 
-	/** Returns cluster role of the controller */
-	virtual EDisplayClusterNodeRole GetClusterRole() const = 0;
+public:
 
 	/** Return node ID */
 	virtual FString GetNodeId() const = 0;
@@ -42,18 +48,28 @@ public:
 	/** Return controller name */
 	virtual FString GetControllerName() const = 0;
 
-	/**
-	 * @note
-	 * This is a temporary workaround for a quick and safe implementation of the generic barriers expansion,
-	 * that would allow to set custom synchronization delegate. This will be refactored during Failover v2
-	 * implementation as we'll need all the nodes to have full set of servers for smooth primary role transition
-	 * and safe node replacement. Currently, the GB server is hosted on the p-nodes only.
-	 */
-	virtual FDisplayClusterService* GetGenericBarriersServer() const
+	/** Returns a set of internal service names */
+	virtual TSet<FName> GetInternalServiceNames() const
+	{
+		return { };
+	}
+
+	/** Access to a specific service */
+	virtual TWeakPtr<FDisplayClusterService> GetService(const FName& ServiceName) const
 	{
 		return nullptr;
 	}
 
+	/** Initializes internal set of GB clients. Returns set ID for external referencing. */
+	virtual int32 InitializeGeneralPurposeBarrierClients()
+	{
+		return INDEX_NONE;
+	}
+
+	/** Releases requested clients set. */
+	virtual void ReleaseGeneralPurposeBarrierClients(int32 ClientSetId)
+	{ }
+
 	/** Drop specific cluster node */
 	virtual bool DropClusterNode(const FString& NodeId)
 	{
diff --git a/Engine/Plugins/Runtime/nDisplay/Source/DisplayCluster/Private/Cluster/DisplayClusterClusterEvent.cpp b/Engine/Plugins/Runtime/nDisplay/Source/DisplayCluster/Private/Cluster/DisplayClusterClusterEvent.cpp
--- a/Engine/Plugins/Runtime/nDisplay/Source/DisplayCluster/Private/Cluster/DisplayClusterClusterEvent.cpp
+++ b/Engine/Plugins/Runtime/nDisplay/Source/DisplayCluster/Private/Cluster/DisplayClusterClusterEvent.cpp
@@ -6,82 +6,59 @@
 #include "Misc/DisplayClusterStrings.h"
 #include "Misc/DisplayClusterTypesConverter.h"
 
+#include "Serialization/Archive.h"
+
 
 //////////////////////////////////////////////////////////////////////////////////////////////
-// FDisplayClusterClusterEventJson
+// FDisplayClusterClusterEventBase
 //////////////////////////////////////////////////////////////////////////////////////////////
-FString FDisplayClusterClusterEventJson::SerializeToString() const
+void FDisplayClusterClusterEventBase::Serialize(FArchive& Ar)
 {
-	return FString::Printf(TEXT("%s:%s:%s:%d:%d:%s"),
-		*Category,
-		*Type,
-		*Name,
-		bIsSystemEvent ? 1 : 0,
-		bShouldDiscardOnRepeat ? 1 : 0,
-		*SerializeParametersToString());
+	Ar << bIsSystemEvent;
+	Ar << bShouldDiscardOnRepeat;
 }
 
-bool FDisplayClusterClusterEventJson::DeserializeFromString(const FString& Arch)
-{
-	const FString TokenSeparator(TEXT(":"));
-	FString TempStr;
-
-	FString StrIsSystemEvent;
-	FString StrShouldDiscardOnRepeat;
-
-	if (Arch.Split(TokenSeparator, &Category, &TempStr, ESearchCase::IgnoreCase, ESearchDir::FromStart) == false ||
-		TempStr.Split(TokenSeparator, &Type, &TempStr, ESearchCase::IgnoreCase, ESearchDir::FromStart) == false ||
-		TempStr.Split(TokenSeparator, &Name, &TempStr, ESearchCase::IgnoreCase, ESearchDir::FromStart) == false ||
-		TempStr.Split(TokenSeparator, &StrIsSystemEvent, &TempStr, ESearchCase::IgnoreCase, ESearchDir::FromStart) == false ||
-		TempStr.Split(TokenSeparator, &StrShouldDiscardOnRepeat, &TempStr, ESearchCase::IgnoreCase, ESearchDir::FromStart) == false)
-	{
-		return false;
-	}
-
-	bIsSystemEvent = DisplayClusterTypesConverter::template FromString<bool>(StrIsSystemEvent);
-	bShouldDiscardOnRepeat = DisplayClusterTypesConverter::template FromString<bool>(StrShouldDiscardOnRepeat);
 
-	if (DeserializeParametersFromString(TempStr) == false)
-	{
-		return false;
-	}
-
-	return true;
-}
-
-FString FDisplayClusterClusterEventJson::SerializeParametersToString() const
+//////////////////////////////////////////////////////////////////////////////////////////////
+// FDisplayClusterClusterEventJson
+//////////////////////////////////////////////////////////////////////////////////////////////
+void FDisplayClusterClusterEventJson::Serialize(FArchive& Ar)
 {
-	FString Result;
-
-	for (const auto& obj : Parameters)
-	{
-		Result += FString::Printf(TEXT("%s%s%s;"), *obj.Key, DisplayClusterStrings::common::KeyValSeparator, *obj.Value);
-	}
+	FDisplayClusterClusterEventBase::Serialize(Ar);
 
-	return Result;
+	Ar << Category;
+	Ar << Type;
+	Ar << Name;
+	Ar << Parameters;
 }
 
-bool FDisplayClusterClusterEventJson::DeserializeParametersFromString(const FString& Arch)
+FString FDisplayClusterClusterEventJson::ToString(bool bWithParams) const
 {
-	Parameters.Empty(Parameters.Num());
-
-	FString TempStr = Arch;
-	FString TempKeyValPair;
-
-	while (TempStr.Split(FString(";"), &TempKeyValPair, &TempStr, ESearchCase::IgnoreCase, ESearchDir::FromStart))
+	if (bWithParams)
 	{
-		FString l;
-		FString r;
-
-		if (TempKeyValPair.Split(FString(DisplayClusterStrings::common::KeyValSeparator), &l, &r, ESearchCase::IgnoreCase, ESearchDir::FromStart) == false)
+		FString ParamText;
+		for (const TPair<FString, FString>& Parameter : Parameters)
 		{
-			return false;
+			ParamText += FString::Printf(TEXT("%s%s%s;"), *Parameter.Key, DisplayClusterStrings::common::KeyValSeparator, *Parameter.Value);
 		}
 
-		Parameters.Add(l, r);
+		return FString::Printf(TEXT("%s:%s:%s:%d:%d:%s"),
+			*Category,
+			*Type,
+			*Name,
+			bIsSystemEvent ? 1 : 0,
+			bShouldDiscardOnRepeat ? 1 : 0,
+			*ParamText);
+	}
+	else
+	{
+		return FString::Printf(TEXT("%s:%s:%s:%d:%d"),
+			*Category,
+			*Type,
+			*Name,
+			bIsSystemEvent ? 1 : 0,
+			bShouldDiscardOnRepeat ? 1 : 0);
 	}
-
-	return true;
 }
 
 
@@ -142,3 +119,11 @@ bool FDisplayClusterClusterEventBinary::DeserializeFromByteArray(const TArray<ui
 
 	return true;
 }
+
+void FDisplayClusterClusterEventBinary::Serialize(FArchive& Ar)
+{
+	FDisplayClusterClusterEventBase::Serialize(Ar);
+
+	Ar << EventId;
+	Ar << EventData;
+}
diff --git a/Engine/Plugins/Runtime/nDisplay/Source/DisplayCluster/Private/Cluster/DisplayClusterClusterEventHandler.h b/Engine/Plugins/Runtime/nDisplay/Source/DisplayCluster/Private/Cluster/DisplayClusterClusterEventHandler.h
--- a/Engine/Plugins/Runtime/nDisplay/Source/DisplayCluster/Private/Cluster/DisplayClusterClusterEventHandler.h
+++ b/Engine/Plugins/Runtime/nDisplay/Source/DisplayCluster/Private/Cluster/DisplayClusterClusterEventHandler.h
@@ -4,11 +4,6 @@
 
 #include "Cluster/IDisplayClusterClusterEventListener.h"
 #include "Cluster/IDisplayClusterClusterManager.h"
-#include "Misc/App.h"
-
-class ADisplayClusterSettings;
-class FJsonObject;
-class FEvent;
 
 
 /**
@@ -18,24 +13,29 @@ class FDisplayClusterClusterEventHandler
 	: public IDisplayClusterClusterEventListener
 {
 protected:
+
 	FDisplayClusterClusterEventHandler();
 
 public:
+
+	/** Singleton access */
 	static FDisplayClusterClusterEventHandler& Get()
 	{
 		static FDisplayClusterClusterEventHandler Instance;
 		return Instance;
 	}
 
+	/** Returns JSON event handler delegate */
 	FOnClusterEventJsonListener& GetJsonListenerDelegate()
 	{
 		return ListenerDelegate;
 	}
 
 protected:
-	// Cluster event processing function
+
+	/** Cluster event processing function */
 	void HandleClusterEvent(const FDisplayClusterClusterEventJson& InEvent);
 
-	// Cluster event listener delegate
+	/** Cluster event listener delegate */
 	FOnClusterEventJsonListener ListenerDelegate;
 };
diff --git a/Engine/Plugins/Runtime/nDisplay/Source/DisplayCluster/Private/Cluster/DisplayClusterClusterManager.cpp b/Engine/Plugins/Runtime/nDisplay/Source/DisplayCluster/Private/Cluster/DisplayClusterClusterManager.cpp
--- a/Engine/Plugins/Runtime/nDisplay/Source/DisplayCluster/Private/Cluster/DisplayClusterClusterManager.cpp
+++ b/Engine/Plugins/Runtime/nDisplay/Source/DisplayCluster/Private/Cluster/DisplayClusterClusterManager.cpp
@@ -1,40 +1,46 @@
 // Copyright Epic Games, Inc. All Rights Reserved.
 
 #include "Cluster/DisplayClusterClusterManager.h"
-#include "Cluster/DisplayClusterClusterEvent.h"
+
 #include "Cluster/DisplayClusterClusterEventHandler.h"
+#include "Cluster/DisplayClusterGenericBarrierAPI.h"
 
 #include "Cluster/IDisplayClusterClusterSyncObject.h"
 #include "Cluster/IDisplayClusterClusterEventListener.h"
+
+#include "Cluster/Controller/DisplayClusterClusterNodeCtrlDisabled.h"
 #include "Cluster/Controller/DisplayClusterClusterNodeCtrlEditor.h"
-#include "Cluster/Controller/DisplayClusterClusterNodeCtrlNull.h"
-#include "Cluster/Controller/DisplayClusterClusterNodeCtrlPrimary.h"
-#include "Cluster/Controller/DisplayClusterClusterNodeCtrlSecondary.h"
-#include "Cluster/Failover/DisplayClusterFailoverNodeCtrlNull.h"
-#include "Cluster/Failover/DisplayClusterFailoverNodeCtrlPrimary.h"
-#include "Cluster/Failover/DisplayClusterFailoverNodeCtrlSecondary.h"
+#include "Cluster/Controller/DisplayClusterClusterNodeCtrlMain.h"
 
-#include "Network/Service/GenericBarrier/DisplayClusterGenericBarrierClient.h"
+#include "Cluster/Failover/DisplayClusterFailoverNodeCtrlDisabled.h"
+#include "Cluster/Failover/DisplayClusterFailoverNodeCtrlEditor.h"
+#include "Cluster/Failover/DisplayClusterFailoverNodeCtrlMain.h"
+
+#include "Cluster/NetAPI/DisplayClusterNetApiFacade.h"
 
 #include "Config/IPDisplayClusterConfigManager.h"
-#include "DisplayClusterConfigurationTypes.h"
 
 #include "Dom/JsonObject.h"
-#include "HAL/Event.h"
 
-#include "Misc/DisplayClusterTypesConverter.h"
+#include "Misc/App.h"
+#include "Misc/DisplayClusterAppExit.h"
 #include "Misc/DisplayClusterGlobals.h"
 #include "Misc/DisplayClusterHelpers.h"
 #include "Misc/DisplayClusterLog.h"
 #include "Misc/DisplayClusterStrings.h"
-
-#include "Misc/App.h"
-#include "Misc/QualifiedFrameTime.h"
+#include "Misc/DisplayClusterTypesConverter.h"
+#include "Misc/ScopeLock.h"
 
 #include "UObject/Interface.h"
 
+#include "DisplayClusterConfigurationTypes.h"
+#include "IDisplayClusterCallbacks.h"
+
 
 FDisplayClusterClusterManager::FDisplayClusterClusterManager()
+	: NodeCtrl(MakeShared<FDisplayClusterClusterNodeCtrlDisabled>())
+	, FailoverCtrl(MakeShared<FDisplayClusterFailoverNodeCtrlDisabled>(NodeCtrl))
+	, NetApi(new FDisplayClusterNetApiFacade(FailoverCtrl))
 {
 	// Sync objects
 	ObjectsToSync.Emplace(EDisplayClusterSyncGroup::PreTick).Reserve(64);
@@ -65,12 +71,10 @@ FDisplayClusterClusterManager::~FDisplayClusterClusterManager()
 	// Trigger all data cache availability events to prevent client session threads to be deadlocked.
 	SetInternalSyncObjectsReleaseState(true);
 
-	// Also stop clients/servers
-	if (ClusterNodeCtrl)
-	{
-		ClusterNodeCtrl->Shutdown();
-	}
+	// Stop networking in case it hasn't been stopped yet
+	ReleaseNetworking();
 
+	// Return sync event objects to the pool
 	for (TPair<EDisplayClusterSyncGroup, FEvent*>& GroupEventIt : ObjectsToSyncCacheReadySignals)
 	{
 		FPlatformProcess::ReturnSynchEventToPool(GroupEventIt.Value);
@@ -89,6 +93,7 @@ bool FDisplayClusterClusterManager::Init(EDisplayClusterOperationMode OperationM
 
 void FDisplayClusterClusterManager::Release()
 {
+	CurrentOperationMode = EDisplayClusterOperationMode::Disabled;
 }
 
 bool FDisplayClusterClusterManager::StartSession(UDisplayClusterConfigurationData* InConfigData, const FString& InNodeId)
@@ -105,61 +110,73 @@ bool FDisplayClusterClusterManager::StartSession(UDisplayClusterConfigurationDat
 	}
 
 	// Get configuration data
-	const UDisplayClusterConfigurationData* ConfigData = GDisplayCluster->GetPrivateConfigMgr()->GetConfig();
+	const UDisplayClusterConfigurationData* const ConfigData = GDisplayCluster->GetPrivateConfigMgr()->GetConfig();
 	if (!ConfigData)
 	{
 		UE_LOG(LogDisplayClusterCluster, Error, TEXT("Couldn't get configuration data"));
 		return false;
 	}
 
+	// Does it exist in the cluster configuration?
 	if (!ConfigData->Cluster->Nodes.Contains(ClusterNodeId))
 	{
 		UE_LOG(LogDisplayClusterCluster, Error, TEXT("Node '%s' not found in the configuration data"), *ClusterNodeId);
 		return false;
 	}
 
-	// Save node IDs
-	ConfigData->Cluster->Nodes.GenerateKeyArray(ClusterNodeIds);
+	// Subscribe for events
+	GDisplayCluster->GetCallbacks().OnDisplayClusterFailoverPrimaryNodeChanged().AddRaw(this, &FDisplayClusterClusterManager::OnPrimaryNodeChangedHandler);
+	GDisplayCluster->GetCallbacks().OnDisplayClusterFailoverNodeDown().AddRaw(this, &FDisplayClusterClusterManager::OnClusterNodeFailed);
 
 	// Reset all internal sync objects
 	SetInternalSyncObjectsReleaseState(false);
 
-	// Instantiate cluster node controller
-	ClusterNodeCtrl = CreateClusterNodeController();
+	// Save initial list of cluster nodes
+	ConfigData->Cluster->Nodes.GetKeys(InitialClusterNodeIds);
 
-	// Initialize the controller
-	UE_LOG(LogDisplayClusterCluster, Log, TEXT("Initializing the controller..."));
-	if (!ClusterNodeCtrl->Initialize())
+	// Also, initialize the active nodes list
 	{
-		UE_LOG(LogDisplayClusterCluster, Error, TEXT("Couldn't initialize a controller."));
-		ClusterNodeCtrl = MakeUnique<FDisplayClusterClusterNodeCtrlNull>(FString("[CTRL-NULL]"), ClusterNodeId);
-		return false;
+		FScopeLock Lock(&ActiveClusterNodeIdsCS);
+		ActiveClusterNodeIds = InitialClusterNodeIds;
 	}
 
-	// Instantiate failover node controller
-	FailoverNodeCtrl = CreateFailoverNodeController();
+	// Determine cluster role for this instance
+	InitializeClusterRole(InNodeId, ConfigData);
+
+	// Set primary node
+	SetPrimaryNode(ConfigData->Cluster->PrimaryNode.Id);
+
+	// Initialize networking internals
+	const bool bNetworkingInternalsInitialized = InitializeNetworking(ConfigData);
+	if (!bNetworkingInternalsInitialized)
+	{
+		UE_LOG(LogDisplayClusterCluster, Error, TEXT("Node '%s' could not initialize networking subsystems"), *ClusterNodeId);
+		return false;
+	}
 
 	return true;
 }
 
 void FDisplayClusterClusterManager::EndSession()
 {
+	// Unsubscribe from the session events
+	GDisplayCluster->GetCallbacks().OnDisplayClusterFailoverPrimaryNodeChanged().RemoveAll(this);
+	GDisplayCluster->GetCallbacks().OnDisplayClusterFailoverNodeDown().RemoveAll(this);
+
 	// Trigger all data cache availability events to prevent
 	// client session threads to be deadlocked.
 	SetInternalSyncObjectsReleaseState(true);
 
-	// Stop local clients/servers
-	if (ClusterNodeCtrl)
-	{
-		ClusterNodeCtrl->Shutdown();
-	}
+	// Stop networking
+	ReleaseNetworking();
 
 	{
-		FScopeLock Lock(&InternalsCS);
-
-		ClusterNodeIds.Reset();
-		ClusterNodeId.Empty();
+		FScopeLock Lock(&ActiveClusterNodeIdsCS);
+		ActiveClusterNodeIds.Reset();
 	}
+
+	InitialClusterNodeIds.Reset();
+	ClusterNodeId.Empty();
 }
 
 bool FDisplayClusterClusterManager::StartScene(UWorld* InWorld)
@@ -247,40 +264,69 @@ void FDisplayClusterClusterManager::PostTick(float DeltaSeconds)
 //////////////////////////////////////////////////////////////////////////////////////////////
 bool FDisplayClusterClusterManager::IsPrimary() const
 {
-	FScopeLock Lock(&InternalsCS);
-	return ClusterNodeCtrl ? ClusterNodeCtrl->GetClusterRole() == EDisplayClusterNodeRole::Primary : false;
+	return HasClusterRole(EDisplayClusterNodeRole::Primary);
 }
 
 bool FDisplayClusterClusterManager::IsSecondary() const
 {
-	FScopeLock Lock(&InternalsCS);
-	return ClusterNodeCtrl ? ClusterNodeCtrl->GetClusterRole() == EDisplayClusterNodeRole::Secondary : false;
+	return HasClusterRole(EDisplayClusterNodeRole::Secondary);
 }
 
 bool FDisplayClusterClusterManager::IsBackup() const
 {
-	FScopeLock Lock(&InternalsCS);
-	return ClusterNodeCtrl ? ClusterNodeCtrl->GetClusterRole() == EDisplayClusterNodeRole::Backup : false;
+	return HasClusterRole(EDisplayClusterNodeRole::Backup);
+}
+
+bool FDisplayClusterClusterManager::HasClusterRole(EDisplayClusterNodeRole Role) const
+{
+	return GetClusterRole() == Role;
 }
 
 EDisplayClusterNodeRole FDisplayClusterClusterManager::GetClusterRole() const
 {
-	FScopeLock Lock(&InternalsCS);
-	return ClusterNodeCtrl ? ClusterNodeCtrl->GetClusterRole() : EDisplayClusterNodeRole::None;
+	FScopeLock Lock(&CurrentNodeRoleCS);
+	return CurrentNodeRole;
+}
+
+FString FDisplayClusterClusterManager::GetPrimaryNodeId() const
+{
+	FScopeLock Lock(&PrimaryNodeIdCS);
+	return PrimaryNodeId;
+}
+
+FString FDisplayClusterClusterManager::GetNodeId() const
+{
+	return ClusterNodeId;
+}
+
+uint32 FDisplayClusterClusterManager::GetNodesAmount() const
+{
+	FScopeLock Lock(&ActiveClusterNodeIdsCS);
+	const int32 NodesNum = ActiveClusterNodeIds.Num();
+	return static_cast<uint32>(NodesNum <= 0 ? 0 : NodesNum);
 }
 
 void FDisplayClusterClusterManager::GetNodeIds(TArray<FString>& OutNodeIds) const
 {
-	FScopeLock Lock(&InternalsCS);
-	OutNodeIds = ClusterNodeIds;
+	FScopeLock Lock(&ActiveClusterNodeIdsCS);
+	OutNodeIds = ActiveClusterNodeIds.Array();
+}
+
+void FDisplayClusterClusterManager::GetNodeIds(TSet<FString>& OutNodeIds) const
+{
+	FScopeLock Lock(&ActiveClusterNodeIdsCS);
+	OutNodeIds = ActiveClusterNodeIds;
 }
 
 bool FDisplayClusterClusterManager::DropClusterNode(const FString& NodeId)
 {
-	FScopeLock Lock(&InternalsCS);
+	if (!IsPrimary())
+	{
+		UE_LOG(LogDisplayClusterCluster, Warning, TEXT("Node drop is allowed on P-nodes only"));
+		return false;
+	}
 
-	// Pass request to the cluster controller
-	return ClusterNodeCtrl ? ClusterNodeCtrl->DropClusterNode(NodeId) : false;
+	return DropNode(NodeId, IPDisplayClusterClusterManager::ENodeDropReason::UserRequest);
 }
 
 void FDisplayClusterClusterManager::RegisterSyncObject(IDisplayClusterClusterSyncObject* SyncObj, EDisplayClusterSyncGroup SyncGroup)
@@ -308,10 +354,9 @@ void FDisplayClusterClusterManager::UnregisterSyncObject(IDisplayClusterClusterS
 	}
 }
 
-TUniquePtr<IDisplayClusterGenericBarriersClient, FDisplayClusterGenericBarriersClientDeleter> FDisplayClusterClusterManager::CreateGenericBarriersClient(const FString& ClientName)
+TSharedRef<IDisplayClusterGenericBarriersClient> FDisplayClusterClusterManager::CreateGenericBarriersClient()
 {
-	return TUniquePtr<FDisplayClusterGenericBarrierClient, FDisplayClusterGenericBarriersClientDeleter>(
-		new FDisplayClusterGenericBarrierClient(ClientName), FDisplayClusterGenericBarriersClientDeleter());
+	return MakeShared<FDisplayClusterGenericBarrierAPI>();
 }
 
 void FDisplayClusterClusterManager::AddClusterEventListener(TScriptInterface<IDisplayClusterClusterEventListener> Listener)
@@ -359,7 +404,7 @@ void FDisplayClusterClusterManager::RemoveClusterEventBinaryListener(const FOnCl
 
 void FDisplayClusterClusterManager::EmitClusterEventJson(const FDisplayClusterClusterEventJson& Event, bool bPrimaryOnly)
 {
-	UE_LOG(LogDisplayClusterCluster, Verbose, TEXT("JSON event emission request: %s:%s:%s"), *Event.Category, *Event.Type, *Event.Name);
+	UE_LOG(LogDisplayClusterCluster, Verbose, TEXT("JSON event emission request: %s"), *Event.ToString());
 
 	FScopeLock Lock(&ClusterEventsJsonCS);
 
@@ -371,7 +416,7 @@ void FDisplayClusterClusterManager::EmitClusterEventJson(const FDisplayClusterCl
 			// Generate event ID
 			const FString EventId = FString::Printf(TEXT("%s-%s-%s"), *Event.Category, *Event.Type, *Event.Name);
 			// Make it shared ptr
-			TSharedPtr<FDisplayClusterClusterEventJson, ESPMode::ThreadSafe> EventPtr = MakeShared<FDisplayClusterClusterEventJson, ESPMode::ThreadSafe>(Event);
+			TSharedPtr<FDisplayClusterClusterEventJson> EventPtr = MakeShared<FDisplayClusterClusterEventJson>(Event);
 			// Store event object
 			if (EventPtr->bShouldDiscardOnRepeat)
 			{
@@ -386,9 +431,9 @@ void FDisplayClusterClusterManager::EmitClusterEventJson(const FDisplayClusterCl
 		else
 		{
 			// An event will be emitted from a secondary node if it's explicitly specified by bPrimaryOnly=false
-			if (!bPrimaryOnly && ClusterNodeCtrl)
+			if (!bPrimaryOnly)
 			{
-				ClusterNodeCtrl->EmitClusterEventJson(Event);
+				FailoverCtrl->EmitClusterEventJson(Event);
 			}
 		}
 	}
@@ -406,7 +451,7 @@ void FDisplayClusterClusterManager::EmitClusterEventBinary(const FDisplayCluster
 		if (IsPrimary())
 		{
 			// Make it shared ptr
-			TSharedPtr<FDisplayClusterClusterEventBinary, ESPMode::ThreadSafe> EventPtr = MakeShared<FDisplayClusterClusterEventBinary, ESPMode::ThreadSafe>(Event);
+			TSharedPtr<FDisplayClusterClusterEventBinary> EventPtr = MakeShared<FDisplayClusterClusterEventBinary>(Event);
 
 			if (EventPtr->bShouldDiscardOnRepeat)
 			{
@@ -421,9 +466,9 @@ void FDisplayClusterClusterManager::EmitClusterEventBinary(const FDisplayCluster
 		else
 		{
 			// An event will be emitted from a secondary node if it's explicitly specified by bPrimaryOnly=false
-			if (!bPrimaryOnly && ClusterNodeCtrl)
+			if (!bPrimaryOnly)
 			{
-				ClusterNodeCtrl->EmitClusterEventBinary(Event);
+				FailoverCtrl->EmitClusterEventBinary(Event);
 			}
 		}
 	}
@@ -433,13 +478,10 @@ void FDisplayClusterClusterManager::SendClusterEventTo(const FString& Address, c
 {
 	if (CurrentOperationMode == EDisplayClusterOperationMode::Cluster || CurrentOperationMode == EDisplayClusterOperationMode::Editor)
 	{
-		if (ClusterNodeCtrl)
+		if (IsPrimary() || !bPrimaryOnly)
 		{
-			if (IsPrimary() || !bPrimaryOnly)
-			{
-				UE_LOG(LogDisplayClusterCluster, Verbose, TEXT("JSON event emission request: recipient=%s:%u, event=%s:%s:%s"), *Address, Port, *Event.Category, *Event.Type, *Event.Name);
-				ClusterNodeCtrl->SendClusterEventTo(Address, Port, Event, bPrimaryOnly);
-			}
+			UE_LOG(LogDisplayClusterCluster, Verbose, TEXT("JSON event emission request: recipient=%s:%u, event=%s:%s:%s"), *Address, Port, *Event.Category, *Event.Type, *Event.Name);
+			NodeCtrl->SendClusterEventTo(Address, Port, Event, bPrimaryOnly);
 		}
 	}
 }
@@ -448,13 +490,10 @@ void FDisplayClusterClusterManager::SendClusterEventTo(const FString& Address, c
 {
 	if (CurrentOperationMode == EDisplayClusterOperationMode::Cluster || CurrentOperationMode == EDisplayClusterOperationMode::Editor)
 	{
-		if (ClusterNodeCtrl)
+		if (IsPrimary() || !bPrimaryOnly)
 		{
-			if (IsPrimary() || !bPrimaryOnly)
-			{
-				UE_LOG(LogDisplayClusterCluster, Verbose, TEXT("BIN event emission request: recipient=%s:%u, event=%d"), *Address, Port, Event.EventId);
-				ClusterNodeCtrl->SendClusterEventTo(Address, Port, Event, bPrimaryOnly);
-			}
+			UE_LOG(LogDisplayClusterCluster, Verbose, TEXT("BIN event emission request: recipient=%s:%u, event=%d"), *Address, Port, Event.EventId);
+			NodeCtrl->SendClusterEventTo(Address, Port, Event, bPrimaryOnly);
 		}
 	}
 }
@@ -463,50 +502,86 @@ void FDisplayClusterClusterManager::SendClusterEventTo(const FString& Address, c
 //////////////////////////////////////////////////////////////////////////////////////////////
 // IPDisplayClusterClusterManager
 //////////////////////////////////////////////////////////////////////////////////////////////
-IDisplayClusterClusterNodeController* FDisplayClusterClusterManager::GetClusterNodeController() const
+FDisplayClusterNetApiFacade& FDisplayClusterClusterManager::GetNetApi()
 {
-	FScopeLock Lock(&InternalsCS);
-	return ClusterNodeCtrl.Get();
+	return *NetApi.Get();
 }
 
-IDisplayClusterFailoverNodeController* FDisplayClusterClusterManager::GetFailoverNodeController() const
+TSharedRef<IDisplayClusterClusterNodeController> FDisplayClusterClusterManager::GetNodeController()
 {
-	FScopeLock Lock(&InternalsCS);
-	return FailoverNodeCtrl.Get();
+	return NodeCtrl;
 }
 
+TSharedRef<FDisplayClusterCommDataCache> FDisplayClusterClusterManager::GetDataCache()
+{
+	return FailoverCtrl->GetDataCache();
+}
 
-void FDisplayClusterClusterManager::SyncTimeData()
+TWeakPtr<FDisplayClusterService> FDisplayClusterClusterManager::GetNodeService(const FName& ServiceName)
 {
-	if (ClusterNodeCtrl)
+	return NodeCtrl->GetService(ServiceName);
+}
+
+bool FDisplayClusterClusterManager::DropNode(const FString& NodeId, IPDisplayClusterClusterManager::ENodeDropReason DropReason)
+{
+	// Ignore invalid requests
 	{
-		double DeltaTime = 0.0f;
-		double GameTime = 0.0f;
-		TOptional<FQualifiedFrameTime> FrameTime;
+		FScopeLock Lock(&ActiveClusterNodeIdsCS);
+
+		if (!ActiveClusterNodeIds.Contains(NodeId))
+		{
+			return false;
+		}
+	}
 
-		UE_LOG(LogDisplayClusterCluster, Verbose, TEXT("Downloading synchronization data (time)..."));
-		ClusterNodeCtrl->GetTimeData(DeltaTime, GameTime, FrameTime);
-		UE_LOG(LogDisplayClusterCluster, Verbose, TEXT("Downloading finished. Delta=%lf, Game=%lf, Frame=%lf"), DeltaTime, GameTime, FrameTime.IsSet() ? FrameTime.GetValue().AsSeconds() : 0.f);
+	UE_LOG(LogDisplayClusterCluster, Log, TEXT("Requested node '%s' drop, reason=%u"), *NodeId, DropReason);
 
-		// Apply new time data (including primary node)
-		ImportTimeData(DeltaTime, GameTime, FrameTime);
+	// User requests are sent to the desired nodes as "exit" like commands
+	if (DropReason == IPDisplayClusterClusterManager::ENodeDropReason::UserRequest)
+	{
+		FailoverCtrl->RequestNodeDrop(NodeId, static_cast<uint8>(DropReason));
+	}
+	// Other requests should go though failover pipeline
+	else if (DropReason == IPDisplayClusterClusterManager::ENodeDropReason::Failed)
+	{
+		HandleNodeDrop(NodeId);
 	}
+
+	return true;
 }
 
-void FDisplayClusterClusterManager::ExportTimeData(double& OutDeltaTime, double& OutGameTime, TOptional<FQualifiedFrameTime>& OutFrameTime)
+void FDisplayClusterClusterManager::CacheTimeData()
 {
-	if (IsInGameThread())
-	{
-		// Cache data so it will be the same for all requests within current frame
-		DeltaTimeCache = FApp::GetDeltaTime();
-		GameTimeCache  = FApp::GetGameTime();
-		FrameTimeCache = FApp::GetCurrentFrameTime();
+	// Cache data so it will be the same for all requests within current frame
+	DeltaTimeCache = FApp::GetDeltaTime();
+	GameTimeCache  = FApp::GetGameTime();
+	FrameTimeCache = FApp::GetCurrentFrameTime();
 
-		UE_LOG(LogDisplayClusterCluster, Verbose, TEXT("Time data cache: Delta=%lf, Game=%lf, Frame=%lf"), DeltaTimeCache, GameTimeCache, FrameTimeCache.IsSet() ? FrameTimeCache.GetValue().AsSeconds() : 0.f);
+	UE_LOG(LogDisplayClusterCluster, Verbose, TEXT("Time data cache: Delta=%lf, Game=%lf, Frame=%lf"),
+		DeltaTimeCache, GameTimeCache, FrameTimeCache.IsSet() ? FrameTimeCache.GetValue().AsSeconds() : 0.f);
 
-		TimeDataCacheReadySignal->Trigger();
-	}
+	TimeDataCacheReadySignal->Trigger();
+}
+
+void FDisplayClusterClusterManager::SyncTimeData()
+{
+	double DeltaTime = 0.0f;
+	double GameTime  = 0.0f;
+	TOptional<FQualifiedFrameTime> FrameTime;
 
+	UE_LOG(LogDisplayClusterCluster, Verbose, TEXT("Downloading synchronization data (time)..."));
+
+	GetNetApi().GetClusterSyncAPI()->GetTimeData(DeltaTime, GameTime, FrameTime);
+
+	UE_LOG(LogDisplayClusterCluster, Verbose, TEXT("Downloading finished. Delta=%lf, Game=%lf, Frame=%lf"),
+		DeltaTime, GameTime, FrameTime.IsSet() ? FrameTime.GetValue().AsSeconds() : 0.f);
+
+	// Apply new time data (including primary node)
+	ImportTimeData(DeltaTime, GameTime, FrameTime);
+}
+
+void FDisplayClusterClusterManager::ExportTimeData(double& OutDeltaTime, double& OutGameTime, TOptional<FQualifiedFrameTime>& OutFrameTime)
+{
 	// Wait until data is available
 	TimeDataCacheReadySignal->Wait();
 
@@ -520,7 +595,7 @@ void FDisplayClusterClusterManager::ImportTimeData(const double& InDeltaTime, co
 {
 	// Compute new 'current' and 'last' time on the local platform timeline
 	const double NewCurrentTime = FPlatformTime::Seconds();
-	const double NewLastTime = NewCurrentTime - InDeltaTime;
+	const double NewLastTime  = NewCurrentTime - InDeltaTime;
 
 	// Store new data
 	FApp::SetCurrentTime(NewLastTime);
@@ -545,62 +620,52 @@ void FDisplayClusterClusterManager::ImportTimeData(const double& InDeltaTime, co
 
 void FDisplayClusterClusterManager::SyncObjects(EDisplayClusterSyncGroup InSyncGroup)
 {
-	if (ClusterNodeCtrl)
-	{
-		TMap<FString, FString> ObjectsData;
+	TMap<FString, FString> ObjectsData;
 
-		UE_LOG(LogDisplayClusterCluster, Verbose, TEXT("Downloading synchronization data (objects)..."));
-		ClusterNodeCtrl->GetObjectsData(InSyncGroup, ObjectsData);
-		UE_LOG(LogDisplayClusterCluster, Verbose, TEXT("Downloading finished. Available %d records (objects)."), ObjectsData.Num());
+	UE_LOG(LogDisplayClusterCluster, Verbose, TEXT("Downloading synchronization data (objects)..."));
+	GetNetApi().GetClusterSyncAPI()->GetObjectsData(InSyncGroup, ObjectsData);
+	UE_LOG(LogDisplayClusterCluster, Verbose, TEXT("Downloading finished. Available %d records (objects)."), ObjectsData.Num());
 
-		// No need to import data on the primary node
-		if (!IsPrimary())
-		{
-			// Perform data load (objects state update)
-			ImportObjectsData(InSyncGroup, ObjectsData);
-		}
-	}
+	// Perform data load (objects state update)
+	ImportObjectsData(InSyncGroup, ObjectsData);
 }
 
-void FDisplayClusterClusterManager::ExportObjectsData(const EDisplayClusterSyncGroup InSyncGroup, TMap<FString, FString>& OutObjectsData)
+void FDisplayClusterClusterManager::CacheObjects(EDisplayClusterSyncGroup SyncGroup)
 {
-	// Clean output buffer just in case
-	OutObjectsData.Reset();
+	FScopeLock Lock(&ObjectsToSyncCS);
 
-	if(IsInGameThread())
+	// Cache data for requested sync group
+	if (TMap<FString, FString>* GroupCache = ObjectsToSyncCache.Find(SyncGroup))
 	{
-		FScopeLock Lock(&ObjectsToSyncCS);
+		UE_LOG(LogDisplayClusterCluster, Verbose, TEXT("Exporting sync data for sync group: %u, items to sync: %d"), (uint8)SyncGroup, GroupCache->Num());
 
-		// Cache data for requested sync group
-		if (TMap<FString, FString>* GroupCache = ObjectsToSyncCache.Find(InSyncGroup))
+		for (IDisplayClusterClusterSyncObject* SyncObj : ObjectsToSync[SyncGroup])
 		{
-			UE_LOG(LogDisplayClusterCluster, Verbose, TEXT("Exporting sync data for sync group: %u, items to sync: %d"), (uint8)InSyncGroup, GroupCache->Num());
-
-			for (IDisplayClusterClusterSyncObject* SyncObj : ObjectsToSync[InSyncGroup])
+			if (SyncObj && SyncObj->IsActive() && SyncObj->IsDirty())
 			{
-				if (SyncObj && SyncObj->IsActive() && SyncObj->IsDirty())
-				{
-					UE_LOG(LogDisplayClusterCluster, Verbose, TEXT("Adding object to sync: %s"), *SyncObj->GetSyncId());
+				UE_LOG(LogDisplayClusterCluster, Verbose, TEXT("Adding object to sync: %s"), *SyncObj->GetSyncId());
 
-					const FString SyncId   = SyncObj->GetSyncId();
-					const FString SyncData = SyncObj->SerializeToString();
+				const FString SyncId   = SyncObj->GetSyncId();
+				const FString SyncData = SyncObj->SerializeToString();
 
-					UE_LOG(LogDisplayClusterCluster, VeryVerbose, TEXT("Sync object: %s - %s"), *SyncId, *SyncData);
+				UE_LOG(LogDisplayClusterCluster, VeryVerbose, TEXT("Sync object: %s - %s"), *SyncId, *SyncData);
 
-					// Cache the object
-					GroupCache->Emplace(SyncId, SyncData);
+				// Cache the object
+				GroupCache->Emplace(SyncId, SyncData);
 
-					SyncObj->ClearDirty();
-				}
+				SyncObj->ClearDirty();
 			}
 		}
+	}
 
-		UE_LOG(LogDisplayClusterCluster, Verbose, TEXT("Objects data cache contains %d records"), ObjectsToSyncCache[InSyncGroup].Num());
+	UE_LOG(LogDisplayClusterCluster, Verbose, TEXT("Objects data cache contains %d records"), ObjectsToSyncCache[SyncGroup].Num());
 
-		// Notify data is available
-		ObjectsToSyncCacheReadySignals[InSyncGroup]->Trigger();
-	}
+	// Notify data is available
+	ObjectsToSyncCacheReadySignals[SyncGroup]->Trigger();
+}
 
+void FDisplayClusterClusterManager::ExportObjectsData(const EDisplayClusterSyncGroup InSyncGroup, TMap<FString, FString>& OutObjectsData)
+{
 	// Wait until primary node provides data
 	ObjectsToSyncCacheReadySignals[InSyncGroup]->Wait();
 	// Return cached value
@@ -644,73 +709,66 @@ void FDisplayClusterClusterManager::ImportObjectsData(const EDisplayClusterSyncG
 
 void FDisplayClusterClusterManager::SyncEvents()
 {
-	if (ClusterNodeCtrl)
-	{
-		TArray<TSharedPtr<FDisplayClusterClusterEventJson,   ESPMode::ThreadSafe>> JsonEvents;
-		TArray<TSharedPtr<FDisplayClusterClusterEventBinary, ESPMode::ThreadSafe>> BinaryEvents;
+	TArray<TSharedPtr<FDisplayClusterClusterEventJson>>   JsonEvents;
+	TArray<TSharedPtr<FDisplayClusterClusterEventBinary>> BinaryEvents;
 
-		// Get events data from a provider
-		UE_LOG(LogDisplayClusterCluster, Verbose, TEXT("Downloading synchronization data (events)..."));
-		ClusterNodeCtrl->GetEventsData(JsonEvents, BinaryEvents);
-		UE_LOG(LogDisplayClusterCluster, Verbose, TEXT("Downloading finished. Available events: json=%d binary=%d"), JsonEvents.Num(), BinaryEvents.Num());
+	// Get events data from a provider
+	UE_LOG(LogDisplayClusterCluster, Verbose, TEXT("Downloading synchronization data (events)..."));
+	GetNetApi().GetClusterSyncAPI()->GetEventsData(JsonEvents, BinaryEvents);
+	UE_LOG(LogDisplayClusterCluster, Verbose, TEXT("Downloading finished. Available events: json=%d binary=%d"), JsonEvents.Num(), BinaryEvents.Num());
 
-		// Import and process them
-		ImportEventsData(JsonEvents, BinaryEvents);
-	}
+	// Import and process them
+	ImportEventsData(JsonEvents, BinaryEvents);
 }
 
-void FDisplayClusterClusterManager::ExportEventsData(TArray<TSharedPtr<FDisplayClusterClusterEventJson, ESPMode::ThreadSafe>>& OutJsonEvents, TArray<TSharedPtr<FDisplayClusterClusterEventBinary, ESPMode::ThreadSafe>>& OutBinaryEvents)
+void FDisplayClusterClusterManager::CacheEvents()
 {
-	// Clear output containers just in case
-	OutJsonEvents.Reset();
-	OutBinaryEvents.Reset();
-
-	if (IsInGameThread())
+	// Export JSON events
 	{
-		// Export JSON events
-		{
-			FScopeLock Lock(&ClusterEventsJsonCS);
+		FScopeLock Lock(&ClusterEventsJsonCS);
 
-			// Export all system and non-system json events that have 'discard on repeat' flag
-			for (const TPair<bool, TMap<FString, TSharedPtr<FDisplayClusterClusterEventJson, ESPMode::ThreadSafe>>>& It : ClusterEventsJson)
-			{
-				TArray<TSharedPtr<FDisplayClusterClusterEventJson, ESPMode::ThreadSafe>> JsonEventsToExport;
-				It.Value.GenerateValueArray(JsonEventsToExport);
-				JsonEventsCache.Append(MoveTemp(JsonEventsToExport));
-			}
-
-			// Clear original containers
-			ClusterEventsJson.Reset();
-
-			// Export all json events that don't have 'discard on repeat' flag
-			JsonEventsCache.Append(MoveTemp(ClusterEventsJsonNonDiscarded));
+		// Export all system and non-system json events that have 'discard on repeat' flag
+		for (const TPair<bool, TMap<FString, TSharedPtr<FDisplayClusterClusterEventJson>>>& It : ClusterEventsJson)
+		{
+			TArray<TSharedPtr<FDisplayClusterClusterEventJson>> JsonEventsToExport;
+			It.Value.GenerateValueArray(JsonEventsToExport);
+			JsonEventsCache.Append(MoveTemp(JsonEventsToExport));
 		}
 
-		// Export binary events
-		{
-			FScopeLock Lock(&ClusterEventsBinaryCS);
+		// Clear original containers
+		ClusterEventsJson.Reset();
 
-			// Export all binary events that have 'discard on repeat' flag
-			for (const TPair<bool, TMap<int32, TSharedPtr<FDisplayClusterClusterEventBinary, ESPMode::ThreadSafe>>>& It : ClusterEventsBinary)
-			{
-				TArray<TSharedPtr<FDisplayClusterClusterEventBinary, ESPMode::ThreadSafe>> BinaryEventsToExport;
-				It.Value.GenerateValueArray(BinaryEventsToExport);
-				BinaryEventsCache.Append(MoveTemp(BinaryEventsToExport));
-			}
+		// Export all json events that don't have 'discard on repeat' flag
+		JsonEventsCache.Append(MoveTemp(ClusterEventsJsonNonDiscarded));
+	}
 
-			// Clear original containers
-			ClusterEventsBinary.Reset();
+	// Export binary events
+	{
+		FScopeLock Lock(&ClusterEventsBinaryCS);
 
-			// Export all binary events that don't have 'discard on repeat' flag
-			BinaryEventsCache.Append(MoveTemp(ClusterEventsBinaryNonDiscarded));
+		// Export all binary events that have 'discard on repeat' flag
+		for (const TPair<bool, TMap<int32, TSharedPtr<FDisplayClusterClusterEventBinary>>>& It : ClusterEventsBinary)
+		{
+			TArray<TSharedPtr<FDisplayClusterClusterEventBinary>> BinaryEventsToExport;
+			It.Value.GenerateValueArray(BinaryEventsToExport);
+			BinaryEventsCache.Append(MoveTemp(BinaryEventsToExport));
 		}
 
-		UE_LOG(LogDisplayClusterCluster, Verbose, TEXT("Cluster events data cache contains: json=%d, binary=%d"), JsonEventsCache.Num(), BinaryEventsCache.Num());
+		// Clear original containers
+		ClusterEventsBinary.Reset();
 
-		// Notify data is available
-		CachedEventsDataSignal->Trigger();
+		// Export all binary events that don't have 'discard on repeat' flag
+		BinaryEventsCache.Append(MoveTemp(ClusterEventsBinaryNonDiscarded));
 	}
 
+	UE_LOG(LogDisplayClusterCluster, Verbose, TEXT("Cluster events data cache contains: json=%d, binary=%d"), JsonEventsCache.Num(), BinaryEventsCache.Num());
+
+	// Notify data is available
+	CachedEventsDataSignal->Trigger();
+}
+
+void FDisplayClusterClusterManager::ExportEventsData(TArray<TSharedPtr<FDisplayClusterClusterEventJson>>& OutJsonEvents, TArray<TSharedPtr<FDisplayClusterClusterEventBinary>>& OutBinaryEvents)
+{
 	// Wait until data is available
 	CachedEventsDataSignal->Wait();
 
@@ -719,14 +777,14 @@ void FDisplayClusterClusterManager::ExportEventsData(TArray<TSharedPtr<FDisplayC
 	OutBinaryEvents = BinaryEventsCache;
 }
 
-void FDisplayClusterClusterManager::ImportEventsData(const TArray<TSharedPtr<FDisplayClusterClusterEventJson, ESPMode::ThreadSafe>>& InJsonEvents, const TArray<TSharedPtr<FDisplayClusterClusterEventBinary, ESPMode::ThreadSafe>>& InBinaryEvents)
+void FDisplayClusterClusterManager::ImportEventsData(const TArray<TSharedPtr<FDisplayClusterClusterEventJson>>& InJsonEvents, const TArray<TSharedPtr<FDisplayClusterClusterEventBinary>>& InBinaryEvents)
 {
 	// Process and fire all JSON events
 	if (InJsonEvents.Num() > 0)
 	{
 		FScopeLock LockListeners(&ClusterEventListenersCS);
 
-		for (const TSharedPtr<FDisplayClusterClusterEventJson, ESPMode::ThreadSafe>& Event : InJsonEvents)
+		for (const TSharedPtr<FDisplayClusterClusterEventJson>& Event : InJsonEvents)
 		{
 			UE_LOG(LogDisplayClusterCluster, Verbose, TEXT("Processing json event %s|%s|%s|s%d|d%d..."), *Event->Category, *Event->Type, *Event->Name, Event->bIsSystemEvent ? 1 : 0, Event->bShouldDiscardOnRepeat ? 1 : 0);
 			// Fire event
@@ -739,7 +797,7 @@ void FDisplayClusterClusterManager::ImportEventsData(const TArray<TSharedPtr<FDi
 	{
 		FScopeLock LockListeners(&ClusterEventListenersCS);
 
-		for (const TSharedPtr<FDisplayClusterClusterEventBinary, ESPMode::ThreadSafe>& Event : InBinaryEvents)
+		for (const TSharedPtr<FDisplayClusterClusterEventBinary>& Event : InBinaryEvents)
 		{
 			UE_LOG(LogDisplayClusterCluster, Verbose, TEXT("Processing binary event %d..."), Event->EventId);
 			// Fire event
@@ -748,20 +806,12 @@ void FDisplayClusterClusterManager::ImportEventsData(const TArray<TSharedPtr<FDi
 	}
 }
 
-void FDisplayClusterClusterManager::ImportNativeInputData(TMap<FString, FString>&& InNativeInputData)
+void FDisplayClusterClusterManager::ImportNativeInputData(TMap<FString, FString>& InNativeInputData)
 {
-	UE_LOG(LogDisplayClusterCluster, Verbose, TEXT("Caching native input data, %d records available"), InNativeInputData.Num());
-
 	// Cache input data
-	NativeInputCache = InNativeInputData;
-
-	UE_LOG(LogDisplayClusterCluster, VeryVerbose, TEXT("Native input data cache:"));
+	NativeInputCache = MoveTemp(InNativeInputData);
 
-	int32 Idx = 0;
-	for (auto It = NativeInputCache.CreateConstIterator(); It; ++It, ++Idx)
-	{
-		UE_LOG(LogDisplayClusterCluster, VeryVerbose, TEXT("Native input data cache: %d - %s - %s"), Idx, *It->Key, *It->Value);
-	}
+	UE_LOG(LogDisplayClusterCluster, VeryVerbose, TEXT("Native input data cache: %d items"), NativeInputCache.Num());
 
 	// Notify the data is available
 	NativeInputCacheReadySignal->Trigger();
@@ -769,85 +819,187 @@ void FDisplayClusterClusterManager::ImportNativeInputData(TMap<FString, FString>
 
 void FDisplayClusterClusterManager::ExportNativeInputData(TMap<FString, FString>& OutNativeInputData)
 {
-	// If we are on the primary node, that means we got network request for input data.
-	// So, provide with data once it's cached.
-	if (IsPrimary())
+	// Wait for data cache to be ready
+	NativeInputCacheReadySignal->Wait();
+	// Export data from cache
+	OutNativeInputData = NativeInputCache;
+}
+
+
+//////////////////////////////////////////////////////////////////////////////////////////////
+// FDisplayClusterClusterManager
+//////////////////////////////////////////////////////////////////////////////////////////////
+bool FDisplayClusterClusterManager::InitializeNetworking(const UDisplayClusterConfigurationData* ConfigData)
+{
+	// Instantiate cluster node controller
+	NodeCtrl = CreateClusterNodeController();
+
+	// Initialize the controller
+	if (!NodeCtrl->Initialize())
 	{
-		// Wait for data cache to be ready
-		NativeInputCacheReadySignal->Wait();
-		// Export data from cache
-		OutNativeInputData = NativeInputCache;
+		UE_LOG(LogDisplayClusterCluster, Error, TEXT("Couldn't initialize the networking controller."));
+		return false;
 	}
-	// Otherwise we need to send network request to download input data
-	else
-	{
-		if (ClusterNodeCtrl)
-		{
-			UE_LOG(LogDisplayClusterCluster, Verbose, TEXT("Downloading native input data..."));
-			ClusterNodeCtrl->GetNativeInputData(OutNativeInputData);
-			UE_LOG(LogDisplayClusterCluster, Verbose, TEXT("Downloading finished. Available %d native input records."), OutNativeInputData.Num());
 
-			int32 Idx = 0;
-			for (auto It = OutNativeInputData.CreateConstIterator(); It; ++It, ++Idx)
-			{
-				UE_LOG(LogDisplayClusterCluster, VeryVerbose, TEXT("Native input data: %d - %s - %s"), Idx, *It->Key, *It->Value);
-			}
-		}
+	// Instantiate failover controller
+	FailoverCtrl = CreateFailoverController(NodeCtrl);
+
+	// Initialize the controller
+	if (!FailoverCtrl->Initialize(ConfigData))
+	{
+		UE_LOG(LogDisplayClusterCluster, Error, TEXT("Couldn't initialize the failover controller."));
+		return false;
 	}
+
+	// Finally, setup API
+	NetApi = MakeUnique<FDisplayClusterNetApiFacade>(FailoverCtrl);
+
+	UE_LOG(LogDisplayClusterCluster, Log, TEXT("Networking internals have been successfully initialized."));
+
+	return true;
 }
 
+void FDisplayClusterClusterManager::ReleaseNetworking()
+{
+	// Stop local clients/servers
+	NodeCtrl->Shutdown();
 
-//////////////////////////////////////////////////////////////////////////////////////////////
-// FDisplayClusterClusterManager
-//////////////////////////////////////////////////////////////////////////////////////////////
-TUniquePtr<IDisplayClusterClusterNodeController> FDisplayClusterClusterManager::CreateClusterNodeController() const
+	// Reset controllers to their 'Disabled' state
+	NodeCtrl = MakeShared<FDisplayClusterClusterNodeCtrlDisabled>();
+	FailoverCtrl = MakeShared<FDisplayClusterFailoverNodeCtrlDisabled>(NodeCtrl);
+	NetApi = MakeUnique<FDisplayClusterNetApiFacade>(FailoverCtrl);
+}
+
+TSharedRef<IDisplayClusterClusterNodeController> FDisplayClusterClusterManager::CreateClusterNodeController() const
 {
-	// Instantiate appropriate controller depending on operation mode and cluster role
+	// Instantiate appropriate controller depending on the operation mode
 	if (CurrentOperationMode == EDisplayClusterOperationMode::Cluster)
 	{
-		// Primary or secondary
-		if (ClusterNodeId.Equals(GDisplayCluster->GetPrivateConfigMgr()->GetPrimaryNodeId(), ESearchCase::IgnoreCase))
-		{
-			UE_LOG(LogDisplayClusterCluster, Log, TEXT("Instantiating primary node controller..."));
-			return MakeUnique<FDisplayClusterClusterNodeCtrlPrimary>(FString("[CTRL-M]"), ClusterNodeId);
-		}
-		else
-		{
-			UE_LOG(LogDisplayClusterCluster, Log, TEXT("Instantiating secondary node controller..."));
-			return MakeUnique<FDisplayClusterClusterNodeCtrlSecondary>(FString("[CTRL-S]"), ClusterNodeId);
-		}
+		UE_LOG(LogDisplayClusterCluster, Log, TEXT("Instantiating 'Main' node controller..."));
+		return MakeShared<FDisplayClusterClusterNodeCtrlMain>(ClusterNodeId);
 	}
 	else if (CurrentOperationMode == EDisplayClusterOperationMode::Editor)
 	{
-		UE_LOG(LogDisplayClusterCluster, Log, TEXT("Instantiating editor node controller..."));
-		return MakeUnique<FDisplayClusterClusterNodeCtrlEditor>(FString("[CTRL-EDTR]"), FString("editor"));
+		UE_LOG(LogDisplayClusterCluster, Log, TEXT("Instantiating 'Editor' node controller..."));
+		return MakeShared<FDisplayClusterClusterNodeCtrlEditor>();
 	}
 
-	// Including CurrentOperationMode == EDisplayClusterOperationMode::Disabled
-	UE_LOG(LogDisplayClusterCluster, Log, TEXT("Controller is not required"));
-	return MakeUnique<FDisplayClusterClusterNodeCtrlNull>(FString("[CTRL-NULL]"), ClusterNodeId);
+	// Otherwise 'Disabled'
+	UE_LOG(LogDisplayClusterCluster, Log, TEXT("Instantiating 'Disabled' node controller..."));
+	return MakeShared<FDisplayClusterClusterNodeCtrlDisabled>();
 }
 
-TUniquePtr<IDisplayClusterFailoverNodeController> FDisplayClusterClusterManager::CreateFailoverNodeController() const
+TSharedRef<IDisplayClusterFailoverNodeController> FDisplayClusterClusterManager::CreateFailoverController(TSharedRef<IDisplayClusterClusterNodeController>& ClusterCtrl) const
 {
-	// Instantiate appropriate controller depending on operation mode and cluster role
+	// Instantiate appropriate controller depending on the operation mode
 	if (CurrentOperationMode == EDisplayClusterOperationMode::Cluster)
 	{
-		const EDisplayClusterNodeRole ClusterRole = GetClusterRole();
-		switch (ClusterRole)
+		UE_LOG(LogDisplayClusterCluster, Log, TEXT("Instantiating 'Main' failover controller..."));
+		return MakeShared<FDisplayClusterFailoverNodeCtrlMain>(ClusterCtrl);
+	}
+	else if (CurrentOperationMode == EDisplayClusterOperationMode::Editor)
+	{
+		UE_LOG(LogDisplayClusterCluster, Log, TEXT("Instantiating 'Editor' failover controller..."));
+		return MakeShared<FDisplayClusterFailoverNodeCtrlEditor>(ClusterCtrl);
+	}
+
+	// Otherwise 'Disabled'
+	UE_LOG(LogDisplayClusterCluster, Log, TEXT("Instantiating 'Disabled' failover controller..."));
+	return MakeShared<FDisplayClusterFailoverNodeCtrlDisabled>(ClusterCtrl);
+}
+
+void FDisplayClusterClusterManager::InitializeClusterRole(const FString& NodeId, const UDisplayClusterConfigurationData* ConfigData)
+{
+	checkSlow(ConfigData);
+
+	const bool bIsPrimary = NodeId.Equals(ConfigData->Cluster->PrimaryNode.Id, ESearchCase::IgnoreCase);
+	if (bIsPrimary)
+	{
+		SetClusterRole(EDisplayClusterNodeRole::Primary);
+		return;
+	}
+	
+	// Currently we don't completely support the backup nodes concept. So this
+	// part remains @todo. If it was supported, we would need to determine
+	// either it's 'secondary' or 'backup'.
+	SetClusterRole(EDisplayClusterNodeRole::Secondary);
+}
+
+void FDisplayClusterClusterManager::SetPrimaryNode(const FString& NewPrimaryNodeId)
+{
+	UE_LOG(LogDisplayClusterCluster, Verbose, TEXT("Requested new primary node: '%s'"), *NewPrimaryNodeId);
+
+	{
+		FScopeLock LockPrimary(&PrimaryNodeIdCS);
+
+		// Nothing to do if already set
+		if (PrimaryNodeId.Equals(NewPrimaryNodeId, ESearchCase::IgnoreCase))
 		{
-		case EDisplayClusterNodeRole::Primary:
-			UE_LOG(LogDisplayClusterCluster, Log, TEXT("Instantiating failover primary controller..."));
-			return MakeUnique<FDisplayClusterFailoverNodeCtrlPrimary>();
+			UE_LOG(LogDisplayClusterCluster, VeryVerbose, TEXT("'%s' is primary already"), *NewPrimaryNodeId);
+			return;
+		}
 
-		case EDisplayClusterNodeRole::Secondary:
-			UE_LOG(LogDisplayClusterCluster, Log, TEXT("Instantiating failover secondary controller..."));
-			return MakeUnique<FDisplayClusterFailoverNodeCtrlSecondary>();
+		// Check if new node is valid
+		{
+			FScopeLock LockActive(&ActiveClusterNodeIdsCS);
+			if (!ActiveClusterNodeIds.Contains(NewPrimaryNodeId))
+			{
+				UE_LOG(LogDisplayClusterCluster, VeryVerbose, TEXT("'%s' was not found in the list of active nodes"), *NewPrimaryNodeId);
+				return;
+			}
+		}
+
+		// Update current primary
+		PrimaryNodeId = NewPrimaryNodeId;
+
+		UE_LOG(LogDisplayClusterCluster, Log, TEXT("New primary node (P-node): '%s'"), *NewPrimaryNodeId);
+
+		// Update the role if we're the new primary.
+		const bool bThisNodeIsNowPrimary = NewPrimaryNodeId.Equals(ClusterNodeId, ESearchCase::IgnoreCase);
+		if (bThisNodeIsNowPrimary)
+		{
+			SetClusterRole(EDisplayClusterNodeRole::Primary);
 		}
 	}
+}
 
-	UE_LOG(LogDisplayClusterCluster, Log, TEXT("Instantiating failover null controller..."));
-	return MakeUnique<FDisplayClusterFailoverNodeCtrlNull>();
+void FDisplayClusterClusterManager::SetClusterRole(EDisplayClusterNodeRole NewRole)
+{
+	FScopeLock Lock(&CurrentNodeRoleCS);
+	UE_LOG(LogDisplayClusterCluster, Log, TEXT("New cluster role: '%u'"), static_cast<uint8>(NewRole));
+	CurrentNodeRole = NewRole;
+}
+
+void FDisplayClusterClusterManager::HandleNodeDrop(const FString& NodeId)
+{
+	// Remove this node from the list of active nodes
+	{
+		FScopeLock Lock(&ActiveClusterNodeIdsCS);
+		if (ActiveClusterNodeIds.Remove(NodeId) <= 0)
+		{
+			// This node has been processed already so nothing to do
+			return;
+		}
+	}
+
+	// Just exit if this node has failed
+	if (NodeId.Equals(GetNodeId(), ESearchCase::IgnoreCase))
+	{
+		FDisplayClusterAppExit::ExitApplication(TEXT("This node has failed. Requesting exit."));
+		return;
+	}
+
+	// Let the node controller drop it
+	NodeCtrl->DropClusterNode(NodeId);
+
+	// Let the failover controller process this
+	if (!FailoverCtrl->HandleFailure(NodeId))
+	{
+		FDisplayClusterAppExit::ExitApplication(TEXT("Failover controller was unable to handle a failure. Requesting exit."));
+	}
+
+	// Finally, broadcast node failed event
+	GDisplayCluster->GetCallbacks().OnDisplayClusterFailoverNodeDown().Broadcast(NodeId);
 }
 
 void FDisplayClusterClusterManager::OnClusterEventJsonHandler(const FDisplayClusterClusterEventJson& Event)
@@ -925,3 +1077,17 @@ void FDisplayClusterClusterManager::SetInternalSyncObjectsReleaseState(bool bRel
 		}
 	}
 }
+
+void FDisplayClusterClusterManager::OnPrimaryNodeChangedHandler(const FString& NewPrimaryId)
+{
+	SetPrimaryNode(NewPrimaryId);
+}
+
+void FDisplayClusterClusterManager::OnClusterNodeFailed(const FString& FailedNodeId)
+{
+	// Remove it from the active nodes list
+	{
+		FScopeLock Lock(&ActiveClusterNodeIdsCS);
+		ActiveClusterNodeIds.Remove(FailedNodeId);
+	}
+}
