[nDisplay] Base failover implementation (v1.5, P-node replacement)

diff --git a/Engine/Plugins/Runtime/nDisplay/Source/DisplayCluster/Private/Cluster/DisplayClusterClusterManager.h b/Engine/Plugins/Runtime/nDisplay/Source/DisplayCluster/Private/Cluster/DisplayClusterClusterManager.h
--- a/Engine/Plugins/Runtime/nDisplay/Source/DisplayCluster/Private/Cluster/DisplayClusterClusterManager.h
+++ b/Engine/Plugins/Runtime/nDisplay/Source/DisplayCluster/Private/Cluster/DisplayClusterClusterManager.h
@@ -4,11 +4,16 @@
 
 #include "Cluster/IPDisplayClusterClusterManager.h"
 #include "Cluster/DisplayClusterClusterEvent.h"
-#include "Misc/App.h"
 #include "HAL/Event.h"
+#include "Misc/Optional.h"
+#include "Misc/QualifiedFrameTime.h"
+
+class IDisplayClusterClusterNodeController;
+class IDisplayClusterFailoverNodeController;
+class FDisplayClusterNetApiFacade;
+class FDisplayClusterService;
+class UDisplayClusterConfigurationData;
 
-class ADisplayClusterSettings;
-class FJsonObject;
 
 /**
  * Cluster manager. Responsible for network communication and data replication.
@@ -17,13 +22,13 @@ class FDisplayClusterClusterManager
 	: public IPDisplayClusterClusterManager
 {
 public:
+
 	FDisplayClusterClusterManager();
 	virtual ~FDisplayClusterClusterManager();
 
 public:
-	//////////////////////////////////////////////////////////////////////////////////////////////
-	// IPDisplayClusterManager
-	//////////////////////////////////////////////////////////////////////////////////////////////
+
+	//~ Begin IPDisplayClusterManager
 	virtual bool Init(EDisplayClusterOperationMode OperationMode) override;
 	virtual void Release() override;
 	virtual bool StartSession(UDisplayClusterConfigurationData* InConfigData, const FString& NodeID) override;
@@ -35,34 +40,30 @@ public:
 	virtual void PreTick(float DeltaSeconds) override;
 	virtual void Tick(float DeltaSeconds) override;
 	virtual void PostTick(float DeltaSeconds) override;
+	//~ End IPDisplayClusterManager
 
 public:
-	//////////////////////////////////////////////////////////////////////////////////////////////
-	// IDisplayClusterClusterManager
-	//////////////////////////////////////////////////////////////////////////////////////////////
+
+	//~ Begin IDisplayClusterClusterManager
 	virtual bool IsPrimary()   const override;
 	virtual bool IsSecondary() const override;
 	virtual bool IsBackup()    const override;
 	virtual EDisplayClusterNodeRole GetClusterRole() const override;
+	virtual bool HasClusterRole(EDisplayClusterNodeRole Role) const override;
 
-	virtual FString GetNodeId() const override
-	{
-		return ClusterNodeId;
-	}
-
-	virtual uint32 GetNodesAmount() const override
-	{
-		return static_cast<uint32>(ClusterNodeIds.Num());
-	}
+	virtual FString GetPrimaryNodeId() const override;
 
+	virtual FString GetNodeId() const override;
+	virtual uint32 GetNodesAmount() const override;
 	virtual void GetNodeIds(TArray<FString>& OutNodeIds) const override;
+	virtual void GetNodeIds(TSet<FString>& OutNodeIds) const override;
 
 	virtual bool DropClusterNode(const FString& NodeId) override;
 
 	virtual void RegisterSyncObject(IDisplayClusterClusterSyncObject* SyncObj, EDisplayClusterSyncGroup SyncGroup) override;
 	virtual void UnregisterSyncObject(IDisplayClusterClusterSyncObject* SyncObj) override;
 
-	virtual TUniquePtr<IDisplayClusterGenericBarriersClient, FDisplayClusterGenericBarriersClientDeleter> CreateGenericBarriersClient(const FString& ClientName) override;
+	virtual TSharedRef<IDisplayClusterGenericBarriersClient> CreateGenericBarriersClient() override;
 
 	virtual void AddClusterEventListener(TScriptInterface<IDisplayClusterClusterEventListener>) override;
 	virtual void RemoveClusterEventListener(TScriptInterface<IDisplayClusterClusterEventListener>) override;
@@ -78,63 +79,133 @@ public:
 
 	virtual void SendClusterEventTo(const FString& Address, const uint16 Port, const FDisplayClusterClusterEventJson&   Event, bool bPrimaryOnly) override;
 	virtual void SendClusterEventTo(const FString& Address, const uint16 Port, const FDisplayClusterClusterEventBinary& Event, bool bPrimaryOnly) override;
+	//~ End IDisplayClusterClusterManager
 
 public:
-	//////////////////////////////////////////////////////////////////////////////////////////////
-	// IPDisplayClusterClusterManager
-	//////////////////////////////////////////////////////////////////////////////////////////////
-	virtual IDisplayClusterClusterNodeController*  GetClusterNodeController()  const override;
-	virtual IDisplayClusterFailoverNodeController* GetFailoverNodeController() const override;
+
+	//~ Begin IPDisplayClusterClusterManager
+
+	virtual FDisplayClusterNetApiFacade& GetNetApi() override;
+	virtual TSharedRef<IDisplayClusterClusterNodeController> GetNodeController() override;
+	virtual TSharedRef<FDisplayClusterCommDataCache> GetDataCache() override;
+	virtual TWeakPtr<FDisplayClusterService> GetNodeService(const FName& ServiceName) override;
+
+	virtual bool DropNode(const FString& NodeId, ENodeDropReason DropReason) override;
 
 	// Sync time data
 	virtual void SyncTimeData() override;
+	virtual void CacheTimeData() override;
 	virtual void ExportTimeData(      double& OutDeltaTime,      double& OutGameTime,      TOptional<FQualifiedFrameTime>& OutFrameTime) override;
 	virtual void ImportTimeData(const double& InDeltaTime, const double& InGameTime, const TOptional<FQualifiedFrameTime>& InFrameTime) override;
 
 	// Sync objects
 	virtual void SyncObjects(EDisplayClusterSyncGroup SyncGroup) override;
+	virtual void CacheObjects(EDisplayClusterSyncGroup SyncGroup) override;
 	virtual void ExportObjectsData(const EDisplayClusterSyncGroup InSyncGroup,       TMap<FString, FString>& OutObjectsData) override;
 	virtual void ImportObjectsData(const EDisplayClusterSyncGroup InSyncGroup, const TMap<FString, FString>& InObjectsData) override;
 
 	// Sync cluster events
 	virtual void SyncEvents()  override;
-	virtual void ExportEventsData(      TArray<TSharedPtr<FDisplayClusterClusterEventJson, ESPMode::ThreadSafe>>& OutJsonEvents,      TArray<TSharedPtr<FDisplayClusterClusterEventBinary, ESPMode::ThreadSafe>>& OutBinaryEvents) override;
-	virtual void ImportEventsData(const TArray<TSharedPtr<FDisplayClusterClusterEventJson, ESPMode::ThreadSafe>>& InJsonEvents, const TArray<TSharedPtr<FDisplayClusterClusterEventBinary, ESPMode::ThreadSafe>>& InBinaryEvents) override;
+	virtual void CacheEvents() override;
+	virtual void ExportEventsData(      TArray<TSharedPtr<FDisplayClusterClusterEventJson>>& OutJsonEvents,      TArray<TSharedPtr<FDisplayClusterClusterEventBinary>>& OutBinaryEvents) override;
+	virtual void ImportEventsData(const TArray<TSharedPtr<FDisplayClusterClusterEventJson>>& InJsonEvents, const TArray<TSharedPtr<FDisplayClusterClusterEventBinary>>& InBinaryEvents) override;
 
 	// Sync native input
-	virtual void ExportNativeInputData(TMap<FString, FString>&  OutNativeInputData) override;
-	virtual void ImportNativeInputData(TMap<FString, FString>&& InNativeInputData) override;
+	virtual void ExportNativeInputData(TMap<FString, FString>& OutNativeInputData) override;
+	virtual void ImportNativeInputData(TMap<FString, FString>& InNativeInputData) override;
+
+	//~ End IPDisplayClusterClusterManager
+
+private:
+
+	/** Performs initialization of networking subsystems */
+	bool InitializeNetworking(const UDisplayClusterConfigurationData* ConfigData);
+
+	/** Release networking subsystems */
+	void ReleaseNetworking();
+
+	/** Factory method to instantiate a networking controller */
+	TSharedRef<IDisplayClusterClusterNodeController> CreateClusterNodeController() const;
+
+	/** Factory method to instantiate a failover controller */
+	TSharedRef<IDisplayClusterFailoverNodeController> CreateFailoverController(TSharedRef<IDisplayClusterClusterNodeController>& ClusterCtrl) const;
+
+	/** Auxiliary method for group set/reset of internal signals */
+	void SetInternalSyncObjectsReleaseState(bool bRelease);
+
+	/** Determines cluster role on session start. */
+	void InitializeClusterRole(const FString& NodeId, const UDisplayClusterConfigurationData* ConfigData);
+
+	/** Changes current primary node ID. It doesn't do any actual role transition, just the ID variable. */
+	void SetPrimaryNode(const FString& NewPrimaryNodeId);
+
+	/** Changes current node role. */
+	void SetClusterRole(EDisplayClusterNodeRole NewRole);
+
+	/** Performs node drop internal cleaning. */
+	void HandleNodeDrop(const FString& NodeId);
 
 private:
-	// Factory methods (never return nullptr)
-	TUniquePtr<IDisplayClusterClusterNodeController>  CreateClusterNodeController() const;
-	TUniquePtr<IDisplayClusterFailoverNodeController> CreateFailoverNodeController() const;
 
-	// Cluster event handlers
+	/** JSON cluster event handler */
 	void OnClusterEventJsonHandler(const FDisplayClusterClusterEventJson& Event);
+
+	/** Binary cluster event handler */
 	void OnClusterEventBinaryHandler(const FDisplayClusterClusterEventBinary& Event);
 
-	// Internal sync objects control
-	void SetInternalSyncObjectsReleaseState(bool bRelease);
+	/** Handles primary node change events */
+	void OnPrimaryNodeChangedHandler(const FString& NewPrimaryId);
+
+	/** Handles node failure events */
+	void OnClusterNodeFailed(const FString& FailedNodeId);
 
 private:
-	// Cluster controller
-	TUniquePtr<IDisplayClusterClusterNodeController>  ClusterNodeCtrl;
-	// Failover controller
-	TUniquePtr<IDisplayClusterFailoverNodeController> FailoverNodeCtrl;
 
-	// Custom list of cluster nodes
-	TArray<FString> ClusterNodeIds;
+	/** Networking controller */
+	TSharedRef<IDisplayClusterClusterNodeController> NodeCtrl;
+
+	/** Failover controller */
+	TSharedRef<IDisplayClusterFailoverNodeController> FailoverCtrl;
+
+	/** Networking API */
+	TUniquePtr<FDisplayClusterNetApiFacade> NetApi;
+
+private:
+
+	/** Current operation mode */
+	EDisplayClusterOperationMode CurrentOperationMode = EDisplayClusterOperationMode::Disabled;
+
+
+	/** Current primary node. It may change in runtime after failure handling. */
+	FString PrimaryNodeId;
+
+	/** Safe access to the primary node ID */
+	mutable FCriticalSection PrimaryNodeIdCS;
+
 
-	// Current operation mode
-	EDisplayClusterOperationMode CurrentOperationMode;
-	// Current node ID
+	/** Current role in the cluster */
+	EDisplayClusterNodeRole CurrentNodeRole = EDisplayClusterNodeRole::None;
+
+	/** Critical section to guard current role access */
+	mutable FCriticalSection CurrentNodeRoleCS;
+
+
+	/** This node ID */
 	FString ClusterNodeId;
-	// Current world
-	UWorld* CurrentWorld;
 
+	/** A full set of cluster node IDs used on cluster start. */
+	TSet<FString> InitialClusterNodeIds;
+
+	/** A subset of cluster node IDs that are currently active */
+	TSet<FString> ActiveClusterNodeIds;
+
+	/** A critical section to guard access to the set of active nodes */
+	mutable FCriticalSection ActiveClusterNodeIdsCS;
+
+
+	// Current world
+	UWorld* CurrentWorld = nullptr;
 
-	// Time data - replication
 	FEventRef TimeDataCacheReadySignal{ EEventMode::ManualReset };
 	double DeltaTimeCache = 0.f;
 	double GameTimeCache  = 0.f;
@@ -152,25 +223,23 @@ private:
 	TMap<FString, FString> NativeInputCache;
 
 	// JSON events
-	TMap<bool, TMap<FString, TSharedPtr<FDisplayClusterClusterEventJson, ESPMode::ThreadSafe>>> ClusterEventsJson;
-	TArray<TSharedPtr<FDisplayClusterClusterEventJson, ESPMode::ThreadSafe>> ClusterEventsJsonNonDiscarded;
-	mutable FCriticalSection     ClusterEventsJsonCS;
-	FOnClusterEventJson          OnClusterEventJson;
+	TMap<bool, TMap<FString, TSharedPtr<FDisplayClusterClusterEventJson>>> ClusterEventsJson;
+	TArray<TSharedPtr<FDisplayClusterClusterEventJson>> ClusterEventsJsonNonDiscarded;
+	mutable FCriticalSection ClusterEventsJsonCS;
+	FOnClusterEventJson OnClusterEventJson;
 	
 	// Binary events
-	TMap<bool, TMap<int32, TSharedPtr<FDisplayClusterClusterEventBinary, ESPMode::ThreadSafe>>> ClusterEventsBinary;
-	TArray<TSharedPtr<FDisplayClusterClusterEventBinary, ESPMode::ThreadSafe>> ClusterEventsBinaryNonDiscarded;
-	mutable FCriticalSection   ClusterEventsBinaryCS;
-	FOnClusterEventBinary      OnClusterEventBinary;
+	TMap<bool, TMap<int32, TSharedPtr<FDisplayClusterClusterEventBinary>>> ClusterEventsBinary;
+	TArray<TSharedPtr<FDisplayClusterClusterEventBinary>> ClusterEventsBinaryNonDiscarded;
+	mutable FCriticalSection ClusterEventsBinaryCS;
+	FOnClusterEventBinary OnClusterEventBinary;
 
 	// JSON/Binary events - replication
 	FEventRef CachedEventsDataSignal{ EEventMode::ManualReset };
-	TArray<TSharedPtr<FDisplayClusterClusterEventJson,   ESPMode::ThreadSafe>> JsonEventsCache;
-	TArray<TSharedPtr<FDisplayClusterClusterEventBinary, ESPMode::ThreadSafe>> BinaryEventsCache;
+	TArray<TSharedPtr<FDisplayClusterClusterEventJson>>   JsonEventsCache;
+	TArray<TSharedPtr<FDisplayClusterClusterEventBinary>> BinaryEventsCache;
 
 	// Cluster event listeners
 	mutable FCriticalSection ClusterEventListenersCS;
 	TArray<TScriptInterface<IDisplayClusterClusterEventListener>> ClusterEventListeners;
-
-	mutable FCriticalSection InternalsCS;
 };
diff --git a/Engine/Plugins/Runtime/nDisplay/Source/DisplayCluster/Private/Cluster/DisplayClusterNetDriverHelper.cpp b/Engine/Plugins/Runtime/nDisplay/Source/DisplayCluster/Private/Cluster/DisplayClusterNetDriverHelper.cpp
--- a/Engine/Plugins/Runtime/nDisplay/Source/DisplayCluster/Private/Cluster/DisplayClusterNetDriverHelper.cpp
+++ b/Engine/Plugins/Runtime/nDisplay/Source/DisplayCluster/Private/Cluster/DisplayClusterNetDriverHelper.cpp
@@ -79,7 +79,7 @@ bool FDisplayClusterNetDriverHelper::RegisterClusterEventsBinaryClient(uint32 Cl
 		return false;
 	}
 	
-	TSharedPtr<FDisplayClusterClusterEventsBinaryClient> Client(new FDisplayClusterClusterEventsBinaryClient(), FDisplayClusterClientDeleter());
+	TSharedPtr<FDisplayClusterClusterEventsBinaryClient> Client = MakeShared<FDisplayClusterClusterEventsBinaryClient>(TEXT("CLN_CEB_NetDrv"));
 	
 	if (!Client->Connect(ClientAddress, ClientPort, 1, 0.f))
 	{
diff --git a/Engine/Plugins/Runtime/nDisplay/Source/DisplayCluster/Private/Cluster/Failover/DisplayClusterFailoverNodeCtrlBase.cpp b/Engine/Plugins/Runtime/nDisplay/Source/DisplayCluster/Private/Cluster/Failover/DisplayClusterFailoverNodeCtrlBase.cpp
--- a/Engine/Plugins/Runtime/nDisplay/Source/DisplayCluster/Private/Cluster/Failover/DisplayClusterFailoverNodeCtrlBase.cpp
+++ b/Engine/Plugins/Runtime/nDisplay/Source/DisplayCluster/Private/Cluster/Failover/DisplayClusterFailoverNodeCtrlBase.cpp
@@ -2,66 +2,8 @@
 
 #include "Cluster/Failover/DisplayClusterFailoverNodeCtrlBase.h"
 
-#include "DisplayClusterConfigurationTypes.h"
-#include "Config/IPDisplayClusterConfigManager.h"
 
-#include "Misc/DisplayClusterGlobals.h"
-#include "Misc/DisplayClusterLog.h"
-#include "Misc/ScopeLock.h"
-
-#include "Network/DisplayClusterNetworkTypes.h"
-
-#include "HAL/IConsoleManager.h"
-
-
-FDisplayClusterFailoverNodeCtrlBase::FDisplayClusterFailoverNodeCtrlBase()
-	: FailoverPolicy( FDisplayClusterFailoverNodeCtrlBase::GetFailoverPolicyFromConfig() )
+FDisplayClusterFailoverNodeCtrlBase::FDisplayClusterFailoverNodeCtrlBase(TSharedRef<IDisplayClusterClusterNodeController>& InNodeController)
+	: NodeController(InNodeController)
 {
-	UE_LOG(LogDisplayClusterCluster, Log, TEXT("Failover: policy %u"), static_cast<uint8>(FailoverPolicy));
-}
-
-
-void FDisplayClusterFailoverNodeCtrlBase::HandleCommResult(EDisplayClusterCommResult CommResult)
-{
-	UE_LOG(LogDisplayClusterCluster, Verbose, TEXT("Failover: CommResult - %u"), (uint32)CommResult);
-
-	FScopeLock Lock(&InternalsCS);
-
-	switch (GetFailoverPolicy())
-	{
-	case EDisplayClusterConfigurationFailoverPolicy::Disabled:
-		return HandleCommResult_Disabled(CommResult);
-
-	case EDisplayClusterConfigurationFailoverPolicy::DropSecondaryNodesOnly:
-		return HandleCommResult_DropSecondaryNodesOnly(CommResult);
-	}
-}
-
-void FDisplayClusterFailoverNodeCtrlBase::HandleNodeFailed(const FString& NodeId, IDisplayClusterServer::ENodeFailType NodeFailType)
-{
-	UE_LOG(LogDisplayClusterCluster, Log, TEXT("Failover: Cluster node failed - node=%s, fail_type=%u"), *NodeId, static_cast<uint32>(NodeFailType));
-
-	FScopeLock Lock(&InternalsCS);
-
-	switch (GetFailoverPolicy())
-	{
-	case EDisplayClusterConfigurationFailoverPolicy::Disabled:
-		return HandleNodeFailed_Disabled(NodeId, NodeFailType);
-
-	case EDisplayClusterConfigurationFailoverPolicy::DropSecondaryNodesOnly:
-		return HandleNodeFailed_DropSecondaryNodesOnly(NodeId, NodeFailType);
-	}
-}
-
-EDisplayClusterConfigurationFailoverPolicy FDisplayClusterFailoverNodeCtrlBase::GetFailoverPolicyFromConfig()
-{
-	if (const IPDisplayClusterConfigManager* const ConfigMgr = GDisplayCluster->GetPrivateConfigMgr())
-	{
-		if (const UDisplayClusterConfigurationData* const ConfigData = ConfigMgr->GetConfig())
-		{
-			return ConfigData->Cluster->Failover.FailoverPolicy;
-		}
-	}
-
-	return EDisplayClusterConfigurationFailoverPolicy::Disabled;
 }
diff --git a/Engine/Plugins/Runtime/nDisplay/Source/DisplayCluster/Private/Cluster/Failover/DisplayClusterFailoverNodeCtrlBase.h b/Engine/Plugins/Runtime/nDisplay/Source/DisplayCluster/Private/Cluster/Failover/DisplayClusterFailoverNodeCtrlBase.h
--- a/Engine/Plugins/Runtime/nDisplay/Source/DisplayCluster/Private/Cluster/Failover/DisplayClusterFailoverNodeCtrlBase.h
+++ b/Engine/Plugins/Runtime/nDisplay/Source/DisplayCluster/Private/Cluster/Failover/DisplayClusterFailoverNodeCtrlBase.h
@@ -5,6 +5,7 @@
 #include "Cluster/Failover/IDisplayClusterFailoverNodeController.h"
 #include "DisplayClusterConfigurationTypes_Enums.h"
 
+class IDisplayClusterClusterNodeController;
 struct FDisplayClusterSessionInfo;
 
 
@@ -15,37 +16,19 @@ class FDisplayClusterFailoverNodeCtrlBase
 	: public IDisplayClusterFailoverNodeController
 {
 public:
-	FDisplayClusterFailoverNodeCtrlBase();
 
-public:
-	virtual EDisplayClusterConfigurationFailoverPolicy GetFailoverPolicy() const override final
-	{
-		return FailoverPolicy;
-	}
-
-public:
-	// Handles in-cluster communication results
-	virtual void HandleCommResult(EDisplayClusterCommResult CommResult) override final;
-
-	// Handles node failure
-	virtual void HandleNodeFailed(const FString& NodeId, IDisplayClusterServer::ENodeFailType NodeFailType) override final;
+	FDisplayClusterFailoverNodeCtrlBase(TSharedRef<IDisplayClusterClusterNodeController>& InNodeController);
 
 protected:
-	// EDisplayClusterConfigurationFailoverPolicy::Disabled
-	virtual void HandleCommResult_Disabled(EDisplayClusterCommResult CommResult) = 0;
-	virtual void HandleNodeFailed_Disabled(const FString& NodeId, IDisplayClusterServer::ENodeFailType NodeFailType) = 0;
-
-	// EDisplayClusterConfigurationFailoverPolicy::DropSecondaryNodesOnly
-	virtual void HandleCommResult_DropSecondaryNodesOnly(EDisplayClusterCommResult CommResult) = 0;
-	virtual void HandleNodeFailed_DropSecondaryNodesOnly(const FString& NodeId, IDisplayClusterServer::ENodeFailType NodeFailType) = 0;
 
-private:
-	// Converts failover policy from CVar integer into enum
-	static EDisplayClusterConfigurationFailoverPolicy GetFailoverPolicyFromConfig();
+	/** Access to the active node controller */
+	const TSharedRef<IDisplayClusterClusterNodeController>& GetNodeController() const
+	{
+		return NodeController;
+	}
 
 private:
-	// Current failover policy
-	const EDisplayClusterConfigurationFailoverPolicy FailoverPolicy;
 
-	mutable FCriticalSection InternalsCS;
+	/** Active node controller */
+	TSharedRef<IDisplayClusterClusterNodeController> NodeController;
 };
diff --git a/Engine/Plugins/Runtime/nDisplay/Source/DisplayCluster/Private/Cluster/Failover/IDisplayClusterFailoverNodeController.h b/Engine/Plugins/Runtime/nDisplay/Source/DisplayCluster/Private/Cluster/Failover/IDisplayClusterFailoverNodeController.h
--- a/Engine/Plugins/Runtime/nDisplay/Source/DisplayCluster/Private/Cluster/Failover/IDisplayClusterFailoverNodeController.h
+++ b/Engine/Plugins/Runtime/nDisplay/Source/DisplayCluster/Private/Cluster/Failover/IDisplayClusterFailoverNodeController.h
@@ -2,28 +2,40 @@
 
 #pragma once
 
-#include "CoreMinimal.h"
-#include "Network/DisplayClusterNetworkTypes.h"
-#include "Network/IDisplayClusterServer.h"
-#include "DisplayClusterConfigurationTypes_Enums.h"
+#include "Network/Protocol/IDisplayClusterProtocolClusterSync.h"
+#include "Network/Protocol/IDisplayClusterProtocolEventsBinary.h"
+#include "Network/Protocol/IDisplayClusterProtocolEventsJson.h"
+#include "Network/Protocol/IDisplayClusterProtocolGenericBarrier.h"
+#include "Network/Protocol/IDisplayClusterProtocolInternalComm.h"
+#include "Network/Protocol/IDisplayClusterProtocolRenderSync.h"
+
+class FDisplayClusterCommDataCache;
+class UDisplayClusterConfigurationData;
 
 
 /**
  * Failover controller interface
  */
 class IDisplayClusterFailoverNodeController
+	: public IDisplayClusterProtocolClusterSync
+	, public IDisplayClusterProtocolEventsBinary
+	, public IDisplayClusterProtocolEventsJson
+	, public IDisplayClusterProtocolGenericBarrier
+	, public IDisplayClusterProtocolInternalComm
+	, public IDisplayClusterProtocolRenderSync
 {
 public:
+
 	virtual ~IDisplayClusterFailoverNodeController() = default;
 
 public:
-	// Returns current failover policy
-	virtual EDisplayClusterConfigurationFailoverPolicy GetFailoverPolicy() const = 0;
 
-public:
-	// Handles in-cluster communication results
-	virtual void HandleCommResult(EDisplayClusterCommResult CommResult) = 0;
+	/** Initialize the controller */
+	virtual bool Initialize(const UDisplayClusterConfigurationData* ConfigData) = 0;
+
+	/** Provides access to the communication data cache */
+	virtual TSharedRef<FDisplayClusterCommDataCache> GetDataCache() = 0;
 
-	// Handles node failure
-	virtual void HandleNodeFailed(const FString& NodeId, IDisplayClusterServer::ENodeFailType NodeFailType) = 0;
+	/** Process node failure */
+	virtual bool HandleFailure(const FString& FailedNodeId) = 0;
 };
diff --git a/Engine/Plugins/Runtime/nDisplay/Source/DisplayCluster/Private/DisplayClusterCallbacks.h b/Engine/Plugins/Runtime/nDisplay/Source/DisplayCluster/Private/DisplayClusterCallbacks.h
--- a/Engine/Plugins/Runtime/nDisplay/Source/DisplayCluster/Private/DisplayClusterCallbacks.h
+++ b/Engine/Plugins/Runtime/nDisplay/Source/DisplayCluster/Private/DisplayClusterCallbacks.h
@@ -87,6 +87,11 @@ public:
 		return DisplayClusterFailoverNodeDownEvent;
 	}
 
+	virtual FDisplayClusterFailoverPrimaryNodeChanged& OnDisplayClusterFailoverPrimaryNodeChanged() override
+	{
+		return DisplayClusterFailoverPrimaryNodeChangedEvent;
+	}
+
 	virtual FDisplayClusterPostTonemapPass_RenderThread& OnDisplayClusterPostTonemapPass_RenderThread() override
 	{
 		return DisplayClusterPostTonemapPassEvent;
@@ -176,6 +181,7 @@ private:
 	FDisplayClusterCustomPresentSetEvent     DisplayClusterCustomPresentSetEvent;
 	FDisplayClusterUpdateViewportMediaState  DisplayClusterUpdateViewportMediaStateEvent;
 	FDisplayClusterFailoverNodeDown          DisplayClusterFailoverNodeDownEvent;
+	FDisplayClusterFailoverPrimaryNodeChanged DisplayClusterFailoverPrimaryNodeChangedEvent;
 
 	FDisplayClusterPresentationPreSynchronization_RHIThread  DisplayClusterPresentationPreSynchronizationEvent;
 	FDisplayClusterPresentationPostSynchronization_RHIThread DisplayClusterPresentationPostSynchronizationEvent;
diff --git a/Engine/Plugins/Runtime/nDisplay/Source/DisplayCluster/Private/Game/DisplayClusterGameManager.h b/Engine/Plugins/Runtime/nDisplay/Source/DisplayCluster/Private/Game/DisplayClusterGameManager.h
--- a/Engine/Plugins/Runtime/nDisplay/Source/DisplayCluster/Private/Game/DisplayClusterGameManager.h
+++ b/Engine/Plugins/Runtime/nDisplay/Source/DisplayCluster/Private/Game/DisplayClusterGameManager.h
@@ -46,13 +46,6 @@ public:
 	virtual UWorld* GetWorld() const override
 	{ return CurrentWorld; }
 
-public:
-	//////////////////////////////////////////////////////////////////////////////////////////////
-	// IPDisplayClusterGameManager
-	//////////////////////////////////////////////////////////////////////////////////////////////
-	virtual bool IsDisplayClusterActive() const override
-	{ return CurrentOperationMode != EDisplayClusterOperationMode::Disabled; }
-
 private:
 	// The main function to look for a proper DCRA in a level and its sublevels
 	ADisplayClusterRootActor* FindRootActor(UWorld* World, UDisplayClusterConfigurationData* ConfigData);
diff --git a/Engine/Plugins/Runtime/nDisplay/Source/DisplayCluster/Private/Game/EngineClasses/Basics/DisplayClusterGameEngine.cpp b/Engine/Plugins/Runtime/nDisplay/Source/DisplayCluster/Private/Game/EngineClasses/Basics/DisplayClusterGameEngine.cpp
--- a/Engine/Plugins/Runtime/nDisplay/Source/DisplayCluster/Private/Game/EngineClasses/Basics/DisplayClusterGameEngine.cpp
+++ b/Engine/Plugins/Runtime/nDisplay/Source/DisplayCluster/Private/Game/EngineClasses/Basics/DisplayClusterGameEngine.cpp
@@ -5,6 +5,7 @@
 #include "Algo/Accumulate.h"
 #include "Cluster/IPDisplayClusterClusterManager.h"
 #include "Cluster/Controller/IDisplayClusterClusterNodeController.h"
+#include "Cluster/NetAPI/DisplayClusterNetApiFacade.h"
 #include "Config/IPDisplayClusterConfigManager.h"
 #include "DisplayClusterEnums.h"
 #include "Engine/DynamicBlueprintBinding.h"
@@ -24,11 +25,6 @@
 #include "Misc/DisplayClusterStrings.h"
 #include "Misc/Parse.h"
 
-#include "Misc/DisplayClusterGlobals.h"
-#include "Misc/DisplayClusterHelpers.h"
-#include "Misc/DisplayClusterLog.h"
-#include "Misc/DisplayClusterStrings.h"
-
 #include "SocketSubsystem.h"
 #include "Interfaces/IPv4/IPv4Endpoint.h"
 
@@ -44,7 +40,6 @@
 
 #include "GameDelegates.h"
 #include "Kismet/GameplayStatics.h"
-#include "Misc/CoreDelegates.h"
 
 
 namespace DisplayClusterGameEngineUtils
@@ -184,11 +179,8 @@ bool UDisplayClusterGameEngine::InitializeInternals()
 	// Store diagnostics settings locally
 	Diagnostics = Config->Diagnostics;
 
-	ClusterMgr     = GDisplayCluster->GetPrivateClusterMgr();
-	checkSlow(ClusterMgr);
-
 	FOnClusterEventJsonListener GameSyncTransition = FOnClusterEventJsonListener::CreateUObject(this, &UDisplayClusterGameEngine::GameSyncChange);
-	ClusterMgr->AddClusterEventJsonListener(GameSyncTransition);
+	GDisplayCluster->GetPrivateClusterMgr()->AddClusterEventJsonListener(GameSyncTransition);
 
 	const UDisplayClusterConfigurationClusterNode* CfgLocalNode = GDisplayCluster->GetPrivateConfigMgr()->GetLocalNode();
 	const bool bSoundEnabled = (CfgLocalNode ? CfgLocalNode->bIsSoundEnabled : false);
@@ -314,13 +306,16 @@ bool UDisplayClusterGameEngine::LoadMap(FWorldContext& WorldContext, FURL URL, c
 
 		if(BarrierAvoidanceOn() && RunningMode != EDisplayClusterRunningMode::Startup)
 		{
+			IPDisplayClusterClusterManager& ClusterMgr = *GDisplayCluster->GetPrivateClusterMgr();
+
 			FDisplayClusterClusterEventJson WaitForGameEvent;
 			WaitForGameEvent.Category = DisplayClusterGameEngineUtils::WaitForGameCategory;
 			WaitForGameEvent.Type = URL.ToString();
-			WaitForGameEvent.Name = ClusterMgr->GetClusterNodeController()->GetNodeId();
+			WaitForGameEvent.Name = ClusterMgr.GetNodeId();
 			WaitForGameEvent.bIsSystemEvent = true;
 			WaitForGameEvent.bShouldDiscardOnRepeat = false;
-			ClusterMgr->EmitClusterEventJson(WaitForGameEvent, false);
+
+			ClusterMgr.EmitClusterEventJson(WaitForGameEvent, false);
 
 			RunningMode = EDisplayClusterRunningMode::WaitingForSync;
 			// Assume that all nodes are now out of sync.
@@ -345,10 +340,12 @@ void UDisplayClusterGameEngine::Tick(float DeltaSeconds, bool bIdleMode)
 
 	if (CanTick())
 	{
+		IPDisplayClusterClusterManager& ClusterMgr = *GDisplayCluster->GetPrivateClusterMgr();
+
 		//////////////////////////////////////////////////////////////////////////////////////////////
 		// Frame start barrier
 		UE_LOG(LogDisplayClusterEngine, Verbose, TEXT("Sync frame start"));
-		ClusterMgr->GetClusterNodeController()->WaitForFrameStart();
+		ClusterMgr.GetNetApi().GetClusterSyncAPI()->WaitForFrameStart();
 
 		// Perform StartFrame notification
 		GDisplayCluster->StartFrame(GFrameCounter);
@@ -376,7 +373,7 @@ void UDisplayClusterGameEngine::Tick(float DeltaSeconds, bool bIdleMode)
 
 		//////////////////////////////////////////////////////////////////////////////////////////////
 		// Frame end barrier
-		ClusterMgr->GetClusterNodeController()->WaitForFrameEnd();
+		ClusterMgr.GetNetApi().GetClusterSyncAPI()->WaitForFrameEnd();
 
 		// Perform EndFrame notification
 		GDisplayCluster->EndFrame(GFrameCounter);
@@ -402,10 +399,10 @@ void UDisplayClusterGameEngine::UpdateTimeAndHandleMaxTickRate()
 
 	UEngine::UpdateTimeAndHandleMaxTickRate();
 
-	if (CanTick() && ClusterMgr)
+	if (CanTick())
 	{
 		// Synchronize time data
-		ClusterMgr->SyncTimeData();
+		GDisplayCluster->GetPrivateClusterMgr()->SyncTimeData();
 	}
 }
 
@@ -431,7 +428,7 @@ void UDisplayClusterGameEngine::ReceivedSync(const FString &Level, const FString
 	UE_LOG(LogDisplayClusterEngine,Display, TEXT("GameSyncChange event received."));
 	TSet<FString> &SyncItem = SyncMap.FindOrAdd(Level);
 	SyncItem.Add(NodeId);
-	if (SyncItem.Num() == ClusterMgr->GetNodesAmount())
+	if (SyncItem.Num() == GDisplayCluster->GetPrivateClusterMgr()->GetNodesAmount())
 	{
 		SyncMap.Remove(Level);
 	}
@@ -450,7 +447,7 @@ void UDisplayClusterGameEngine::CheckGameStartBarrier()
 {
 	if (!BarrierAvoidanceOn())
 	{
-		ClusterMgr->GetClusterNodeController()->WaitForGameStart();
+		GDisplayCluster->GetPrivateClusterMgr()->GetNetApi().GetClusterSyncAPI()->WaitForGameStart();
 	}
 	else
 	{
@@ -459,7 +456,7 @@ void UDisplayClusterGameEngine::CheckGameStartBarrier()
 			UE_LOG(LogDisplayClusterEngine, Display, TEXT("CheckGameStartBarrier - we are no longer out of sync. Restoring Play."));
 			if (RunningMode == EDisplayClusterRunningMode::Startup)
 			{
-				ClusterMgr->GetClusterNodeController()->WaitForGameStart();
+				GDisplayCluster->GetPrivateClusterMgr()->GetNetApi().GetClusterSyncAPI()->WaitForGameStart();
 			}
 			UGameplayStatics::SetGamePaused(WorldContextObject,false);
 			RunningMode = EDisplayClusterRunningMode::Synced;
diff --git a/Engine/Plugins/Runtime/nDisplay/Source/DisplayCluster/Private/Game/EngineClasses/Basics/DisplayClusterPlayerInput.cpp b/Engine/Plugins/Runtime/nDisplay/Source/DisplayCluster/Private/Game/EngineClasses/Basics/DisplayClusterPlayerInput.cpp
--- a/Engine/Plugins/Runtime/nDisplay/Source/DisplayCluster/Private/Game/EngineClasses/Basics/DisplayClusterPlayerInput.cpp
+++ b/Engine/Plugins/Runtime/nDisplay/Source/DisplayCluster/Private/Game/EngineClasses/Basics/DisplayClusterPlayerInput.cpp
@@ -9,6 +9,7 @@
 #include "Components/InputComponent.h"
 
 #include "Cluster/IPDisplayClusterClusterManager.h"
+#include "Cluster/NetAPI/DisplayClusterNetApiFacade.h"
 #include "Config/IPDisplayClusterConfigManager.h"
 
 #include "Misc/DisplayClusterGlobals.h"
@@ -151,16 +152,17 @@ void UDisplayClusterPlayerInput::ProcessPolicy_ReplicatePrimary(const TArray<UIn
 	{
 		TMap<FString, FString> KeyStates;
 
-		if (ClusterMgr->IsPrimary())
+		// Always cache local inputs whether we're primary or not. If we're secondary now,
+		// we may turn into primary in a moment, and the data would be available for others.
 		{
-			// Export key states data for the stack
 			SerializeKeyStateMap(KeyStates);
-			ClusterMgr->ImportNativeInputData(MoveTemp(KeyStates));
+			ClusterMgr->ImportNativeInputData(KeyStates);
 		}
-		else
+
+		// Get external input data only if we're not primary
+		if (!ClusterMgr->IsPrimary())
 		{
-			// Import key states data to the stack
-			ClusterMgr->ExportNativeInputData(KeyStates);
+			ClusterMgr->GetNetApi().GetClusterSyncAPI()->GetNativeInputData(KeyStates);
 			DeserializeKeyStateMap(KeyStates);
 		}
 	}
diff --git a/Engine/Plugins/Runtime/nDisplay/Source/DisplayCluster/Private/Game/EngineClasses/Scene/DisplayClusterSceneComponentSync.cpp b/Engine/Plugins/Runtime/nDisplay/Source/DisplayCluster/Private/Game/EngineClasses/Scene/DisplayClusterSceneComponentSync.cpp
--- a/Engine/Plugins/Runtime/nDisplay/Source/DisplayCluster/Private/Game/EngineClasses/Scene/DisplayClusterSceneComponentSync.cpp
+++ b/Engine/Plugins/Runtime/nDisplay/Source/DisplayCluster/Private/Game/EngineClasses/Scene/DisplayClusterSceneComponentSync.cpp
@@ -6,8 +6,7 @@
 
 #include "GameFramework/Actor.h"
 
-#include "Cluster/IPDisplayClusterClusterManager.h"
-#include "Game/IPDisplayClusterGameManager.h"
+#include "Cluster/IDisplayClusterClusterManager.h"
 
 #include "Misc/DisplayClusterGlobals.h"
 #include "Misc/DisplayClusterLog.h"
@@ -24,27 +23,20 @@ void UDisplayClusterSceneComponentSync::BeginPlay()
 	Super::BeginPlay();
 
 	const EDisplayClusterOperationMode OpMode = GDisplayCluster->GetOperationMode();
-	if (OpMode != EDisplayClusterOperationMode::Cluster)
+	if (OpMode == EDisplayClusterOperationMode::Cluster)
 	{
-		return;
-	}
+		// Generate unique sync id
+		SyncId = GenerateSyncId();
 
-	// Generate unique sync id
-	SyncId = GenerateSyncId();
-
-	GameMgr = GDisplayCluster->GetPrivateGameMgr();
-	if (GameMgr && GameMgr->IsDisplayClusterActive())
-	{
 		// Register sync object
-		ClusterMgr = GDisplayCluster->GetPrivateClusterMgr();
-		if (ClusterMgr)
+		if (IDisplayClusterClusterManager* ClusterMgr = GDisplayCluster->GetClusterMgr())
 		{
 			UE_LOG(LogDisplayClusterGame, Log, TEXT("Registering sync object %s..."), *SyncId);
 			ClusterMgr->RegisterSyncObject(this, EDisplayClusterSyncGroup::Tick);
 		}
 		else
 		{
-			UE_LOG(LogDisplayClusterGame, Warning, TEXT("Couldn't register %s scene component sync. Looks like we're in non-DisplayCluster mode."), *SyncId);
+			UE_LOG(LogDisplayClusterGame, Warning, TEXT("Couldn't register %s scene component sync."), *SyncId);
 		}
 	}
 }
@@ -54,13 +46,11 @@ void UDisplayClusterSceneComponentSync::EndPlay(const EEndPlayReason::Type EndPl
 	const EDisplayClusterOperationMode OpMode = GDisplayCluster->GetOperationMode();
 	if (OpMode == EDisplayClusterOperationMode::Cluster)
 	{
-		if (GameMgr && GameMgr->IsDisplayClusterActive())
+		// Unregister sync object
+		if (IDisplayClusterClusterManager* ClusterMgr = GDisplayCluster->GetClusterMgr())
 		{
-			if (ClusterMgr)
-			{
-				UE_LOG(LogDisplayClusterGame, Log, TEXT("Unregistering sync object %s..."), *SyncId);
-				ClusterMgr->UnregisterSyncObject(this);
-			}
+			UE_LOG(LogDisplayClusterGame, Log, TEXT("Unregistering sync object %s..."), *SyncId);
+			ClusterMgr->UnregisterSyncObject(this);
 		}
 	}
 
diff --git a/Engine/Plugins/Runtime/nDisplay/Source/DisplayCluster/Private/Game/IPDisplayClusterGameManager.h b/Engine/Plugins/Runtime/nDisplay/Source/DisplayCluster/Private/Game/IPDisplayClusterGameManager.h
--- a/Engine/Plugins/Runtime/nDisplay/Source/DisplayCluster/Private/Game/IPDisplayClusterGameManager.h
+++ b/Engine/Plugins/Runtime/nDisplay/Source/DisplayCluster/Private/Game/IPDisplayClusterGameManager.h
@@ -15,7 +15,4 @@ class IPDisplayClusterGameManager
 {
 public:
 	virtual ~IPDisplayClusterGameManager() = default;
-
-public:
-	virtual bool IsDisplayClusterActive() const = 0;
 };
diff --git a/Engine/Plugins/Runtime/nDisplay/Source/DisplayCluster/Private/Misc/DisplayClusterLog.cpp b/Engine/Plugins/Runtime/nDisplay/Source/DisplayCluster/Private/Misc/DisplayClusterLog.cpp
--- a/Engine/Plugins/Runtime/nDisplay/Source/DisplayCluster/Private/Misc/DisplayClusterLog.cpp
+++ b/Engine/Plugins/Runtime/nDisplay/Source/DisplayCluster/Private/Misc/DisplayClusterLog.cpp
@@ -3,14 +3,16 @@
 #include "DisplayClusterLog.h"
 
 // Plugin-wide log categories
-DEFINE_LOG_CATEGORY(LogDisplayClusterEngine);
-DEFINE_LOG_CATEGORY(LogDisplayClusterModule);
+DEFINE_LOG_CATEGORY(LogDisplayClusterBarrier);
+DEFINE_LOG_CATEGORY(LogDisplayClusterBarrierGP);
+DEFINE_LOG_CATEGORY(LogDisplayClusterBlueprint);
 DEFINE_LOG_CATEGORY(LogDisplayClusterCluster);
 DEFINE_LOG_CATEGORY(LogDisplayClusterConfig);
-DEFINE_LOG_CATEGORY(LogDisplayClusterBarrier);
+DEFINE_LOG_CATEGORY(LogDisplayClusterEngine);
+DEFINE_LOG_CATEGORY(LogDisplayClusterFailover);
 DEFINE_LOG_CATEGORY(LogDisplayClusterGame);
+DEFINE_LOG_CATEGORY(LogDisplayClusterModule);
 DEFINE_LOG_CATEGORY(LogDisplayClusterNetwork);
 DEFINE_LOG_CATEGORY(LogDisplayClusterRender);
 DEFINE_LOG_CATEGORY(LogDisplayClusterRenderSync);
 DEFINE_LOG_CATEGORY(LogDisplayClusterViewport);
-DEFINE_LOG_CATEGORY(LogDisplayClusterBlueprint);
diff --git a/Engine/Plugins/Runtime/nDisplay/Source/DisplayCluster/Private/Misc/DisplayClusterLog.h b/Engine/Plugins/Runtime/nDisplay/Source/DisplayCluster/Private/Misc/DisplayClusterLog.h
--- a/Engine/Plugins/Runtime/nDisplay/Source/DisplayCluster/Private/Misc/DisplayClusterLog.h
+++ b/Engine/Plugins/Runtime/nDisplay/Source/DisplayCluster/Private/Misc/DisplayClusterLog.h
@@ -2,21 +2,23 @@
 
 #pragma once
 
-#include "CoreMinimal.h"
+#include "Logging/LogMacros.h"
 #include "Stats/Stats.h"
 
 // Module-wide log categories
-DECLARE_LOG_CATEGORY_EXTERN(LogDisplayClusterEngine,     Log, All);
-DECLARE_LOG_CATEGORY_EXTERN(LogDisplayClusterModule,     Log, All);
+DECLARE_LOG_CATEGORY_EXTERN(LogDisplayClusterBarrier,    Log, All);
+DECLARE_LOG_CATEGORY_EXTERN(LogDisplayClusterBarrierGP,  Log, All);
+DECLARE_LOG_CATEGORY_EXTERN(LogDisplayClusterBlueprint,  Log, All);
 DECLARE_LOG_CATEGORY_EXTERN(LogDisplayClusterCluster,    Log, All);
 DECLARE_LOG_CATEGORY_EXTERN(LogDisplayClusterConfig,     Log, All);
-DECLARE_LOG_CATEGORY_EXTERN(LogDisplayClusterBarrier,    Log, All);
+DECLARE_LOG_CATEGORY_EXTERN(LogDisplayClusterEngine,     Log, All);
+DECLARE_LOG_CATEGORY_EXTERN(LogDisplayClusterFailover,   Log, All);
 DECLARE_LOG_CATEGORY_EXTERN(LogDisplayClusterGame,       Log, All);
+DECLARE_LOG_CATEGORY_EXTERN(LogDisplayClusterModule,     Log, All);
 DECLARE_LOG_CATEGORY_EXTERN(LogDisplayClusterNetwork,    Log, All);
 DECLARE_LOG_CATEGORY_EXTERN(LogDisplayClusterRender,     Log, All);
 DECLARE_LOG_CATEGORY_EXTERN(LogDisplayClusterRenderSync, Log, All);
 DECLARE_LOG_CATEGORY_EXTERN(LogDisplayClusterViewport,   Log, All);
-DECLARE_LOG_CATEGORY_EXTERN(LogDisplayClusterBlueprint,  Log, All);
 
 
-DECLARE_STATS_GROUP(TEXT("nDisplay"), STATGROUP_NDisplay, STATCAT_Advanced)
\ No newline at end of file
+DECLARE_STATS_GROUP(TEXT("nDisplay"), STATGROUP_NDisplay, STATCAT_Advanced)
diff --git a/Engine/Plugins/Runtime/nDisplay/Source/DisplayCluster/Private/Network/Barrier/DisplayClusterBarrier.cpp b/Engine/Plugins/Runtime/nDisplay/Source/DisplayCluster/Private/Network/Barrier/DisplayClusterBarrier.cpp
--- a/Engine/Plugins/Runtime/nDisplay/Source/DisplayCluster/Private/Network/Barrier/DisplayClusterBarrier.cpp
+++ b/Engine/Plugins/Runtime/nDisplay/Source/DisplayCluster/Private/Network/Barrier/DisplayClusterBarrier.cpp
@@ -6,18 +6,18 @@
 #include "Misc/ScopeLock.h"
 
 
-FDisplayClusterBarrier::FDisplayClusterBarrier(const FString& InName, const TArray<FString>& InCallersAllowed, const uint32 InTimeout)
+FDisplayClusterBarrier::FDisplayClusterBarrier(const FString& InName, const TSet<FString>& InCallersAllowed, const uint32 InTimeout)
 	: Name(InName)
 	, CallersAllowed(InCallersAllowed)
 	, Timeout(InTimeout)
-	, ThreadLimit(static_cast<uint32>(CallersAllowed.Num()))
-	, WatchdogTimer(InName + FString("_watchdog"))
+	, WatchdogTimer(InName + TEXT("_watchdog"))
 {
-	UE_LOG(LogDisplayClusterBarrier, Log, TEXT("Initialized barrier '%s' with timeout %u ms and threads limit: %u"), *Name, Timeout, ThreadLimit);
+	UE_LOG(LogDisplayClusterBarrier, Log, TEXT("Initialized barrier '%s' with timeout %u ms and threads limit: %u"), *Name, Timeout, CallersAllowed.Num());
 
-	for (int32 Idx = 0; Idx < CallersAllowed.Num(); ++Idx)
+	int32 Idx = 0;
+	for (const FString& CallerId : CallersAllowed)
 	{
-		UE_LOG(LogDisplayClusterBarrier, Verbose, TEXT("Barrier '%s': client (%d): '%s'"), *Name, Idx, *CallersAllowed[Idx]);
+		UE_LOG(LogDisplayClusterBarrier, Verbose, TEXT("Barrier '%s': client (%d): '%s'"), *Name, Idx++, *CallerId);
 	}
 
 	// Subscribe for timeout events
@@ -39,7 +39,6 @@ bool FDisplayClusterBarrier::Activate()
 
 	if (!bActive)
 	{
-		ThreadCount = 0;
 		bActive = true;
 		CallersAwaiting.Reset();
 
@@ -128,15 +127,12 @@ EDisplayClusterBarrierWaitResult FDisplayClusterBarrier::Wait(const FString& Cal
 				// Fixate awaiting start for this particular thread
 				ThreadWaitTimeStart = FPlatformTime::Seconds();
 
-				// One more thread awaiting
-				++ThreadCount;
-
 				// In case this thread came first to the barrier, we need:
 				// - to fixate barrier awaiting start time
 				// - start watchdog timer
-				if (ThreadCount == 1)
+				if (CallersAwaiting.Num() == 1)
 				{
-					UE_LOG(LogDisplayClusterBarrier, Verbose, TEXT("Barrier '%s': sync start"), *Name);
+					UE_LOG(LogDisplayClusterBarrier, Verbose, TEXT("Barrier '%s': sync start, cycle %lu"), *Name, SyncCycleCounter);
 
 					// Prepare for new sync iteration
 					HandleBarrierPreSyncStart();
@@ -145,19 +141,22 @@ EDisplayClusterBarrierWaitResult FDisplayClusterBarrier::Wait(const FString& Cal
 					WatchdogTimer.SetTimer(Timeout);
 				}
 
-				UE_LOG(LogDisplayClusterBarrier, Verbose, TEXT("Barrier '%s': awaiting threads amount - %d"), *Name, ThreadCount);
+				UE_LOG(LogDisplayClusterBarrier, Verbose, TEXT("Barrier '%s': awaiting threads amount - %d"), *Name, CallersAwaiting.Num());
 
 				// In case this thread is the last one the barrier is awaiting for, we need:
 				// - to fixate barrier awaiting finish time
 				// - to open the output gate (release the barrier)
 				// - to close the input gate
 				// - reset watchdog timer
-				if (ThreadCount == ThreadLimit)
+				if (CallersAwaiting.Num() == CallersAllowed.Num())
 				{
 					BarrierWaitTimeFinish = FPlatformTime::Seconds();
 					BarrierWaitTimeOverall = BarrierWaitTimeFinish - BarrierWaitTimeStart;
 
-					UE_LOG(LogDisplayClusterBarrier, Verbose, TEXT("Barrier '%s': sync end, barrier wait time %f"), *Name, BarrierWaitTimeOverall);
+					UE_LOG(LogDisplayClusterBarrier, Verbose, TEXT("Barrier '%s': sync end, cycle %lu, barrier wait time %f"), *Name, SyncCycleCounter, BarrierWaitTimeOverall);
+
+					// Increment cycle counter
+					++SyncCycleCounter;
 
 					// All callers are here, pet the watchdog
 					WatchdogTimer.ResetTimer();
@@ -185,6 +184,15 @@ EDisplayClusterBarrierWaitResult FDisplayClusterBarrier::Wait(const FString& Cal
 		{
 			FScopeLock LockData(&DataCS);
 
+			// In case there are any 'simulated' callers that have been dropped actually,
+			// modify "Awaiting" and "Allowed" lists so they could properly handle the next sync cycle.
+			if (!CallersToForget.IsEmpty())
+			{
+				CallersAwaiting = CallersAwaiting.Difference(CallersToForget);
+				CallersAllowed  = CallersAllowed.Difference(CallersToForget);
+				CallersToForget.Reset();
+			}
+
 			// Unregister caller
 			CallersAwaiting.Remove(CallerId);
 
@@ -192,7 +200,7 @@ EDisplayClusterBarrierWaitResult FDisplayClusterBarrier::Wait(const FString& Cal
 			if (bActive)
 			{
 				// In case this thread is leaving last, close output and open input
-				if (--ThreadCount == 0)
+				if (CallersAwaiting.IsEmpty())
 				{
 					EventOutputGateOpen->Reset();
 					EventInputGateOpen->Trigger();
@@ -218,24 +226,24 @@ EDisplayClusterBarrierWaitResult FDisplayClusterBarrier::Wait(const FString& Cal
 	return EDisplayClusterBarrierWaitResult::Ok;
 }
 
-EDisplayClusterBarrierWaitResult FDisplayClusterBarrier::WaitWithData(const FString& ThreadMarker, const TArray<uint8>& RequestData, TArray<uint8>& OutResponseData, double* OutThreadWaitTime /* = nullptr */, double* OutBarrierWaitTime /* = nullptr */)
+EDisplayClusterBarrierWaitResult FDisplayClusterBarrier::WaitWithData(const FString& CallerId, const TArray<uint8>& RequestData, TArray<uint8>& OutResponseData, double* OutThreadWaitTime /* = nullptr */, double* OutBarrierWaitTime /* = nullptr */)
 {
 	{
 		FScopeLock Lock(&CommDataCS);
 		// Store request data so it can be used once all the threads arrived
-		ClientsRequestData.Emplace(ThreadMarker, RequestData);
+		ClientsRequestData.Emplace(CallerId, RequestData);
 	}
 
 	// Wait at the barrier
-	const EDisplayClusterBarrierWaitResult WaitResult = Wait(ThreadMarker, OutThreadWaitTime, OutBarrierWaitTime);
+	const EDisplayClusterBarrierWaitResult WaitResult = Wait(CallerId, OutThreadWaitTime, OutBarrierWaitTime);
 
 	{
 		FScopeLock Lock(&CommDataCS);
 
-		if (TArray<uint8>* const FoundThreadResponseData = ClientsResponseData.Find(ThreadMarker))
+		if (TArray<uint8>* const FoundThreadResponseData = ClientsResponseData.Find(CallerId))
 		{
 			OutResponseData = MoveTemp(*FoundThreadResponseData);
-			ClientsResponseData.Remove(ThreadMarker);
+			ClientsResponseData.Remove(CallerId);
 		}
 	}
 
@@ -248,33 +256,52 @@ void FDisplayClusterBarrier::UnregisterSyncCaller(const FString& CallerId)
 
 	FScopeLock Lock(&DataCS);
 
-	// Remove specified callers from the management lists
-	if (CallersAllowed.Remove(CallerId) > 0 || CallersTimedout.Remove(CallerId) > 0)
+	// Ignore if it has been processed it already
+	if (CallersDropped.Contains(CallerId))
+	{
+		return;
+	}
+	else
 	{
-		// Update thread limit
-		ThreadLimit = static_cast<uint32>(CallersAllowed.Num());
+		CallersDropped.Add(CallerId);
+	}
 
-		UE_LOG(LogDisplayClusterBarrier, Log, TEXT("Barrier '%s': caller '%s' has been unregistered, new thread limit is %u"), *Name, *CallerId, ThreadLimit);
+	// If synchronization cycle has not started yet, just remove it from 'Allowed' list
+	if (CallersAwaiting.IsEmpty())
+	{
+		CallersAllowed.Remove(CallerId);
+		return;
+	}
+	// Otherwise, simulate this node just arrived to synchronize. But remember that to clean up later.
+	else
+	{
+		CallersAwaiting.Add(CallerId);
+		CallersToForget.Add(CallerId);
+	}
 
-		// In case it's a last missing caller, we need to open the barrier
-		if (ThreadCount == ThreadLimit)
+	// In case it's a last missing caller, we need to open the barrier
+	if (CallersAwaiting.Num() == CallersAllowed.Num())
+	{
+		// In case there are any 'simulated' callers that have been dropped actually,
+		// modify "Awaiting" and "Allowed" lists so they could properly handle the next sync cycle.
+		if (!CallersToForget.IsEmpty())
 		{
-			BarrierWaitTimeFinish  = FPlatformTime::Seconds();
-			BarrierWaitTimeOverall = BarrierWaitTimeFinish - BarrierWaitTimeStart;
+			CallersAwaiting = CallersAwaiting.Difference(CallersToForget);
+			CallersAllowed  = CallersAllowed.Difference(CallersToForget);
+			CallersToForget.Reset();
+		}
 
-			UE_LOG(LogDisplayClusterBarrier, Verbose, TEXT("Barrier '%s': sync end, barrier wait time %f"), *Name, BarrierWaitTimeOverall);
+		BarrierWaitTimeFinish  = FPlatformTime::Seconds();
+		BarrierWaitTimeOverall = BarrierWaitTimeFinish - BarrierWaitTimeStart;
 
-			// All callers here, pet the watchdog
-			WatchdogTimer.ResetTimer();
+		UE_LOG(LogDisplayClusterBarrier, Verbose, TEXT("Barrier '%s': sync end, barrier wait time %f"), *Name, BarrierWaitTimeOverall);
 
-			// Close the input gate, and open the output gate
-			EventInputGateOpen->Reset();
-			EventOutputGateOpen->Trigger();
-		}
-	}
-	else
-	{
-		UE_LOG(LogDisplayClusterBarrier, Log, TEXT("Barrier '%s': caller '%s' not found"), *Name, *CallerId);
+		// All callers here, pet the watchdog
+		WatchdogTimer.ResetTimer();
+
+		// Close the input gate, and open the output gate
+		EventInputGateOpen->Reset();
+		EventOutputGateOpen->Trigger();
 	}
 }
 
@@ -305,27 +332,21 @@ void FDisplayClusterBarrier::HandleBarrierTimeout()
 	BarrierWaitTimeFinish = FPlatformTime::Seconds();
 	BarrierWaitTimeOverall = BarrierWaitTimeFinish - BarrierWaitTimeStart;
 
-	// List of timed out threads
-	TArray<FString> CallersTimedOutOnLastSync;
-
-	// Update the list of timed out threads. Copy all missing callers to the list of timed out ones.
-	for (const FString& CallerAllowed : CallersAllowed)
+	// List of the threads that timed out at the current synchronization cycle
+	const TSet<FString> CallersTimedOutOnLastSync = CallersAllowed.Difference(CallersAwaiting);
+	for (const FString& CallerId : CallersTimedOutOnLastSync)
 	{
-		if (!CallersAwaiting.Contains(CallerAllowed))
-		{
-			UE_LOG(LogDisplayClusterBarrier, Log, TEXT("Barrier '%s': caller '%s' was moved to the 'TimedOut' list"), *Name, *CallerAllowed);
-			CallersTimedOutOnLastSync.Add(CallerAllowed);
-		}
+		UE_LOG(LogDisplayClusterBarrier, Log, TEXT("Barrier '%s': caller '%s' was moved to the 'TimedOut' list"), *Name, *CallerId);
 	}
 
-	// Update timedout list
+	// Update timedout and dropped list
 	CallersTimedout.Append(CallersTimedOutOnLastSync);
+	CallersDropped.Append(CallersTimedOutOnLastSync);
+
 	// Update the list of permitted callers
 	CallersAllowed = CallersAwaiting;
-	// Update thread limit
-	ThreadLimit = ThreadCount; // Same as CallersAllowed.Num()
 
-	UE_LOG(LogDisplayClusterBarrier, Log, TEXT("Barrier '%s': new threads limit %d"), *Name, ThreadLimit);
+	UE_LOG(LogDisplayClusterBarrier, Log, TEXT("Barrier '%s': new threads limit %d"), *Name, CallersAllowed.Num());
 
 	// Notify listeners
 	OnBarrierTimeout().Broadcast(Name, CallersTimedOutOnLastSync);
diff --git a/Engine/Plugins/Runtime/nDisplay/Source/DisplayCluster/Private/Network/Barrier/DisplayClusterBarrier.h b/Engine/Plugins/Runtime/nDisplay/Source/DisplayCluster/Private/Network/Barrier/DisplayClusterBarrier.h
--- a/Engine/Plugins/Runtime/nDisplay/Source/DisplayCluster/Private/Network/Barrier/DisplayClusterBarrier.h
+++ b/Engine/Plugins/Runtime/nDisplay/Source/DisplayCluster/Private/Network/Barrier/DisplayClusterBarrier.h
@@ -2,7 +2,6 @@
 
 #pragma once
 
-#include "CoreMinimal.h"
 #include "Network/Barrier/IDisplayClusterBarrier.h"
 #include "Misc/DisplayClusterWatchdogTimer.h"
 
@@ -16,7 +15,7 @@ class FDisplayClusterBarrier
 	: public IDisplayClusterBarrier
 {
 public:
-	FDisplayClusterBarrier(const FString& Name, const TArray<FString>& CallersAllowed, const uint32 Timeout);
+	FDisplayClusterBarrier(const FString& Name, const TSet<FString>& CallersAllowed, const uint32 Timeout);
 	virtual  ~FDisplayClusterBarrier();
 
 public:
@@ -37,7 +36,7 @@ public:
 	virtual EDisplayClusterBarrierWaitResult Wait(const FString& CallerId, double* OutThreadWaitTime = nullptr, double* OutBarrierWaitTime = nullptr) override;
 
 	// Wait until all threads arrive (with data)
-	virtual EDisplayClusterBarrierWaitResult WaitWithData(const FString& ThreadMarker, const TArray<uint8>& RequestData, TArray<uint8>& OutResponseData, double* OutThreadWaitTime = nullptr, double* OutBarrierWaitTime = nullptr) override;
+	virtual EDisplayClusterBarrierWaitResult WaitWithData(const FString& CallerId, const TArray<uint8>& RequestData, TArray<uint8>& OutResponseData, double* OutThreadWaitTime = nullptr, double* OutBarrierWaitTime = nullptr) override;
 
 	// Remove specified caller from the sync pipeline
 	virtual void UnregisterSyncCaller(const FString& CallerId) override;
@@ -72,24 +71,25 @@ private:
 
 
 	// Caller that are allowed to join the barrier
-	TArray<FString> CallersAllowed;
+	TSet<FString> CallersAllowed;
 
 	// Callers that are already waiting at the barrier
-	TArray<FString> CallersAwaiting;
+	TSet<FString> CallersAwaiting;
+
+	// Callers that were dropped during synchronization cycle. Holds them temporary to provide correct cycle enclosure.
+	TSet<FString> CallersToForget;
 
 	// Cluster Callers that have been timed out previously
-	TArray<FString> CallersTimedout;
+	TSet<FString> CallersTimedout;
 
+	// Cluster Callers that have been detached from barrier synchronization
+	TSet<FString> CallersDropped;
 
 	// Timeout for the barrier
 	const uint32 Timeout = 0;
 
-	// Overall amount of threads to wait at the barrier
-	uint32 ThreadLimit = 0;
-
-	// Amount of threads at the barrier
-	uint32 ThreadCount = 0;
-
+	// Synchronization cycles counter
+	uint64 SyncCycleCounter = 0;
 
 	// Events to control when threads can join (input) and leave (output) the barrier
 	FEventRef EventInputGateOpen { EEventMode::ManualReset };
diff --git a/Engine/Plugins/Runtime/nDisplay/Source/DisplayCluster/Private/Network/Barrier/DisplayClusterBarrierFactory.cpp b/Engine/Plugins/Runtime/nDisplay/Source/DisplayCluster/Private/Network/Barrier/DisplayClusterBarrierFactory.cpp
--- a/Engine/Plugins/Runtime/nDisplay/Source/DisplayCluster/Private/Network/Barrier/DisplayClusterBarrierFactory.cpp
+++ b/Engine/Plugins/Runtime/nDisplay/Source/DisplayCluster/Private/Network/Barrier/DisplayClusterBarrierFactory.cpp
@@ -6,8 +6,8 @@
 #include "Misc/DisplayClusterLog.h"
 
 
-IDisplayClusterBarrier* FDisplayClusterBarrierFactory::CreateBarrier(const FString& BarrierId, const TArray<FString>& ThreadMarkers, const uint32 Timeout)
+IDisplayClusterBarrier* FDisplayClusterBarrierFactory::CreateBarrier(const FString& BarrierId, const TSet<FString>& CallerIds, const uint32 Timeout)
 {
-	UE_LOG(LogDisplayClusterBarrier, Log, TEXT("Instantiating unique barrier '%s': Threads=%d, Timeout=%u ms"), *BarrierId, ThreadMarkers.Num(), Timeout);
-	return new FDisplayClusterBarrier(BarrierId, ThreadMarkers, Timeout);
+	UE_LOG(LogDisplayClusterBarrier, Log, TEXT("Instantiating unique barrier '%s': Threads=%d, Timeout=%u ms"), *BarrierId, CallerIds.Num(), Timeout);
+	return new FDisplayClusterBarrier(BarrierId, CallerIds, Timeout);
 }
diff --git a/Engine/Plugins/Runtime/nDisplay/Source/DisplayCluster/Private/Network/Barrier/DisplayClusterBarrierFactory.h b/Engine/Plugins/Runtime/nDisplay/Source/DisplayCluster/Private/Network/Barrier/DisplayClusterBarrierFactory.h
--- a/Engine/Plugins/Runtime/nDisplay/Source/DisplayCluster/Private/Network/Barrier/DisplayClusterBarrierFactory.h
+++ b/Engine/Plugins/Runtime/nDisplay/Source/DisplayCluster/Private/Network/Barrier/DisplayClusterBarrierFactory.h
@@ -2,7 +2,7 @@
 
 #pragma once
 
-#include "CoreMinimal.h"
+#include "Containers/Set.h"
 
 class IDisplayClusterBarrier;
 
@@ -13,9 +13,11 @@ class IDisplayClusterBarrier;
 class FDisplayClusterBarrierFactory
 {
 private:
+
 	FDisplayClusterBarrierFactory();
 
 public:
-	// Factory method to instantiate a barrier based on the user settings
-	static IDisplayClusterBarrier* CreateBarrier(const FString& BarrierId, const TArray<FString>& ThreadMarkers, const uint32 Timeout);
+
+	/** Factory method to instantiate a barrier based on the user settings */
+	static IDisplayClusterBarrier* CreateBarrier(const FString& BarrierId, const TSet<FString>& CallerIds, const uint32 Timeout);
 };
diff --git a/Engine/Plugins/Runtime/nDisplay/Source/DisplayCluster/Private/Network/Barrier/IDisplayClusterBarrier.h b/Engine/Plugins/Runtime/nDisplay/Source/DisplayCluster/Private/Network/Barrier/IDisplayClusterBarrier.h
--- a/Engine/Plugins/Runtime/nDisplay/Source/DisplayCluster/Private/Network/Barrier/IDisplayClusterBarrier.h
+++ b/Engine/Plugins/Runtime/nDisplay/Source/DisplayCluster/Private/Network/Barrier/IDisplayClusterBarrier.h
@@ -27,10 +27,10 @@ struct FDisplayClusterBarrierPreSyncEndDelegateData
 	/** Barrier ID */
 	const FString& BarrierId;
 
-	/** Binary data provided on sync request (thread marker - to - data mapping) */
+	/** Binary data provided on sync request (CallerId-to-Data mapping) */
 	const TMap<FString, TArray<uint8>>& RequestData;
 
-	/** Binary data to respond (thread marker - to - data mapping) */
+	/** Binary data to respond (CallerId-to-Data mapping) */
 	TMap<FString, TArray<uint8>>& ResponseData;
 };
 
@@ -41,35 +41,37 @@ struct FDisplayClusterBarrierPreSyncEndDelegateData
 class IDisplayClusterBarrier
 {
 public:
+
 	virtual ~IDisplayClusterBarrier() = default;
 
 public:
-	// Barrier name
+
+	/** Barrier name */
 	virtual const FString& GetName() const = 0;
 
-	// Activate barrier
+	/** Activate barrier */
 	virtual bool Activate() = 0;
 
-	// Deactivate barrier, no threads will be blocked
+	/** Deactivate barrier, no threads will be blocked */
 	virtual void Deactivate() = 0;
 
-	// Returns true if the barrier has been activated
+	/** Returns true if the barrier has been activated */
 	virtual bool IsActivated() const = 0;
 
-	// Wait until all threads arrive
-	virtual EDisplayClusterBarrierWaitResult Wait(const FString& ThreadMarker, double* OutThreadWaitTime = nullptr, double* OutBarrierWaitTime = nullptr) = 0;
+	/** Wait until all threads arrive */
+	virtual EDisplayClusterBarrierWaitResult Wait(const FString& CallerId, double* OutThreadWaitTime = nullptr, double* OutBarrierWaitTime = nullptr) = 0;
 
-	// Wait until all threads arrive (with data)
-	virtual EDisplayClusterBarrierWaitResult WaitWithData(const FString& ThreadMarker, const TArray<uint8>& RequestData, TArray<uint8>& OutResponseData, double* OutThreadWaitTime = nullptr, double* OutBarrierWaitTime = nullptr) = 0;
+	/** Wait until all threads arrive(with data) */
+	virtual EDisplayClusterBarrierWaitResult WaitWithData(const FString& CallerId, const TArray<uint8>& RequestData, TArray<uint8>& OutResponseData, double* OutThreadWaitTime = nullptr, double* OutBarrierWaitTime = nullptr) = 0;
 
-	// Remove specified caller from the sync pipeline
+	/** Remove specified caller from the sync pipeline */
 	virtual void UnregisterSyncCaller(const FString& CallerId) = 0;
 
 	/** Barrier PreSyncEnd delegate. Called when all calling threads arrived right before opening the gate. */
 	DECLARE_DELEGATE_OneParam(FDisplayClusterBarrierPreSyncEndDelegate, FDisplayClusterBarrierPreSyncEndDelegateData&);
 	virtual FDisplayClusterBarrierPreSyncEndDelegate& GetPreSyncEndDelegate() = 0;
 
-	// Barrier timout notification (provides BarrierName and CallersTimedOut in parameters)
-	DECLARE_EVENT_TwoParams(IDisplayClusterBarrier, FDisplayClusterBarrierTimeoutEvent, const FString&, const TArray<FString>&);
+	/** Barrier timout notification (provides BarrierName and CallersTimedOut in parameters) */
+	DECLARE_EVENT_TwoParams(IDisplayClusterBarrier, FDisplayClusterBarrierTimeoutEvent, const FString&, const TSet<FString>&);
 	virtual FDisplayClusterBarrierTimeoutEvent& OnBarrierTimeout() = 0;
 };
diff --git a/Engine/Plugins/Runtime/nDisplay/Source/DisplayCluster/Private/Network/Conversion/DisplayClusterNetworkDataConversion.cpp b/Engine/Plugins/Runtime/nDisplay/Source/DisplayCluster/Private/Network/Conversion/DisplayClusterNetworkDataConversion.cpp
--- a/Engine/Plugins/Runtime/nDisplay/Source/DisplayCluster/Private/Network/Conversion/DisplayClusterNetworkDataConversion.cpp
+++ b/Engine/Plugins/Runtime/nDisplay/Source/DisplayCluster/Private/Network/Conversion/DisplayClusterNetworkDataConversion.cpp
@@ -3,69 +3,77 @@
 #include "Network/Conversion/DisplayClusterNetworkDataConversion.h"
 
 #include "Cluster/DisplayClusterClusterEvent.h"
+#include "Dom/JsonObject.h"
+
 #include "Network/Packet/DisplayClusterPacketBinary.h"
 #include "Network/Packet/DisplayClusterPacketJson.h"
 #include "Network/Packet/DisplayClusterPacketInternal.h"
 #include "Network/Service/ClusterSync/DisplayClusterClusterSyncStrings.h"
 
+#include "Serialization/MemoryReader.h"
+#include "Serialization/MemoryWriter.h"
 
-#include "Dom/JsonObject.h"
 #include "JsonObjectConverter.h"
 
 
-namespace DisplayClusterNetworkDataConversion
+namespace UE::nDisplay::DisplayClusterNetworkDataConversion
 {
-	// Internal packet helpers
 	void JsonEventsFromInternalPacket(const TSharedPtr<FDisplayClusterPacketInternal>& Packet, TArray<TSharedPtr<FDisplayClusterClusterEventJson, ESPMode::ThreadSafe>>& JsonEvents)
 	{
-		TArray<FString> TextObjects;
-		Packet->GetTextObjects(DisplayClusterClusterSyncStrings::ArgumentsDefaultCategory, TextObjects, false);
+		TArray<TArray<uint8>> BinaryObjects;
+		Packet->GetBinObjects(DisplayClusterClusterSyncStrings::ArgumentsJsonEvents, BinaryObjects, false);
 
-		for (const auto& it : TextObjects)
+		for (const TArray<uint8>& BinaryObject : BinaryObjects)
 		{
 			TSharedPtr<FDisplayClusterClusterEventJson, ESPMode::ThreadSafe> JsonEvent = MakeShared<FDisplayClusterClusterEventJson, ESPMode::ThreadSafe>();
-			if (JsonEvent->DeserializeFromString(*it))
-			{
-				JsonEvents.Add(JsonEvent);
-			}
+
+			FMemoryReader MemoryReader(BinaryObject);
+			JsonEvent->Serialize(MemoryReader);
+
+			JsonEvents.Add(JsonEvent);
 		}
 	}
 
 	void JsonEventsToInternalPacket(const TArray<TSharedPtr<FDisplayClusterClusterEventJson, ESPMode::ThreadSafe>>& JsonEvents, TSharedPtr<FDisplayClusterPacketInternal>& Packet)
 	{
-		for (const auto& it : JsonEvents)
+		for (const TSharedPtr<FDisplayClusterClusterEventJson, ESPMode::ThreadSafe>& JsonEvent : JsonEvents)
 		{
-			const FString TextObject = it->SerializeToString();
-			Packet->AddTextObject(DisplayClusterClusterSyncStrings::ArgumentsDefaultCategory, TextObject);
+			TArray<uint8> BinaryObject;
+			FMemoryWriter MemoryWriter(BinaryObject);
+			JsonEvent->Serialize(MemoryWriter);
+
+			Packet->AddBinObject(DisplayClusterClusterSyncStrings::ArgumentsJsonEvents, BinaryObject);
 		}
 	}
 
 	void BinaryEventsFromInternalPacket(const TSharedPtr<FDisplayClusterPacketInternal>& Packet, TArray<TSharedPtr<FDisplayClusterClusterEventBinary, ESPMode::ThreadSafe>>& BinaryEvents)
 	{
 		TArray<TArray<uint8>> BinaryObjects;
-		Packet->GetBinObjects(DisplayClusterClusterSyncStrings::ArgumentsDefaultCategory, BinaryObjects, false);
+		Packet->GetBinObjects(DisplayClusterClusterSyncStrings::ArgumentsBinaryEvents, BinaryObjects, false);
 
-		for (const auto& it : BinaryObjects)
+		for (const TArray<uint8>& BinaryObject : BinaryObjects)
 		{
 			TSharedPtr<FDisplayClusterClusterEventBinary, ESPMode::ThreadSafe> BinaryEvent = MakeShared<FDisplayClusterClusterEventBinary, ESPMode::ThreadSafe>();
-			if (BinaryEvent->DeserializeFromByteArray(it))
-			{
-				BinaryEvents.Add(BinaryEvent);
-			}
+
+			FMemoryReader MemoryReader(BinaryObject);
+			BinaryEvent->Serialize(MemoryReader);
+
+			BinaryEvents.Add(BinaryEvent);
 		}
 	}
 
 	void BinaryEventsToInternalPacket(const TArray<TSharedPtr<FDisplayClusterClusterEventBinary, ESPMode::ThreadSafe>>& BinaryEvents, TSharedPtr<FDisplayClusterPacketInternal>& Packet)
 	{
-		for (const auto& it : BinaryEvents)
+		for (const TSharedPtr<FDisplayClusterClusterEventBinary, ESPMode::ThreadSafe>& BinaryEvent : BinaryEvents)
 		{
 			TArray<uint8> BinaryObject;
-			it->SerializeToByteArray(BinaryObject);
-			Packet->AddBinObject(DisplayClusterClusterSyncStrings::ArgumentsDefaultCategory, BinaryObject);
+			FMemoryWriter MemoryWriter(BinaryObject);
+			BinaryEvent->Serialize(MemoryWriter);
+
+			Packet->AddBinObject(DisplayClusterClusterSyncStrings::ArgumentsBinaryEvents, BinaryObject);
 		}
 	}
 
-	// JSON events conversion
 	bool JsonPacketToJsonEvent(const TSharedPtr<FDisplayClusterPacketJson>& Packet, FDisplayClusterClusterEventJson& OutBinaryEvent)
 	{
 		return FJsonObjectConverter::JsonObjectToUStruct(Packet->GetJsonData().ToSharedRef(), &OutBinaryEvent);
@@ -85,14 +93,13 @@ namespace DisplayClusterNetworkDataConversion
 		return Packet;
 	}
 
-	// Binary events conversion
 	bool BinaryPacketToBinaryEvent(const TSharedPtr<FDisplayClusterPacketBinary>& Packet, FDisplayClusterClusterEventBinary& OutBinaryEvent)
 	{
 		const TArray<uint8>& Buffer = Packet->GetPacketData();
 		return OutBinaryEvent.DeserializeFromByteArray(Buffer);
 	}
 
-	TSharedPtr<FDisplayClusterPacketBinary> BinaryEventToBinaryPacket(const FDisplayClusterClusterEventBinary& BinaryEvent)
+	TSharedPtr<FDisplayClusterPacketBinary> BinaryEventToBinaryPacket(FDisplayClusterClusterEventBinary& BinaryEvent)
 	{
 		TSharedPtr<FDisplayClusterPacketBinary> Packet = MakeShared<FDisplayClusterPacketBinary>();
 		TArray<uint8>& Buffer = Packet->GetPacketData();
diff --git a/Engine/Plugins/Runtime/nDisplay/Source/DisplayCluster/Private/Network/Conversion/DisplayClusterNetworkDataConversion.h b/Engine/Plugins/Runtime/nDisplay/Source/DisplayCluster/Private/Network/Conversion/DisplayClusterNetworkDataConversion.h
--- a/Engine/Plugins/Runtime/nDisplay/Source/DisplayCluster/Private/Network/Conversion/DisplayClusterNetworkDataConversion.h
+++ b/Engine/Plugins/Runtime/nDisplay/Source/DisplayCluster/Private/Network/Conversion/DisplayClusterNetworkDataConversion.h
@@ -2,7 +2,7 @@
 
 #pragma once
 
-#include "CoreMinimal.h"
+#include "Templates/SharedPointer.h"
 
 struct FDisplayClusterClusterEventBinary;
 struct FDisplayClusterClusterEventJson;
@@ -11,20 +11,33 @@ class FDisplayClusterPacketJson;
 class FDisplayClusterPacketInternal;
 
 
-// Some helpers that simplify export/import of internal non-trivial data types
-namespace DisplayClusterNetworkDataConversion
+/*
+ * Some helpers that simplify export/import of internal non-trivial data types
+ */ 
+namespace UE::nDisplay::DisplayClusterNetworkDataConversion
 {
-	// Internal packet helpers
+	/* [Internal] Extracts JSON events from binary objects (internal in-cluster replication only) */
 	void JsonEventsFromInternalPacket(const TSharedPtr<FDisplayClusterPacketInternal>& Packet, TArray<TSharedPtr<FDisplayClusterClusterEventJson, ESPMode::ThreadSafe>>& JsonEvents);
+
+	/* [Internal] Converts JSON events into binary objects (internal in-cluster replication only) */
 	void JsonEventsToInternalPacket(const TArray<TSharedPtr<FDisplayClusterClusterEventJson, ESPMode::ThreadSafe>>& JsonEvents, TSharedPtr<FDisplayClusterPacketInternal>& Packet);
+
+	/* [Internal] Extracts Binary events from binary objects (internal in-cluster replication only) */
 	void BinaryEventsFromInternalPacket(const TSharedPtr<FDisplayClusterPacketInternal>& Packet, TArray<TSharedPtr<FDisplayClusterClusterEventBinary, ESPMode::ThreadSafe>>& BinaryEvents);
+
+	/* [Internal] Converts Binary events into binary objects (internal in-cluster replication only) */
 	void BinaryEventsToInternalPacket(const TArray<TSharedPtr<FDisplayClusterClusterEventBinary, ESPMode::ThreadSafe>>& BinaryEvents, TSharedPtr<FDisplayClusterPacketInternal>& Packet);
 
-	// Json events conversion
+
+	/** [External] Extracts JSON event data from an external JSON packet */
 	bool JsonPacketToJsonEvent(const TSharedPtr<FDisplayClusterPacketJson>& Packet, FDisplayClusterClusterEventJson& OutBinaryEvent);
+
+	/** [External] Creates an external JSON packet with a specific JSON event set */
 	TSharedPtr<FDisplayClusterPacketJson> JsonEventToJsonPacket(const FDisplayClusterClusterEventJson& BinaryEvent);
 
-	// Binary events conversion
+	/** [External] Extracts Binary event data from an external Binary packet */
 	bool BinaryPacketToBinaryEvent(const TSharedPtr<FDisplayClusterPacketBinary>& Packet, FDisplayClusterClusterEventBinary& OutBinaryEvent);
-	TSharedPtr<FDisplayClusterPacketBinary> BinaryEventToBinaryPacket(const FDisplayClusterClusterEventBinary& BinaryEvent);
+
+	/** [External] Creates an external Binary packet with a specific Binary event set */
+	TSharedPtr<FDisplayClusterPacketBinary> BinaryEventToBinaryPacket(FDisplayClusterClusterEventBinary& BinaryEvent);
 }
diff --git a/Engine/Plugins/Runtime/nDisplay/Source/DisplayCluster/Private/Network/DisplayClusterClient.cpp b/Engine/Plugins/Runtime/nDisplay/Source/DisplayCluster/Private/Network/DisplayClusterClient.cpp
--- a/Engine/Plugins/Runtime/nDisplay/Source/DisplayCluster/Private/Network/DisplayClusterClient.cpp
+++ b/Engine/Plugins/Runtime/nDisplay/Source/DisplayCluster/Private/Network/DisplayClusterClient.cpp
@@ -45,22 +45,18 @@ bool FDisplayClusterClientBase::Connect(const FString& InAddr, const uint16 InPo
 
 void FDisplayClusterClientBase::Disconnect()
 {
-	FScopeLock Lock(&GetSyncObj());
-
 	UE_LOG(LogDisplayClusterNetwork, Log, TEXT("%s disconnecting..."), *GetName());
 
 	CloseSocket();
 }
 
-FSocket* FDisplayClusterClientBase::CreateSocket(const FString& InName)
+FSocket* FDisplayClusterClientBase::CreateSocket(const FString& InName, int32 LingerTime)
 {
-	FSocket* NewSocket = FTcpSocketBuilder(*InName).AsBlocking().Lingering(0);
+	FSocket* NewSocket = FTcpSocketBuilder(*InName).AsBlocking().Lingering(LingerTime);
 	check(NewSocket);
 
-	// Set socket properties (non-blocking, no delay)
+	// Set other socket properties
 	NewSocket->SetNoDelay(true);
-	NewSocket->SetLinger(true, 0);
-	NewSocket->SetNonBlocking(false);
 
 	return NewSocket;
 }
diff --git a/Engine/Plugins/Runtime/nDisplay/Source/DisplayCluster/Private/Network/DisplayClusterClient.h b/Engine/Plugins/Runtime/nDisplay/Source/DisplayCluster/Private/Network/DisplayClusterClient.h
--- a/Engine/Plugins/Runtime/nDisplay/Source/DisplayCluster/Private/Network/DisplayClusterClient.h
+++ b/Engine/Plugins/Runtime/nDisplay/Source/DisplayCluster/Private/Network/DisplayClusterClient.h
@@ -2,14 +2,11 @@
 
 #pragma once
 
-#include "CoreMinimal.h"
-
 #include "Network/IDisplayClusterClient.h"
+#include "Network/Configuration/DisplayClusterNetworkConfiguration.h"
 #include "Network/Transport/DisplayClusterSocketOperations.h"
 #include "Network/Transport/DisplayClusterSocketOperationsHelper.h"
 
-#include "Misc/DisplayClusterConstants.h"
-
 
 /**
  * Base DisplayCluster TCP client
@@ -19,32 +16,35 @@ class FDisplayClusterClientBase
 	, public FDisplayClusterSocketOperations
 {
 public:
-	FDisplayClusterClientBase(const FString& InName)
-		: FDisplayClusterSocketOperations(CreateSocket(InName), DisplayClusterConstants::net::PacketBufferSize, InName)
+
+	FDisplayClusterClientBase(const FString& InName, int32 InLingerTime = 0)
+		: FDisplayClusterSocketOperations(CreateSocket(InName, InLingerTime), UE::nDisplay::Network::Configuration::PacketBufferSize, InName)
 	{ }
 
 public:
-	// Connects to a server
-	virtual bool Connect(const FString& Address, const uint16 Port, const uint32 ConnectRetriesAmount, const uint32 ConnectRetryDelay) override;
 
-	// Terminates current connection
+	/** Connects to a server */
+	virtual bool Connect(const FString& Address, const uint16 Port, const uint32 ConnectRetriesAmount = 0, const uint32 ConnectRetryDelay = 0) override;
+
+	/** Terminates current connection */
 	virtual void Disconnect() override;
 
-	// Provides with net unit name
+	/** Returns client name */
 	virtual FString GetName() const override
 	{
 		return GetConnectionName();
 	}
 
+	/** Returns true if the client is currently connected */
 	virtual bool IsConnected() const override
 	{
 		return IsOpen();
 	}
 
 protected:
-	// Creates client socket
-	FSocket* CreateSocket(const FString& InName);
 
+	/** Creates client socket */
+	FSocket* CreateSocket(const FString& InName, int32 LingerTime);
 };
 
 
@@ -54,8 +54,8 @@ class FDisplayClusterClient
 	, protected FDisplayClusterSocketOperationsHelper<TPacketType>
 {
 public:
-	FDisplayClusterClient(const FString& InName)
-		: FDisplayClusterClientBase(InName)
+	FDisplayClusterClient(const FString& InName, int32 InLingerTime = 0)
+		: FDisplayClusterClientBase(InName, InLingerTime)
 		, FDisplayClusterSocketOperationsHelper<TPacketType>(*this, InName)
 	{ }
 };
diff --git a/Engine/Plugins/Runtime/nDisplay/Source/DisplayCluster/Private/Network/DisplayClusterNetworkTypes.h b/Engine/Plugins/Runtime/nDisplay/Source/DisplayCluster/Private/Network/DisplayClusterNetworkTypes.h
--- a/Engine/Plugins/Runtime/nDisplay/Source/DisplayCluster/Private/Network/DisplayClusterNetworkTypes.h
+++ b/Engine/Plugins/Runtime/nDisplay/Source/DisplayCluster/Private/Network/DisplayClusterNetworkTypes.h
@@ -2,8 +2,8 @@
 
 #pragma once
 
-#include "CoreMinimal.h"
 #include "Interfaces/IPv4/IPv4Endpoint.h"
+#include "Misc/Optional.h"
 
 class FSocket;
 
diff --git a/Engine/Plugins/Runtime/nDisplay/Source/DisplayCluster/Private/Network/DisplayClusterServer.cpp b/Engine/Plugins/Runtime/nDisplay/Source/DisplayCluster/Private/Network/DisplayClusterServer.cpp
--- a/Engine/Plugins/Runtime/nDisplay/Source/DisplayCluster/Private/Network/DisplayClusterServer.cpp
+++ b/Engine/Plugins/Runtime/nDisplay/Source/DisplayCluster/Private/Network/DisplayClusterServer.cpp
@@ -1,17 +1,18 @@
 // Copyright Epic Games, Inc. All Rights Reserved.
 
 #include "Network/DisplayClusterServer.h"
+
+#include "Engine/GameEngine.h"
+
+#include "Misc/DateTime.h"
+#include "Misc/DisplayClusterLog.h"
+#include "Misc/ScopeLock.h"
+
 #include "Network/Session/IDisplayClusterSession.h"
 #include "Network/Listener/DisplayClusterTcpListener.h"
 #include "Network/DisplayClusterNetworkTypes.h"
 
-#include "Engine/GameEngine.h"
 #include "TimerManager.h"
-#include "Misc/DateTime.h"
-
-#include "Misc/ScopeLock.h"
-#include "Misc/DisplayClusterConstants.h"
-#include "Misc/DisplayClusterLog.h"
 
 
 // Give 10 seconds minimum for every session that finished its job to finalize
@@ -19,14 +20,24 @@
 const double FDisplayClusterServer::CleanSessionResourcesSafePeriod = 10.f;
 
 
-FDisplayClusterServer::FDisplayClusterServer(const FString& InName)
-	: Name(InName)
+FDisplayClusterServer::FDisplayClusterServer(const FString& InInstanceName)
+	: InstanceName(InInstanceName)
 {
 	// Subscribe for session events
 	OnSessionOpened().AddRaw(this, &FDisplayClusterServer::ProcessSessionOpened);
 	OnSessionClosed().AddRaw(this, &FDisplayClusterServer::ProcessSessionClosed);
 }
 
+FDisplayClusterServer::~FDisplayClusterServer()
+{
+	// Unsubscribe from events
+	OnSessionOpened().RemoveAll(this);
+	OnSessionClosed().RemoveAll(this);
+
+	ShutdownImpl();
+}
+
+
 bool FDisplayClusterServer::Start(const FString& Address, const uint16 Port)
 {
 	FScopeLock Lock(&ServerStateCritSec);
@@ -37,16 +48,16 @@
 		return true;
 	}
 
-	UE_LOG(LogDisplayClusterNetwork, Log, TEXT("%s starting..."), *Name);
+	UE_LOG(LogDisplayClusterNetwork, Log, TEXT("%s starting..."), *InstanceName);
 
 	// Instantiate internal listener
-	Listener = MakeShared<FDisplayClusterTcpListener>(false, GetName() + FString("_listener"));
+	Listener = MakeShared<FDisplayClusterTcpListener>(false, InstanceName + FString(TEXT("_listener")));
 	// Bind connection handler method
 	Listener->OnConnectionAccepted().BindRaw(this, &FDisplayClusterServer::ConnectionHandler);
 
 	if (!Listener->StartListening(Address, Port))
 	{
-		UE_LOG(LogDisplayClusterNetwork, Error, TEXT("%s couldn't start connection listener [%s:%d]"), *Name, *Address, Port);
+		UE_LOG(LogDisplayClusterNetwork, Error, TEXT("%s couldn't start connection listener [%s:%d]"), *InstanceName, *Address, Port);
 		return false;
 	}
 
@@ -56,10 +67,8 @@
 	return bIsRunning;
 }
 
-bool FDisplayClusterServer::Start(TSharedPtr<FDisplayClusterTcpListener>& ExternalListener)
+bool FDisplayClusterServer::Start(TSharedRef<FDisplayClusterTcpListener>& ExternalListener)
 {
-	check(ExternalListener.IsValid());
-
 	FScopeLock Lock(&ServerStateCritSec);
 
 	// Set connection handler up
@@ -83,49 +92,7 @@
 
 void FDisplayClusterServer::Shutdown()
 {
-	TArray<TSharedPtr<IDisplayClusterSession>> ActiveSessionsToWait;
-
-	{
-		FScopeLock Lock(&ServerStateCritSec);
-
-		// Nothing to do if not running
-		if (!IsRunning())
-		{
-			return;
-		}
-
-		UE_LOG(LogDisplayClusterNetwork, Log, TEXT("%s stopping..."), *Name);
-
-		// Update server state so no new sessions allowed anymore
-		bIsRunning = false;
-		// Stop listening for new connections
-		Listener.Reset();
-
-		// Make copy of all active sessions to wait for completion later
-		ActiveSessions.GenerateValueArray(ActiveSessionsToWait);
-
-		// Send stop command to all active sessions
-		for (auto It = ActiveSessions.CreateIterator(); It; ++It)
-		{
-			It->Value->StopSession(false);
-		}
-	}
-
-	// Now wait unless all session threads are finished. We're waiting outside of the
-	// critical section so the sessions won't be deadlocked notifying the server about their state changes.
-	for (TSharedPtr<IDisplayClusterSession>& Session : ActiveSessionsToWait)
-	{
-		Session->WaitForCompletion();
-	}
-
-	{
-		FScopeLock Lock(&ServerStateCritSec);
-
-		// Finally, cleanup internals
-		ActiveSessions.Reset();
-		PendingSessions.Reset();
-		PendingKillSessions.Reset();
-	}
+	ShutdownImpl();
 }
 
 bool FDisplayClusterServer::IsRunning() const
@@ -188,6 +155,7 @@
 		SessionInfo.Socket->SetNoDelay(true);
 		SessionInfo.Socket->SetLinger(true, 0);
 		SessionInfo.Socket->SetNonBlocking(false);
+		SessionInfo.Socket->SetReuseAddr(false);
 		// Update other session info
 		SessionInfo.Protocol  = GetProtocolName();
 		SessionInfo.SessionId = IncrementalSessionId++;
@@ -248,6 +216,54 @@
 	}
 }
 
+void FDisplayClusterServer::ShutdownImpl()
+{
+
+	TArray<TSharedPtr<IDisplayClusterSession>> ActiveSessionsToWait;
+
+	{
+		FScopeLock Lock(&ServerStateCritSec);
+
+		// Nothing to do if not running
+		if (!IsRunning())
+		{
+			return;
+		}
+
+		UE_LOG(LogDisplayClusterNetwork, Log, TEXT("%s stopping..."), *InstanceName);
+
+		// Update server state so no new sessions allowed anymore
+		bIsRunning = false;
+		// Stop listening for new connections
+		Listener.Reset();
+
+		// Make copy of all active sessions to wait for completion later
+		ActiveSessions.GenerateValueArray(ActiveSessionsToWait);
+
+		// Send stop command to all active sessions
+		for (auto It = ActiveSessions.CreateIterator(); It; ++It)
+		{
+			It->Value->StopSession(false);
+		}
+	}
+
+	// Now wait unless all session threads are finished. We're waiting outside of the
+	// critical section so the sessions won't be deadlocked notifying the server about their state changes.
+	for (TSharedPtr<IDisplayClusterSession>& Session : ActiveSessionsToWait)
+	{
+		Session->WaitForCompletion();
+	}
+
+	{
+		FScopeLock Lock(&ServerStateCritSec);
+
+		// Finally, cleanup internals
+		ActiveSessions.Reset();
+		PendingSessions.Reset();
+		PendingKillSessions.Reset();
+	}
+}
+
 void FDisplayClusterServer::CleanPendingKillSessions()
 {
 	FScopeLock Lock(&SessionsCritSec);
@@ -259,7 +275,7 @@
 	{
 		if ((CurrentTime - SessionIt->Value->GetSessionInfo().TimeEnd) > FDisplayClusterServer::CleanSessionResourcesSafePeriod)
 		{
-			UE_LOG(LogDisplayClusterNetwork, Verbose, TEXT("%s - cleaning session resources, SessionId=%llu..."), *Name, SessionIt->Key);
+			UE_LOG(LogDisplayClusterNetwork, Verbose, TEXT("%s - cleaning session resources, SessionId=%llu..."), *InstanceName, SessionIt->Key);
 			SessionIt.RemoveCurrent();
 		}
 	}
diff --git a/Engine/Plugins/Runtime/nDisplay/Source/DisplayCluster/Private/Network/DisplayClusterServer.h b/Engine/Plugins/Runtime/nDisplay/Source/DisplayCluster/Private/Network/DisplayClusterServer.h
--- a/Engine/Plugins/Runtime/nDisplay/Source/DisplayCluster/Private/Network/DisplayClusterServer.h
+++ b/Engine/Plugins/Runtime/nDisplay/Source/DisplayCluster/Private/Network/DisplayClusterServer.h
@@ -2,13 +2,9 @@
 
 #pragma once
 
-#include "CoreMinimal.h"
-
 #include "Network/IDisplayClusterServer.h"
 #include "Network/DisplayClusterNetworkTypes.h"
 
-#include "Containers/Queue.h"
-
 struct FIPv4Endpoint;
 class FSocket;
 class IDisplayClusterSession;
@@ -22,113 +18,128 @@ class FDisplayClusterServer
 	: public IDisplayClusterServer
 {
 public:
-	// Minimal time (seconds) before cleaning resources of the 'pending kill' sessions
+
+	/** Minimal time(seconds) before cleaning resources of the 'pending kill' sessions */
 	static const double CleanSessionResourcesSafePeriod;
 
 public:
-	FDisplayClusterServer(const FString& Name);
-	virtual ~FDisplayClusterServer() = default;
+
+	FDisplayClusterServer(const FString& InstanceName);
+	virtual ~FDisplayClusterServer();
 
 public:
-	//////////////////////////////////////////////////////////////////////////////////////////////
-	// IDisplayClusterServer
-	//////////////////////////////////////////////////////////////////////////////////////////////
+
+	//~ Begin IDisplayClusterServer
 	
-	// Start server on a specific socket
+	/** Start server on a specific socket */
 	virtual bool Start(const FString& Address, const uint16 Port) override;
 
-	// Start server with a specified listener
-	virtual bool Start(TSharedPtr<FDisplayClusterTcpListener>& ExternalListener) override;
+	/** Start server with a specified listener */
+	virtual bool Start(TSharedRef<FDisplayClusterTcpListener>& ExternalListener) override;
 
-	// Stop server
+	/** Stop server */
 	virtual void Shutdown() override;
 
-	// Returns current server state
+	/** Returns current server state */
 	virtual bool IsRunning() const override;
 
-	// Server name
+	/** Server name */
 	virtual FString GetName() const override
 	{
-		return Name;
+		return InstanceName;
 	}
 
-	// Server address
+	/** Server address */
 	virtual FString GetAddress() const override;
 
-	// Server port
+	/** Server port */
 	virtual uint16 GetPort() const override;
 
-	// Kills all sessions of a specified cluster node
+	/** Kills all sessions of a specified cluster node */
 	virtual void KillSession(const FString& NodeId) override;
 
-	// Connection validation delegate
+	/** Connection validation delegate */
 	virtual FIsConnectionAllowedDelegate& OnIsConnectionAllowed() override
 	{
 		return IsConnectionAllowedDelegate;
 	}
 
-	// Session opened event
+	/** Session opened event */
 	virtual FSessionOpenedEvent& OnSessionOpened() override
 	{
 		return SessionOpenedEvent;
 	}
 
-	// Session closed event
+	/** Session closed event */
 	virtual FSessionClosedEvent& OnSessionClosed() override
 	{
 		return SessionClosedEvent;
 	}
 
-	// Node failed event
-	virtual FNodeFailedEvent& OnNodeFailed() override
-	{
-		return NodeFailedEvent;
-	}
+	//~ End IDisplayClusterServer
 
 protected:
-	// Allow to specify custom session class
+
+	/** Allow to specify custom session class */
 	virtual TSharedPtr<IDisplayClusterSession> CreateSession(FDisplayClusterSessionInfo& SessionInfo) = 0;
 
-	// Handle incoming connections
+private:
+
+	/** Handle incoming connections */
 	bool ConnectionHandler(FDisplayClusterSessionInfo& SessionInfo);
 
-	// Callback on session opened
+	/** Callback on session opened */
 	void ProcessSessionOpened(const FDisplayClusterSessionInfo& SessionInfo);
-	// Callback on session closed
+
+	/** Callback on session closed */
 	void ProcessSessionClosed(const FDisplayClusterSessionInfo& SessionInfo);
 
-private:
-	// Free resources of the sessions that already finished their job
+	/** Non-virtual shutdown implementation */
+	void ShutdownImpl();
+
+	/** Free resources of the sessions that already finished their job */
 	void CleanPendingKillSessions();
 
 private:
-	// Server data
-	const FString Name;
+
+	/** Server instance name */
+	const FString InstanceName;
 	
-	// Server running state
+	/** Server running state */
 	bool bIsRunning = false;
 
-	// Socket listeners
+	/** Socket listener */
 	TSharedPtr<FDisplayClusterTcpListener> Listener;
 
-	// Server events/delegates
+	/** Connection approval delegate */
 	FIsConnectionAllowedDelegate IsConnectionAllowedDelegate;
+
+	/** Session opened event */
 	FSessionOpenedEvent SessionOpenedEvent;
+
+	/** Session closed event */
 	FSessionClosedEvent SessionClosedEvent;
-	FNodeFailedEvent    NodeFailedEvent;
 
 private:
-	// Session counter used for session ID generation
+
+	/** Session counter used for session ID generation */
 	uint64 IncrementalSessionId = 0;
 
-	// Pending sessions
+
+	/** Pending sessions */
 	TMap<uint64, TSharedPtr<IDisplayClusterSession>> PendingSessions;
-	// Active sessions
+
+	/** Active sessions */
 	TMap<uint64, TSharedPtr<IDisplayClusterSession>> ActiveSessions;
-	// Closed sessions, awaiting for cleaning
+
+	/** Closed sessions, awaiting for cleaning */
 	TMap<uint64, TSharedPtr<IDisplayClusterSession>> PendingKillSessions;
 
-	// Sync access
+private:
+
+	/** Critical section to manipulate server states */
 	mutable FCriticalSection ServerStateCritSec;
+
+	/** Critical section to manipulate the connection sessions */
 	mutable FCriticalSection SessionsCritSec;
 };
diff --git a/Engine/Plugins/Runtime/nDisplay/Source/DisplayCluster/Private/Network/IDisplayClusterClient.h b/Engine/Plugins/Runtime/nDisplay/Source/DisplayCluster/Private/Network/IDisplayClusterClient.h
--- a/Engine/Plugins/Runtime/nDisplay/Source/DisplayCluster/Private/Network/IDisplayClusterClient.h
+++ b/Engine/Plugins/Runtime/nDisplay/Source/DisplayCluster/Private/Network/IDisplayClusterClient.h
@@ -2,10 +2,9 @@
 
 #pragma once
 
-#include "CoreMinimal.h"
-
 #include "Cluster/IDisplayClusterGenericBarriersClient.h"
 
+
 /**
  * DisplayCluster TCP client interface
  */
@@ -15,15 +14,16 @@ public:
 	virtual ~IDisplayClusterClient() = default;
 
 public:
-	// Connects to a server
+
+	/** Connects to a server */
 	virtual bool Connect(const FString& Address, const uint16 Port, const uint32 ConnectRetriesAmount, const uint32 ConnectRetryDelay) = 0;
-	// Terminates current connection
+
+	/** Terminates current connection */
 	virtual void Disconnect() = 0;
-	// Returns connection status
+
+	/** Returns connection status */
 	virtual bool IsConnected() const = 0;
-	// Returns client name
+
+	/** Returns client name */
 	virtual FString GetName() const = 0;
 };
-
-// Custom deleter for IDisplayClusterClient. See FDisplayClusterGenericClientDeleter for more info.
-using FDisplayClusterClientDeleter = FDisplayClusterGenericClientDeleter<IDisplayClusterClient>;
diff --git a/Engine/Plugins/Runtime/nDisplay/Source/DisplayCluster/Private/Network/IDisplayClusterServer.h b/Engine/Plugins/Runtime/nDisplay/Source/DisplayCluster/Private/Network/IDisplayClusterServer.h
--- a/Engine/Plugins/Runtime/nDisplay/Source/DisplayCluster/Private/Network/IDisplayClusterServer.h
+++ b/Engine/Plugins/Runtime/nDisplay/Source/DisplayCluster/Private/Network/IDisplayClusterServer.h
@@ -2,7 +2,6 @@
 
 #pragma once
 
-#include "CoreMinimal.h"
 #include "Network/DisplayClusterNetworkTypes.h"
 
 class FDisplayClusterTcpListener;
@@ -14,51 +13,49 @@ class FDisplayClusterTcpListener;
 class IDisplayClusterServer
 {
 public:
+
 	virtual ~IDisplayClusterServer() = default;
 
 public:
-	// Start server on a specific socket
+
+	/** Start server on a specific socket */
 	virtual bool Start(const FString& Address, const uint16 Port) = 0;
-	// Start server with a specified listener
-	virtual bool Start(TSharedPtr<FDisplayClusterTcpListener>& Listener) = 0;
-	// Stop server
+
+	/** Start server with a specified listener */
+	virtual bool Start(TSharedRef<FDisplayClusterTcpListener>& Listener) = 0;
+
+	/** Stop server */
 	virtual void Shutdown() = 0;
 
-	// Returns current server state
+	/** Returns true if server is currently running */
 	virtual bool IsRunning() const = 0;
 
-	// Returns server name
+	/** Returns server instance name */
 	virtual FString GetName() const = 0;
-	// Returns server address
+
+	/** Returns server address */
 	virtual FString GetAddress() const = 0;
-	// Returns server port
+
+	/** Returns server port */
 	virtual uint16 GetPort() const = 0;
-	// Returns server protocol name
+
+	/** Returns server protocol name */
 	virtual FString GetProtocolName() const = 0;
 
-	// Kill all sessions of a specific node
+	/** Kill all sessions of a specific node */
 	virtual void KillSession(const FString& NodeId) = 0;
 
-	// Returns connection validation delegate
+public:
+
+	/** Connection validation delegate */
 	DECLARE_DELEGATE_RetVal_OneParam(bool, FIsConnectionAllowedDelegate, const FDisplayClusterSessionInfo&);
 	virtual FIsConnectionAllowedDelegate& OnIsConnectionAllowed() = 0;
 
-	// Session opened event
+	/** Session opened event */
 	DECLARE_EVENT_OneParam(IDisplayClusterServer, FSessionOpenedEvent, const FDisplayClusterSessionInfo&);
 	virtual FSessionOpenedEvent& OnSessionOpened() = 0;
 
-	// Session closed event
+	/** Session closed event */
 	DECLARE_EVENT_OneParam(IDisplayClusterServer, FSessionClosedEvent, const FDisplayClusterSessionInfo&);
 	virtual FSessionClosedEvent& OnSessionClosed() = 0;
-
-
-	// Cluster node lost
-	enum class ENodeFailType : uint8
-	{
-		BarrierTimeOut,
-		ConnectionLost,
-	};
-
-	DECLARE_EVENT_TwoParams(IDisplayClusterServer, FNodeFailedEvent, const FString&, ENodeFailType);
-	virtual FNodeFailedEvent& OnNodeFailed() = 0;
 };
diff --git a/Engine/Plugins/Runtime/nDisplay/Source/DisplayCluster/Private/Network/Listener/DisplayClusterHelloMessageStrings.h b/Engine/Plugins/Runtime/nDisplay/Source/DisplayCluster/Private/Network/Listener/DisplayClusterHelloMessageStrings.h
--- a/Engine/Plugins/Runtime/nDisplay/Source/DisplayCluster/Private/Network/Listener/DisplayClusterHelloMessageStrings.h
+++ b/Engine/Plugins/Runtime/nDisplay/Source/DisplayCluster/Private/Network/Listener/DisplayClusterHelloMessageStrings.h
@@ -2,17 +2,20 @@
 
 #pragma once
 
+#include "HAL/Platform.h"
+
 
 /**
  * Hello message strings
  */
 namespace DisplayClusterHelloMessageStrings
 {
-	constexpr static auto ArgumentsDefaultCategory = "HC";
+	constexpr static const TCHAR* ArgumentsDefaultCategory = TEXT("HC");
 
 	namespace Hello
 	{
-		constexpr static auto Name      = "Hello";
-		constexpr static auto ArgNodeId = "NodeId";
+		constexpr static const TCHAR* Name        = TEXT("Hello");
+		constexpr static const TCHAR* TypeRequest = TEXT("Request");
+		constexpr static const TCHAR* ArgNodeId   = TEXT("NodeId");
 	};
 };
diff --git a/Engine/Plugins/Runtime/nDisplay/Source/DisplayCluster/Private/Network/Listener/DisplayClusterTcpListener.cpp b/Engine/Plugins/Runtime/nDisplay/Source/DisplayCluster/Private/Network/Listener/DisplayClusterTcpListener.cpp
--- a/Engine/Plugins/Runtime/nDisplay/Source/DisplayCluster/Private/Network/Listener/DisplayClusterTcpListener.cpp
+++ b/Engine/Plugins/Runtime/nDisplay/Source/DisplayCluster/Private/Network/Listener/DisplayClusterTcpListener.cpp
@@ -30,19 +30,14 @@ FDisplayClusterTcpListener::FDisplayClusterTcpListener(bool bIsShared, const FSt
 FDisplayClusterTcpListener::~FDisplayClusterTcpListener()
 {
 	// Just free resources by stopping the listening
-	StopListening(false);
-
-	// ThreadObj is a TUniquePtr that is deleted after the destructor of this instance has completed.
-	// But internally ThreadObj accesses this instance and calls its methods from its destructor.
-	// Therefore, we need to delete the ThreadObj earlier, inside the destructor of this instance.
-	ThreadObj.Reset();
+	StopListening(true);
 }
 
 bool FDisplayClusterTcpListener::StartListening(const FString& InAddr, const uint16 InPort)
 {
 	FScopeLock Lock(&InternalsCS);
 
-	if (bIsListening == true)
+	if (IsListening())
 	{
 		return true;
 	}
@@ -60,7 +55,7 @@ bool FDisplayClusterTcpListener::StartListening(const FIPv4Endpoint& InEndpoint)
 {
 	FScopeLock Lock(&InternalsCS);
 
-	if (bIsListening == true)
+	if (IsListening())
 	{
 		return true;
 	}
@@ -68,15 +63,13 @@ bool FDisplayClusterTcpListener::StartListening(const FIPv4Endpoint& InEndpoint)
 	// Save new endpoint
 	Endpoint = InEndpoint;
 
-	UE_LOG(LogDisplayClusterNetwork, Log, TEXT("TCP listener %s: started listening to %s:%d..."), *Name, *InEndpoint.Address.ToString(), Endpoint.Port);
-
 	// Create listening thread
 	ThreadObj.Reset(FRunnableThread::Create(this, *(Name + FString("_thread")), 128 * 1024, TPri_Normal));
 	ensure(ThreadObj);
 
-	// Update state
-	bIsListening = ThreadObj.IsValid();
-	
+	// This must have been updated from the worker thread. FRunnableThread::Create uses
+	// internal event ThreadInitSyncEvent to wait until FRunnable::Init is done (it's
+	// FDisplayClusterTcpListener::Init here).
 	return bIsListening;
 }
 
@@ -84,12 +77,14 @@ void FDisplayClusterTcpListener::StopListening(bool bWaitForCompletion)
 {
 	FScopeLock Lock(&InternalsCS);
 
-	if (bIsListening == false)
+	if (!IsListening())
 	{
 		return;
 	}
 
-	UE_LOG(LogDisplayClusterNetwork, Log, TEXT("TCP listener %s: stopped listening to %s:%d..."), *Name, *Endpoint.Address.ToString(), Endpoint.Port);
+	FString ListeningAddr;
+	uint16 ListeningPort = 0;
+	GetListeningParams(ListeningAddr, ListeningPort);
 
 	// Ask runnable to stop
 	Stop();
@@ -99,6 +94,8 @@ void FDisplayClusterTcpListener::StopListening(bool bWaitForCompletion)
 	{
 		WaitForCompletion();
 	}
+
+	UE_LOG(LogDisplayClusterNetwork, Log, TEXT("TCP listener %s: stopped listening to %s:%d..."), *Name, *ListeningAddr, ListeningPort);
 }
 
 void FDisplayClusterTcpListener::WaitForCompletion()
@@ -112,6 +109,51 @@ void FDisplayClusterTcpListener::WaitForCompletion()
 	}
 }
 
+bool FDisplayClusterTcpListener::GetListeningParams(FString& OutAddr, uint16& OutPort)
+{
+	if (!IsListening())
+	{
+		return false;
+	}
+
+	OutAddr = GetListeningHost();
+	OutPort = GetListeningPort();
+
+	return true;
+}
+
+FString FDisplayClusterTcpListener::GetListeningHost() const
+{
+	if (IsListening() && SocketObj)
+	{
+		TSharedRef<FInternetAddr> ListeningAddress = ISocketSubsystem::Get(PLATFORM_SOCKETSUBSYSTEM)->CreateInternetAddr();
+		SocketObj->GetAddress(*ListeningAddress);
+
+		uint32 IpAddrAsNum = 0;
+		ListeningAddress->GetIp(IpAddrAsNum);
+
+		FIPv4Address IpAddress(IpAddrAsNum);
+		const FString IpAddrAsString = IpAddress.ToString();
+
+		return IpAddrAsString;
+	}
+
+	return FString();
+}
+
+uint16 FDisplayClusterTcpListener::GetListeningPort() const
+{
+	if (IsListening() && SocketObj)
+	{
+		const int32 PortNum32 = SocketObj->GetPortNo();
+		check(PortNum32 > 0 && PortNum32 <= TNumericLimits<uint16>::Max());
+		const uint16 PortNum16 = static_cast<uint16>(PortNum32 & TNumericLimits<uint16>::Max());
+		return PortNum16;
+	}
+
+	return 0;
+}
+
 FDisplayClusterTcpListener::FConnectionAcceptedDelegate& FDisplayClusterTcpListener::OnConnectionAccepted(const FString& ProtocolName)
 {
 	FScopeLock Lock(&InternalsCS);
@@ -129,14 +171,23 @@ FDisplayClusterTcpListener::FConnectionAcceptedDelegate& FDisplayClusterTcpListe
 bool FDisplayClusterTcpListener::Init()
 {
 	// Create socket
-	SocketObj = FTcpSocketBuilder(*Name).AsBlocking().BoundToEndpoint(Endpoint).Listening(128);
+	SocketObj = FTcpSocketBuilder(*Name)
+		.AsBlocking()
+		.Lingering(0)
+		.AsReusable(false)
+		.Listening(128)
+		.BoundToEndpoint(Endpoint);
 
 	if (SocketObj)
 	{
 		// Set socket properties (non-blocking, no delay)
 		SocketObj->SetNoDelay(true);
-		SocketObj->SetLinger(true, 0);
 		SocketObj->SetNonBlocking(false);
+		SocketObj->SetLinger(true, 0);
+		SocketObj->SetReuseAddr(false);
+	
+		// Update listening state
+		bIsListening = true;
 	}
 
 	return SocketObj != nullptr;
@@ -151,6 +202,12 @@ uint32 FDisplayClusterTcpListener::Run()
 
 	if (SocketObj)
 	{
+		FString ListeningAddr;
+		uint16 ListeningPort = 0;
+		GetListeningParams(ListeningAddr, ListeningPort);
+
+		UE_LOG(LogDisplayClusterNetwork, Log, TEXT("TCP listener %s: started listening to %s:%d..."), *Name, *ListeningAddr, ListeningPort);
+
 		while (FSocket* NewSock = SocketObj->Accept(*RemoteAddress, TEXT("DisplayCluster session")))
 		{
 			UE_LOG(LogDisplayClusterNetwork, Verbose, TEXT("TCP listener %s: New incoming connection: %s"), *Name, *RemoteAddress->ToString(true));
diff --git a/Engine/Plugins/Runtime/nDisplay/Source/DisplayCluster/Private/Network/Listener/DisplayClusterTcpListener.h b/Engine/Plugins/Runtime/nDisplay/Source/DisplayCluster/Private/Network/Listener/DisplayClusterTcpListener.h
--- a/Engine/Plugins/Runtime/nDisplay/Source/DisplayCluster/Private/Network/Listener/DisplayClusterTcpListener.h
+++ b/Engine/Plugins/Runtime/nDisplay/Source/DisplayCluster/Private/Network/Listener/DisplayClusterTcpListener.h
@@ -2,17 +2,13 @@
 
 #pragma once
 
-#include "CoreMinimal.h"
-
-#include "Network/DisplayClusterNetworkTypes.h"
-
-#include "Sockets.h"
 #include "HAL/Runnable.h"
-
 #include "Interfaces/IPv4/IPv4Endpoint.h"
+#include "Network/DisplayClusterNetworkTypes.h"
 
-#include "Misc/DisplayClusterConstants.h"
+#include "Sockets.h"
 
+#include <atomic>
 
 
 /**
@@ -25,95 +21,93 @@ class FDisplayClusterTcpListener
 	: protected FRunnable
 {
 public:
+
 	FDisplayClusterTcpListener(bool bIsShared, const FString& InName);
 	virtual ~FDisplayClusterTcpListener();
 
 public:
-	// Start listening to a specified socket
+
+	/** Start listening to address:port */
 	bool StartListening(const FString& InAddr, const uint16 InPort);
-	// Start listening to a specified endpoint
+
+	/** Start listening to an endpoint */
 	bool StartListening(const FIPv4Endpoint& Endpoint);
-	// Stop listening
+
+	/** Stop listening */
 	void StopListening(bool bWaitForCompletion);
-	// Wait unless working thread is finished
+
+	/** Wait unless working thread is finished */
 	void WaitForCompletion();
 
-	// Is currently listening
+	/** Is currently listening */
 	bool IsListening() const
 	{
 		return bIsListening;
 	}
 
-	// Returns listenting parameters
-	bool GetListeningParams(FString& OutAddr, uint16& OutPort)
-	{
-		if (IsListening())
-		{
-			OutAddr = GetListeningHost();
-			OutPort = GetListeningPort();
-		}
+	/** Returns listenting address & port */
+	bool GetListeningParams(FString& OutAddr, uint16& OutPort);
 
-		return false;
-	}
+	/** Returns listening host */
+	FString GetListeningHost() const;
 
-	// Returns listening host
-	FString GetListeningHost() const
-	{
-		return IsListening() ? Endpoint.Address.ToString() : FString();
-	}
+	/** Returns listening port */
+	uint16 GetListeningPort() const;
 
-	// Returns listening port
-	uint16 GetListeningPort() const
-	{
-		return IsListening() ? Endpoint.Port : 0;
-	}
-
-	// Delegate for processing incoming connections
+	/** Delegate for processing incoming connections */
 	DECLARE_DELEGATE_RetVal_OneParam(bool, FConnectionAcceptedDelegate, FDisplayClusterSessionInfo&);
 
-	// Returns connection validation delegate
+	/** Returns connection validation delegate */
 	FConnectionAcceptedDelegate& OnConnectionAccepted()
 	{
 		return ConnectionAcceptedDelegate;
 	}
 
-	// Returns protocol-bound delegate
+	/** Returns protocol - bound delegate */
 	FConnectionAcceptedDelegate& OnConnectionAccepted(const FString& ProtocolName);
 
 protected:
-	//////////////////////////////////////////////////////////////////////////////////////////////
-	// FRunnable
-	//////////////////////////////////////////////////////////////////////////////////////////////
+
+	//~Begin FRunnable
 	virtual bool Init() override;
 	virtual uint32 Run() override;
 	virtual void Stop() override;
 	virtual void Exit() override;
+	//~End FRunnable
 
 protected:
+
+	/** Fills endpoint data with given address and port */
 	bool GenIPv4Endpoint(const FString& Addr, const uint16 Port, FIPv4Endpoint& EP) const;
 
 private:
-	// Socket name
+
+	/** Socket name */
 	FString Name;
-	// Listening socket
+
+	/** Listening socket */
 	FSocket* SocketObj = nullptr;
-	// Listening endpoint
+
+	/** Listening endpoint */
 	FIPv4Endpoint Endpoint;
-	// Holds the thread object
+
+	/** Holds the listening thread object */
 	TUniquePtr<FRunnableThread> ThreadObj;
-	// Listening state
-	bool bIsListening = false;
 
-	// Holds a delegate to be invoked when an incoming connection has been accepted.
+	/** Current listening state */
+	std::atomic<bool> bIsListening = false;
+
+	/** Holds a delegate to be invoked when an incoming connection has been accepted. */
 	FConnectionAcceptedDelegate ConnectionAcceptedDelegate;
 
 private:
-	// Handles incoming connections
+
+	/** Handles incoming connections */
 	bool ProcessIncomingConnection(FDisplayClusterSessionInfo& SessionInfo);
 
-	// ProtocolName-to-ServiceDelegate map for transferring connection ownership to an appropriate server
+	/** ProtocolName - to - ServiceDelegate map for transferring connection ownership to an appropriate server */
 	TMap<FString, FConnectionAcceptedDelegate> ProtocolDispatchingMap;
 
-	// Critical section for access control
+	/** Critical section for access control */
 	FCriticalSection InternalsCS;
 };
diff --git a/Engine/Plugins/Runtime/nDisplay/Source/DisplayCluster/Private/Network/Packet/DisplayClusterPacketBinary.h b/Engine/Plugins/Runtime/nDisplay/Source/DisplayCluster/Private/Network/Packet/DisplayClusterPacketBinary.h
--- a/Engine/Plugins/Runtime/nDisplay/Source/DisplayCluster/Private/Network/Packet/DisplayClusterPacketBinary.h
+++ b/Engine/Plugins/Runtime/nDisplay/Source/DisplayCluster/Private/Network/Packet/DisplayClusterPacketBinary.h
@@ -2,8 +2,6 @@
 
 #pragma once
 
-#include "CoreMinimal.h"
-
 #include "Network/Packet/IDisplayClusterPacket.h"
 
 
