[nDisplay] Base failover implementation (v1.5, P-node replacement)

diff --git a/Engine/Plugins/Runtime/nDisplay/Source/DisplayCluster/Private/Network/Packet/DisplayClusterPacketJson.h b/Engine/Plugins/Runtime/nDisplay/Source/DisplayCluster/Private/Network/Packet/DisplayClusterPacketJson.h
--- a/Engine/Plugins/Runtime/nDisplay/Source/DisplayCluster/Private/Network/Packet/DisplayClusterPacketJson.h
+++ b/Engine/Plugins/Runtime/nDisplay/Source/DisplayCluster/Private/Network/Packet/DisplayClusterPacketJson.h
@@ -2,8 +2,6 @@
 
 #pragma once
 
-#include "CoreMinimal.h"
-
 #include "Network/Packet/IDisplayClusterPacket.h"
 #include "Dom/JsonObject.h"
 
@@ -41,7 +39,7 @@ public:
 private:
 	struct FPacketHeader
 	{
-		uint32 PacketBodyLength;
+		uint32 PacketBodyLength = 0;
 
 		FString ToString()
 		{
diff --git a/Engine/Plugins/Runtime/nDisplay/Source/DisplayCluster/Private/Network/Protocol/IDisplayClusterProtocolClusterSync.h b/Engine/Plugins/Runtime/nDisplay/Source/DisplayCluster/Private/Network/Protocol/IDisplayClusterProtocolClusterSync.h
--- a/Engine/Plugins/Runtime/nDisplay/Source/DisplayCluster/Private/Network/Protocol/IDisplayClusterProtocolClusterSync.h
+++ b/Engine/Plugins/Runtime/nDisplay/Source/DisplayCluster/Private/Network/Protocol/IDisplayClusterProtocolClusterSync.h
@@ -2,15 +2,14 @@
 
 #pragma once
 
-#include "CoreMinimal.h"
-
 #include "DisplayClusterEnums.h"
+
+#include "Cluster/DisplayClusterClusterEvent.h"
+
 #include "Misc/Optional.h"
 #include "Misc/QualifiedFrameTime.h"
 
-#include "Cluster/DisplayClusterClusterEvent.h"
 #include "Network/DisplayClusterNetworkTypes.h"
-#include "Network/Barrier/IDisplayClusterBarrier.h"
 
 
 /**
@@ -20,27 +19,29 @@
 class IDisplayClusterProtocolClusterSync
 {
 public:
+
 	virtual ~IDisplayClusterProtocolClusterSync() = default;
 
 public:
-	// Game start barrier
+
+	/** Game start barrier syncrhonization */
 	virtual EDisplayClusterCommResult WaitForGameStart() = 0;
 
-	// Frame start barrier
+	/** Frame start barrier syncrhonization */
 	virtual EDisplayClusterCommResult WaitForFrameStart() = 0;
 
-	// Frame end barrier
+	/** Frame end barrier syncrhonization */
 	virtual EDisplayClusterCommResult WaitForFrameEnd() = 0;
 
-	// Engine time
+	/** Engine time syncrhonization */
 	virtual EDisplayClusterCommResult GetTimeData(double& OutDeltaTime, double& OutGameTime, TOptional<FQualifiedFrameTime>& OutFrameTime) = 0;
 
-	// Sync objects
+	/** Custom objects syncrhonization */
 	virtual EDisplayClusterCommResult GetObjectsData(const EDisplayClusterSyncGroup InSyncGroup, TMap<FString, FString>& OutObjectsData) = 0;
 
-	// Sync events
+	/** Cluster events syncrhonization */
 	virtual EDisplayClusterCommResult GetEventsData(TArray<TSharedPtr<FDisplayClusterClusterEventJson, ESPMode::ThreadSafe>>& OutJsonEvents, TArray<TSharedPtr<FDisplayClusterClusterEventBinary, ESPMode::ThreadSafe>>& OutBinaryEvents) = 0;
 
-	// Sync native UE input
+	/** Native UE input syncrhonization */
 	virtual EDisplayClusterCommResult GetNativeInputData(TMap<FString, FString>& OutNativeInputData) = 0;
 };
diff --git a/Engine/Plugins/Runtime/nDisplay/Source/DisplayCluster/Private/Network/Protocol/IDisplayClusterProtocolEventsBinary.h b/Engine/Plugins/Runtime/nDisplay/Source/DisplayCluster/Private/Network/Protocol/IDisplayClusterProtocolEventsBinary.h
--- a/Engine/Plugins/Runtime/nDisplay/Source/DisplayCluster/Private/Network/Protocol/IDisplayClusterProtocolEventsBinary.h
+++ b/Engine/Plugins/Runtime/nDisplay/Source/DisplayCluster/Private/Network/Protocol/IDisplayClusterProtocolEventsBinary.h
@@ -2,7 +2,6 @@
 
 #pragma once
 
-#include "CoreMinimal.h"
 #include "Network/DisplayClusterNetworkTypes.h"
 
 struct FDisplayClusterClusterEventBinary;
@@ -14,8 +13,11 @@ struct FDisplayClusterClusterEventBinary;
 class IDisplayClusterProtocolEventsBinary
 {
 public:
+
 	virtual ~IDisplayClusterProtocolEventsBinary() = default;
 
 public:
+
+	/** Emits binary cluster event */
 	virtual EDisplayClusterCommResult EmitClusterEventBinary(const FDisplayClusterClusterEventBinary& Event) = 0;
 };
diff --git a/Engine/Plugins/Runtime/nDisplay/Source/DisplayCluster/Private/Network/Protocol/IDisplayClusterProtocolEventsJson.h b/Engine/Plugins/Runtime/nDisplay/Source/DisplayCluster/Private/Network/Protocol/IDisplayClusterProtocolEventsJson.h
--- a/Engine/Plugins/Runtime/nDisplay/Source/DisplayCluster/Private/Network/Protocol/IDisplayClusterProtocolEventsJson.h
+++ b/Engine/Plugins/Runtime/nDisplay/Source/DisplayCluster/Private/Network/Protocol/IDisplayClusterProtocolEventsJson.h
@@ -2,7 +2,6 @@
 
 #pragma once
 
-#include "CoreMinimal.h"
 #include "Network/DisplayClusterNetworkTypes.h"
 
 struct FDisplayClusterClusterEventJson;
@@ -14,8 +13,11 @@ struct FDisplayClusterClusterEventJson;
 class IDisplayClusterProtocolEventsJson
 {
 public:
+
 	virtual ~IDisplayClusterProtocolEventsJson() = default;
 
 public:
+
+	/** Emits JSON cluster event */
 	virtual EDisplayClusterCommResult EmitClusterEventJson(const FDisplayClusterClusterEventJson& Event) = 0;
 };
diff --git a/Engine/Plugins/Runtime/nDisplay/Source/DisplayCluster/Private/Network/Protocol/IDisplayClusterProtocolGenericBarrier.h b/Engine/Plugins/Runtime/nDisplay/Source/DisplayCluster/Private/Network/Protocol/IDisplayClusterProtocolGenericBarrier.h
--- a/Engine/Plugins/Runtime/nDisplay/Source/DisplayCluster/Private/Network/Protocol/IDisplayClusterProtocolGenericBarrier.h
+++ b/Engine/Plugins/Runtime/nDisplay/Source/DisplayCluster/Private/Network/Protocol/IDisplayClusterProtocolGenericBarrier.h
@@ -2,7 +2,6 @@
 
 #pragma once
 
-#include "CoreMinimal.h"
 #include "Network/DisplayClusterNetworkTypes.h"
 
 
@@ -13,10 +12,14 @@
 class IDisplayClusterProtocolGenericBarrier
 {
 public:
+
 	virtual ~IDisplayClusterProtocolGenericBarrier() = default;
 
 public:
-	// Just a union of different barrier control operation results to simplify management API
+
+	/**
+	 * Just a union of different barrier control operation results to simplify the management API
+	 */
 	enum class EBarrierControlResult : uint8
 	{
 		CreatedSuccessfully = 0,
@@ -28,21 +31,22 @@ public:
 	};
 
 public:
-	// Creates new barrier
-	virtual EDisplayClusterCommResult CreateBarrier(const FString& BarrierId, const TArray<FString>& UniqueThreadMarkers, uint32 Timeout, EBarrierControlResult& Result) = 0;
 
-	// Wait until a barrier with specific ID is created and ready to go
+	/** Creates new barrier */
+	virtual EDisplayClusterCommResult CreateBarrier(const FString& BarrierId, const TMap<FString, TSet<FString>>& NodeToSyncCallers, uint32 Timeout, EBarrierControlResult& Result) = 0;
+
+	/** Wait until a barrier with specific ID is created and ready to go */
 	virtual EDisplayClusterCommResult WaitUntilBarrierIsCreated(const FString& BarrierId, EBarrierControlResult& Result) = 0;
 
-	// Checks if a specific barrier exists
+	/** Checks if a specific barrier exists */
 	virtual EDisplayClusterCommResult IsBarrierAvailable(const FString& BarrierId, EBarrierControlResult& Result) = 0;
 
-	// Releases specific barrier
+	/** Releases specific barrier */
 	virtual EDisplayClusterCommResult ReleaseBarrier(const FString& BarrierId, EBarrierControlResult& Result) = 0;
 
-	// Synchronize calling thread on a specific barrier
-	virtual EDisplayClusterCommResult SyncOnBarrier(const FString& BarrierId, const FString& UniqueThreadMarker, EBarrierControlResult& Result) = 0;
+	/** Synchronize calling thread on a specific barrier */
+	virtual EDisplayClusterCommResult SyncOnBarrier(const FString& BarrierId, const FString& CallerId, EBarrierControlResult& Result) = 0;
 
-	// Synchronize calling thread on a specific barrier (with custom data)
-	virtual EDisplayClusterCommResult SyncOnBarrierWithData(const FString& BarrierId, const FString& UniqueThreadMarker, const TArray<uint8>& RequestData, TArray<uint8>& OutResponseData, EBarrierControlResult& Result) = 0;
+	/** Synchronize calling thread on a specific barrier(with custom data) */
+	virtual EDisplayClusterCommResult SyncOnBarrierWithData(const FString& BarrierId, const FString& CallerId, const TArray<uint8>& RequestData, TArray<uint8>& OutResponseData, EBarrierControlResult& Result) = 0;
 };
diff --git a/Engine/Plugins/Runtime/nDisplay/Source/DisplayCluster/Private/Network/Protocol/IDisplayClusterProtocolRenderSync.h b/Engine/Plugins/Runtime/nDisplay/Source/DisplayCluster/Private/Network/Protocol/IDisplayClusterProtocolRenderSync.h
--- a/Engine/Plugins/Runtime/nDisplay/Source/DisplayCluster/Private/Network/Protocol/IDisplayClusterProtocolRenderSync.h
+++ b/Engine/Plugins/Runtime/nDisplay/Source/DisplayCluster/Private/Network/Protocol/IDisplayClusterProtocolRenderSync.h
@@ -2,7 +2,6 @@
 
 #pragma once
 
-#include "CoreMinimal.h"
 #include "Network/DisplayClusterNetworkTypes.h"
 
 
@@ -13,9 +12,11 @@
 class IDisplayClusterProtocolRenderSync
 {
 public:
+
 	virtual ~IDisplayClusterProtocolRenderSync() = default;
 
 public:
-	// Synchronize RHI threads on a network barrier
-	virtual EDisplayClusterCommResult SyncOnBarrier() = 0;
+
+	/** Synchronize RHI thread on a network barrier */
+	virtual EDisplayClusterCommResult SynchronizeOnBarrier() = 0;
 };
diff --git a/Engine/Plugins/Runtime/nDisplay/Source/DisplayCluster/Private/Network/Service/ClusterEventsBinary/DisplayClusterClusterEventsBinaryClient.cpp b/Engine/Plugins/Runtime/nDisplay/Source/DisplayCluster/Private/Network/Service/ClusterEventsBinary/DisplayClusterClusterEventsBinaryClient.cpp
--- a/Engine/Plugins/Runtime/nDisplay/Source/DisplayCluster/Private/Network/Service/ClusterEventsBinary/DisplayClusterClusterEventsBinaryClient.cpp
+++ b/Engine/Plugins/Runtime/nDisplay/Source/DisplayCluster/Private/Network/Service/ClusterEventsBinary/DisplayClusterClusterEventsBinaryClient.cpp
@@ -1,20 +1,57 @@
 // Copyright Epic Games, Inc. All Rights Reserved.
 
 #include "Network/Service/ClusterEventsBinary/DisplayClusterClusterEventsBinaryClient.h"
-#include "Network/Conversion/DisplayClusterNetworkDataConversion.h"
+
 #include "Cluster/DisplayClusterClusterEvent.h"
+#include "Cluster/IPDisplayClusterClusterManager.h"
+
+#include "Network/Configuration/DisplayClusterNetworkConfiguration.h"
+#include "Network/Conversion/DisplayClusterNetworkDataConversion.h"
+#include "Network/Listener/DisplayClusterHelloMessageStrings.h"
+#include "Network/Packet/DisplayClusterPacketInternal.h"
+#include "Network/Service/ClusterEventsBinary/DisplayClusterClusterEventsBinaryStrings.h"
 
+#include "Misc/DisplayClusterGlobals.h"
 #include "Misc/DisplayClusterLog.h"
 
 
-FDisplayClusterClusterEventsBinaryClient::FDisplayClusterClusterEventsBinaryClient()
-	: FDisplayClusterClusterEventsBinaryClient(FString("CLN_CEB"))
+FDisplayClusterClusterEventsBinaryClient::FDisplayClusterClusterEventsBinaryClient(const FName& InName, bool bIsInternal)
+	: FDisplayClusterClient(InName.ToString(), 1)
+	, bIsInternalClient(bIsInternal)
 {
 }
 
-FDisplayClusterClusterEventsBinaryClient::FDisplayClusterClusterEventsBinaryClient(const FString& InName)
-	: FDisplayClusterClient(InName)
+
+//////////////////////////////////////////////////////////////////////////////////////////////
+// IDisplayClusterClient
+//////////////////////////////////////////////////////////////////////////////////////////////
+bool FDisplayClusterClusterEventsBinaryClient::Connect(const FString& Address, const uint16 Port, const uint32 ConnectRetriesAmount, const uint32 ConnectRetryDelay)
 {
+	// First, allow base class to perform connection
+	if (!FDisplayClusterClient::Connect(Address, Port, ConnectRetriesAmount, ConnectRetryDelay))
+	{
+		return false;
+	}
+
+	// Only internal clients greet the server on connection
+	if (bIsInternalClient)
+	{
+		// Prepare 'hello' message
+		TSharedPtr<FDisplayClusterPacketInternal> HelloMsg = MakeShared<FDisplayClusterPacketInternal>(
+			DisplayClusterHelloMessageStrings::Hello::Name,
+			DisplayClusterHelloMessageStrings::Hello::TypeRequest,
+			DisplayClusterClusterEventsBinaryStrings::ProtocolName
+		);
+
+		// Fill in the message with data
+		const FString NodeId = GDisplayCluster->GetPrivateClusterMgr()->GetNodeId();
+		HelloMsg->SetTextArg(DisplayClusterHelloMessageStrings::ArgumentsDefaultCategory, DisplayClusterHelloMessageStrings::Hello::ArgNodeId, NodeId);
+
+		// Send message (no response awaiting)
+		return SendPacket(HelloMsg);
+	}
+
+	return true;
 }
 
 
@@ -24,7 +61,7 @@ FDisplayClusterClusterEventsBinaryClient::FDisplayClusterClusterEventsBinaryClie
 EDisplayClusterCommResult FDisplayClusterClusterEventsBinaryClient::EmitClusterEventBinary(const FDisplayClusterClusterEventBinary& Event)
 {
 	// Convert internal binary event type to binary net packet
-	TSharedPtr<FDisplayClusterPacketBinary> Request = DisplayClusterNetworkDataConversion::BinaryEventToBinaryPacket(Event);
+	TSharedPtr<FDisplayClusterPacketBinary> Request = UE::nDisplay::DisplayClusterNetworkDataConversion::BinaryEventToBinaryPacket(const_cast<FDisplayClusterClusterEventBinary&>(Event));
 	if (!Request)
 	{
 		UE_LOG(LogDisplayClusterNetwork, Error, TEXT("Couldn't convert binary cluster event data to net packet"));
@@ -35,7 +72,7 @@ EDisplayClusterCommResult FDisplayClusterClusterEventsBinaryClient::EmitClusterE
 
 	{
 		// Send event
-		TRACE_CPUPROFILER_EVENT_SCOPE(nD CLN_CEB::EmitClusterEventBinary);
+		TRACE_CPUPROFILER_EVENT_SCOPE(CLN_CEB::EmitClusterEventBinary);
 		bResult = SendPacket(Request);
 	}
 
diff --git a/Engine/Plugins/Runtime/nDisplay/Source/DisplayCluster/Private/Network/Service/ClusterEventsBinary/DisplayClusterClusterEventsBinaryClient.h b/Engine/Plugins/Runtime/nDisplay/Source/DisplayCluster/Private/Network/Service/ClusterEventsBinary/DisplayClusterClusterEventsBinaryClient.h
--- a/Engine/Plugins/Runtime/nDisplay/Source/DisplayCluster/Private/Network/Service/ClusterEventsBinary/DisplayClusterClusterEventsBinaryClient.h
+++ b/Engine/Plugins/Runtime/nDisplay/Source/DisplayCluster/Private/Network/Service/ClusterEventsBinary/DisplayClusterClusterEventsBinaryClient.h
@@ -17,12 +17,23 @@ class FDisplayClusterClusterEventsBinaryClient
 	, public IDisplayClusterProtocolEventsBinary
 {
 public:
-	FDisplayClusterClusterEventsBinaryClient();
-	FDisplayClusterClusterEventsBinaryClient(const FString& InName);
+
+	FDisplayClusterClusterEventsBinaryClient(const FName& InName, bool bIsInternal = true);
+
+public:
+
+	//~ Begin IDisplayClusterClient
+	virtual bool Connect(const FString& Address, const uint16 Port, const uint32 ConnectRetriesAmount = 0, const uint32 ConnectRetryDelay = 0) override;
+	//~ End IDisplayClusterClient
 
 public:
-	//////////////////////////////////////////////////////////////////////////////////////////////
-	// IDisplayClusterProtocolEventsBinary
-	//////////////////////////////////////////////////////////////////////////////////////////////
+
+	//~ Begin IDisplayClusterProtocolEventsBinary
 	EDisplayClusterCommResult EmitClusterEventBinary(const FDisplayClusterClusterEventBinary& Event) override;
+	//~ End IDisplayClusterProtocolEventsBinary
+
+private:
+
+	/** Whether this client is intended to be used with an internal server, and therefore greet on connection */
+	const bool bIsInternalClient;
 };
diff --git a/Engine/Plugins/Runtime/nDisplay/Source/DisplayCluster/Private/Network/Service/ClusterEventsBinary/DisplayClusterClusterEventsBinaryService.cpp b/Engine/Plugins/Runtime/nDisplay/Source/DisplayCluster/Private/Network/Service/ClusterEventsBinary/DisplayClusterClusterEventsBinaryService.cpp
--- a/Engine/Plugins/Runtime/nDisplay/Source/DisplayCluster/Private/Network/Service/ClusterEventsBinary/DisplayClusterClusterEventsBinaryService.cpp
+++ b/Engine/Plugins/Runtime/nDisplay/Source/DisplayCluster/Private/Network/Service/ClusterEventsBinary/DisplayClusterClusterEventsBinaryService.cpp
@@ -1,7 +1,9 @@
 // Copyright Epic Games, Inc. All Rights Reserved.
 
 #include "Network/Service/ClusterEventsBinary/DisplayClusterClusterEventsBinaryService.h"
+#include "Network/Service/ClusterEventsBinary/DisplayClusterClusterEventsBinaryStrings.h"
 #include "Network/Session/DisplayClusterSession.h"
+#include "Network/Configuration/DisplayClusterNetworkConfiguration.h"
 #include "Network/Conversion/DisplayClusterNetworkDataConversion.h"
 
 #include "Cluster/IPDisplayClusterClusterManager.h"
@@ -13,8 +15,8 @@
 #include "Interfaces/IPv4/IPv4Endpoint.h"
 
 
-FDisplayClusterClusterEventsBinaryService::FDisplayClusterClusterEventsBinaryService()
-	: FDisplayClusterService(FString("SRV_CEB"))
+FDisplayClusterClusterEventsBinaryService::FDisplayClusterClusterEventsBinaryService(const FName& InInstanceName)
+	: FDisplayClusterService(InInstanceName.ToString())
 {
 	// Subscribe for SessionClosed events
 	OnSessionClosed().AddRaw(this, &FDisplayClusterClusterEventsBinaryService::ProcessSessionClosed);
@@ -24,14 +26,12 @@ FDisplayClusterClusterEventsBinaryService::~FDisplayClusterClusterEventsBinarySe
 {
 	// Unsubscribe from SessionClosed notifications
 	OnSessionClosed().RemoveAll(this);
-
-	Shutdown();
 }
 
 
 FString FDisplayClusterClusterEventsBinaryService::GetProtocolName() const
 {
-	static const FString ProtocolName("ClusterEventsBinary");
+	static const FString ProtocolName(DisplayClusterClusterEventsBinaryStrings::ProtocolName);
 	return ProtocolName;
 }
 
@@ -47,6 +47,24 @@ TSharedPtr<IDisplayClusterSession> FDisplayClusterClusterEventsBinaryService::Cr
 	return MakeShared<FDisplayClusterSession<FDisplayClusterPacketBinary, false>>(SessionInfo, *this, *this, FDisplayClusterService::GetThreadPriority());
 }
 
+void FDisplayClusterClusterEventsBinaryService::ProcessSessionClosed(const FDisplayClusterSessionInfo& SessionInfo)
+{
+	if (!SessionInfo.IsTerminatedByServer())
+	{
+		// Ignore sessions with empty NodeId as those could be external
+		if (SessionInfo.NodeId.IsSet())
+		{
+			// Prepare failure info
+			FDisplayClusterServiceFailureEvent EventInfo;
+			EventInfo.NodeFailed = SessionInfo.NodeId;
+			EventInfo.FailureType = FDisplayClusterServiceFailureEvent::ENodeFailType::ConnectionLost;
+
+			// Notify others about node fail
+			OnNodeFailed().Broadcast(EventInfo);
+		}
+	}
+}
+
 //////////////////////////////////////////////////////////////////////////////////////////////
 // IDisplayClusterSessionListener
 //////////////////////////////////////////////////////////////////////////////////////////////
@@ -61,7 +79,7 @@ typename IDisplayClusterSessionPacketHandler<FDisplayClusterPacketBinary, false>
 
 	// Convert net packet to the internal event data type
 	FDisplayClusterClusterEventBinary ClusterEvent;
-	if (!DisplayClusterNetworkDataConversion::BinaryPacketToBinaryEvent(Request, ClusterEvent))
+	if (!UE::nDisplay::DisplayClusterNetworkDataConversion::BinaryPacketToBinaryEvent(Request, ClusterEvent))
 	{
 		UE_LOG(LogDisplayClusterNetwork, Error, TEXT("%s - couldn't translate net packet data to binary event"), *GetName());
 		return typename IDisplayClusterSessionPacketHandler<FDisplayClusterPacketBinary, false>::ReturnType();
@@ -80,26 +98,8 @@ typename IDisplayClusterSessionPacketHandler<FDisplayClusterPacketBinary, false>
 //////////////////////////////////////////////////////////////////////////////////////////////
 EDisplayClusterCommResult FDisplayClusterClusterEventsBinaryService::EmitClusterEventBinary(const FDisplayClusterClusterEventBinary& Event)
 {
-	TRACE_CPUPROFILER_EVENT_SCOPE(nD SRV_CEB::EmitClusterEventBinary);
+	TRACE_CPUPROFILER_EVENT_SCOPE(SRV_CEB::EmitClusterEventBinary);
 
 	GDisplayCluster->GetPrivateClusterMgr()->EmitClusterEventBinary(Event, true);
 	return EDisplayClusterCommResult::Ok;
 }
-
-
-//////////////////////////////////////////////////////////////////////////////////////////////
-// FDisplayClusterClusterEventsBinaryService
-//////////////////////////////////////////////////////////////////////////////////////////////
-void FDisplayClusterClusterEventsBinaryService::ProcessSessionClosed(const FDisplayClusterSessionInfo& SessionInfo)
-{
-	if (!SessionInfo.IsTerminatedByServer())
-	{
-		// Get node ID
-		const FString NodeId = SessionInfo.NodeId.Get(FString());
-		if (!NodeId.IsEmpty())
-		{
-			// Notify others about node fail
-			OnNodeFailed().Broadcast(NodeId, ENodeFailType::ConnectionLost);
-		}
-	}
-}
diff --git a/Engine/Plugins/Runtime/nDisplay/Source/DisplayCluster/Private/Network/Service/ClusterEventsBinary/DisplayClusterClusterEventsBinaryService.h b/Engine/Plugins/Runtime/nDisplay/Source/DisplayCluster/Private/Network/Service/ClusterEventsBinary/DisplayClusterClusterEventsBinaryService.h
--- a/Engine/Plugins/Runtime/nDisplay/Source/DisplayCluster/Private/Network/Service/ClusterEventsBinary/DisplayClusterClusterEventsBinaryService.h
+++ b/Engine/Plugins/Runtime/nDisplay/Source/DisplayCluster/Private/Network/Service/ClusterEventsBinary/DisplayClusterClusterEventsBinaryService.h
@@ -2,8 +2,6 @@
 
 #pragma once
 
-#include "CoreMinimal.h"
-
 #include "Network/Service/DisplayClusterService.h"
 #include "Network/Session/IDisplayClusterSessionPacketHandler.h"
 #include "Network/Protocol/IDisplayClusterProtocolEventsBinary.h"
@@ -19,32 +17,34 @@ class FDisplayClusterClusterEventsBinaryService
 	, protected IDisplayClusterProtocolEventsBinary
 {
 public:
-	FDisplayClusterClusterEventsBinaryService();
+
+	FDisplayClusterClusterEventsBinaryService(const FName& InstanceName);
+
 	virtual ~FDisplayClusterClusterEventsBinaryService();
 
 public:
-	//////////////////////////////////////////////////////////////////////////////////////////////
-	// IDisplayClusterServer
-	//////////////////////////////////////////////////////////////////////////////////////////////
+
+	//~ Begin IDisplayClusterServer
 	virtual FString GetProtocolName() const override;
+	//~ End IDisplayClusterServer
 
 protected:
-	// Creates session instance for this service
+
+	/** Creates session instance for this service */
 	virtual TSharedPtr<IDisplayClusterSession> CreateSession(FDisplayClusterSessionInfo& SessionInfo) override;
 
+	/** Callback when a session is closed */
+	void ProcessSessionClosed(const FDisplayClusterSessionInfo& SessionInfo);
+
 protected:
-	//////////////////////////////////////////////////////////////////////////////////////////////
-	// IDisplayClusterSessionPacketHandler
-	//////////////////////////////////////////////////////////////////////////////////////////////
+
+	//~ Begin IDisplayClusterSessionPacketHandler
 	virtual typename IDisplayClusterSessionPacketHandler<FDisplayClusterPacketBinary, false>::ReturnType ProcessPacket(const TSharedPtr<FDisplayClusterPacketBinary>& Request, const FDisplayClusterSessionInfo& SessionInfo) override;
+	//~ End IDisplayClusterSessionPacketHandler
 
 protected:
-	//////////////////////////////////////////////////////////////////////////////////////////////
-	// IDisplayClusterProtocolEventsBinary
-	//////////////////////////////////////////////////////////////////////////////////////////////
-	virtual EDisplayClusterCommResult EmitClusterEventBinary(const FDisplayClusterClusterEventBinary& Event) override;
 
-protected:
-	// Callback when a session is closed
-	void ProcessSessionClosed(const FDisplayClusterSessionInfo& SessionInfo);
+	//~ Begin IDisplayClusterProtocolEventsBinary
+	virtual EDisplayClusterCommResult EmitClusterEventBinary(const FDisplayClusterClusterEventBinary& Event) override;
+	//~ End IDisplayClusterProtocolEventsBinary
 };
diff --git a/Engine/Plugins/Runtime/nDisplay/Source/DisplayCluster/Private/Network/Service/ClusterEventsBinary/DisplayClusterClusterEventsBinaryStrings.h b/Engine/Plugins/Runtime/nDisplay/Source/DisplayCluster/Private/Network/Service/ClusterEventsBinary/DisplayClusterClusterEventsBinaryStrings.h
new file mode 100644
--- /dev/null
+++ b/Engine/Plugins/Runtime/nDisplay/Source/DisplayCluster/Private/Network/Service/ClusterEventsBinary/DisplayClusterClusterEventsBinaryStrings.h
@@ -0,0 +1,15 @@
+// Copyright Epic Games, Inc. All Rights Reserved.
+
+#pragma once
+
+#include "HAL/Platform.h"
+
+
+/**
+ * Binary cluster event protocol strings
+ */
+namespace DisplayClusterClusterEventsBinaryStrings
+{
+	constexpr static const TCHAR* ProtocolName = TEXT("EventsBinary");
+	constexpr static const TCHAR* TypeRequest  = TEXT("Request");
+};
diff --git a/Engine/Plugins/Runtime/nDisplay/Source/DisplayCluster/Private/Network/Service/ClusterEventsJson/DisplayClusterClusterEventsJsonClient.cpp b/Engine/Plugins/Runtime/nDisplay/Source/DisplayCluster/Private/Network/Service/ClusterEventsJson/DisplayClusterClusterEventsJsonClient.cpp
--- a/Engine/Plugins/Runtime/nDisplay/Source/DisplayCluster/Private/Network/Service/ClusterEventsJson/DisplayClusterClusterEventsJsonClient.cpp
+++ b/Engine/Plugins/Runtime/nDisplay/Source/DisplayCluster/Private/Network/Service/ClusterEventsJson/DisplayClusterClusterEventsJsonClient.cpp
@@ -1,22 +1,57 @@
 // Copyright Epic Games, Inc. All Rights Reserved.
 
 #include "Network/Service/ClusterEventsJson/DisplayClusterClusterEventsJsonClient.h"
-#include "Network/Service/ClusterEventsJson/DisplayClusterClusterEventsJsonStrings.h"
-#include "Network/Conversion/DisplayClusterNetworkDataConversion.h"
 
 #include "Cluster/DisplayClusterClusterEvent.h"
+#include "Cluster/IPDisplayClusterClusterManager.h"
 
+#include "Misc/DisplayClusterGlobals.h"
 #include "Misc/DisplayClusterLog.h"
 
+#include "Network/Configuration/DisplayClusterNetworkConfiguration.h"
+#include "Network/Conversion/DisplayClusterNetworkDataConversion.h"
+#include "Network/Listener/DisplayClusterHelloMessageStrings.h"
+#include "Network/Packet/DisplayClusterPacketInternal.h"
+#include "Network/Service/ClusterEventsJson/DisplayClusterClusterEventsJsonStrings.h"
+
 
-FDisplayClusterClusterEventsJsonClient::FDisplayClusterClusterEventsJsonClient()
-	: FDisplayClusterClusterEventsJsonClient(FString("CLN_CEJ"))
+FDisplayClusterClusterEventsJsonClient::FDisplayClusterClusterEventsJsonClient(const FName& InName, bool bIsInternal)
+	: FDisplayClusterClient(InName.ToString(), 1)
+	, bIsInternalClient(bIsInternal)
 {
 }
 
-FDisplayClusterClusterEventsJsonClient::FDisplayClusterClusterEventsJsonClient(const FString& InName)
-	: FDisplayClusterClient(InName)
+
+//////////////////////////////////////////////////////////////////////////////////////////////
+// IDisplayClusterClient
+//////////////////////////////////////////////////////////////////////////////////////////////
+bool FDisplayClusterClusterEventsJsonClient::Connect(const FString& Address, const uint16 Port, const uint32 ConnectRetriesAmount, const uint32 ConnectRetryDelay)
 {
+	// First, allow base class to perform connection
+	if (!FDisplayClusterClient::Connect(Address, Port, ConnectRetriesAmount, ConnectRetryDelay))
+	{
+		return false;
+	}
+
+	// Only internal clients greet the server on connection
+	if (bIsInternalClient)
+	{
+		// Prepare 'hello' message
+		TSharedPtr<FDisplayClusterPacketInternal> HelloMsg = MakeShared<FDisplayClusterPacketInternal>(
+			DisplayClusterHelloMessageStrings::Hello::Name,
+			DisplayClusterHelloMessageStrings::Hello::TypeRequest,
+			DisplayClusterClusterEventsJsonStrings::ProtocolName
+		);
+
+		// Fill in the message with data
+		const FString NodeId = GDisplayCluster->GetPrivateClusterMgr()->GetNodeId();
+		HelloMsg->SetTextArg(DisplayClusterHelloMessageStrings::ArgumentsDefaultCategory, DisplayClusterHelloMessageStrings::Hello::ArgNodeId, NodeId);
+
+		// Send message (no response awaiting)
+		return SendPacket(HelloMsg);
+	}
+
+	return true;
 }
 
 
@@ -26,7 +61,7 @@ FDisplayClusterClusterEventsJsonClient::FDisplayClusterClusterEventsJsonClient(c
 EDisplayClusterCommResult FDisplayClusterClusterEventsJsonClient::EmitClusterEventJson(const FDisplayClusterClusterEventJson& Event)
 {
 	// Convert internal json event type to json net packet
-	TSharedPtr<FDisplayClusterPacketJson> Request = DisplayClusterNetworkDataConversion::JsonEventToJsonPacket(Event);
+	TSharedPtr<FDisplayClusterPacketJson> Request = UE::nDisplay::DisplayClusterNetworkDataConversion::JsonEventToJsonPacket(Event);
 	if (!Request)
 	{
 		UE_LOG(LogDisplayClusterNetwork, Error, TEXT("Couldn't convert json cluster event data to net packet"));
@@ -37,7 +72,7 @@ EDisplayClusterCommResult FDisplayClusterClusterEventsJsonClient::EmitClusterEve
 
 	{
 		// Send event
-		TRACE_CPUPROFILER_EVENT_SCOPE(nD CLN_CEJ::EmitClusterEventJson);
+		TRACE_CPUPROFILER_EVENT_SCOPE(CLN_CEJ::EmitClusterEventJson);
 		bResult = SendPacket(Request);
 	}
 
diff --git a/Engine/Plugins/Runtime/nDisplay/Source/DisplayCluster/Private/Network/Service/ClusterEventsJson/DisplayClusterClusterEventsJsonClient.h b/Engine/Plugins/Runtime/nDisplay/Source/DisplayCluster/Private/Network/Service/ClusterEventsJson/DisplayClusterClusterEventsJsonClient.h
--- a/Engine/Plugins/Runtime/nDisplay/Source/DisplayCluster/Private/Network/Service/ClusterEventsJson/DisplayClusterClusterEventsJsonClient.h
+++ b/Engine/Plugins/Runtime/nDisplay/Source/DisplayCluster/Private/Network/Service/ClusterEventsJson/DisplayClusterClusterEventsJsonClient.h
@@ -15,12 +15,23 @@ class FDisplayClusterClusterEventsJsonClient
 	, public IDisplayClusterProtocolEventsJson
 {
 public:
-	FDisplayClusterClusterEventsJsonClient();
-	FDisplayClusterClusterEventsJsonClient(const FString& InName);
+
+	FDisplayClusterClusterEventsJsonClient(const FName& InName, bool bIsInternal = true);
+
+public:
+
+	//~ Begin IDisplayClusterClient
+	virtual bool Connect(const FString& Address, const uint16 Port, const uint32 ConnectRetriesAmount = 0, const uint32 ConnectRetryDelay = 0) override;
+	//~ End IDisplayClusterClient
 
 public:
-	//////////////////////////////////////////////////////////////////////////////////////////////
-	// IDisplayClusterProtocolEventsJson
-	//////////////////////////////////////////////////////////////////////////////////////////////
+
+	//~ Begin IDisplayClusterProtocolEventsJson
 	EDisplayClusterCommResult EmitClusterEventJson(const FDisplayClusterClusterEventJson& Event) override;
+	//~ End IDisplayClusterProtocolEventsJson
+
+private:
+
+	/** Whether this client is intended to be used with an internal server, and therefore greet on connection */
+	const bool bIsInternalClient;
 };
diff --git a/Engine/Plugins/Runtime/nDisplay/Source/DisplayCluster/Private/Network/Service/ClusterEventsJson/DisplayClusterClusterEventsJsonService.cpp b/Engine/Plugins/Runtime/nDisplay/Source/DisplayCluster/Private/Network/Service/ClusterEventsJson/DisplayClusterClusterEventsJsonService.cpp
--- a/Engine/Plugins/Runtime/nDisplay/Source/DisplayCluster/Private/Network/Service/ClusterEventsJson/DisplayClusterClusterEventsJsonService.cpp
+++ b/Engine/Plugins/Runtime/nDisplay/Source/DisplayCluster/Private/Network/Service/ClusterEventsJson/DisplayClusterClusterEventsJsonService.cpp
@@ -3,6 +3,7 @@
 #include "Network/Service/ClusterEventsJson/DisplayClusterClusterEventsJsonService.h"
 #include "Network/Service/ClusterEventsJson/DisplayClusterClusterEventsJsonStrings.h"
 #include "Network/Session/DisplayClusterSession.h"
+#include "Network/Configuration/DisplayClusterNetworkConfiguration.h"
 #include "Network/Conversion/DisplayClusterNetworkDataConversion.h"
 
 #include "Cluster/IPDisplayClusterClusterManager.h"
@@ -16,8 +17,8 @@
 #include "Interfaces/IPv4/IPv4Endpoint.h"
 
 
-FDisplayClusterClusterEventsJsonService::FDisplayClusterClusterEventsJsonService()
-	: FDisplayClusterService(FString("SRV_CEJ"))
+FDisplayClusterClusterEventsJsonService::FDisplayClusterClusterEventsJsonService(const FName& InInstanceName)
+	: FDisplayClusterService(InInstanceName.ToString())
 {
 	// Subscribe for SessionClosed events
 	OnSessionClosed().AddRaw(this, &FDisplayClusterClusterEventsJsonService::ProcessSessionClosed);
@@ -27,14 +28,12 @@ FDisplayClusterClusterEventsJsonService::~FDisplayClusterClusterEventsJsonServic
 {
 	// Unsubscribe from SessionClosed notifications
 	OnSessionClosed().RemoveAll(this);
-
-	Shutdown();
 }
 
 
 FString FDisplayClusterClusterEventsJsonService::GetProtocolName() const
 {
-	static const FString ProtocolName("ClusterEventsJson");
+	static const FString ProtocolName(DisplayClusterClusterEventsJsonStrings::ProtocolName);
 	return ProtocolName;
 }
 
@@ -50,6 +49,25 @@ TSharedPtr<IDisplayClusterSession> FDisplayClusterClusterEventsJsonService::Crea
 	return MakeShared<FDisplayClusterSession<FDisplayClusterPacketJson, false>>(SessionInfo, *this, *this, FDisplayClusterService::GetThreadPriority());
 }
 
+void FDisplayClusterClusterEventsJsonService::ProcessSessionClosed(const FDisplayClusterSessionInfo& SessionInfo)
+{
+	if (!SessionInfo.IsTerminatedByServer())
+	{
+		// Ignore the session with empty NodeId as those could be the external ones
+		if (SessionInfo.NodeId.IsSet())
+		{
+			// Prepare failure info
+			FDisplayClusterServiceFailureEvent EventInfo;
+			EventInfo.NodeFailed = SessionInfo.NodeId;
+			EventInfo.FailureType = FDisplayClusterServiceFailureEvent::ENodeFailType::ConnectionLost;
+
+			// Notify others about node fail
+			OnNodeFailed().Broadcast(EventInfo);
+		}
+	}
+}
+
+
 //////////////////////////////////////////////////////////////////////////////////////////////
 // IDisplayClusterSessionListener
 //////////////////////////////////////////////////////////////////////////////////////////////
@@ -84,7 +102,7 @@ typename IDisplayClusterSessionPacketHandler<FDisplayClusterPacketJson, false>::
 
 	// Convert net packet to the internal event data type
 	FDisplayClusterClusterEventJson ClusterEvent;
-	if(!DisplayClusterNetworkDataConversion::JsonPacketToJsonEvent(Request, ClusterEvent))
+	if(!UE::nDisplay::DisplayClusterNetworkDataConversion::JsonPacketToJsonEvent(Request, ClusterEvent))
 	{
 		UE_LOG(LogDisplayClusterNetwork, Error, TEXT("%s - couldn't translate net packet data to json event"), *GetName());
 		return IDisplayClusterSessionPacketHandler<FDisplayClusterPacketJson, false>::ReturnType();
@@ -103,26 +121,8 @@ typename IDisplayClusterSessionPacketHandler<FDisplayClusterPacketJson, false>::
 //////////////////////////////////////////////////////////////////////////////////////////////
 EDisplayClusterCommResult FDisplayClusterClusterEventsJsonService::EmitClusterEventJson(const FDisplayClusterClusterEventJson& Event)
 {
-	TRACE_CPUPROFILER_EVENT_SCOPE(nD SRV_CEJ::EmitClusterEventJson);
+	TRACE_CPUPROFILER_EVENT_SCOPE(SRV_CEJ::EmitClusterEventJson);
 
 	GDisplayCluster->GetPrivateClusterMgr()->EmitClusterEventJson(Event, true);
 	return EDisplayClusterCommResult::Ok;
 }
-
-
-//////////////////////////////////////////////////////////////////////////////////////////////
-// FDisplayClusterClusterEventsJsonService
-//////////////////////////////////////////////////////////////////////////////////////////////
-void FDisplayClusterClusterEventsJsonService::ProcessSessionClosed(const FDisplayClusterSessionInfo& SessionInfo)
-{
-	if (!SessionInfo.IsTerminatedByServer())
-	{
-		// Get node ID
-		const FString NodeId = SessionInfo.NodeId.Get(FString());
-		if (!NodeId.IsEmpty())
-		{
-			// Notify others about node fail
-			OnNodeFailed().Broadcast(NodeId, ENodeFailType::ConnectionLost);
-		}
-	}
-}
diff --git a/Engine/Plugins/Runtime/nDisplay/Source/DisplayCluster/Private/Network/Service/ClusterEventsJson/DisplayClusterClusterEventsJsonService.h b/Engine/Plugins/Runtime/nDisplay/Source/DisplayCluster/Private/Network/Service/ClusterEventsJson/DisplayClusterClusterEventsJsonService.h
--- a/Engine/Plugins/Runtime/nDisplay/Source/DisplayCluster/Private/Network/Service/ClusterEventsJson/DisplayClusterClusterEventsJsonService.h
+++ b/Engine/Plugins/Runtime/nDisplay/Source/DisplayCluster/Private/Network/Service/ClusterEventsJson/DisplayClusterClusterEventsJsonService.h
@@ -2,8 +2,6 @@
 
 #pragma once
 
-#include "CoreMinimal.h"
-
 #include "Network/Service/DisplayClusterService.h"
 #include "Network/Session/IDisplayClusterSessionPacketHandler.h"
 #include "Network/Packet/DisplayClusterPacketJson.h"
@@ -20,44 +18,36 @@ class FDisplayClusterClusterEventsJsonService
 	, protected IDisplayClusterProtocolEventsJson
 {
 public:
-	FDisplayClusterClusterEventsJsonService();
+
+	FDisplayClusterClusterEventsJsonService(const FName& InInstanceName);
+
 	virtual ~FDisplayClusterClusterEventsJsonService();
 
 public:
-	//////////////////////////////////////////////////////////////////////////////////////////////
-	// IDisplayClusterServer
-	//////////////////////////////////////////////////////////////////////////////////////////////
-	virtual FString GetProtocolName() const override;
 
-public:
-	enum class EDisplayClusterJsonError : uint8
-	{
-		Ok = 0,
-		NotSupported = 1,
-		MissedMandatoryFields = 2,
-		UnknownError = 255
-	};
+	//~ Begin IDisplayClusterServer
+	virtual FString GetProtocolName() const override;
+	//~ End IDisplayClusterServer
 
 protected:
-	// Creates session instance for this service
+
+	/** Creates session instance for this service */
 	virtual TSharedPtr<IDisplayClusterSession> CreateSession(FDisplayClusterSessionInfo& SessionInfo) override;
 
-protected:
-	//////////////////////////////////////////////////////////////////////////////////////////////
-	// IDisplayClusterSessionPacketHandler
-	//////////////////////////////////////////////////////////////////////////////////////////////
-	virtual typename IDisplayClusterSessionPacketHandler<FDisplayClusterPacketJson, false>::ReturnType ProcessPacket(const TSharedPtr<FDisplayClusterPacketJson>& Request, const FDisplayClusterSessionInfo& SessionInfo) override;
+private:
+
+	/** Callback when a session is closed */
+	void ProcessSessionClosed(const FDisplayClusterSessionInfo& SessionInfo);
 
 protected:
-	//////////////////////////////////////////////////////////////////////////////////////////////
-	// IDisplayClusterProtocolEventsJson
-	//////////////////////////////////////////////////////////////////////////////////////////////
-	virtual EDisplayClusterCommResult EmitClusterEventJson(const FDisplayClusterClusterEventJson& Event) override;
+
+	//~ Begin IDisplayClusterSessionPacketHandler
+	virtual typename IDisplayClusterSessionPacketHandler<FDisplayClusterPacketJson, false>::ReturnType ProcessPacket(const TSharedPtr<FDisplayClusterPacketJson>& Request, const FDisplayClusterSessionInfo& SessionInfo) override;
+	//~ End IDisplayClusterSessionPacketHandler
 
 protected:
-	// Callback when a session is closed
-	void ProcessSessionClosed(const FDisplayClusterSessionInfo& SessionInfo);
 
-private:
-	FDisplayClusterClusterEventJson BuildClusterEventFromJson(const FString& EventName, const FString& EventType, const FString& EventCategory, const TSharedPtr<FJsonObject>& JsonPacket) const;
+	//~ Begin IDisplayClusterProtocolEventsJson
+	virtual EDisplayClusterCommResult EmitClusterEventJson(const FDisplayClusterClusterEventJson& Event) override;
+	//~ End IDisplayClusterProtocolEventsJson
 };
diff --git a/Engine/Plugins/Runtime/nDisplay/Source/DisplayCluster/Private/Network/Service/ClusterEventsJson/DisplayClusterClusterEventsJsonStrings.h b/Engine/Plugins/Runtime/nDisplay/Source/DisplayCluster/Private/Network/Service/ClusterEventsJson/DisplayClusterClusterEventsJsonStrings.h
--- a/Engine/Plugins/Runtime/nDisplay/Source/DisplayCluster/Private/Network/Service/ClusterEventsJson/DisplayClusterClusterEventsJsonStrings.h
+++ b/Engine/Plugins/Runtime/nDisplay/Source/DisplayCluster/Private/Network/Service/ClusterEventsJson/DisplayClusterClusterEventsJsonStrings.h
@@ -4,14 +4,16 @@
 
 #include "HAL/Platform.h"
 
+
 /**
  * JSON cluster event protocol strings
  */
 namespace DisplayClusterClusterEventsJsonStrings
 {
+	constexpr static const TCHAR* ProtocolName = TEXT("EventsJSON");
+	constexpr static const TCHAR* TypeRequest  = TEXT("Request");
+
 	constexpr static auto ArgName       = TEXT("Name");
 	constexpr static auto ArgType       = TEXT("Type");
 	constexpr static auto ArgCategory   = TEXT("Category");
-	constexpr static auto ArgParameters = TEXT("Parameters");
-	constexpr static auto ArgError      = TEXT("Error");
 };
diff --git a/Engine/Plugins/Runtime/nDisplay/Source/DisplayCluster/Private/Network/Service/ClusterSync/DisplayClusterClusterSyncClient.cpp b/Engine/Plugins/Runtime/nDisplay/Source/DisplayCluster/Private/Network/Service/ClusterSync/DisplayClusterClusterSyncClient.cpp
--- a/Engine/Plugins/Runtime/nDisplay/Source/DisplayCluster/Private/Network/Service/ClusterSync/DisplayClusterClusterSyncClient.cpp
+++ b/Engine/Plugins/Runtime/nDisplay/Source/DisplayCluster/Private/Network/Service/ClusterSync/DisplayClusterClusterSyncClient.cpp
@@ -2,6 +2,7 @@
 
 #include "Network/Service/ClusterSync/DisplayClusterClusterSyncClient.h"
 #include "Network/Service/ClusterSync/DisplayClusterClusterSyncStrings.h"
+#include "Network/Configuration/DisplayClusterNetworkConfiguration.h"
 #include "Network/Conversion/DisplayClusterNetworkDataConversion.h"
 #include "Network/Packet/DisplayClusterPacketInternal.h"
 #include "Network/Listener/DisplayClusterHelloMessageStrings.h"
@@ -17,13 +18,8 @@
 #include "Misc/DisplayClusterLog.h"
 
 
-FDisplayClusterClusterSyncClient::FDisplayClusterClusterSyncClient()
-	: FDisplayClusterClusterSyncClient(FString("CLN_CS"))
-{
-}
-
-FDisplayClusterClusterSyncClient::FDisplayClusterClusterSyncClient(const FString& InName)
-	: FDisplayClusterClient(InName)
+FDisplayClusterClusterSyncClient::FDisplayClusterClusterSyncClient(const FName& InName)
+	: FDisplayClusterClient(InName.ToString())
 {
 }
 
@@ -42,7 +38,7 @@ bool FDisplayClusterClusterSyncClient::Connect(const FString& Address, const uin
 	// Prepare 'hello' message
 	TSharedPtr<FDisplayClusterPacketInternal> HelloMsg = MakeShared<FDisplayClusterPacketInternal>(
 		DisplayClusterHelloMessageStrings::Hello::Name,
-		DisplayClusterClusterSyncStrings::TypeRequest,
+		DisplayClusterHelloMessageStrings::Hello::TypeRequest,
 		DisplayClusterClusterSyncStrings::ProtocolName
 	);
 
@@ -68,7 +64,7 @@ EDisplayClusterCommResult FDisplayClusterClusterSyncClient::WaitForGameStart()
 	TSharedPtr<FDisplayClusterPacketInternal> Response;
 
 	{
-		TRACE_CPUPROFILER_EVENT_SCOPE(nD CLN_CS::WaitForGameStart);
+		TRACE_CPUPROFILER_EVENT_SCOPE(CLN_CS::WaitForGameStart);
 		Response = SendRecvPacket(Request);
 	}
 
@@ -91,7 +87,7 @@ EDisplayClusterCommResult FDisplayClusterClusterSyncClient::WaitForFrameStart()
 	TSharedPtr<FDisplayClusterPacketInternal> Response;
 
 	{
-		TRACE_CPUPROFILER_EVENT_SCOPE(nD CLN_CS::WaitForFrameStart);
+		TRACE_CPUPROFILER_EVENT_SCOPE(CLN_CS::WaitForFrameStart);
 		Response = SendRecvPacket(Request);
 	}
 
@@ -114,7 +110,7 @@ EDisplayClusterCommResult FDisplayClusterClusterSyncClient::WaitForFrameEnd()
 	TSharedPtr<FDisplayClusterPacketInternal> Response;
 
 	{
-		TRACE_CPUPROFILER_EVENT_SCOPE(nD CLN_CS::WaitForFrameEnd);
+		TRACE_CPUPROFILER_EVENT_SCOPE(CLN_CS::WaitForFrameEnd);
 		Response = SendRecvPacket(Request);
 	}
 
@@ -137,7 +133,7 @@ EDisplayClusterCommResult FDisplayClusterClusterSyncClient::GetTimeData(double&
 	TSharedPtr<FDisplayClusterPacketInternal> Response;
 
 	{
-		TRACE_CPUPROFILER_EVENT_SCOPE(nD CLN_CS::GetTimeData);
+		TRACE_CPUPROFILER_EVENT_SCOPE(CLN_CS::GetTimeData);
 		Response = SendRecvPacket(Request);
 	}
 
@@ -199,11 +195,11 @@ EDisplayClusterCommResult FDisplayClusterClusterSyncClient::GetObjectsData(const
 			DisplayClusterClusterSyncStrings::ProtocolName);
 	
 	Request->SetTextArg(DisplayClusterClusterSyncStrings::ArgumentsDefaultCategory, DisplayClusterClusterSyncStrings::GetObjectsData::ArgSyncGroup, static_cast<uint8>(InSyncGroup));
-	
+
 	TSharedPtr<FDisplayClusterPacketInternal> Response;
 
 	{
-		TRACE_CPUPROFILER_EVENT_SCOPE(nD CLN_CS::GetObjectsData);
+		TRACE_CPUPROFILER_EVENT_SCOPE(CLN_CS::GetObjectsData);
 		Response = SendRecvPacket(Request);
 	}
 
@@ -229,7 +225,7 @@ EDisplayClusterCommResult FDisplayClusterClusterSyncClient::GetEventsData(TArray
 	TSharedPtr<FDisplayClusterPacketInternal> Response;
 
 	{
-		TRACE_CPUPROFILER_EVENT_SCOPE(nD CLN_CS::GetEventsData);
+		TRACE_CPUPROFILER_EVENT_SCOPE(CLN_CS::GetEventsData);
 		Response = SendRecvPacket(Request);
 	}
 
@@ -240,8 +236,8 @@ EDisplayClusterCommResult FDisplayClusterClusterSyncClient::GetEventsData(TArray
 	}
 
 	// Extract events data from response packet
-	DisplayClusterNetworkDataConversion::JsonEventsFromInternalPacket(Response,   OutJsonEvents);
-	DisplayClusterNetworkDataConversion::BinaryEventsFromInternalPacket(Response, OutBinaryEvents);
+	UE::nDisplay::DisplayClusterNetworkDataConversion::JsonEventsFromInternalPacket(Response,   OutJsonEvents);
+	UE::nDisplay::DisplayClusterNetworkDataConversion::BinaryEventsFromInternalPacket(Response, OutBinaryEvents);
 
 	return Response->GetCommResult();
 }
@@ -256,7 +252,7 @@ EDisplayClusterCommResult FDisplayClusterClusterSyncClient::GetNativeInputData(T
 	TSharedPtr<FDisplayClusterPacketInternal> Response;
 
 	{
-		TRACE_CPUPROFILER_EVENT_SCOPE(nD CLN_CS::GetNativeInputData);
+		TRACE_CPUPROFILER_EVENT_SCOPE(CLN_CS::GetNativeInputData);
 		Response = SendRecvPacket(Request);
 	}
 
diff --git a/Engine/Plugins/Runtime/nDisplay/Source/DisplayCluster/Private/Network/Service/ClusterSync/DisplayClusterClusterSyncClient.h b/Engine/Plugins/Runtime/nDisplay/Source/DisplayCluster/Private/Network/Service/ClusterSync/DisplayClusterClusterSyncClient.h
--- a/Engine/Plugins/Runtime/nDisplay/Source/DisplayCluster/Private/Network/Service/ClusterSync/DisplayClusterClusterSyncClient.h
+++ b/Engine/Plugins/Runtime/nDisplay/Source/DisplayCluster/Private/Network/Service/ClusterSync/DisplayClusterClusterSyncClient.h
@@ -15,19 +15,18 @@ class FDisplayClusterClusterSyncClient
 	, public IDisplayClusterProtocolClusterSync
 {
 public:
-	FDisplayClusterClusterSyncClient();
-	FDisplayClusterClusterSyncClient(const FString& InName);
+
+	FDisplayClusterClusterSyncClient(const FName& InName);
 
 public:
-	//////////////////////////////////////////////////////////////////////////////////////////////
-	// IDisplayClusterClient
-	//////////////////////////////////////////////////////////////////////////////////////////////
-	virtual bool Connect(const FString& Address, const uint16 Port, const uint32 ConnectRetriesAmount, const uint32 ConnectRetryDelay) override;
+
+	//~ Begin IDisplayClusterClient
+	virtual bool Connect(const FString& Address, const uint16 Port, const uint32 ConnectRetriesAmount = 0, const uint32 ConnectRetryDelay = 0) override;
+	//~ End IDisplayClusterClient
 
 public:
-	//////////////////////////////////////////////////////////////////////////////////////////////
-	// IDisplayClusterProtocolClusterSync
-	//////////////////////////////////////////////////////////////////////////////////////////////
+
+	//~ Begin IDisplayClusterProtocolClusterSync
 	virtual EDisplayClusterCommResult WaitForGameStart() override;
 	virtual EDisplayClusterCommResult WaitForFrameStart() override;
 	virtual EDisplayClusterCommResult WaitForFrameEnd() override;
@@ -35,4 +34,5 @@ public:
 	virtual EDisplayClusterCommResult GetObjectsData(const EDisplayClusterSyncGroup InSyncGroup, TMap<FString, FString>& OutObjectsData) override;
 	virtual EDisplayClusterCommResult GetEventsData(TArray<TSharedPtr<FDisplayClusterClusterEventJson, ESPMode::ThreadSafe>>& OutJsonEvents, TArray<TSharedPtr<FDisplayClusterClusterEventBinary, ESPMode::ThreadSafe>>& OutBinaryEvents) override;
 	virtual EDisplayClusterCommResult GetNativeInputData(TMap<FString, FString>& OutNativeInputData) override;
+	//~ End IDisplayClusterProtocolClusterSync
 };
diff --git a/Engine/Plugins/Runtime/nDisplay/Source/DisplayCluster/Private/Network/Service/ClusterSync/DisplayClusterClusterSyncService.cpp b/Engine/Plugins/Runtime/nDisplay/Source/DisplayCluster/Private/Network/Service/ClusterSync/DisplayClusterClusterSyncService.cpp
--- a/Engine/Plugins/Runtime/nDisplay/Source/DisplayCluster/Private/Network/Service/ClusterSync/DisplayClusterClusterSyncService.cpp
+++ b/Engine/Plugins/Runtime/nDisplay/Source/DisplayCluster/Private/Network/Service/ClusterSync/DisplayClusterClusterSyncService.cpp
@@ -2,45 +2,72 @@
 
 #include "Network/Service/ClusterSync/DisplayClusterClusterSyncService.h"
 #include "Network/Service/ClusterSync/DisplayClusterClusterSyncStrings.h"
-#include "Network/Barrier/DisplayClusterBarrierFactory.h"
-#include "Network/Barrier/IDisplayClusterBarrier.h"
-#include "Network/Conversion/DisplayClusterNetworkDataConversion.h"
-#include "Network/Listener/DisplayClusterTcpListener.h"
-#include "Network/Session/DisplayClusterSession.h"
 
+#include "DisplayClusterConfigurationTypes.h"
+#include "DisplayClusterEnums.h"
+#include "IDisplayCluster.h"
+
+#include "Cluster/DisplayClusterClusterEvent.h"
 #include "Cluster/IPDisplayClusterClusterManager.h"
 #include "Cluster/Controller/IDisplayClusterClusterNodeController.h"
-#include "Cluster/Failover/IDisplayClusterFailoverNodeController.h"
-#include "Cluster/DisplayClusterClusterEvent.h"
+#include "Cluster/Failover/DisplayClusterCommDataCache.h"
+#include "Cluster/RequestHandler/DisplayClusterCommRequestHandlerLocal.h"
+#include "Cluster/RequestHandler/DisplayClusterCommRequestHandlerRemote.h"
 
-#include "Game/IPDisplayClusterGameManager.h"
 #include "Config/IPDisplayClusterConfigManager.h"
+#include "Game/IPDisplayClusterGameManager.h"
+#include "Interfaces/IPv4/IPv4Endpoint.h"
 
-#include "DisplayClusterConfigurationTypes.h"
-
-#include "Misc/QualifiedFrameTime.h"
 #include "Misc/DisplayClusterGlobals.h"
 #include "Misc/DisplayClusterLog.h"
+#include "Misc/QualifiedFrameTime.h"
 
-#include "DisplayClusterEnums.h"
-
-#include "Interfaces/IPv4/IPv4Endpoint.h"
+#include "Network/Barrier/DisplayClusterBarrierFactory.h"
+#include "Network/Barrier/IDisplayClusterBarrier.h"
+#include "Network/Configuration/DisplayClusterNetworkConfiguration.h"
+#include "Network/Conversion/DisplayClusterNetworkDataConversion.h"
+#include "Network/Listener/DisplayClusterTcpListener.h"
+#include "Network/Session/DisplayClusterSession.h"
 
 
-FDisplayClusterClusterSyncService::FDisplayClusterClusterSyncService()
-	: FDisplayClusterService(FString("SRV_CS"))
+FDisplayClusterClusterSyncService::FDisplayClusterClusterSyncService(const FName& InInstanceName)
+	: FDisplayClusterService(InInstanceName.ToString())
 {
+	IDisplayCluster& DCAPI = IDisplayCluster::Get();
+
 	// Get list of cluster node IDs
-	TArray<FString> NodeIds;
-	GDisplayCluster->GetPrivateClusterMgr()->GetNodeIds(NodeIds);
+	TSet<FString> NodeIds;
+	if (const IDisplayClusterClusterManager* const ClusterMgr = DCAPI.GetClusterMgr())
+	{
+		ClusterMgr->GetNodeIds(NodeIds);
+	}
 
-	// Network settings from config data
-	const FDisplayClusterConfigurationNetworkSettings& NetworkSettings = GDisplayCluster->GetConfigMgr()->GetConfig()->Cluster->Network;
+	// Get cluster configuration
+	UDisplayClusterConfigurationData* ConfigData = nullptr;
+	if (const IDisplayClusterConfigManager* const ConfigMgr = DCAPI.GetConfigMgr())
+	{
+		ConfigData = ConfigMgr->GetConfig();
+	}
 
 	// Instantiate service barriers
-	BarrierGameStart.Reset(FDisplayClusterBarrierFactory::CreateBarrier(TEXT("GameStart_barrier"), NodeIds, NetworkSettings.GameStartBarrierTimeout));
-	BarrierFrameStart.Reset(FDisplayClusterBarrierFactory::CreateBarrier(TEXT("FrameStart_barrier"), NodeIds, NetworkSettings.FrameStartBarrierTimeout));
-	BarrierFrameEnd.Reset(FDisplayClusterBarrierFactory::CreateBarrier(TEXT("FrameEnd_barrier"), NodeIds, NetworkSettings.FrameEndBarrierTimeout));
+	if (ConfigData)
+	{
+		// Network settings from config data
+		const FDisplayClusterConfigurationNetworkSettings& NetworkSettings = ConfigData->Cluster->Network;
+
+		BarrierGameStart.Reset( FDisplayClusterBarrierFactory::CreateBarrier(TEXT("GameStart_barrier"),  NodeIds, NetworkSettings.GameStartBarrierTimeout));
+		BarrierFrameStart.Reset(FDisplayClusterBarrierFactory::CreateBarrier(TEXT("FrameStart_barrier"), NodeIds, NetworkSettings.FrameStartBarrierTimeout));
+		BarrierFrameEnd.Reset(  FDisplayClusterBarrierFactory::CreateBarrier(TEXT("FrameEnd_barrier"),   NodeIds, NetworkSettings.FrameEndBarrierTimeout));
+	}
+	else
+	{
+		UE_LOG(LogDisplayClusterNetwork, Warning, TEXT("%s - Couldn't obtain cluster configuration. Using default barrier timeouts."), *GetName());
+
+		FDisplayClusterConfigurationNetworkSettings DefaultNetworkSettings;
+		BarrierGameStart.Reset( FDisplayClusterBarrierFactory::CreateBarrier(TEXT("GameStart_barrier"),  NodeIds, DefaultNetworkSettings.GameStartBarrierTimeout));
+		BarrierFrameStart.Reset(FDisplayClusterBarrierFactory::CreateBarrier(TEXT("FrameStart_barrier"), NodeIds, DefaultNetworkSettings.FrameStartBarrierTimeout));
+		BarrierFrameEnd.Reset(  FDisplayClusterBarrierFactory::CreateBarrier(TEXT("FrameEnd_barrier"),   NodeIds, DefaultNetworkSettings.FrameEndBarrierTimeout));
+	}
 
 	// Put the barriers into an aux container
 	ServiceBarriers.Emplace(BarrierGameStart->GetName(),  &BarrierGameStart);
@@ -68,39 +95,30 @@ FDisplayClusterClusterSyncService::~FDisplayClusterClusterSyncService()
 	// Unsubscribe from SessionClosed notifications
 	OnSessionClosed().RemoveAll(this);
 
-	Shutdown();
+	ShutdownImpl();
 }
 
 
 bool FDisplayClusterClusterSyncService::Start(const FString& Address, const uint16 Port)
 {
-	// Activate all barriers
-	for (TPair<FString, TUniquePtr<IDisplayClusterBarrier>*>& BarrierIt : ServiceBarriers)
-	{
-		(*BarrierIt.Value)->Activate();
-	}
+	// Start internals
+	StartInternal();
 
 	return FDisplayClusterServer::Start(Address, Port);
 }
 
-bool FDisplayClusterClusterSyncService::Start(TSharedPtr<FDisplayClusterTcpListener>& ExternalListener)
+bool FDisplayClusterClusterSyncService::Start(TSharedRef<FDisplayClusterTcpListener>& ExternalListener)
 {
-	// Activate all barriers
-	for (TPair<FString, TUniquePtr<IDisplayClusterBarrier>*>& BarrierIt : ServiceBarriers)
-	{
-		(*BarrierIt.Value)->Activate();
-	}
+	// Start internals
+	StartInternal();
 
 	return FDisplayClusterServer::Start(ExternalListener);
 }
 
 void FDisplayClusterClusterSyncService::Shutdown()
 {
-	// Deactivate all barriers
-	for (TPair<FString, TUniquePtr<IDisplayClusterBarrier>*>& BarrierIt : ServiceBarriers)
-	{
-		(*BarrierIt.Value)->Deactivate();
-	}
+	// Release internals
+	ShutdownImpl();
 
 	return FDisplayClusterServer::Shutdown();
 }
@@ -135,30 +153,44 @@ TSharedPtr<IDisplayClusterSession> FDisplayClusterClusterSyncService::CreateSess
 	return MakeShared<FDisplayClusterSession<FDisplayClusterPacketInternal, true>>(SessionInfo, *this, *this, FDisplayClusterService::GetThreadPriority());
 }
 
+void FDisplayClusterClusterSyncService::StartInternal()
+{
+	// Activate all barriers
+	for (TPair<FString, TUniquePtr<IDisplayClusterBarrier>*>& BarrierIt : ServiceBarriers)
+	{
+		(*BarrierIt.Value)->Activate();
+	}
+}
+
 void FDisplayClusterClusterSyncService::ProcessSessionClosed(const FDisplayClusterSessionInfo& SessionInfo)
 {
 	if (!SessionInfo.IsTerminatedByServer())
 	{
-		// Get node ID
-		const FString NodeId = SessionInfo.NodeId.Get(FString());
-
-		if (!NodeId.IsEmpty())
+		// Ignore sessions with empty NodeId
+		if (SessionInfo.NodeId.IsSet())
 		{
+			const FString NodeId = SessionInfo.NodeId.Get(FString());
+
 			// We have to unregister the node that just disconnected from all barriers
 			for (TPair<FString, TUniquePtr<IDisplayClusterBarrier>*>& BarrierIt : ServiceBarriers)
 			{
 				(*BarrierIt.Value)->UnregisterSyncCaller(NodeId);
 			}
 
+			// Prepare failure info
+			FDisplayClusterServiceFailureEvent EventInfo;
+			EventInfo.NodeFailed  = SessionInfo.NodeId;
+			EventInfo.FailureType = FDisplayClusterServiceFailureEvent::ENodeFailType::ConnectionLost;
+
 			// Notify others about node fail
-			OnNodeFailed().Broadcast(NodeId, ENodeFailType::ConnectionLost);
+			OnNodeFailed().Broadcast(EventInfo);
 		}
 	}
 }
 
-void FDisplayClusterClusterSyncService::ProcessBarrierTimeout(const FString& BarrierName, const TArray<FString>& NodesTimedOut)
+void FDisplayClusterClusterSyncService::ProcessBarrierTimeout(const FString& BarrierName, const TSet<FString>& NodesTimedOut)
 {
-	// We have to unregister the node that just timed out from all barriers
+	// We have to unregister the node that just timed out from all the barriers
 	for (TPair<FString, TUniquePtr<IDisplayClusterBarrier>*>& BarrierIt : ServiceBarriers)
 	{
 		for (const FString& NodeId : NodesTimedOut)
@@ -170,7 +202,22 @@ void FDisplayClusterClusterSyncService::ProcessBarrierTimeout(const FString& Bar
 	// Notify others about timeout
 	for (const FString& NodeId : NodesTimedOut)
 	{
-		OnNodeFailed().Broadcast(NodeId, ENodeFailType::BarrierTimeOut);
+		// Prepare failure info
+		FDisplayClusterServiceFailureEvent EventInfo;
+		EventInfo.NodeFailed  = NodeId;
+		EventInfo.FailureType = FDisplayClusterServiceFailureEvent::ENodeFailType::BarrierTimeOut;
+
+		// Notify others about node fail
+		OnNodeFailed().Broadcast(EventInfo);
+	}
+}
+
+void FDisplayClusterClusterSyncService::ShutdownImpl()
+{
+	// Deactivate all barriers
+	for (TPair<FString, TUniquePtr<IDisplayClusterBarrier>*>& BarrierIt : ServiceBarriers)
+	{
+		(*BarrierIt.Value)->Deactivate();
 	}
 }
 
@@ -276,8 +323,8 @@ TSharedPtr<FDisplayClusterPacketInternal> FDisplayClusterClusterSyncService::Pro
 		
 		const EDisplayClusterCommResult CommResult = GetEventsData(JsonEvents, BinaryEvents);
 
-		DisplayClusterNetworkDataConversion::JsonEventsToInternalPacket(JsonEvents, Response);
-		DisplayClusterNetworkDataConversion::BinaryEventsToInternalPacket(BinaryEvents, Response);
+		UE::nDisplay::DisplayClusterNetworkDataConversion::JsonEventsToInternalPacket(JsonEvents, Response);
+		UE::nDisplay::DisplayClusterNetworkDataConversion::BinaryEventsToInternalPacket(BinaryEvents, Response);
 
 		Response->SetCommResult(CommResult);
 
@@ -308,51 +355,71 @@ TSharedPtr<FDisplayClusterPacketInternal> FDisplayClusterClusterSyncService::Pro
 //////////////////////////////////////////////////////////////////////////////////////////////
 EDisplayClusterCommResult FDisplayClusterClusterSyncService::WaitForGameStart()
 {
-	TRACE_CPUPROFILER_EVENT_SCOPE(nD SRV_CS::WaitForGameStart);
+	TRACE_CPUPROFILER_EVENT_SCOPE(SRV_CS::WaitForGameStart);
 
-	const FString CachedNodeId = GetSessionInfoCache().NodeId.Get("");
+	const FString CachedNodeId = GetSessionInfoCache().NodeId.Get(TEXT(""));
 	BarrierGameStart->Wait(CachedNodeId);
 	return EDisplayClusterCommResult::Ok;
 }
 
 EDisplayClusterCommResult FDisplayClusterClusterSyncService::WaitForFrameStart()
 {
-	TRACE_CPUPROFILER_EVENT_SCOPE(nD SRV_CS::WaitForFrameStart);
+	TRACE_CPUPROFILER_EVENT_SCOPE(SRV_CS::WaitForFrameStart);
 
-	const FString CachedNodeId = GetSessionInfoCache().NodeId.Get("");
+	const FString CachedNodeId = GetSessionInfoCache().NodeId.Get(TEXT(""));
 	BarrierFrameStart->Wait(CachedNodeId);
 	return EDisplayClusterCommResult::Ok;
 }
 
 EDisplayClusterCommResult FDisplayClusterClusterSyncService::WaitForFrameEnd()
 {
-	TRACE_CPUPROFILER_EVENT_SCOPE(nD SRV_CS::WaitForFrameEnd);
+	TRACE_CPUPROFILER_EVENT_SCOPE(SRV_CS::WaitForFrameEnd);
 
-	const FString CachedNodeId = GetSessionInfoCache().NodeId.Get("");
+	const FString CachedNodeId = GetSessionInfoCache().NodeId.Get(TEXT(""));
 	BarrierFrameEnd->Wait(CachedNodeId);
 	return EDisplayClusterCommResult::Ok;
 }
 
 EDisplayClusterCommResult FDisplayClusterClusterSyncService::GetTimeData(double& OutDeltaTime, double& OutGameTime, TOptional<FQualifiedFrameTime>& OutFrameTime)
 {
-	TRACE_CPUPROFILER_EVENT_SCOPE(nD SRV_CS::GetTimeData);
-	return GDisplayCluster->GetPrivateClusterMgr()->GetClusterNodeController()->GetTimeData(OutDeltaTime, OutGameTime, OutFrameTime);
+	TRACE_CPUPROFILER_EVENT_SCOPE(SRV_CS::GetTimeData);
+
+	FDisplayClusterCommRequestHandlerBase& Handler = (IsLocalRequest() ? 
+		static_cast<FDisplayClusterCommRequestHandlerBase&>(FDisplayClusterCommRequestHandlerLocal::Get()) :
+		static_cast<FDisplayClusterCommRequestHandlerBase&>(FDisplayClusterCommRequestHandlerRemote::Get()));
+
+	return Handler.GetTimeData(OutDeltaTime, OutGameTime, OutFrameTime);
 }
 
 EDisplayClusterCommResult FDisplayClusterClusterSyncService::GetObjectsData(EDisplayClusterSyncGroup InSyncGroup, TMap<FString, FString>& OutObjectsData)
 {
-	TRACE_CPUPROFILER_EVENT_SCOPE(nD SRV_CS::GetObjectsData);
-	return GDisplayCluster->GetPrivateClusterMgr()->GetClusterNodeController()->GetObjectsData(InSyncGroup, OutObjectsData);
+	TRACE_CPUPROFILER_EVENT_SCOPE(SRV_CS::GetObjectsData);
+
+	FDisplayClusterCommRequestHandlerBase& Handler = (IsLocalRequest() ?
+		static_cast<FDisplayClusterCommRequestHandlerBase&>(FDisplayClusterCommRequestHandlerLocal::Get()) :
+		static_cast<FDisplayClusterCommRequestHandlerBase&>(FDisplayClusterCommRequestHandlerRemote::Get()));
+
+	return Handler.GetObjectsData(InSyncGroup, OutObjectsData);
 }
 
 EDisplayClusterCommResult FDisplayClusterClusterSyncService::GetEventsData(TArray<TSharedPtr<FDisplayClusterClusterEventJson, ESPMode::ThreadSafe>>& OutJsonEvents, TArray<TSharedPtr<FDisplayClusterClusterEventBinary, ESPMode::ThreadSafe>>& OutBinaryEvents)
 {
-	TRACE_CPUPROFILER_EVENT_SCOPE(nD SRV_CS::GetEventsData);
-	return GDisplayCluster->GetPrivateClusterMgr()->GetClusterNodeController()->GetEventsData(OutJsonEvents, OutBinaryEvents);
+	TRACE_CPUPROFILER_EVENT_SCOPE(SRV_CS::GetEventsData);
+
+	FDisplayClusterCommRequestHandlerBase& Handler = (IsLocalRequest() ?
+		static_cast<FDisplayClusterCommRequestHandlerBase&>(FDisplayClusterCommRequestHandlerLocal::Get()) :
+		static_cast<FDisplayClusterCommRequestHandlerBase&>(FDisplayClusterCommRequestHandlerRemote::Get()));
+
+	return Handler.GetEventsData(OutJsonEvents, OutBinaryEvents);
 }
 
 EDisplayClusterCommResult FDisplayClusterClusterSyncService::GetNativeInputData(TMap<FString, FString>& OutNativeInputData)
 {
-	TRACE_CPUPROFILER_EVENT_SCOPE(nD SRV_CS::GetNativeInputData);
-	return GDisplayCluster->GetPrivateClusterMgr()->GetClusterNodeController()->GetNativeInputData(OutNativeInputData);
+	TRACE_CPUPROFILER_EVENT_SCOPE(SRV_CS::GetNativeInputData);
+
+	FDisplayClusterCommRequestHandlerBase& Handler = (IsLocalRequest() ?
+		static_cast<FDisplayClusterCommRequestHandlerBase&>(FDisplayClusterCommRequestHandlerLocal::Get()) :
+		static_cast<FDisplayClusterCommRequestHandlerBase&>(FDisplayClusterCommRequestHandlerRemote::Get()));
+
+	return Handler.GetNativeInputData(OutNativeInputData);
 }
diff --git a/Engine/Plugins/Runtime/nDisplay/Source/DisplayCluster/Private/Network/Service/ClusterSync/DisplayClusterClusterSyncService.h b/Engine/Plugins/Runtime/nDisplay/Source/DisplayCluster/Private/Network/Service/ClusterSync/DisplayClusterClusterSyncService.h
--- a/Engine/Plugins/Runtime/nDisplay/Source/DisplayCluster/Private/Network/Service/ClusterSync/DisplayClusterClusterSyncService.h
+++ b/Engine/Plugins/Runtime/nDisplay/Source/DisplayCluster/Private/Network/Service/ClusterSync/DisplayClusterClusterSyncService.h
@@ -2,8 +2,6 @@
 
 #pragma once
 
-#include "CoreMinimal.h"
-
 #include "Network/Service/DisplayClusterService.h"
 #include "Network/Session/IDisplayClusterSessionPacketHandler.h"
 #include "Network/Protocol/IDisplayClusterProtocolClusterSync.h"
@@ -21,39 +19,49 @@ class FDisplayClusterClusterSyncService
 	, protected IDisplayClusterProtocolClusterSync
 {
 public:
-	FDisplayClusterClusterSyncService();
+
+	FDisplayClusterClusterSyncService(const FName& InInstanceName);
+
 	virtual ~FDisplayClusterClusterSyncService();
 
 public:
-	//////////////////////////////////////////////////////////////////////////////////////////////
-	// IDisplayClusterServer
-	//////////////////////////////////////////////////////////////////////////////////////////////
+
+	//~ Begin IDisplayClusterServer
 	virtual bool Start(const FString& Address, const uint16 Port) override;
-	virtual bool Start(TSharedPtr<FDisplayClusterTcpListener>& ExternalListener) override;
-	virtual void Shutdown() override final;
+	virtual bool Start(TSharedRef<FDisplayClusterTcpListener>& ExternalListener) override;
+	virtual void Shutdown() override;
 	virtual FString GetProtocolName() const override;
 	virtual void KillSession(const FString& NodeId) override;
+	//~ End IDisplayClusterServer
 
 protected:
-	// Creates session instance for this service
+
+	/** Creates session instance for this service */
 	virtual TSharedPtr<IDisplayClusterSession> CreateSession(FDisplayClusterSessionInfo& SessionInfo) override;
 
-	// Callback when a session is closed
+private:
+
+	/** Performs internal initialization during server start */
+	void StartInternal();
+
+	/** Callback when a session is closed */
 	void ProcessSessionClosed(const FDisplayClusterSessionInfo& SessionInfo);
 
-	// Callbck on barrier timeout
-	void ProcessBarrierTimeout(const FString& BarrierName, const TArray<FString>& NodesTimedOut);
+	/** Callbck on barrier timeout */
+	void ProcessBarrierTimeout(const FString& BarrierName, const TSet<FString>& NodesTimedOut);
+
+	/** Non-virtual shutdown implementation */
+	void ShutdownImpl();
 
 protected:
-	//////////////////////////////////////////////////////////////////////////////////////////////
-	// IDisplayClusterSessionPacketHandler
-	//////////////////////////////////////////////////////////////////////////////////////////////
+
+	//~ Begin IDisplayClusterSessionPacketHandler
 	virtual TSharedPtr<FDisplayClusterPacketInternal> ProcessPacket(const TSharedPtr<FDisplayClusterPacketInternal>& Request, const FDisplayClusterSessionInfo& SessionInfo) override;
+	//~ End IDisplayClusterSessionPacketHandler
 
-private:
-	//////////////////////////////////////////////////////////////////////////////////////////////
-	// IDisplayClusterProtocolClusterSync
-	//////////////////////////////////////////////////////////////////////////////////////////////
+protected:
+
+	//~ Begin IDisplayClusterProtocolClusterSync
 	virtual EDisplayClusterCommResult WaitForGameStart() override;
 	virtual EDisplayClusterCommResult WaitForFrameStart() override;
 	virtual EDisplayClusterCommResult WaitForFrameEnd() override;
@@ -61,15 +69,19 @@ private:
 	virtual EDisplayClusterCommResult GetObjectsData(EDisplayClusterSyncGroup InSyncGroup, TMap<FString, FString>& OutObjectsData)  override;
 	virtual EDisplayClusterCommResult GetEventsData(TArray<TSharedPtr<FDisplayClusterClusterEventJson, ESPMode::ThreadSafe>>& OutJsonEvents, TArray<TSharedPtr<FDisplayClusterClusterEventBinary, ESPMode::ThreadSafe>>& OutBinaryEvents) override;
 	virtual EDisplayClusterCommResult GetNativeInputData(TMap<FString, FString>& OutNativeInputData) override;
+	//~ End IDisplayClusterProtocolClusterSync
 
 private:
-	// Game start sync barrier
+
+	/** Game start sync barrier */
 	TUniquePtr<IDisplayClusterBarrier> BarrierGameStart;
-	// Frame start barrier
+
+	/** Frame start barrier */
 	TUniquePtr<IDisplayClusterBarrier> BarrierFrameStart;
-	// Frame end barrier
+
+	/** Frame end barrier */
 	TUniquePtr<IDisplayClusterBarrier> BarrierFrameEnd;
 
-	// Auxiliary container that keeps all the barriers
+	/** Auxiliary container that keeps all the barriers */
 	TMap<FString, TUniquePtr<IDisplayClusterBarrier>*> ServiceBarriers;
 };
diff --git a/Engine/Plugins/Runtime/nDisplay/Source/DisplayCluster/Private/Network/Service/ClusterSync/DisplayClusterClusterSyncStrings.h b/Engine/Plugins/Runtime/nDisplay/Source/DisplayCluster/Private/Network/Service/ClusterSync/DisplayClusterClusterSyncStrings.h
--- a/Engine/Plugins/Runtime/nDisplay/Source/DisplayCluster/Private/Network/Service/ClusterSync/DisplayClusterClusterSyncStrings.h
+++ b/Engine/Plugins/Runtime/nDisplay/Source/DisplayCluster/Private/Network/Service/ClusterSync/DisplayClusterClusterSyncStrings.h
@@ -2,7 +2,7 @@
 
 #pragma once
 
-#include "CoreMinimal.h"
+#include "HAL/Platform.h"
 
 
 /**
@@ -57,5 +57,6 @@ namespace DisplayClusterClusterSyncStrings
 	namespace GetNativeInputData
 	{
 		constexpr static const TCHAR* Name = TEXT("GetNativeInputData");
+		constexpr static const TCHAR* ArgNativeInputData = TEXT("NativeInputData");
 	}
 };
diff --git a/Engine/Plugins/Runtime/nDisplay/Source/DisplayCluster/Private/Network/Service/DisplayClusterService.cpp b/Engine/Plugins/Runtime/nDisplay/Source/DisplayCluster/Private/Network/Service/DisplayClusterService.cpp
--- a/Engine/Plugins/Runtime/nDisplay/Source/DisplayCluster/Private/Network/Service/DisplayClusterService.cpp
+++ b/Engine/Plugins/Runtime/nDisplay/Source/DisplayCluster/Private/Network/Service/DisplayClusterService.cpp
@@ -2,18 +2,20 @@
 
 #include "Network/Service/DisplayClusterService.h"
 
-#include "Network/DisplayClusterNetworkTypes.h"
-
+#include "Cluster/IPDisplayClusterClusterManager.h"
 #include "Config/IPDisplayClusterConfigManager.h"
-#include "DisplayClusterConfigurationTypes.h"
 
 #include "Misc/DisplayClusterGlobals.h"
 #include "Misc/DisplayClusterLog.h"
 
+#include "Network/DisplayClusterNetworkTypes.h"
+
 #include "Interfaces/IPv4/IPv4Endpoint.h"
 #include "HAL/IConsoleManager.h"
 #include "HAL/RunnableThread.h"
 
+#include "DisplayClusterConfigurationTypes.h"
+
 
 // nDisplay service threads priority
 static TAutoConsoleVariable<int32> CVarServiceThreadsPriority(
@@ -37,6 +39,7 @@ FDisplayClusterService::FDisplayClusterService(const FString& Name)
 {
 }
 
+
 EThreadPriority FDisplayClusterService::ConvertThreadPriorityFromCvarValue(int ThreadPriority)
 {
 	switch (ThreadPriority)
@@ -67,24 +70,6 @@ EThreadPriority FDisplayClusterService::GetThreadPriority()
 	return FDisplayClusterService::ConvertThreadPriorityFromCvarValue(CVarServiceThreadsPriority.GetValueOnAnyThread());
 }
 
-EDisplayClusterCommResult FDisplayClusterService::TranslateBarrierWaitResultIntoCommResult(EDisplayClusterBarrierWaitResult WaitResult)
-{
-	// Translate BarrierWaitResult into CommResult
-	switch (WaitResult)
-	{
-	case EDisplayClusterBarrierWaitResult::Ok:
-	case EDisplayClusterBarrierWaitResult::NotActive:
-		return EDisplayClusterCommResult::Ok;
-
-	case EDisplayClusterBarrierWaitResult::NotAllowed:
-	case EDisplayClusterBarrierWaitResult::TimeOut:
-		return EDisplayClusterCommResult::NotAllowed;
-
-	default:
-		return EDisplayClusterCommResult::NotImplemented;
-	}
-}
-
 void FDisplayClusterService::SetSessionInfoCache(const FDisplayClusterSessionInfo& SessionInfo)
 {
 	const uint32 ThreadId = FPlatformTLS::GetCurrentThreadId();
@@ -120,3 +105,12 @@ void FDisplayClusterService::ClearCache()
 	FScopeLock Lock(&SessionInfoCacheCS);
 	SessionInfoCache.Reset();
 }
+
+bool FDisplayClusterService::IsLocalRequest() const
+{
+	static const FString ThisNodeId = GDisplayCluster->GetClusterMgr()->GetNodeId();
+
+	// Check if this session belongs to this node. If so, it's a local request.
+	const bool bLocalRequest = ThisNodeId.Equals(GetSessionInfoCache().NodeId.Get(TEXT("")), ESearchCase::IgnoreCase);
+	return bLocalRequest;
+}
diff --git a/Engine/Plugins/Runtime/nDisplay/Source/DisplayCluster/Private/Network/Service/DisplayClusterService.h b/Engine/Plugins/Runtime/nDisplay/Source/DisplayCluster/Private/Network/Service/DisplayClusterService.h
--- a/Engine/Plugins/Runtime/nDisplay/Source/DisplayCluster/Private/Network/Service/DisplayClusterService.h
+++ b/Engine/Plugins/Runtime/nDisplay/Source/DisplayCluster/Private/Network/Service/DisplayClusterService.h
@@ -3,9 +3,13 @@
 #pragma once
 
 #include "Network/DisplayClusterServer.h"
+
+#include "GenericPlatform/GenericPlatformAffinity.h"
+
+#include "Misc/Optional.h"
+
 #include "Network/DisplayClusterNetworkTypes.h"
 #include "Network/Barrier/IDisplayClusterBarrier.h"
-#include "GenericPlatform/GenericPlatformAffinity.h"
 
 
 class FSocket;
@@ -13,28 +17,83 @@ struct FIPv4Endpoint;
 struct FDisplayClusterSessionInfo;
 
 
+/**
+ * A container to keep all the information about any potential networking failures
+ */
+struct FDisplayClusterServiceFailureEvent
+{
+public:
+
+	/** Cluster node lost */
+	enum class ENodeFailType : uint8
+	{
+		Unknown,
+		BarrierTimeOut,
+		ConnectionLost,
+	};
+
+public:
+
+	/** The ID of the cluster node that failed */
+	TOptional<FString> NodeFailed;
+
+	/** Failure type */
+	ENodeFailType FailureType = ENodeFailType::Unknown;
+};
+
+
 /**
  * Abstract DisplayCluster service
  */
 class FDisplayClusterService
 	: public FDisplayClusterServer
+
 {
 public:
+
 	FDisplayClusterService(const FString& Name);
 
 public:
+
+	/** A helper function to convert a cvar integer into the thread priority EThreadPriority enum value */
 	static EThreadPriority ConvertThreadPriorityFromCvarValue(int ThreadPriority);
+
+	/** Returns thread priority that is currently set by the corresponding CVar */
 	static EThreadPriority GetThreadPriority();
-	static EDisplayClusterCommResult TranslateBarrierWaitResultIntoCommResult(EDisplayClusterBarrierWaitResult WaitResult);
+
+public:
+
+	/** Networking failure notification event */
+	DECLARE_EVENT_OneParam(FDisplayClusterService, FNodeFailedEvent, const FDisplayClusterServiceFailureEvent&);
+	FNodeFailedEvent& OnNodeFailed()
+	{
+		return NodeFailedEvent;
+	}
 
 protected:
-	// Cache session info data if needed for child service implementations
+
+	/** Cache session info data (calling thread) if needed for child service implementations */
 	void SetSessionInfoCache(const FDisplayClusterSessionInfo& SessionInfo);
+
+	/** Returns session info of the calling thread */
 	const FDisplayClusterSessionInfo& GetSessionInfoCache() const;
+
+	/** Resets whole session info cache */
 	void ClearCache();
 
+	/** Checks if an incoming request is local (was sent by this node) */
+	bool IsLocalRequest() const;
+
 private:
-	// Session info cache
+
+	/** Session info cache */
 	TMap<uint32, FDisplayClusterSessionInfo> SessionInfoCache;
+
+	/** A critical section to safely accesss the session info cache */
 	mutable FCriticalSection SessionInfoCacheCS;
+
+private:
+
+	/** Failure reporting event */
+	FNodeFailedEvent NodeFailedEvent;
 };
diff --git a/Engine/Plugins/Runtime/nDisplay/Source/DisplayCluster/Private/Network/Service/GenericBarrier/DisplayClusterGenericBarrierClient.cpp b/Engine/Plugins/Runtime/nDisplay/Source/DisplayCluster/Private/Network/Service/GenericBarrier/DisplayClusterGenericBarrierClient.cpp
--- a/Engine/Plugins/Runtime/nDisplay/Source/DisplayCluster/Private/Network/Service/GenericBarrier/DisplayClusterGenericBarrierClient.cpp
+++ b/Engine/Plugins/Runtime/nDisplay/Source/DisplayCluster/Private/Network/Service/GenericBarrier/DisplayClusterGenericBarrierClient.cpp
@@ -14,6 +14,7 @@
 #include "Misc/DisplayClusterTypesConverter.h"
 
 #include "Network/Barrier/IDisplayClusterBarrier.h"
+#include "Network/Configuration/DisplayClusterNetworkConfiguration.h"
 #include "Network/Service/GenericBarrier/DisplayClusterGenericBarrierService.h"
 #include "Network/Service/GenericBarrier/DisplayClusterGenericBarrierStrings.h"
 #include "Network/Packet/DisplayClusterPacketInternal.h"
@@ -23,13 +24,8 @@
 #include "DisplayClusterEnums.h"
 
 
-FDisplayClusterGenericBarrierClient::FDisplayClusterGenericBarrierClient()
-	: FDisplayClusterGenericBarrierClient(FString("CLN_GB"))
-{
-}
-
-FDisplayClusterGenericBarrierClient::FDisplayClusterGenericBarrierClient(const FString& InName)
-	: FDisplayClusterClient(InName)
+FDisplayClusterGenericBarrierClient::FDisplayClusterGenericBarrierClient(const FName& InName)
+	: FDisplayClusterClient(InName.ToString())
 {
 }
 
@@ -48,7 +44,7 @@ bool FDisplayClusterGenericBarrierClient::Connect(const FString& Address, const
 	// Prepare 'hello' message
 	TSharedPtr<FDisplayClusterPacketInternal> HelloMsg = MakeShared<FDisplayClusterPacketInternal>(
 		DisplayClusterHelloMessageStrings::Hello::Name,
-		DisplayClusterGenericBarrierStrings::TypeRequest,
+		DisplayClusterHelloMessageStrings::Hello::TypeRequest,
 		DisplayClusterGenericBarrierStrings::ProtocolName
 	);
 
@@ -61,114 +57,10 @@ bool FDisplayClusterGenericBarrierClient::Connect(const FString& Address, const
 }
 
 
-//////////////////////////////////////////////////////////////////////////////////////////////
-// IDisplayClusterGenericBarriersClient
-//////////////////////////////////////////////////////////////////////////////////////////////
-bool FDisplayClusterGenericBarrierClient::Connect()
-{
-	// Generic barriers are available in runtime only
-	if (!GDisplayCluster || (GDisplayCluster->GetOperationMode() != EDisplayClusterOperationMode::Cluster))
-	{
-		return false;
-	}
-
-	// Now we need to get server address/port data from the config
-	if (const IPDisplayClusterConfigManager* const ConfigMgr = GDisplayCluster->GetPrivateConfigMgr())
-	{
-		if (const UDisplayClusterConfigurationData* const Config = ConfigMgr->GetConfig())
-		{
-			if (const UDisplayClusterConfigurationClusterNode* const PrimaryNode = Config->Cluster->GetNode(Config->Cluster->PrimaryNode.Id))
-			{
-				// Now try to establish connection
-				return Connect(PrimaryNode->Host, Config->Cluster->PrimaryNode.Ports.ClusterSync, 3, 1000);
-			}
-		}
-	}
-
-	return false;
-}
-
-void FDisplayClusterGenericBarrierClient::Disconnect()
-{
-	// Forward IDisplayClusterGenericBarriersClient::Disconnect() to FDisplayClusterClient::Disconnect() explicitly
-	FDisplayClusterClient::Disconnect();
-}
-
-bool FDisplayClusterGenericBarrierClient::IsConnected() const
-{
-	// Forward IDisplayClusterGenericBarriersClient::IsConnected() to FDisplayClusterClient::IsConnected() explicitly
-	return FDisplayClusterClient::IsConnected();
-}
-
-FString FDisplayClusterGenericBarrierClient::GetName() const
-{
-	// Forward IDisplayClusterGenericBarriersClient::GetName() to FDisplayClusterClient::GetName() explicitly
-	return FDisplayClusterClient::GetName();
-}
-
-bool FDisplayClusterGenericBarrierClient::CreateBarrier(const FString& BarrierId, const TArray<FString>& UniqueThreadMarkers, uint32 Timeout)
-{
-	EBarrierControlResult CtrlResult = EBarrierControlResult::UnknownError;
-	const EDisplayClusterCommResult Result = CreateBarrier(BarrierId, UniqueThreadMarkers, Timeout, CtrlResult);
-	const bool bResult = (CtrlResult == EBarrierControlResult::CreatedSuccessfully || CtrlResult == EBarrierControlResult::AlreadyExists) && (Result == EDisplayClusterCommResult::Ok);
-
-	if (bResult)
-	{
-		// Setup sync delegate
-		ConfigureBarrierSyncDelegate(BarrierId, true);
-	}
-
-	return bResult;
-}
-
-bool FDisplayClusterGenericBarrierClient::WaitUntilBarrierIsCreated(const FString& BarrierId)
-{
-	EBarrierControlResult CtrlResult = EBarrierControlResult::UnknownError;
-	const EDisplayClusterCommResult Result = WaitUntilBarrierIsCreated(BarrierId, CtrlResult);
-	return (CtrlResult == EBarrierControlResult::AlreadyExists) && (Result == EDisplayClusterCommResult::Ok);
-}
-
-bool FDisplayClusterGenericBarrierClient::IsBarrierAvailable(const FString& BarrierId)
-{
-	EBarrierControlResult CtrlResult = EBarrierControlResult::UnknownError;
-	const EDisplayClusterCommResult Result = IsBarrierAvailable(BarrierId, CtrlResult);
-	return (CtrlResult == EBarrierControlResult::AlreadyExists) && (Result == EDisplayClusterCommResult::Ok);
-}
-
-IDisplayClusterGenericBarriersClient::FOnGenericBarrierSynchronizationDelegate* FDisplayClusterGenericBarrierClient::GetBarrierSyncDelegate(const FString& BarrierId)
-{
-	return BarrierSyncDelegates.Find(BarrierId);
-}
-
-bool FDisplayClusterGenericBarrierClient::ReleaseBarrier(const FString& BarrierId)
-{
-	// Release sync delegate
-	ConfigureBarrierSyncDelegate(BarrierId, false);
-
-	EBarrierControlResult CtrlResult = EBarrierControlResult::UnknownError;
-	const EDisplayClusterCommResult Result = ReleaseBarrier(BarrierId, CtrlResult);
-	return (CtrlResult == EBarrierControlResult::ReleasedSuccessfully) && (Result == EDisplayClusterCommResult::Ok);
-}
-
-bool FDisplayClusterGenericBarrierClient::Synchronize(const FString& BarrierId, const FString& UniqueThreadMarker)
-{
-	EBarrierControlResult CtrlResult = EBarrierControlResult::UnknownError;
-	const EDisplayClusterCommResult Result = SyncOnBarrier(BarrierId, UniqueThreadMarker, CtrlResult);
-	return (CtrlResult == EBarrierControlResult::SynchronizedSuccessfully) && (Result == EDisplayClusterCommResult::Ok);
-}
-
-bool FDisplayClusterGenericBarrierClient::Synchronize(const FString& BarrierId, const FString& UniqueThreadMarker, const TArray<uint8>& RequestData, TArray<uint8>& OutResponseData)
-{
-	EBarrierControlResult CtrlResult = EBarrierControlResult::UnknownError;
-	const EDisplayClusterCommResult Result = SyncOnBarrierWithData(BarrierId, UniqueThreadMarker, RequestData, OutResponseData, CtrlResult);
-	return (CtrlResult == EBarrierControlResult::SynchronizedSuccessfully) && (Result == EDisplayClusterCommResult::Ok);
-}
-
-
 //////////////////////////////////////////////////////////////////////////////////////////////
 // IDisplayClusterProtocolGenericBarrier
 //////////////////////////////////////////////////////////////////////////////////////////////
-EDisplayClusterCommResult FDisplayClusterGenericBarrierClient::CreateBarrier(const FString& BarrierId, const TArray<FString>& UniqueThreadMarkers, uint32 Timeout, EBarrierControlResult& Result)
+EDisplayClusterCommResult FDisplayClusterGenericBarrierClient::CreateBarrier(const FString& BarrierId, const TMap<FString, TSet<FString>>& NodeToSyncCallers, uint32 Timeout, EBarrierControlResult& Result)
 {
 	const TSharedPtr<FDisplayClusterPacketInternal> Request = MakeShared<FDisplayClusterPacketInternal>(
 		DisplayClusterGenericBarrierStrings::CreateBarrier::Name,
@@ -179,9 +71,12 @@ EDisplayClusterCommResult FDisplayClusterGenericBarrierClient::CreateBarrier(con
 	Request->SetTextArg(DisplayClusterGenericBarrierStrings::ArgumentsDefaultCategory, DisplayClusterGenericBarrierStrings::ArgBarrierId, BarrierId);
 	Request->SetTextArg(DisplayClusterGenericBarrierStrings::ArgumentsDefaultCategory, DisplayClusterGenericBarrierStrings::CreateBarrier::ArgTimeout, Timeout);
 
-	// Stringify markers array and put as a single parameter
-	const FString ThreadMarkersPacked = DisplayClusterHelpers::str::template ArrayToStr(UniqueThreadMarkers, DisplayClusterStrings::common::ArrayValSeparator, false);
-	Request->SetTextArg(DisplayClusterGenericBarrierStrings::ArgumentsDefaultCategory, DisplayClusterGenericBarrierStrings::CreateBarrier::ArgThreadMarkers, ThreadMarkersPacked);
+	{
+		TArray<uint8> NodeToSyncCallersData;
+		FMemoryWriter MemoryWriter(NodeToSyncCallersData);
+		MemoryWriter << const_cast<TMap<FString, TSet<FString>>&>(NodeToSyncCallers);
+		Request->SetBinArg(DisplayClusterGenericBarrierStrings::ArgumentsDefaultCategory, DisplayClusterGenericBarrierStrings::CreateBarrier::ArgCallers, NodeToSyncCallersData);
+	}
 
 	TSharedPtr<FDisplayClusterPacketInternal> Response;
 
@@ -297,7 +192,7 @@ EDisplayClusterCommResult FDisplayClusterGenericBarrierClient::ReleaseBarrier(co
 	return Response->GetCommResult();
 }
 
-EDisplayClusterCommResult FDisplayClusterGenericBarrierClient::SyncOnBarrier(const FString& BarrierId, const FString& UniqueThreadMarker, EBarrierControlResult& Result)
+EDisplayClusterCommResult FDisplayClusterGenericBarrierClient::SyncOnBarrier(const FString& BarrierId, const FString& CallerId, EBarrierControlResult& Result)
 {
 	const TSharedPtr<FDisplayClusterPacketInternal> Request = MakeShared<FDisplayClusterPacketInternal>(
 		DisplayClusterGenericBarrierStrings::SyncOnBarrier::Name,
@@ -306,7 +201,7 @@ EDisplayClusterCommResult FDisplayClusterGenericBarrierClient::SyncOnBarrier(con
 
 	// Request arguments
 	Request->SetTextArg(DisplayClusterGenericBarrierStrings::ArgumentsDefaultCategory, DisplayClusterGenericBarrierStrings::ArgBarrierId, BarrierId);
-	Request->SetTextArg(DisplayClusterGenericBarrierStrings::ArgumentsDefaultCategory, DisplayClusterGenericBarrierStrings::SyncOnBarrier::ArgThreadMarker, UniqueThreadMarker);
+	Request->SetTextArg(DisplayClusterGenericBarrierStrings::ArgumentsDefaultCategory, DisplayClusterGenericBarrierStrings::SyncOnBarrier::ArgCallerId, CallerId);
 
 	TSharedPtr<FDisplayClusterPacketInternal> Response;
 
@@ -329,7 +224,7 @@ EDisplayClusterCommResult FDisplayClusterGenericBarrierClient::SyncOnBarrier(con
 	return Response->GetCommResult();
 }
 
-EDisplayClusterCommResult FDisplayClusterGenericBarrierClient::SyncOnBarrierWithData(const FString& BarrierId, const FString& UniqueThreadMarker, const TArray<uint8>& RequestData, TArray<uint8>& OutResponseData, EBarrierControlResult& Result)
+EDisplayClusterCommResult FDisplayClusterGenericBarrierClient::SyncOnBarrierWithData(const FString& BarrierId, const FString& CallerId, const TArray<uint8>& RequestData, TArray<uint8>& OutResponseData, EBarrierControlResult& Result)
 {
 	const TSharedPtr<FDisplayClusterPacketInternal> Request = MakeShared<FDisplayClusterPacketInternal>(
 		DisplayClusterGenericBarrierStrings::SyncOnBarrierWithData::Name,
@@ -338,7 +233,7 @@ EDisplayClusterCommResult FDisplayClusterGenericBarrierClient::SyncOnBarrierWith
 
 	// Request arguments
 	Request->SetTextArg(DisplayClusterGenericBarrierStrings::ArgumentsDefaultCategory, DisplayClusterGenericBarrierStrings::ArgBarrierId, BarrierId);
-	Request->SetTextArg(DisplayClusterGenericBarrierStrings::ArgumentsDefaultCategory, DisplayClusterGenericBarrierStrings::SyncOnBarrierWithData::ArgThreadMarker, UniqueThreadMarker);
+	Request->SetTextArg(DisplayClusterGenericBarrierStrings::ArgumentsDefaultCategory, DisplayClusterGenericBarrierStrings::SyncOnBarrierWithData::ArgCallerId, CallerId);
 	Request->SetBinArg(DisplayClusterGenericBarrierStrings::ArgumentsDefaultCategory, DisplayClusterGenericBarrierStrings::SyncOnBarrierWithData::ArgRequestData, RequestData);
 
 	TSharedPtr<FDisplayClusterPacketInternal> Response;
@@ -363,104 +258,3 @@ EDisplayClusterCommResult FDisplayClusterGenericBarrierClient::SyncOnBarrierWith
 
 	return Response->GetCommResult();
 }
-
-FDisplayClusterGenericBarrierService* FDisplayClusterGenericBarrierClient::GetGenericBarrierService() const
-{
-	if (IPDisplayClusterClusterManager* const ClusterMgr = GDisplayCluster->GetPrivateClusterMgr())
-	{
-		if (IDisplayClusterClusterNodeController* const ClusterNodeCtrl = ClusterMgr->GetClusterNodeController())
-		{
-			if (FDisplayClusterService* const Service = ClusterNodeCtrl->GetGenericBarriersServer())
-			{
-				return static_cast<FDisplayClusterGenericBarrierService*>(Service);
-			}
-		}
-	}
-
-	return nullptr;
-}
-
-bool FDisplayClusterGenericBarrierClient::ConfigureBarrierSyncDelegate(const FString& BarrierId, bool bSetup)
-{
-	// @note
-	// This is a temporary workaround. It will be refactored during failover v2 implementation in more clear and generic way.
-	//
-	// Once a barrier is created, we can explicitly set custom sync handler to that specific barrier. As it was mentioned above,
-	// this is a temporary workaround. We need to pick GB service on the p-node, and set the synchronization delegate so it's
-	// called when all the barrier clients have arrived.
-	if (GDisplayCluster->GetOperationMode() == EDisplayClusterOperationMode::Cluster)
-	{
-		if (IPDisplayClusterClusterManager* ClusterMgr = GDisplayCluster->GetPrivateClusterMgr())
-		{
-			if (ClusterMgr->IsPrimary())
-			{
-				if (IDisplayClusterClusterNodeController* const ClusterNodeCtrl = ClusterMgr->GetClusterNodeController())
-				{
-					if (FDisplayClusterService* const Service = ClusterNodeCtrl->GetGenericBarriersServer())
-					{
-						FDisplayClusterGenericBarrierService* const GBService = static_cast<FDisplayClusterGenericBarrierService*>(Service);
-						if (TSharedPtr<IDisplayClusterBarrier, ESPMode::ThreadSafe> Barrier = GBService->GetBarrier(BarrierId))
-						{
-							// Setup
-							if (bSetup)
-							{
-								if (!BarrierSyncDelegates.Contains(BarrierId))
-								{
-									BarrierSyncDelegates.Emplace(BarrierId);
-									Barrier->GetPreSyncEndDelegate().BindRaw(this, &FDisplayClusterGenericBarrierClient::OnBarrierSync);
-								}
-							}
-							// Release
-							else
-							{
-								if (BarrierSyncDelegates.Remove(BarrierId) > 0)
-								{
-									Barrier->GetPreSyncEndDelegate().Unbind();
-								}
-							}
-
-							return true;
-						}
-					}
-				}
-			}
-		}
-	}
-
-	return false;
-}
-
-void FDisplayClusterGenericBarrierClient::OnBarrierSync(FDisplayClusterBarrierPreSyncEndDelegateData& SyncData)
-{
-	// Make sure the delegate is set
-	FOnGenericBarrierSynchronizationDelegate* const Delegate = BarrierSyncDelegates.Find(SyncData.BarrierId);
-	if (!Delegate || !Delegate->IsBound())
-	{
-		return;
-	}
-
-	// Access GB server
-	FDisplayClusterGenericBarrierService* const GBService = GetGenericBarrierService();
-	checkSlow(GBService);
-	if (!GBService)
-	{
-		return;
-	}
-
-	// Get barrier info
-	const TSharedPtr<FDisplayClusterGenericBarrierService::FBarrierInfo> BarrierInfo = GBService->GetBarrierInfo(SyncData.BarrierId);
-	checkSlow(BarrierInfo.IsValid());
-	if (!BarrierInfo.IsValid())
-	{
-		return;
-	}
-
-	// Last thing before triggering the callback. For optimization purposes, we disable further barrier info updates.
-	// At this point we got a sync callback, therefore all the thread markers and cluster nodes for this barrier
-	// are known. No need to update anymore as we won't get any new information.
-	GBService->SetBarrierInfoUpdateLocked(SyncData.BarrierId, true);
-
-	// Now forward data to the handler
-	FGenericBarrierSynchronizationDelegateData CallbackData{ SyncData.BarrierId, BarrierInfo->ThreadToNodeMapping, SyncData.RequestData, SyncData.ResponseData };
-	Delegate->Execute(CallbackData);
-}
diff --git a/Engine/Plugins/Runtime/nDisplay/Source/DisplayCluster/Private/Network/Service/GenericBarrier/DisplayClusterGenericBarrierClient.h b/Engine/Plugins/Runtime/nDisplay/Source/DisplayCluster/Private/Network/Service/GenericBarrier/DisplayClusterGenericBarrierClient.h
--- a/Engine/Plugins/Runtime/nDisplay/Source/DisplayCluster/Private/Network/Service/GenericBarrier/DisplayClusterGenericBarrierClient.h
+++ b/Engine/Plugins/Runtime/nDisplay/Source/DisplayCluster/Private/Network/Service/GenericBarrier/DisplayClusterGenericBarrierClient.h
@@ -2,11 +2,9 @@
 
 #pragma once
 
-#include "CoreMinimal.h"
 #include "Network/DisplayClusterClient.h"
 #include "Network/Packet/DisplayClusterPacketInternal.h"
 #include "Network/Protocol/IDisplayClusterProtocolGenericBarrier.h"
-#include "Cluster/IDisplayClusterGenericBarriersClient.h"
 
 class FDisplayClusterGenericBarrierService;
 struct FDisplayClusterBarrierPreSyncEndDelegateData;
@@ -18,59 +16,25 @@ struct FDisplayClusterBarrierPreSyncEndDelegateData;
 class FDisplayClusterGenericBarrierClient
 	: public FDisplayClusterClient<FDisplayClusterPacketInternal>
 	, public IDisplayClusterProtocolGenericBarrier
-	, public IDisplayClusterGenericBarriersClient
 {
 public:
-	FDisplayClusterGenericBarrierClient();
-	FDisplayClusterGenericBarrierClient(const FString& InName);
 
-public:
-	//////////////////////////////////////////////////////////////////////////////////////////////
-	// IDisplayClusterClient
-	//////////////////////////////////////////////////////////////////////////////////////////////
-	virtual bool Connect(const FString& Address, const uint16 Port, const uint32 ConnectRetriesAmount, const uint32 ConnectRetryDelay) override;
+	FDisplayClusterGenericBarrierClient(const FName& InName);
 
 public:
-	//////////////////////////////////////////////////////////////////////////////////////////////
-	// IDisplayClusterGenericBarriersClient
-	//////////////////////////////////////////////////////////////////////////////////////////////
-	virtual bool Connect() override;
-	virtual void Disconnect() override;
-	virtual bool IsConnected() const override;
-
-	virtual FString GetName() const override;
 
-	virtual bool CreateBarrier(const FString& BarrierId, const TArray<FString>& UniqueThreadMarkers, uint32 Timeout) override;
-	virtual bool WaitUntilBarrierIsCreated(const FString& BarrierId) override;
-	virtual bool IsBarrierAvailable(const FString& BarrierId) override;
-	virtual FOnGenericBarrierSynchronizationDelegate* GetBarrierSyncDelegate(const FString& BarrierId) override;
-	virtual bool ReleaseBarrier(const FString& BarrierId) override;
-	virtual bool Synchronize(const FString& BarrierId, const FString& UniqueThreadMarker) override;
-	virtual bool Synchronize(const FString& BarrierId, const FString& UniqueThreadMarker, const TArray<uint8>& RequestData, TArray<uint8>& OutResponseData) override;
+	//~ Begin IDisplayClusterClient
+	virtual bool Connect(const FString& Address, const uint16 Port, const uint32 ConnectRetriesAmount = 0, const uint32 ConnectRetryDelay = 0) override;
+	//~ End IDisplayClusterClient
 
 public:
-	//////////////////////////////////////////////////////////////////////////////////////////////
-	// IDisplayClusterProtocolGenericBarrier
-	//////////////////////////////////////////////////////////////////////////////////////////////
-	virtual EDisplayClusterCommResult CreateBarrier(const FString& BarrierId, const TArray<FString>& UniqueThreadMarkers, uint32 Timeout, EBarrierControlResult& Result) override;
+
+	//~ Begin IDisplayClusterProtocolGenericBarrier
+	virtual EDisplayClusterCommResult CreateBarrier(const FString& BarrierId, const TMap<FString, TSet<FString>>& NodeToSyncCallers, uint32 Timeout, EBarrierControlResult& Result) override;
 	virtual EDisplayClusterCommResult WaitUntilBarrierIsCreated(const FString& BarrierId, EBarrierControlResult& Result) override;
 	virtual EDisplayClusterCommResult IsBarrierAvailable(const FString& BarrierId, EBarrierControlResult& Result) override;
 	virtual EDisplayClusterCommResult ReleaseBarrier(const FString& BarrierId, EBarrierControlResult& Result) override;
-	virtual EDisplayClusterCommResult SyncOnBarrier(const FString& BarrierId, const FString& UniqueThreadMarker, EBarrierControlResult& Result) override;
-	virtual EDisplayClusterCommResult SyncOnBarrierWithData(const FString& BarrierId, const FString& UniqueThreadMarker, const TArray<uint8>& RequestData, TArray<uint8>& OutResponseData, EBarrierControlResult& Result) override;
-
-private:
-
-	/** Provides easy access to the GB service */
-	FDisplayClusterGenericBarrierService* GetGenericBarrierService() const;
-
-	/** Setup/release sync delegate for a specific barrier */
-	bool ConfigureBarrierSyncDelegate(const FString& BarrierId, bool bSetup);
-
-	/** Callback on barrier sync phase end */
-	void OnBarrierSync(FDisplayClusterBarrierPreSyncEndDelegateData& SyncData);
-
-private:
-	// Holds synchronization delegates for the barriers operated by this client
-	TMap<FString, FOnGenericBarrierSynchronizationDelegate> BarrierSyncDelegates;
+	virtual EDisplayClusterCommResult SyncOnBarrier(const FString& BarrierId, const FString& CallerId, EBarrierControlResult& Result) override;
+	virtual EDisplayClusterCommResult SyncOnBarrierWithData(const FString& BarrierId, const FString& CallerId, const TArray<uint8>& RequestData, TArray<uint8>& OutResponseData, EBarrierControlResult& Result) override;
+	//~ End IDisplayClusterProtocolGenericBarrier
 };
