[nDisplay] Base failover implementation (v1.5, P-node replacement)

diff --git a/Engine/Plugins/Runtime/nDisplay/Source/DisplayCluster/Private/Network/Service/GenericBarrier/DisplayClusterGenericBarrierService.cpp b/Engine/Plugins/Runtime/nDisplay/Source/DisplayCluster/Private/Network/Service/GenericBarrier/DisplayClusterGenericBarrierService.cpp
--- a/Engine/Plugins/Runtime/nDisplay/Source/DisplayCluster/Private/Network/Service/GenericBarrier/DisplayClusterGenericBarrierService.cpp
+++ b/Engine/Plugins/Runtime/nDisplay/Source/DisplayCluster/Private/Network/Service/GenericBarrier/DisplayClusterGenericBarrierService.cpp
@@ -4,6 +4,7 @@
 #include "Network/Service/GenericBarrier/DisplayClusterGenericBarrierStrings.h"
 #include "Network/Barrier/DisplayClusterBarrierFactory.h"
 #include "Network/Barrier/IDisplayClusterBarrier.h"
+#include "Network/Configuration/DisplayClusterNetworkConfiguration.h"
 #include "Network/Packet/DisplayClusterPacketInternal.h"
 #include "Network/Session/DisplayClusterSession.h"
 
@@ -14,39 +15,25 @@
 
 #include "HAL/Event.h"
 
+#include "Serialization/MemoryReader.h"
 
-FDisplayClusterGenericBarrierService::FDisplayClusterGenericBarrierService()
-	: FDisplayClusterService(FString("SRV_GB"))
+
+FDisplayClusterGenericBarrierService::FDisplayClusterGenericBarrierService(const FName& InInstanceName)
+	: FDisplayClusterService(InInstanceName.ToString())
 {
+	// Subscribe for SessionClosed events
+	OnSessionClosed().AddRaw(this, &FDisplayClusterGenericBarrierService::ProcessSessionClosed);
 }
 
 FDisplayClusterGenericBarrierService::~FDisplayClusterGenericBarrierService()
 {
-	Shutdown();
+	ShutdownImpl();
 }
 
 
 void FDisplayClusterGenericBarrierService::Shutdown()
 {
-	{
-		FScopeLock Lock(&BarriersCS);
-
-		// Deactivate all barriers
-		for (TPair<FString, TSharedPtr<IDisplayClusterBarrier>>& Barrier : Barriers)
-		{
-			Barrier.Value->Deactivate();
-		}
-
-		// Release events
-		for (TPair<FString, FEvent*>& EventIt : BarrierCreationEvents)
-		{
-			EventIt.Value->Trigger();
-			FPlatformProcess::ReturnSynchEventToPool(EventIt.Value);
-		}
-
-		BarrierCreationEvents.Empty();
-	}
-
+	ShutdownImpl();
 	return FDisplayClusterServer::Shutdown();
 }
 
@@ -56,34 +43,28 @@ FString FDisplayClusterGenericBarrierService::GetProtocolName() const
 	return ProtocolName;
 }
 
-TSharedPtr<IDisplayClusterBarrier, ESPMode::ThreadSafe> FDisplayClusterGenericBarrierService::GetBarrier(const FString& BarrierId)
+void FDisplayClusterGenericBarrierService::KillSession(const FString& NodeId)
 {
-	FScopeLock Lock(&BarriersCS);
-	TSharedPtr<IDisplayClusterBarrier, ESPMode::ThreadSafe>* BarrierFound = Barriers.Find(BarrierId);
-	return BarrierFound ? *BarrierFound : nullptr;
+	// Unbind its threads from the barriers
+	UnregisterClusterNode(NodeId);
+
+	// Now do the session related job
+	FDisplayClusterServer::KillSession(NodeId);
 }
 
-TSharedPtr<FDisplayClusterGenericBarrierService::FBarrierInfo> FDisplayClusterGenericBarrierService::GetBarrierInfo(const FString& BarrierId) const
+TSharedPtr<IDisplayClusterBarrier> FDisplayClusterGenericBarrierService::GetBarrier(const FString& BarrierId)
 {
-	FScopeLock Lock(&BarriersInfoCS);
-
-	const FBarrierInfoWrapper* BarrierInfoWrapper = BarriersInfo.Find(BarrierId);
-	checkSlow(BarrierInfoWrapper);
-
-	return BarrierInfoWrapper ? BarrierInfoWrapper->BarrierInfo : TSharedPtr<FDisplayClusterGenericBarrierService::FBarrierInfo>();
+	FScopeLock Lock(&BarriersCS);
+	TSharedPtr<IDisplayClusterBarrier>* BarrierFound = Barriers.Find(BarrierId);
+	return BarrierFound ? *BarrierFound : nullptr;
 }
 
-void FDisplayClusterGenericBarrierService::SetBarrierInfoUpdateLocked(const FString& BarrierId, bool bLocked)
+TSharedPtr<const FDisplayClusterGenericBarrierService::FBarrierInfo> FDisplayClusterGenericBarrierService::GetBarrierInfo(const FString& BarrierId) const
 {
 	FScopeLock Lock(&BarriersInfoCS);
 
-	FBarrierInfoWrapper* BarrierInfoWrapper = BarriersInfo.Find(BarrierId);
-	checkSlow(BarrierInfoWrapper);
-
-	if (BarrierInfoWrapper)
-	{
-		BarrierInfoWrapper->bBarrierInfoLockedOut = bLocked;
-	}
+	const TSharedRef<FBarrierInfo>* BarrierInfo = BarriersInfo.Find(BarrierId);
+	return BarrierInfo ? BarrierInfo->ToSharedPtr() : nullptr;
 }
 
 TSharedPtr<IDisplayClusterSession> FDisplayClusterGenericBarrierService::CreateSession(FDisplayClusterSessionInfo& SessionInfo)
@@ -133,13 +114,13 @@ TSharedPtr<FDisplayClusterPacketInternal> FDisplayClusterGenericBarrierService::
 		FString BarrierId;
 		Request->GetTextArg(DisplayClusterGenericBarrierStrings::ArgumentsDefaultCategory, DisplayClusterGenericBarrierStrings::ArgBarrierId, BarrierId);
 
-		// Extract parameter: thread marker
-		FString UniqueThreadMarker;
-		Request->GetTextArg(DisplayClusterGenericBarrierStrings::ArgumentsDefaultCategory, DisplayClusterGenericBarrierStrings::SyncOnBarrier::ArgThreadMarker, UniqueThreadMarker);
+		// Extract parameter: caller ID
+		FString CallerId;
+		Request->GetTextArg(DisplayClusterGenericBarrierStrings::ArgumentsDefaultCategory, DisplayClusterGenericBarrierStrings::SyncOnBarrier::ArgCallerId, CallerId);
 
 		// Process command
 		EBarrierControlResult CtrlResult = EBarrierControlResult::UnknownError;
-		const EDisplayClusterCommResult Result = SyncOnBarrier(BarrierId, UniqueThreadMarker, CtrlResult);
+		const EDisplayClusterCommResult Result = SyncOnBarrier(BarrierId, CallerId, CtrlResult);
 
 		// Set response data
 		Response->SetTextArg(DisplayClusterGenericBarrierStrings::ArgumentsDefaultCategory, DisplayClusterGenericBarrierStrings::ArgResult, DisplayClusterTypesConverter::template ToString((uint8)CtrlResult));
@@ -153,9 +134,9 @@ TSharedPtr<FDisplayClusterPacketInternal> FDisplayClusterGenericBarrierService::
 		FString BarrierId;
 		Request->GetTextArg(DisplayClusterGenericBarrierStrings::ArgumentsDefaultCategory, DisplayClusterGenericBarrierStrings::ArgBarrierId, BarrierId);
 
-		// Extract parameter: thread marker
-		FString UniqueThreadMarker;
-		Request->GetTextArg(DisplayClusterGenericBarrierStrings::ArgumentsDefaultCategory, DisplayClusterGenericBarrierStrings::SyncOnBarrierWithData::ArgThreadMarker, UniqueThreadMarker);
+		// Extract parameter: caller ID
+		FString CallerId;
+		Request->GetTextArg(DisplayClusterGenericBarrierStrings::ArgumentsDefaultCategory, DisplayClusterGenericBarrierStrings::SyncOnBarrierWithData::ArgCallerId, CallerId);
 
 		// Extract parameter: request data
 		TArray<uint8> RequestData;
@@ -164,7 +145,7 @@ TSharedPtr<FDisplayClusterPacketInternal> FDisplayClusterGenericBarrierService::
 		// Process command
 		EBarrierControlResult CtrlResult = EBarrierControlResult::UnknownError;
 		TArray<uint8> ResponseData;
-		const EDisplayClusterCommResult Result = SyncOnBarrierWithData(BarrierId, UniqueThreadMarker, RequestData, ResponseData, CtrlResult);
+		const EDisplayClusterCommResult Result = SyncOnBarrierWithData(BarrierId, CallerId, RequestData, ResponseData, CtrlResult);
 
 		// Set response data
 		Response->SetBinArg(DisplayClusterGenericBarrierStrings::ArgumentsDefaultCategory, DisplayClusterGenericBarrierStrings::SyncOnBarrierWithData::ArgResponseData, ResponseData);
@@ -179,19 +160,23 @@ TSharedPtr<FDisplayClusterPacketInternal> FDisplayClusterGenericBarrierService::
 		FString BarrierId;
 		Request->GetTextArg(DisplayClusterGenericBarrierStrings::ArgumentsDefaultCategory, DisplayClusterGenericBarrierStrings::ArgBarrierId, BarrierId);
 
-		// Extract parameter: thread markers
-		FString UniqueThreadMarkersStringified;
-		Request->GetTextArg(DisplayClusterGenericBarrierStrings::ArgumentsDefaultCategory, DisplayClusterGenericBarrierStrings::CreateBarrier::ArgThreadMarkers, UniqueThreadMarkersStringified);
-		TArray<FString> UniqueThreadMarkers;
-		DisplayClusterHelpers::str::template StrToArray(UniqueThreadMarkersStringified, DisplayClusterStrings::common::ArrayValSeparator, UniqueThreadMarkers);
-
 		// Extract parameter: timeout
 		uint32 Timeout = 0;
 		Request->GetTextArg(DisplayClusterGenericBarrierStrings::ArgumentsDefaultCategory, DisplayClusterGenericBarrierStrings::CreateBarrier::ArgTimeout, Timeout);
 
+		// Extract parameter: sync callers
+		TMap<FString, TSet<FString>> NodeToSyncCallers;
+		{
+			TArray<uint8> NodeToSyncCallersData;
+			Request->GetBinArg(DisplayClusterGenericBarrierStrings::ArgumentsDefaultCategory, DisplayClusterGenericBarrierStrings::CreateBarrier::ArgCallers, NodeToSyncCallersData);
+
+			FMemoryReader MemoryReader(NodeToSyncCallersData);
+			MemoryReader << NodeToSyncCallers;
+		}
+
 		// Process command
 		EBarrierControlResult CtrlResult = EBarrierControlResult::UnknownError;
-		const EDisplayClusterCommResult Result = CreateBarrier(BarrierId, UniqueThreadMarkers, Timeout, CtrlResult);
+		const EDisplayClusterCommResult Result = CreateBarrier(BarrierId, NodeToSyncCallers, Timeout, CtrlResult);
 
 		// Set response data
 		Response->SetTextArg(DisplayClusterGenericBarrierStrings::ArgumentsDefaultCategory, DisplayClusterGenericBarrierStrings::ArgResult, DisplayClusterTypesConverter::template ToString((uint8)CtrlResult));
@@ -258,11 +243,14 @@ TSharedPtr<FDisplayClusterPacketInternal> FDisplayClusterGenericBarrierService::
 //////////////////////////////////////////////////////////////////////////////////////////////
 // IDisplayClusterProtocolGenericBarrier
 //////////////////////////////////////////////////////////////////////////////////////////////
-EDisplayClusterCommResult FDisplayClusterGenericBarrierService::CreateBarrier(const FString& BarrierId, const TArray<FString>& UniqueThreadMarkers, uint32 Timeout, EBarrierControlResult& Result)
+EDisplayClusterCommResult FDisplayClusterGenericBarrierService::CreateBarrier(const FString& BarrierId, const TMap<FString, TSet<FString>>& NodeToSyncCallers, uint32 Timeout, EBarrierControlResult& Result)
 {
+	TRACE_CPUPROFILER_EVENT_SCOPE(SRV_GB::CreateBarrier);
+
 	// Validate input data
-	if (BarrierId.IsEmpty() || UniqueThreadMarkers.Num() < 1 || Timeout == 0)
+	if (BarrierId.IsEmpty() || NodeToSyncCallers.Num() < 1 || Timeout == 0)
 	{
+		UE_LOG(LogDisplayClusterBarrierGP, Warning, TEXT("%s - CreateBarrier - invalid request data"), *GetName());
 		return EDisplayClusterCommResult::WrongRequestData;
 	}
 
@@ -272,14 +260,25 @@ EDisplayClusterCommResult FDisplayClusterGenericBarrierService::CreateBarrier(co
 		// Check if exists
 		if (Barriers.Contains(BarrierId))
 		{
+			UE_LOG(LogDisplayClusterBarrierGP, Verbose, TEXT("%s - CreateBarrier - Barrier '%s' already exists"), *GetName(), *BarrierId);
 			Result = EBarrierControlResult::AlreadyExists;
 			return EDisplayClusterCommResult::Ok;
 		}
 
+		// Make full set of callers
+		TSet<FString> AllSyncCallers;
+		for (const TPair<FString, TSet<FString>>& PerNodeCallers : NodeToSyncCallers)
+		{
+			AllSyncCallers.Append(PerNodeCallers.Value);
+		}
+
 		// Create a new one
-		TSharedPtr<IDisplayClusterBarrier, ESPMode::ThreadSafe> NewBarrier(FDisplayClusterBarrierFactory::CreateBarrier(BarrierId, UniqueThreadMarkers, Timeout));
+		TSharedPtr<IDisplayClusterBarrier, ESPMode::ThreadSafe> NewBarrier(FDisplayClusterBarrierFactory::CreateBarrier(BarrierId, AllSyncCallers, Timeout));
 		checkSlow(NewBarrier);
 
+		// Set up an info container for this new barrier
+		InitializeBarrierInfo(BarrierId, NodeToSyncCallers);
+
 		// Activate barrier
 		NewBarrier->Activate();
 
@@ -292,6 +291,8 @@ EDisplayClusterCommResult FDisplayClusterGenericBarrierService::CreateBarrier(co
 			(*BarrierCreatedEvent)->Trigger();
 		}
 
+		UE_LOG(LogDisplayClusterBarrierGP, Verbose, TEXT("%s - CreateBarrier - Barrier '%s' created successfully"), *GetName(), *BarrierId);
+
 		Result = EBarrierControlResult::CreatedSuccessfully;
 	}
 
@@ -300,6 +301,15 @@ EDisplayClusterCommResult FDisplayClusterGenericBarrierService::CreateBarrier(co
 
 EDisplayClusterCommResult FDisplayClusterGenericBarrierService::WaitUntilBarrierIsCreated(const FString& BarrierId, EBarrierControlResult& Result)
 {
+	TRACE_CPUPROFILER_EVENT_SCOPE(SRV_GB::WaitUntilBarrierIsCreated);
+
+	// Validate input data
+	if (BarrierId.IsEmpty())
+	{
+		UE_LOG(LogDisplayClusterBarrierGP, Warning, TEXT("%s - WaitUntilBarrierIsCreated - invalid request data"), *GetName());
+		return EDisplayClusterCommResult::WrongRequestData;
+	}
+
 	FEvent** BarrierAvailableEvent = nullptr;
 
 	{
@@ -308,12 +318,15 @@ EDisplayClusterCommResult FDisplayClusterGenericBarrierService::WaitUntilBarrier
 		// Check if the barrier exists already
 		if (Barriers.Contains(BarrierId))
 		{
+			UE_LOG(LogDisplayClusterBarrierGP, Verbose, TEXT("%s - WaitUntilBarrierIsCreated - Barrier '%s' already exists"), *GetName(), *BarrierId);
 			Result = EBarrierControlResult::AlreadyExists;
 			return EDisplayClusterCommResult::Ok;
 		}
 		// If no, set up a notification event
 		else
 		{
+			UE_LOG(LogDisplayClusterBarrierGP, Verbose, TEXT("%s - WaitUntilBarrierIsCreated - Barrier '%s' not found. Waiting until it's created."), *GetName(), *BarrierId);
+
 			// Check if notification event has been created previously
 			BarrierAvailableEvent = BarrierCreationEvents.Find(BarrierId);
 			if (!BarrierAvailableEvent)
@@ -336,21 +349,27 @@ EDisplayClusterCommResult FDisplayClusterGenericBarrierService::WaitUntilBarrier
 	// Wait until created
 	(*BarrierAvailableEvent)->Wait();
 
+	UE_LOG(LogDisplayClusterBarrierGP, Verbose, TEXT("%s - WaitUntilBarrierIsCreated - Barrier '%s' already exists"), *GetName(), *BarrierId);
+
 	Result = EBarrierControlResult::AlreadyExists;
 	return EDisplayClusterCommResult::Ok;
 }
 
 EDisplayClusterCommResult FDisplayClusterGenericBarrierService::IsBarrierAvailable(const FString& BarrierId, EBarrierControlResult& Result)
 {
+	TRACE_CPUPROFILER_EVENT_SCOPE(SRV_GB::IsBarrierAvailable);
+
 	// Validate input data
 	if (BarrierId.IsEmpty())
 	{
+		UE_LOG(LogDisplayClusterBarrierGP, Warning, TEXT("%s - IsBarrierAvailable - invalid request data"), *GetName());
 		return EDisplayClusterCommResult::WrongRequestData;
 	}
 
 	{
 		FScopeLock Lock(&BarriersCS);
 		Result = (Barriers.Contains(BarrierId) ? EBarrierControlResult::AlreadyExists : EBarrierControlResult::NotFound);
+		UE_LOG(LogDisplayClusterBarrierGP, Verbose, TEXT("%s - IsBarrierAvailable - Barrier '%s' is %s"), *GetName(), *BarrierId, Barriers.Contains(BarrierId) ? TEXT("available") : TEXT("not available"));
 	}
 
 	return EDisplayClusterCommResult::Ok;
@@ -358,9 +377,12 @@ EDisplayClusterCommResult FDisplayClusterGenericBarrierService::IsBarrierAvailab
 
 EDisplayClusterCommResult FDisplayClusterGenericBarrierService::ReleaseBarrier(const FString& BarrierId, EBarrierControlResult& Result)
 {
+	TRACE_CPUPROFILER_EVENT_SCOPE(SRV_GB::ReleaseBarrier);
+
 	// Validate input data
 	if (BarrierId.IsEmpty())
 	{
+		UE_LOG(LogDisplayClusterBarrierGP, Warning, TEXT("%s - ReleaseBarrier - invalid request data"), *GetName());
 		return EDisplayClusterCommResult::WrongRequestData;
 	}
 
@@ -384,30 +406,31 @@ EDisplayClusterCommResult FDisplayClusterGenericBarrierService::ReleaseBarrier(c
 			}
 
 			// Remove clients info of this barrier
-			{
-				FScopeLock LockInfo(&BarriersInfoCS);
-				BarriersInfo.Remove(BarrierId);
-			}
+			ReleaseBarrierInfo(BarrierId);
+
+			UE_LOG(LogDisplayClusterBarrierGP, Verbose, TEXT("%s - ReleaseBarrier - Barrier '%s' has been released"), *GetName(), *BarrierId);
+
+			Result = EBarrierControlResult::ReleasedSuccessfully;
 		}
 		else
 		{
 			// Not exists
+			UE_LOG(LogDisplayClusterBarrierGP, Verbose, TEXT("%s - ReleaseBarrier - Barrier '%s' not found"), *GetName(), *BarrierId);
 			Result = EBarrierControlResult::NotFound;
-			return EDisplayClusterCommResult::WrongRequestData;
 		}
 	}
 
-	Result = EBarrierControlResult::ReleasedSuccessfully;
 	return EDisplayClusterCommResult::Ok;
 }
 
-EDisplayClusterCommResult FDisplayClusterGenericBarrierService::SyncOnBarrier(const FString& BarrierId, const FString& UniqueThreadMarker, EBarrierControlResult& Result)
+EDisplayClusterCommResult FDisplayClusterGenericBarrierService::SyncOnBarrier(const FString& BarrierId, const FString& CallerId, EBarrierControlResult& Result)
 {
 	TRACE_CPUPROFILER_EVENT_SCOPE(SRV_GB::SyncOnBarrier);
 
 	// Validate input data
-	if (BarrierId.IsEmpty() || UniqueThreadMarker.IsEmpty())
+	if (BarrierId.IsEmpty() || CallerId.IsEmpty())
 	{
+		UE_LOG(LogDisplayClusterBarrierGP, Warning, TEXT("%s - SyncOnBarrier - invalid request data"), *GetName());
 		Result = EBarrierControlResult::UnknownError;
 		return EDisplayClusterCommResult::WrongRequestData;
 	}
@@ -435,23 +458,23 @@ EDisplayClusterCommResult FDisplayClusterGenericBarrierService::SyncOnBarrier(co
 		return EDisplayClusterCommResult::NotAllowed;
 	}
 
-	// Save info about this caller
-	UpdateBarrierInformation(BarrierId, GetSessionInfoCache().NodeId.Get(FString()), UniqueThreadMarker);
-
 	// Sync on the barrier
-	(*Barrier)->Wait(UniqueThreadMarker);
+	UE_LOG(LogDisplayClusterBarrierGP, VeryVerbose, TEXT("%s - SyncOnBarrier - Barrier '%s' wait start"), *GetName(), *BarrierId);
+	(*Barrier)->Wait(CallerId);
+	UE_LOG(LogDisplayClusterBarrierGP, VeryVerbose, TEXT("%s - SyncOnBarrier - Barrier '%s' wait end"), *GetName(), *BarrierId);
 
 	Result = EBarrierControlResult::SynchronizedSuccessfully;
 	return EDisplayClusterCommResult::Ok;
 }
 
-EDisplayClusterCommResult FDisplayClusterGenericBarrierService::SyncOnBarrierWithData(const FString& BarrierId, const FString& UniqueThreadMarker, const TArray<uint8>& RequestData, TArray<uint8>& OutResponseData, EBarrierControlResult& Result)
+EDisplayClusterCommResult FDisplayClusterGenericBarrierService::SyncOnBarrierWithData(const FString& BarrierId, const FString& CallerId, const TArray<uint8>& RequestData, TArray<uint8>& OutResponseData, EBarrierControlResult& Result)
 {
 	TRACE_CPUPROFILER_EVENT_SCOPE(SRV_GB::SyncOnBarrierWithData);
 
 	// Validate input data
-	if (BarrierId.IsEmpty() || UniqueThreadMarker.IsEmpty())
+	if (BarrierId.IsEmpty() || CallerId.IsEmpty())
 	{
+		UE_LOG(LogDisplayClusterBarrierGP, Warning, TEXT("%s - SyncOnBarrier - invalid request data"), *GetName());
 		Result = EBarrierControlResult::UnknownError;
 		return EDisplayClusterCommResult::WrongRequestData;
 	}
@@ -470,39 +493,113 @@ EDisplayClusterCommResult FDisplayClusterGenericBarrierService::SyncOnBarrierWit
 	if (!Barrier)
 	{
 		// Barrier not found
+		UE_LOG(LogDisplayClusterBarrierGP, Verbose, TEXT("%s - SyncOnBarrierWithData - Barrier '%s' not found"), *GetName(), *BarrierId);
 		Result = EBarrierControlResult::NotFound;
 		return EDisplayClusterCommResult::WrongRequestData;
 	}
 	else if (!(*Barrier)->IsActivated())
 	{
+		UE_LOG(LogDisplayClusterBarrierGP, Verbose, TEXT("%s - SyncOnBarrierWithData - Barrier '%s' is not active"), *GetName(), *BarrierId);
 		Result = EBarrierControlResult::UnknownError;
 		return EDisplayClusterCommResult::NotAllowed;
 	}
 
-	// Save info about this caller
-	UpdateBarrierInformation(BarrierId, GetSessionInfoCache().NodeId.Get(FString()), UniqueThreadMarker);
-
 	// Sync on the barrier
-	(*Barrier)->WaitWithData(UniqueThreadMarker, RequestData, OutResponseData);
+	UE_LOG(LogDisplayClusterBarrierGP, VeryVerbose, TEXT("%s - SyncOnBarrierWithData - Barrier '%s' wait start"), *GetName(), *BarrierId);
+	(*Barrier)->WaitWithData(CallerId, RequestData, OutResponseData);
+	UE_LOG(LogDisplayClusterBarrierGP, VeryVerbose, TEXT("%s - SyncOnBarrierWithData - Barrier '%s' wait end"), *GetName(), *BarrierId);
 
 	Result = EBarrierControlResult::SynchronizedSuccessfully;
 	return EDisplayClusterCommResult::Ok;
 }
 
-void FDisplayClusterGenericBarrierService::UpdateBarrierInformation(const FString& BarrierId, const FString& NodeId, const FString& ThreadMarker)
+void FDisplayClusterGenericBarrierService::ShutdownImpl()
+{
+	{
+		FScopeLock Lock(&BarriersCS);
+
+		// Deactivate all barriers
+		for (TPair<FString, TSharedPtr<IDisplayClusterBarrier>>& Barrier : Barriers)
+		{
+			Barrier.Value->Deactivate();
+		}
+
+		// Release events
+		for (TPair<FString, FEvent*>& EventIt : BarrierCreationEvents)
+		{
+			EventIt.Value->Trigger();
+			FPlatformProcess::ReturnSynchEventToPool(EventIt.Value);
+		}
+
+		BarrierCreationEvents.Empty();
+	}
+
+	OnSessionClosed().RemoveAll(this);
+}
+
+void FDisplayClusterGenericBarrierService::InitializeBarrierInfo(const FString& BarrierId, const TMap<FString, TSet<FString>>& NodeToSyncCallers)
 {
 	FScopeLock Lock(&BarriersInfoCS);
 
-	FBarrierInfoWrapper& BarrierInfoWrapper = BarriersInfo.FindOrAdd(BarrierId);
+	if (BarriersInfo.Contains(BarrierId))
+	{
+		return;
+	}
+
+	// Create new info slot
+	TSharedRef<FBarrierInfo> NewBarrierInfo = BarriersInfo.Emplace(BarrierId, MakeShared<FBarrierInfo>());
+
+	// Node-to-SyncCallers mapping
+	NewBarrierInfo->NodeToThreadsMapping = NodeToSyncCallers;
+
+	// SyncCaller-to-Node mapping
+	for (const TPair<FString, TSet<FString>>& NodeMapping : NodeToSyncCallers)
+	{
+		for (const FString& SyncCaller : NodeMapping.Value)
+		{
+			NewBarrierInfo->ThreadToNodeMapping.Emplace(SyncCaller, NodeMapping.Key);
+		}
+	}
+}
+
+void FDisplayClusterGenericBarrierService::ReleaseBarrierInfo(const FString& BarrierId)
+{
+	FScopeLock LockInfo(&BarriersInfoCS);
+	BarriersInfo.Remove(BarrierId);
+}
 
-	// Update if allowed
-	if (!BarrierInfoWrapper.bBarrierInfoLockedOut)
+void FDisplayClusterGenericBarrierService::ProcessSessionClosed(const FDisplayClusterSessionInfo& SessionInfo)
+{
+	if (!SessionInfo.IsTerminatedByServer())
 	{
-		// Update NodeToThread map
-		TSet<FString>& NodeInfo = BarrierInfoWrapper.BarrierInfo->NodeToThreadsMapping.FindOrAdd(NodeId);
-		NodeInfo.Add(ThreadMarker);
+		// Ignore sessions with empty NodeId
+		if (SessionInfo.NodeId.IsSet())
+		{
+			const FString NodeId = SessionInfo.NodeId.Get(FString());
+			UnregisterClusterNode(NodeId);
+		}
+	}
+}
 
-		// Update ThreadToNode map
-		BarrierInfoWrapper.BarrierInfo->ThreadToNodeMapping.Emplace(ThreadMarker, NodeId);
+void FDisplayClusterGenericBarrierService::UnregisterClusterNode(const FString& NodeId)
+{
+	FScopeLock Lock(&BarriersCS);
+
+	// For every registered barrier
+	for (TPair<FString, TSharedPtr<IDisplayClusterBarrier>>& Barrier : Barriers)
+	{
+		// Get its corresponding barrier info
+		if (const TSharedPtr<const FBarrierInfo> BarrierInfo = GetBarrierInfo(Barrier.Key))
+		{
+			// See if there are any barrier sync callers associated with the node
+			if (const TSet<FString>* const CallersAssociatedWithNode = BarrierInfo->NodeToThreadsMapping.Find(NodeId))
+			{
+				// If so, unregister all the caller IDs of that node
+				for (const FString& CallerId : *CallersAssociatedWithNode)
+				{
+					Barrier.Value->UnregisterSyncCaller(CallerId);
+				}
+			}
+		}
 	}
 }
diff --git a/Engine/Plugins/Runtime/nDisplay/Source/DisplayCluster/Private/Network/Service/GenericBarrier/DisplayClusterGenericBarrierService.h b/Engine/Plugins/Runtime/nDisplay/Source/DisplayCluster/Private/Network/Service/GenericBarrier/DisplayClusterGenericBarrierService.h
--- a/Engine/Plugins/Runtime/nDisplay/Source/DisplayCluster/Private/Network/Service/GenericBarrier/DisplayClusterGenericBarrierService.h
+++ b/Engine/Plugins/Runtime/nDisplay/Source/DisplayCluster/Private/Network/Service/GenericBarrier/DisplayClusterGenericBarrierService.h
@@ -27,92 +27,90 @@ public:
 	 */
 	struct FBarrierInfo
 	{
-		// Holds a set of thread markers bound to the owning cluster node
+		/** Holds ClusterNodeId-to-CallerIDs mapping */
 		TMap<FString, TSet<FString>> NodeToThreadsMapping;
 
-		// Holds thread marker - to - cluster node mapping
+		/** Holds CallerID-to-ClusterNodeId mapping (reversed version of NodeToThreadsMapping) */
 		TMap<FString, FString> ThreadToNodeMapping;
 	};
 
 public:
-	FDisplayClusterGenericBarrierService();
+
+	FDisplayClusterGenericBarrierService(const FName& InInstanceName);
+
 	virtual ~FDisplayClusterGenericBarrierService();
 
 public:
-	//////////////////////////////////////////////////////////////////////////////////////////////
-	// IDisplayClusterServer
-	//////////////////////////////////////////////////////////////////////////////////////////////
-	virtual void Shutdown() override final;
+
+	//~ Begin IDisplayClusterServer
+	virtual void Shutdown() override;
 	virtual FString GetProtocolName() const override;
+	virtual void KillSession(const FString& NodeId) override;
+	//~ End IDisplayClusterServer
 
 public:
 
-	// Returns barrier by ID
-	TSharedPtr<IDisplayClusterBarrier, ESPMode::ThreadSafe> GetBarrier(const FString& BarrierId);
+	/** Returns barrier by ID */
+	TSharedPtr<IDisplayClusterBarrier> GetBarrier(const FString& BarrierId);
 
-	// Returns barrier information
-	TSharedPtr<FDisplayClusterGenericBarrierService::FBarrierInfo> GetBarrierInfo(const FString& BarrierId) const;
-
-	// Enables/disables barrier info update. The update procedure is called every time a thread joins the barrier
-	// to synchronize. More threads we have more CPU time is consumed for this. For optimization purposes, this
-	// function allows to stop collecting barrier information and continue updating at any time.
-	void SetBarrierInfoUpdateLocked(const FString& BarrierId, bool bLocked);
+	/** Returns barrier information */
+	TSharedPtr<const FBarrierInfo> GetBarrierInfo(const FString& BarrierId) const;
 
 protected:
-	// Creates session instance for this service
+
+	/** Creates session instance for this service */
 	virtual TSharedPtr<IDisplayClusterSession> CreateSession(FDisplayClusterSessionInfo& SessionInfo) override;
 
 protected:
-	//////////////////////////////////////////////////////////////////////////////////////////////
-	// IDisplayClusterSessionPacketHandler
-	//////////////////////////////////////////////////////////////////////////////////////////////
+
+	//~ Begin IDisplayClusterSessionPacketHandler
 	virtual TSharedPtr<FDisplayClusterPacketInternal> ProcessPacket(const TSharedPtr<FDisplayClusterPacketInternal>& Request, const FDisplayClusterSessionInfo& SessionInfo) override;
+	//~ End IDisplayClusterSessionPacketHandler
 
-private:
-	//////////////////////////////////////////////////////////////////////////////////////////////
-	// IDisplayClusterProtocolGenericBarrier
-	//////////////////////////////////////////////////////////////////////////////////////////////
-	virtual EDisplayClusterCommResult CreateBarrier(const FString& BarrierId, const TArray<FString>& UniqueThreadMarkers, uint32 Timeout, EBarrierControlResult& Result) override;
+protected:
+
+	//~ Begin IDisplayClusterProtocolGenericBarrier
+	virtual EDisplayClusterCommResult CreateBarrier(const FString& BarrierId, const TMap<FString, TSet<FString>>& NodeToSyncCallers, uint32 Timeout, EBarrierControlResult& Result) override;
 	virtual EDisplayClusterCommResult WaitUntilBarrierIsCreated(const FString& BarrierId, EBarrierControlResult& Result) override;
 	virtual EDisplayClusterCommResult IsBarrierAvailable(const FString& BarrierId, EBarrierControlResult& Result) override;
 	virtual EDisplayClusterCommResult ReleaseBarrier(const FString& BarrierId, EBarrierControlResult& Result) override;
-	virtual EDisplayClusterCommResult SyncOnBarrier(const FString& BarrierId, const FString& UniqueThreadMarker, EBarrierControlResult& Result) override;
-	virtual EDisplayClusterCommResult SyncOnBarrierWithData(const FString& BarrierId, const FString& UniqueThreadMarker, const TArray<uint8>& RequestData, TArray<uint8>& OutResponseData, EBarrierControlResult& Result) override;
+	virtual EDisplayClusterCommResult SyncOnBarrier(const FString& BarrierId, const FString& CallerId, EBarrierControlResult& Result) override;
+	virtual EDisplayClusterCommResult SyncOnBarrierWithData(const FString& BarrierId, const FString& CallerId, const TArray<uint8>& RequestData, TArray<uint8>& OutResponseData, EBarrierControlResult& Result) override;
+	//~ End IDisplayClusterProtocolGenericBarrier
 
 private:
 
-	// Caches information about barrier users
-	void UpdateBarrierInformation(const FString& BarrierId, const FString& NodeId, const FString& ThreadMarker);
+	/** Non-virtual shutdown implementation */
+	void ShutdownImpl();
+
+	/** An auxiliary function that is responsible for barrier information slot initialization */
+	void InitializeBarrierInfo(const FString& BarrierId, const TMap<FString, TSet<FString>>& NodeToSyncCallers);
+
+	/** An auxiliary function that is responsible for barrier information cleanup */
+	void ReleaseBarrierInfo(const FString& BarrierId);
+
+	/** Callback when a session is closed */
+	void ProcessSessionClosed(const FDisplayClusterSessionInfo& SessionInfo);
+
+	/** Unbinds specific cluster node and all its sync callers from all GP barriers */
+	void UnregisterClusterNode(const FString& NodeId);
 
 private:
 
-	// Barriers
-	TMap<FString, TSharedPtr<IDisplayClusterBarrier, ESPMode::ThreadSafe>> Barriers;
+	/** The barriers managed by this server */
+	TMap<FString, TSharedPtr<IDisplayClusterBarrier>> Barriers;
 
-	// Barrier creation events
+	/** Barrier creation events */
 	TMap<FString, FEvent*> BarrierCreationEvents;
 
-	// Critical section for internal data access synchronization
+	/** Critical section for internal data access synchronization */
 	mutable FCriticalSection BarriersCS;
 
 private:
 
-	/**
-	 * A helper structure that wraps the actual barrier info holder
-	 * with an additional data for optimization purposes.
-	 */
-	struct FBarrierInfoWrapper
-	{
-		// The actual barrier information holder
-		TSharedRef<FBarrierInfo> BarrierInfo = MakeShared<FBarrierInfo>();
-
-		// Whether barrier information is locked from being updated
-		bool bBarrierInfoLockedOut = false;
-	};
-
-	// Holds extra information per-barrier
-	TMap<FString, FBarrierInfoWrapper> BarriersInfo;
+	/** Holds extra information per - barrier */
+	TMap<FString, TSharedRef<FBarrierInfo>> BarriersInfo;
 
-	// Critical section to synchronize access to the barriers information container
+	/** Critical section to synchronize access to the barriers information container */
 	mutable FCriticalSection BarriersInfoCS;
 };
diff --git a/Engine/Plugins/Runtime/nDisplay/Source/DisplayCluster/Private/Network/Service/GenericBarrier/DisplayClusterGenericBarrierStrings.h b/Engine/Plugins/Runtime/nDisplay/Source/DisplayCluster/Private/Network/Service/GenericBarrier/DisplayClusterGenericBarrierStrings.h
--- a/Engine/Plugins/Runtime/nDisplay/Source/DisplayCluster/Private/Network/Service/GenericBarrier/DisplayClusterGenericBarrierStrings.h
+++ b/Engine/Plugins/Runtime/nDisplay/Source/DisplayCluster/Private/Network/Service/GenericBarrier/DisplayClusterGenericBarrierStrings.h
@@ -2,7 +2,7 @@
 
 #pragma once
 
-#include "CoreMinimal.h"
+#include "HAL/Platform.h"
 
 
 /**
@@ -26,8 +26,8 @@ namespace DisplayClusterGenericBarrierStrings
 	{
 		constexpr static const TCHAR* Name = TEXT("CreateBarrier");
 
-		constexpr static const TCHAR* ArgThreadMarkers = TEXT("ThreadMarkers");
-		constexpr static const TCHAR* ArgTimeout       = TEXT("Timeout");
+		constexpr static const TCHAR* ArgCallers = TEXT("SyncCallers");
+		constexpr static const TCHAR* ArgTimeout = TEXT("Timeout");
 	}
 
 	namespace WaitUntilBarrierIsCreated
@@ -49,15 +49,15 @@ namespace DisplayClusterGenericBarrierStrings
 	{
 		constexpr static const TCHAR* Name = TEXT("Sync");
 
-		constexpr static const TCHAR* ArgThreadMarker = TEXT("ThreadMarker");
+		constexpr static const TCHAR* ArgCallerId = TEXT("CallerId");
 	}
 
 	namespace SyncOnBarrierWithData
 	{
 		constexpr static const TCHAR* Name = TEXT("SyncWithData");
 
-		constexpr static const TCHAR* ArgThreadMarker  = TEXT("ThreadMarker");
-		constexpr static const TCHAR* ArgRequestData   = TEXT("ReqData");
-		constexpr static const TCHAR* ArgResponseData  = TEXT("RespData");
+		constexpr static const TCHAR* ArgCallerId     = TEXT("CallerId");
+		constexpr static const TCHAR* ArgRequestData  = TEXT("ReqData");
+		constexpr static const TCHAR* ArgResponseData = TEXT("RespData");
 	}
 };
diff --git a/Engine/Plugins/Runtime/nDisplay/Source/DisplayCluster/Private/Network/Service/InternalComm/DisplayClusterInternalCommStrings.h b/Engine/Plugins/Runtime/nDisplay/Source/DisplayCluster/Private/Network/Service/InternalComm/DisplayClusterInternalCommStrings.h
new file mode 100644
--- /dev/null
+++ b/Engine/Plugins/Runtime/nDisplay/Source/DisplayCluster/Private/Network/Service/InternalComm/DisplayClusterInternalCommStrings.h
@@ -0,0 +1,43 @@
+// Copyright Epic Games, Inc. All Rights Reserved.
+
+#pragma once
+
+#include "HAL/Platform.h"
+
+
+/**
+ * Internal comm related strings
+ */
+namespace DisplayClusterInternalCommStrings
+{
+	constexpr static const TCHAR* ProtocolName = TEXT("InternalComm");
+	
+	constexpr static const TCHAR* TypeRequest  = TEXT("Request");
+	constexpr static const TCHAR* TypeResponse = TEXT("Response");
+
+	constexpr static const TCHAR* ArgumentsDefaultCategory = TEXT("IC");
+
+	namespace GatherServicesHostingInfo
+	{
+		constexpr static const TCHAR* Name = TEXT("GatherServicesHostingInfo");
+
+		constexpr static const TCHAR* ArgNodeHostingInfo    = TEXT("NodeHostingInfo");
+		constexpr static const TCHAR* ArgClusterHostingInfo = TEXT("ClusterHostingInfo");
+	}
+
+	namespace PostFailureNegotiate
+	{
+		constexpr static const TCHAR* Name = TEXT("PostFailureNegotiate");
+
+		constexpr static const TCHAR* ArgSyncStateData = TEXT("SyncState");
+		constexpr static const TCHAR* ArgRecoveryData  = TEXT("RecoveryData");
+	}
+
+	namespace RequestNodeDrop
+	{
+		constexpr static const TCHAR* Name = TEXT("RequestNodeDrop");
+
+		constexpr static const TCHAR* ArgNodeId     = TEXT("NodeId");
+		constexpr static const TCHAR* ArgDropReason = TEXT("DropReason");
+	}
+};
diff --git a/Engine/Plugins/Runtime/nDisplay/Source/DisplayCluster/Private/Network/Service/RenderSync/DisplayClusterRenderSyncClient.cpp b/Engine/Plugins/Runtime/nDisplay/Source/DisplayCluster/Private/Network/Service/RenderSync/DisplayClusterRenderSyncClient.cpp
--- a/Engine/Plugins/Runtime/nDisplay/Source/DisplayCluster/Private/Network/Service/RenderSync/DisplayClusterRenderSyncClient.cpp
+++ b/Engine/Plugins/Runtime/nDisplay/Source/DisplayCluster/Private/Network/Service/RenderSync/DisplayClusterRenderSyncClient.cpp
@@ -2,6 +2,7 @@
 
 #include "Network/Service/RenderSync/DisplayClusterRenderSyncClient.h"
 #include "Network/Service/RenderSync/DisplayClusterRenderSyncStrings.h"
+#include "Network/Configuration/DisplayClusterNetworkConfiguration.h"
 #include "Network/Packet/DisplayClusterPacketInternal.h"
 #include "Network/Listener/DisplayClusterHelloMessageStrings.h"
 
@@ -11,13 +12,8 @@
 #include "Misc/DisplayClusterLog.h"
 
 
-FDisplayClusterRenderSyncClient::FDisplayClusterRenderSyncClient()
-	: FDisplayClusterRenderSyncClient(FString("CLN_RS"))
-{
-}
-
-FDisplayClusterRenderSyncClient::FDisplayClusterRenderSyncClient(const FString& InName)
-	: FDisplayClusterClient(InName)
+FDisplayClusterRenderSyncClient::FDisplayClusterRenderSyncClient(const FName& InName)
+	: FDisplayClusterClient(InName.ToString())
 {
 }
 
@@ -36,7 +32,7 @@ bool FDisplayClusterRenderSyncClient::Connect(const FString& Address, const uint
 	// Prepare 'hello' message
 	TSharedPtr<FDisplayClusterPacketInternal> HelloMsg = MakeShared<FDisplayClusterPacketInternal>(
 		DisplayClusterHelloMessageStrings::Hello::Name,
-		DisplayClusterRenderSyncStrings::TypeRequest,
+		DisplayClusterHelloMessageStrings::Hello::TypeRequest,
 		DisplayClusterRenderSyncStrings::ProtocolName
 	);
 
@@ -52,11 +48,11 @@ bool FDisplayClusterRenderSyncClient::Connect(const FString& Address, const uint
 //////////////////////////////////////////////////////////////////////////////////////////////
 // IDisplayClusterProtocolRenderSync
 //////////////////////////////////////////////////////////////////////////////////////////////
-EDisplayClusterCommResult FDisplayClusterRenderSyncClient::SyncOnBarrier()
+EDisplayClusterCommResult FDisplayClusterRenderSyncClient::SynchronizeOnBarrier()
 {
 	static const TSharedPtr<FDisplayClusterPacketInternal> Request(
 		new FDisplayClusterPacketInternal(
-			DisplayClusterRenderSyncStrings::SyncOnBarrier::Name,
+			DisplayClusterRenderSyncStrings::SynchronizeOnBarrier::Name,
 			DisplayClusterRenderSyncStrings::TypeRequest,
 			DisplayClusterRenderSyncStrings::ProtocolName)
 	);
@@ -64,7 +60,7 @@ EDisplayClusterCommResult FDisplayClusterRenderSyncClient::SyncOnBarrier()
 	TSharedPtr<FDisplayClusterPacketInternal> Response;
 
 	{
-		TRACE_CPUPROFILER_EVENT_SCOPE(nD CLN_RS::WaitForSwapSync);
+		TRACE_CPUPROFILER_EVENT_SCOPE(CLN_RS::WaitForSwapSync);
 		Response = SendRecvPacket(Request);
 	}
 
diff --git a/Engine/Plugins/Runtime/nDisplay/Source/DisplayCluster/Private/Network/Service/RenderSync/DisplayClusterRenderSyncClient.h b/Engine/Plugins/Runtime/nDisplay/Source/DisplayCluster/Private/Network/Service/RenderSync/DisplayClusterRenderSyncClient.h
--- a/Engine/Plugins/Runtime/nDisplay/Source/DisplayCluster/Private/Network/Service/RenderSync/DisplayClusterRenderSyncClient.h
+++ b/Engine/Plugins/Runtime/nDisplay/Source/DisplayCluster/Private/Network/Service/RenderSync/DisplayClusterRenderSyncClient.h
@@ -15,18 +15,18 @@ class FDisplayClusterRenderSyncClient
 	, public IDisplayClusterProtocolRenderSync
 {
 public:
-	FDisplayClusterRenderSyncClient();
-	FDisplayClusterRenderSyncClient(const FString& InName);
+
+	FDisplayClusterRenderSyncClient(const FName& InName);
 
 public:
-	//////////////////////////////////////////////////////////////////////////////////////////////
-	// IDisplayClusterClient
-	//////////////////////////////////////////////////////////////////////////////////////////////
-	virtual bool Connect(const FString& Address, const uint16 Port, const uint32 ConnectRetriesAmount, const uint32 ConnectRetryDelay) override;
+
+	//~ Begin IDisplayClusterClient
+	virtual bool Connect(const FString& Address, const uint16 Port, const uint32 ConnectRetriesAmount = 0, const uint32 ConnectRetryDelay = 0) override;
+	//~ End IDisplayClusterClient
 
 public:
-	//////////////////////////////////////////////////////////////////////////////////////////////
-	// IDisplayClusterProtocolRenderSync
-	//////////////////////////////////////////////////////////////////////////////////////////////
-	virtual EDisplayClusterCommResult SyncOnBarrier() override;
+
+	//~ Begin IDisplayClusterProtocolRenderSync
+	virtual EDisplayClusterCommResult SynchronizeOnBarrier() override;
+	//~ End IDisplayClusterProtocolRenderSync
 };
diff --git a/Engine/Plugins/Runtime/nDisplay/Source/DisplayCluster/Private/Network/Service/RenderSync/DisplayClusterRenderSyncService.cpp b/Engine/Plugins/Runtime/nDisplay/Source/DisplayCluster/Private/Network/Service/RenderSync/DisplayClusterRenderSyncService.cpp
--- a/Engine/Plugins/Runtime/nDisplay/Source/DisplayCluster/Private/Network/Service/RenderSync/DisplayClusterRenderSyncService.cpp
+++ b/Engine/Plugins/Runtime/nDisplay/Source/DisplayCluster/Private/Network/Service/RenderSync/DisplayClusterRenderSyncService.cpp
@@ -4,12 +4,12 @@
 #include "Network/Service/RenderSync/DisplayClusterRenderSyncStrings.h"
 #include "Network/Barrier/DisplayClusterBarrierFactory.h"
 #include "Network/Barrier/IDisplayClusterBarrier.h"
+#include "Network/Configuration/DisplayClusterNetworkConfiguration.h"
 #include "Network/Listener/DisplayClusterTcpListener.h"
 #include "Network/Packet/DisplayClusterPacketInternal.h"
 #include "Network/Session/DisplayClusterSession.h"
 
 #include "Cluster/IPDisplayClusterClusterManager.h"
-#include "Cluster/Failover/IDisplayClusterFailoverNodeController.h"
 #include "Config/IPDisplayClusterConfigManager.h"
 
 #include "DisplayClusterConfigurationStrings.h"
@@ -19,8 +19,8 @@
 #include "Misc/DisplayClusterLog.h"
 
 
-FDisplayClusterRenderSyncService::FDisplayClusterRenderSyncService()
-	: FDisplayClusterService(FString("SRV_RS"))
+FDisplayClusterRenderSyncService::FDisplayClusterRenderSyncService(const FName& InInstanceName)
+	: FDisplayClusterService(InInstanceName.ToString())
 {
 	// Perform barriers initialization depending on current circumstances
 	InitializeBarriers();
@@ -30,35 +30,34 @@
 
 FDisplayClusterRenderSyncService::~FDisplayClusterRenderSyncService()
 {
+	ShutdownImpl();
+
 	// Unsubscribe from barrier timeout events
 	UnsubscribeFromAllBarrierEvents();
 	// Unsubscribe from SessionClosed notifications
 	OnSessionClosed().RemoveAll(this);
-
-	Shutdown();
 }
 
 
 bool FDisplayClusterRenderSyncService::Start(const FString& Address, const uint16 Port)
 {
-	// Activate all barriers
-	ActivateAllBarriers();
+	// Start internals
+	StartInternal();
 
 	return FDisplayClusterServer::Start(Address, Port);
 }
 
-bool FDisplayClusterRenderSyncService::Start(TSharedPtr<FDisplayClusterTcpListener>& ExternalListener)
+bool FDisplayClusterRenderSyncService::Start(TSharedRef<FDisplayClusterTcpListener>& ExternalListener)
 {
-	// Activate all barriers
-	ActivateAllBarriers();
+	// Start internals
+	StartInternal();
 
 	return FDisplayClusterServer::Start(ExternalListener);
 }
 
 void FDisplayClusterRenderSyncService::Shutdown()
 {
-	// Deactivate all barriers
-	DeactivateAllBarriers();
+	ShutdownImpl();
 
 	return FDisplayClusterServer::Shutdown();
 }
@@ -94,94 +93,45 @@
 {
 	if (!SessionInfo.IsTerminatedByServer())
 	{
-		// Get node ID
-		const FString NodeId = SessionInfo.NodeId.Get(FString());
-
-		if (!NodeId.IsEmpty())
+		// Ignore sessions with empty NodeId
+		if (SessionInfo.NodeId.IsSet())
 		{
+			const FString NodeId = SessionInfo.NodeId.Get(FString());
+
 			// We have to unregister the node that just disconnected from all barriers
 			UnregisterClusterNode(NodeId);
 
+			FDisplayClusterServiceFailureEvent EventInfo;
+			EventInfo.NodeFailed  = SessionInfo.NodeId;
+			EventInfo.FailureType = FDisplayClusterServiceFailureEvent::ENodeFailType::ConnectionLost;
+
 			// Notify others about node fail
-			OnNodeFailed().Broadcast(NodeId, ENodeFailType::ConnectionLost);
+			OnNodeFailed().Broadcast(EventInfo);
 		}
 	}
 }
 
 // Callbck on barrier timeout
-void FDisplayClusterRenderSyncService::ProcessBarrierTimeout(const FString& BarrierName, const TArray<FString>& NodesTimedOut)
+void FDisplayClusterRenderSyncService::ProcessBarrierTimeout(const FString& BarrierName, const TSet<FString>& NodesTimedOut)
 {
 	for (const FString& NodeId : NodesTimedOut)
 	{
 		// We have to unregister the node that just timed out from all barriers
 		UnregisterClusterNode(NodeId);
 
+		FDisplayClusterServiceFailureEvent EventInfo;
+		EventInfo.NodeFailed  = NodeId;
+		EventInfo.FailureType = FDisplayClusterServiceFailureEvent::ENodeFailType::BarrierTimeOut;
+
 		// Notify others about timeout
-		OnNodeFailed().Broadcast(NodeId, ENodeFailType::BarrierTimeOut);
+		OnNodeFailed().Broadcast(EventInfo);
 	}
 }
 
-
-//////////////////////////////////////////////////////////////////////////////////////////////
-// IDisplayClusterSessionPacketHandler
-//////////////////////////////////////////////////////////////////////////////////////////////
-TSharedPtr<FDisplayClusterPacketInternal> FDisplayClusterRenderSyncService::ProcessPacket(const TSharedPtr<FDisplayClusterPacketInternal>& Request, const FDisplayClusterSessionInfo& SessionInfo)
+void FDisplayClusterRenderSyncService::StartInternal()
 {
-	// Check the pointer
-	if (!Request.IsValid())
-	{
-		UE_LOG(LogDisplayClusterNetwork, Error, TEXT("%s - Invalid request data (nullptr)"), *GetName());
-		return nullptr;
-	}
-
-	// Cache session info
-	SetSessionInfoCache(SessionInfo);
-
-	UE_LOG(LogDisplayClusterNetwork, Verbose, TEXT("%s - Processing packet: %s"), *GetName(), *Request->ToLogString());
-
-	// Check protocol and type
-	if (Request->GetProtocol() != DisplayClusterRenderSyncStrings::ProtocolName || Request->GetType() != DisplayClusterRenderSyncStrings::TypeRequest)
-	{
-		UE_LOG(LogDisplayClusterNetwork, Error, TEXT("%s - Unsupported packet type: %s"), *GetName(), *Request->ToLogString());
-		return nullptr;
-	}
-
-	// Initialize response packet
-	TSharedPtr<FDisplayClusterPacketInternal> Response = MakeShared<FDisplayClusterPacketInternal>(Request->GetName(), DisplayClusterRenderSyncStrings::TypeResponse, Request->GetProtocol());
-
-	// Dispatch the packet
-	if (Request->GetName().Equals(DisplayClusterRenderSyncStrings::SyncOnBarrier::Name, ESearchCase::IgnoreCase))
-	{
-		// Process command
-		const EDisplayClusterCommResult CommResult = SyncOnBarrier();
-		Response->SetCommResult(CommResult);
-
-		return Response;
-	}
-
-	// Being here means that we have no appropriate dispatch logic for this packet
-	UE_LOG(LogDisplayClusterNetwork, Warning, TEXT("%s - No dispatcher found for packet '%s'"), *GetName(), *Request->GetName());
-
-	return nullptr;
+	// Activate all barriers
+	ActivateAllBarriers();
 }
 
-
-//////////////////////////////////////////////////////////////////////////////////////////////
-// IDisplayClusterProtocolRenderSync
-//////////////////////////////////////////////////////////////////////////////////////////////
-EDisplayClusterCommResult FDisplayClusterRenderSyncService::SyncOnBarrier()
-{
-	TRACE_CPUPROFILER_EVENT_SCOPE(nD SRV_RS::SyncOnBarrier);
-	
-	const FString NodeId = GetSessionInfoCache().NodeId.Get("");
-
-	IDisplayClusterBarrier* Barrier = GetBarrierForNode(NodeId);
-	if (ensureMsgf(Barrier, TEXT("%s could not find a barrier for node '%s'"), *GetName(), *NodeId))
-	{
-		Barrier->Wait(NodeId);
-	}
-
-	return EDisplayClusterCommResult::Ok;
-}
-
 void FDisplayClusterRenderSyncService::InitializeBarriers()
@@ -195,18 +145,18 @@
 	}
 
 	// Get list of cluster node IDs (runtime nodes)
-	TArray<FString> RuntimeNodeIds;
+	TSet<FString> RuntimeNodeIds;
 	GDisplayCluster->GetPrivateClusterMgr()->GetNodeIds(RuntimeNodeIds);
 
 	// Filter nodes that are expected to present in runtime
 	const TMap<FString, TObjectPtr<UDisplayClusterConfigurationClusterNode>> RuntimeNodeCfgs = Config->Cluster->Nodes.FilterByPredicate(
-		[RuntimeNodeIds](const TPair<FString, TObjectPtr<UDisplayClusterConfigurationClusterNode>>& Pair)
+		[&RuntimeNodeIds](const TPair<FString, TObjectPtr<UDisplayClusterConfigurationClusterNode>>& Pair)
 		{
 			return RuntimeNodeIds.Contains(Pair.Key);
 		});
 
 	// Build sync groups (this one is an aux mapping of policy IDs to their node IDs)
-	TMap<FString, TArray<FString>> SyncGroups;
+	TMap<FString, TSet<FString>> SyncGroups;
 	for (const TPair<FString, TObjectPtr<UDisplayClusterConfigurationClusterNode>>& RuntimeNodeCfg : RuntimeNodeCfgs)
 	{
 		FString SyncPolicyId;
@@ -234,7 +184,7 @@
 		SyncPolicyId.ToLowerInline();
 
 		// Create new sync group if not yet available
-		TArray<FString>* NodesInSyncGroup = SyncGroups.Find(SyncPolicyId);
+		TSet<FString>* NodesInSyncGroup = SyncGroups.Find(SyncPolicyId);
 		if (!NodesInSyncGroup)
 		{
 			NodesInSyncGroup = &SyncGroups.Emplace(SyncPolicyId);
@@ -247,9 +197,9 @@
 	}
 
 	// Finally, initialize barriers for every sync group
-	for (const TPair<FString, TArray<FString>>& SyncGroup : SyncGroups)
+	for (const TPair<FString, TSet<FString>>& SyncGroup : SyncGroups)
 	{
-		const FString BarrierName = SyncGroup.Key + TEXT("_barrier");
+		const FString BarrierName = SyncGroup.Key + TEXT("_present_barrier");
 
 		// Instantiate the barrier
 		if (IDisplayClusterBarrier* const Barrier = FDisplayClusterBarrierFactory::CreateBarrier(BarrierName, SyncGroup.Value, Config->Cluster->Network.RenderSyncBarrierTimeout))
@@ -306,3 +256,72 @@
 		Barrier->UnregisterSyncCaller(NodeId);
 	}
 }
+
+void FDisplayClusterRenderSyncService::ShutdownImpl()
+{
+	// Deactivate all barriers
+	DeactivateAllBarriers();
+}
+
+
+//////////////////////////////////////////////////////////////////////////////////////////////
+// IDisplayClusterSessionPacketHandler
+//////////////////////////////////////////////////////////////////////////////////////////////
+TSharedPtr<FDisplayClusterPacketInternal> FDisplayClusterRenderSyncService::ProcessPacket(const TSharedPtr<FDisplayClusterPacketInternal>& Request, const FDisplayClusterSessionInfo& SessionInfo)
+{
+	// Check the pointer
+	if (!Request.IsValid())
+	{
+		UE_LOG(LogDisplayClusterNetwork, Error, TEXT("%s - Invalid request data (nullptr)"), *GetName());
+		return nullptr;
+	}
+
+	// Cache session info
+	SetSessionInfoCache(SessionInfo);
+
+	UE_LOG(LogDisplayClusterNetwork, Verbose, TEXT("%s - Processing packet: %s"), *GetName(), *Request->ToLogString());
+
+	// Check protocol and type
+	if (Request->GetProtocol() != DisplayClusterRenderSyncStrings::ProtocolName || Request->GetType() != DisplayClusterRenderSyncStrings::TypeRequest)
+	{
+		UE_LOG(LogDisplayClusterNetwork, Error, TEXT("%s - Unsupported packet type: %s"), *GetName(), *Request->ToLogString());
+		return nullptr;
+	}
+
+	// Initialize response packet
+	TSharedPtr<FDisplayClusterPacketInternal> Response = MakeShared<FDisplayClusterPacketInternal>(Request->GetName(), DisplayClusterRenderSyncStrings::TypeResponse, Request->GetProtocol());
+
+	// Dispatch the packet
+	if (Request->GetName().Equals(DisplayClusterRenderSyncStrings::SynchronizeOnBarrier::Name, ESearchCase::IgnoreCase))
+	{
+		// Process command
+		const EDisplayClusterCommResult CommResult = SynchronizeOnBarrier();
+		Response->SetCommResult(CommResult);
+
+		return Response;
+	}
+
+	// Being here means that we have no appropriate dispatch logic for this packet
+	UE_LOG(LogDisplayClusterNetwork, Warning, TEXT("%s - No dispatcher found for packet '%s'"), *GetName(), *Request->GetName());
+
+	return nullptr;
+}
+
+
+//////////////////////////////////////////////////////////////////////////////////////////////
+// IDisplayClusterProtocolRenderSync
+//////////////////////////////////////////////////////////////////////////////////////////////
+EDisplayClusterCommResult FDisplayClusterRenderSyncService::SynchronizeOnBarrier()
+{
+	TRACE_CPUPROFILER_EVENT_SCOPE(SRV_RS::SyncOnBarrier);
+
+	const FString NodeId = GetSessionInfoCache().NodeId.Get(TEXT(""));
+
+	IDisplayClusterBarrier* Barrier = GetBarrierForNode(NodeId);
+	if (ensureMsgf(Barrier, TEXT("%s could not find a barrier for node '%s'"), *GetName(), *NodeId))
+	{
+		Barrier->Wait(NodeId);
+	}
+
+	return EDisplayClusterCommResult::Ok;
+}
diff --git a/Engine/Plugins/Runtime/nDisplay/Source/DisplayCluster/Private/Network/Service/RenderSync/DisplayClusterRenderSyncService.h b/Engine/Plugins/Runtime/nDisplay/Source/DisplayCluster/Private/Network/Service/RenderSync/DisplayClusterRenderSyncService.h
--- a/Engine/Plugins/Runtime/nDisplay/Source/DisplayCluster/Private/Network/Service/RenderSync/DisplayClusterRenderSyncService.h
+++ b/Engine/Plugins/Runtime/nDisplay/Source/DisplayCluster/Private/Network/Service/RenderSync/DisplayClusterRenderSyncService.h
@@ -2,8 +2,6 @@
 
 #pragma once
 
-#include "CoreMinimal.h"
-
 #include "Network/Service/DisplayClusterService.h"
 #include "Network/Session/IDisplayClusterSessionPacketHandler.h"
 #include "Network/Protocol/IDisplayClusterProtocolRenderSync.h"
@@ -22,66 +20,78 @@ class FDisplayClusterRenderSyncService
 	, protected IDisplayClusterProtocolRenderSync
 {
 public:
-	FDisplayClusterRenderSyncService();
+
+	FDisplayClusterRenderSyncService(const FName& InInstanceName);
+
 	virtual ~FDisplayClusterRenderSyncService();
 
 public:
-	//////////////////////////////////////////////////////////////////////////////////////////////
-	// IDisplayClusterServer
-	//////////////////////////////////////////////////////////////////////////////////////////////
+
+	//~ Begin IDisplayClusterServer
 	virtual bool Start(const FString& Address, const uint16 Port) override;
-	virtual bool Start(TSharedPtr<FDisplayClusterTcpListener>& ExternalListener) override;
-	virtual void Shutdown() override final;
+	virtual bool Start(TSharedRef<FDisplayClusterTcpListener>& ExternalListener) override;
+	virtual void Shutdown() override;
 	virtual FString GetProtocolName() const override;
 	virtual void KillSession(const FString& NodeId) override;
+	//~ End IDisplayClusterServer
 
 protected:
-	// Creates session instance for this service
+
+	/** Creates session instance for this service */
 	virtual TSharedPtr<IDisplayClusterSession> CreateSession(FDisplayClusterSessionInfo& SessionInfo) override;
 
-	// Callback when a session is closed
+	/** Callback when a session is closed */
 	void ProcessSessionClosed(const FDisplayClusterSessionInfo& SessionInfo);
 
-	// Callback on barrier timeout
-	void ProcessBarrierTimeout(const FString& BarrierName, const TArray<FString>& NodesTimedOut);
+	/** Callback on barrier timeout */
+	void ProcessBarrierTimeout(const FString& BarrierName, const TSet<FString>& NodesTimedOut);
 
 protected:
-	//////////////////////////////////////////////////////////////////////////////////////////////
-	// IDisplayClusterSessionPacketHandler
-	//////////////////////////////////////////////////////////////////////////////////////////////
+
+	//~ Begin IDisplayClusterSessionPacketHandler
 	virtual TSharedPtr<FDisplayClusterPacketInternal> ProcessPacket(const TSharedPtr<FDisplayClusterPacketInternal>& Request, const FDisplayClusterSessionInfo& SessionInfo) override;
+	//~ End IDisplayClusterSessionPacketHandler
 
 private:
-	//////////////////////////////////////////////////////////////////////////////////////////////
-	// IDisplayClusterProtocolRenderSync
-	//////////////////////////////////////////////////////////////////////////////////////////////
-	virtual EDisplayClusterCommResult SyncOnBarrier() override;
+
+	//~ Begin IDisplayClusterProtocolRenderSync
+	virtual EDisplayClusterCommResult SynchronizeOnBarrier() override;
+	//~ End IDisplayClusterProtocolRenderSync
 
 private:
-	// Initializes barriers based on available nodes and their sync policies
+
+	/** Performs internal initialization during server start */
+	void StartInternal();
+
+	/** Initializes barriers based on available nodes and their sync policies */
 	void InitializeBarriers();
 
-	// Activates all service barriers
+	/** Activates all service barriers */
 	void ActivateAllBarriers();
-	// Deactivates all service barriers
+
+	/** Deactivates all service barriers */
 	void DeactivateAllBarriers();
 
-	// Unsubscribe from all events of all internal barriers
+	/** Unsubscribe from all events of all internal barriers */
 	void UnsubscribeFromAllBarrierEvents();
 
-	// Returns barrier of a node's sync group
+	/** Returns barrier of a node's sync group */
 	IDisplayClusterBarrier* GetBarrierForNode(const FString& NodeId) const;
 
-	// Unregister cluster node from a barrier
+	/** Unregister cluster node from a barrier */
 	void UnregisterClusterNode(const FString& NodeId);
 
+	/** Non-virtual shutdown implementation */
+	void ShutdownImpl();
+
 private:
+
 	// We can now use different sync policies within the same cluster. Since each policy may
 	// have its own synchronization logic, mainly barriers utilization, we need to have
 	// individual sync barriers for every sync group. A sync group is a list of nodes that
 	// use the same sync policy.
 	TMap<FString, TUniquePtr<IDisplayClusterBarrier>> PolicyToBarrierMap;
 
-	// Node ID to sync policy ID (or sync group) mapping
+	/** Node ID to sync policy ID(or sync group) mapping */
 	TMap<FString, FString> NodeToPolicyMap;
 };
diff --git a/Engine/Plugins/Runtime/nDisplay/Source/DisplayCluster/Private/Network/Service/RenderSync/DisplayClusterRenderSyncStrings.h b/Engine/Plugins/Runtime/nDisplay/Source/DisplayCluster/Private/Network/Service/RenderSync/DisplayClusterRenderSyncStrings.h
--- a/Engine/Plugins/Runtime/nDisplay/Source/DisplayCluster/Private/Network/Service/RenderSync/DisplayClusterRenderSyncStrings.h
+++ b/Engine/Plugins/Runtime/nDisplay/Source/DisplayCluster/Private/Network/Service/RenderSync/DisplayClusterRenderSyncStrings.h
@@ -2,7 +2,7 @@
 
 #pragma once
 
-#include "CoreMinimal.h"
+#include "HAL/Platform.h"
 
 
 /**
@@ -17,7 +17,7 @@ namespace DisplayClusterRenderSyncStrings
 
 	constexpr static const TCHAR* ArgumentsDefaultCategory = TEXT("RS");
 
-	namespace SyncOnBarrier
+	namespace SynchronizeOnBarrier
 	{
 		constexpr static const TCHAR* Name = TEXT("SyncOnBarrier");
 	};
diff --git a/Engine/Plugins/Runtime/nDisplay/Source/DisplayCluster/Private/Network/Session/DisplayClusterSession.h b/Engine/Plugins/Runtime/nDisplay/Source/DisplayCluster/Private/Network/Session/DisplayClusterSession.h
--- a/Engine/Plugins/Runtime/nDisplay/Source/DisplayCluster/Private/Network/Session/DisplayClusterSession.h
+++ b/Engine/Plugins/Runtime/nDisplay/Source/DisplayCluster/Private/Network/Session/DisplayClusterSession.h
@@ -2,22 +2,21 @@
 
 #pragma once
 
-#include "CoreMinimal.h"
+#include "GenericPlatform/GenericPlatformAffinity.h"
 
+#include "HAL/Runnable.h"
+#include "HAL/RunnableThread.h"
+
+#include "Misc/DisplayClusterLog.h"
+
+#include "Network/DisplayClusterNetworkTypes.h"
+#include "Network/IDisplayClusterServer.h"
+#include "Network/Configuration/DisplayClusterNetworkConfiguration.h"
+#include "Network/Packet/IDisplayClusterPacket.h"
 #include "Network/Session/IDisplayClusterSession.h"
 #include "Network/Session/IDisplayClusterSessionPacketHandler.h"
 #include "Network/Transport/DisplayClusterSocketOperations.h"
 #include "Network/Transport/DisplayClusterSocketOperationsHelper.h"
-#include "Network/Packet/IDisplayClusterPacket.h"
-#include "Network/DisplayClusterNetworkTypes.h"
-#include "Network/IDisplayClusterServer.h"
-
-#include "Misc/DisplayClusterConstants.h"
-#include "Misc/DisplayClusterLog.h"
-
-#include "HAL/Runnable.h"
-#include "HAL/RunnableThread.h"
-#include "GenericPlatform/GenericPlatformAffinity.h"
 
 
 /**
@@ -37,7 +36,7 @@ public:
 			IDisplayClusterSessionPacketHandler<TPacketType, bIsBidirectional>& InPacketHandler,
 			EThreadPriority InThreadPriority = EThreadPriority::TPri_Normal)
 
-		: FDisplayClusterSocketOperations(InSessionInfo.Socket, DisplayClusterConstants::net::PacketBufferSize, InSessionInfo.SessionName)
+		: FDisplayClusterSocketOperations(InSessionInfo.Socket, UE::nDisplay::Network::Configuration::PacketBufferSize, InSessionInfo.SessionName)
 		, FDisplayClusterSocketOperationsHelper<TPacketType>(*this, InSessionInfo.SessionName)
 		, SessionInfo(InSessionInfo)
 		, OwningServer(InOwningServer)
diff --git a/Engine/Plugins/Runtime/nDisplay/Source/DisplayCluster/Private/Network/Transport/DisplayClusterSocketOperations.cpp b/Engine/Plugins/Runtime/nDisplay/Source/DisplayCluster/Private/Network/Transport/DisplayClusterSocketOperations.cpp
--- a/Engine/Plugins/Runtime/nDisplay/Source/DisplayCluster/Private/Network/Transport/DisplayClusterSocketOperations.cpp
+++ b/Engine/Plugins/Runtime/nDisplay/Source/DisplayCluster/Private/Network/Transport/DisplayClusterSocketOperations.cpp
@@ -2,11 +2,10 @@
 
 #include "Network/Transport/DisplayClusterSocketOperations.h"
 
+#include "Misc/DisplayClusterLog.h"
 #include "Misc/ScopeLock.h"
-#include "SocketSubsystem.h"
 
-#include "Misc/DisplayClusterConstants.h"
-#include "Misc/DisplayClusterLog.h"
+#include "SocketSubsystem.h"
 
 
 FDisplayClusterSocketOperations::FDisplayClusterSocketOperations(FSocket* InSocket, int32 PersistentBufferSize, const FString& InConnectionName, bool bReleaseSocketOnDestroy)
@@ -24,6 +23,7 @@ FDisplayClusterSocketOperations::~FDisplayClusterSocketOperations()
 	// Release the socket only if requested on creation
 	if (bReleaseSocket && Socket)
 	{
+		CloseSocket();
 		ISocketSubsystem::Get(PLATFORM_SOCKETSUBSYSTEM)->DestroySocket(Socket);
 		Socket = nullptr;
 	}
@@ -105,12 +105,12 @@ bool FDisplayClusterSocketOperations::SendChunk(FSocket* const InSocket, const T
 
 bool FDisplayClusterSocketOperations::RecvChunk(TArray<uint8>& ChunkBuffer, const uint32 ChunkSize, const FString& ChunkName)
 {
-	FScopeLock Loc(&CritSecInternals);
+	FScopeLock Lock(&CritSecInternals);
 	return FDisplayClusterSocketOperations::RecvChunk(Socket, ChunkBuffer, ChunkSize, ChunkName);
 }
 
 bool FDisplayClusterSocketOperations::SendChunk(const TArray<uint8>& ChunkBuffer, const uint32 ChunkSize, const FString& ChunkName)
 {
-	FScopeLock Loc(&CritSecInternals);
+	FScopeLock Lock(&CritSecInternals);
 	return FDisplayClusterSocketOperations::SendChunk(Socket, ChunkBuffer, ChunkSize, ChunkName);
 }
diff --git a/Engine/Plugins/Runtime/nDisplay/Source/DisplayCluster/Private/Network/Transport/DisplayClusterSocketOperations.h b/Engine/Plugins/Runtime/nDisplay/Source/DisplayCluster/Private/Network/Transport/DisplayClusterSocketOperations.h
--- a/Engine/Plugins/Runtime/nDisplay/Source/DisplayCluster/Private/Network/Transport/DisplayClusterSocketOperations.h
+++ b/Engine/Plugins/Runtime/nDisplay/Source/DisplayCluster/Private/Network/Transport/DisplayClusterSocketOperations.h
@@ -2,8 +2,6 @@
 
 #pragma once
 
-#include "CoreMinimal.h"
-
 #include "Sockets.h"
 
 
@@ -44,7 +42,7 @@ public:
 	// Close socket
 	inline void CloseSocket()
 	{
-		if (IsOpen())
+		if (Socket)
 		{
 			Socket->Close();
 		}
diff --git a/Engine/Plugins/Runtime/nDisplay/Source/DisplayCluster/Private/Network/Transport/DisplayClusterSocketOperationsHelper.h b/Engine/Plugins/Runtime/nDisplay/Source/DisplayCluster/Private/Network/Transport/DisplayClusterSocketOperationsHelper.h
--- a/Engine/Plugins/Runtime/nDisplay/Source/DisplayCluster/Private/Network/Transport/DisplayClusterSocketOperationsHelper.h
+++ b/Engine/Plugins/Runtime/nDisplay/Source/DisplayCluster/Private/Network/Transport/DisplayClusterSocketOperationsHelper.h
@@ -2,13 +2,11 @@
 
 #pragma once
 
-#include "CoreMinimal.h"
+#include "Misc/DisplayClusterLog.h"
 
 #include "Network/Transport/DisplayClusterSocketOperations.h"
 #include "Network/Packet/IDisplayClusterPacket.h"
 
-#include "Misc/DisplayClusterLog.h"
-
 
 /**
  * Socket operations helper. High level operations with specific packet types.
@@ -19,6 +17,7 @@ class FDisplayClusterSocketOperationsHelper
 	static_assert(std::is_base_of<IDisplayClusterPacket, TPacketType>::value, "TPacketType is not derived from IDisplayClusterPacket");
 
 public:
+
 	FDisplayClusterSocketOperationsHelper(FDisplayClusterSocketOperations& InSocketOps, const FString& InLogHeader = FString())
 		: SocketOps(InSocketOps)
 		, LogHeader(InLogHeader)
@@ -28,6 +27,7 @@ public:
 	{ }
 
 public:
+
 	TSharedPtr<TPacketType> SendRecvPacket(const TSharedPtr<TPacketType>& Request)
 	{
 		if (Request)
@@ -41,15 +41,27 @@ public:
 		return nullptr;
 	}
 
-	bool SendPacket(const TSharedPtr<TPacketType>& Packet)
+	bool SendPacket(const TSharedPtr<IDisplayClusterPacket>& Packet)
 	{
 		if (Packet)
 		{
+			constexpr bool bDetailedLog = true;
+
 			UE_LOG(LogDisplayClusterNetwork, VeryVerbose, TEXT("%s: sending packet - %s"),
 				LogHeader.IsEmpty() ? *SocketOps.GetConnectionName() : *LogHeader,
-				*StaticCastSharedPtr<IDisplayClusterPacket>(Packet)->ToLogString(true));
+				*Packet->ToLogString(bDetailedLog));
 
-			return StaticCastSharedPtr<IDisplayClusterPacket>(Packet)->SendPacket(SocketOps);
+			return Packet->SendPacket(SocketOps);
+		}
+
+		return false;
+	}
+
+	bool SendPacket(const TSharedPtr<TPacketType>& Packet)
+	{
+		if (Packet)
+		{
+			return SendPacket(StaticCastSharedPtr<IDisplayClusterPacket>(Packet));
 		}
 
 		return false;
@@ -60,9 +72,11 @@ public:
 		TSharedPtr<TPacketType> Packet = MakeShared<TPacketType>();
 		if (StaticCastSharedPtr<IDisplayClusterPacket>(Packet)->RecvPacket(SocketOps))
 		{
+			constexpr bool bDetailedLog = true;
+
 			UE_LOG(LogDisplayClusterNetwork, VeryVerbose, TEXT("%s: received packet - %s"),
 				LogHeader.IsEmpty() ? *SocketOps.GetConnectionName() : *LogHeader,
-				*StaticCastSharedPtr<IDisplayClusterPacket>(Packet)->ToLogString(true));
+				*StaticCastSharedPtr<IDisplayClusterPacket>(Packet)->ToLogString(bDetailedLog));
 
 			return Packet;
 		}
diff --git a/Engine/Plugins/Runtime/nDisplay/Source/DisplayCluster/Private/Render/DisplayClusterRenderManager.cpp b/Engine/Plugins/Runtime/nDisplay/Source/DisplayCluster/Private/Render/DisplayClusterRenderManager.cpp
--- a/Engine/Plugins/Runtime/nDisplay/Source/DisplayCluster/Private/Render/DisplayClusterRenderManager.cpp
+++ b/Engine/Plugins/Runtime/nDisplay/Source/DisplayCluster/Private/Render/DisplayClusterRenderManager.cpp
@@ -700,7 +700,7 @@ TSharedPtr<IDisplayClusterRenderSyncPolicy> FDisplayClusterRenderManager::Create
 {
 	if (CurrentOperationMode != EDisplayClusterOperationMode::Cluster)
 	{
-		UE_LOG(LogDisplayClusterRender, Warning, TEXT("Synchronization policy is not available for the current operation mode"));
+		UE_LOG(LogDisplayClusterRender, Log, TEXT("Synchronization policy is not available for current operation mode"));
 		return nullptr;
 	}
 
diff --git a/Engine/Plugins/Runtime/nDisplay/Source/DisplayCluster/Private/Render/Synchronization/DisplayClusterRenderSyncPolicyBase.cpp b/Engine/Plugins/Runtime/nDisplay/Source/DisplayCluster/Private/Render/Synchronization/DisplayClusterRenderSyncPolicyBase.cpp
--- a/Engine/Plugins/Runtime/nDisplay/Source/DisplayCluster/Private/Render/Synchronization/DisplayClusterRenderSyncPolicyBase.cpp
+++ b/Engine/Plugins/Runtime/nDisplay/Source/DisplayCluster/Private/Render/Synchronization/DisplayClusterRenderSyncPolicyBase.cpp
@@ -6,6 +6,7 @@
 
 #include "Cluster/IPDisplayClusterClusterManager.h"
 #include "Cluster/Controller/IDisplayClusterClusterNodeController.h"
+#include "Cluster/NetAPI/DisplayClusterNetApiFacade.h"
 
 #include "Misc/DisplayClusterGlobals.h"
 #include "Misc/DisplayClusterLog.h"
@@ -26,7 +27,7 @@ void FDisplayClusterRenderSyncPolicyBase::SyncOnBarrier()
 {
 	{
 		TRACE_CPUPROFILER_EVENT_SCOPE(nDisplay SyncOnBarrier);
-		GDisplayCluster->GetPrivateClusterMgr()->GetClusterNodeController()->SyncOnBarrier();
+		GDisplayCluster->GetPrivateClusterMgr()->GetNetApi().GetRenderSyncAPI()->SynchronizeOnBarrier();
 	}
 }
 
diff --git a/Engine/Plugins/Runtime/nDisplay/Source/DisplayCluster/Public/Blueprints/DisplayClusterBlueprintLib.h b/Engine/Plugins/Runtime/nDisplay/Source/DisplayCluster/Public/Blueprints/DisplayClusterBlueprintLib.h
--- a/Engine/Plugins/Runtime/nDisplay/Source/DisplayCluster/Public/Blueprints/DisplayClusterBlueprintLib.h
+++ b/Engine/Plugins/Runtime/nDisplay/Source/DisplayCluster/Public/Blueprints/DisplayClusterBlueprintLib.h
@@ -68,6 +68,10 @@ public:
 	UFUNCTION(BlueprintPure, meta = (DisplayName = "Is Backup Node"), Category = "NDisplay|Cluster")
 	static DISPLAYCLUSTER_API bool IsBackup();
 
+	/** Checks if the node has a specific role. */
+	UFUNCTION(BlueprintPure, meta = (DisplayName = "Has a Role"), Category = "NDisplay|Cluster")
+	static DISPLAYCLUSTER_API bool HasClusterRole(EDisplayClusterNodeRole Role);
+
 	/** Returns the role of the current cluster node. */
 	UFUNCTION(BlueprintPure, meta = (DisplayName = "Get Cluster Node Role"), Category = "NDisplay|Cluster")
 	static DISPLAYCLUSTER_API EDisplayClusterNodeRole GetClusterRole();
diff --git a/Engine/Plugins/Runtime/nDisplay/Source/DisplayCluster/Public/Cluster/DisplayClusterClusterEvent.h b/Engine/Plugins/Runtime/nDisplay/Source/DisplayCluster/Public/Cluster/DisplayClusterClusterEvent.h
--- a/Engine/Plugins/Runtime/nDisplay/Source/DisplayCluster/Public/Cluster/DisplayClusterClusterEvent.h
+++ b/Engine/Plugins/Runtime/nDisplay/Source/DisplayCluster/Public/Cluster/DisplayClusterClusterEvent.h
@@ -8,29 +8,40 @@
 
 constexpr const auto DisplayClusterResetSyncType = TEXT("nDCReset");
 
+class FArchive;
+
 
 //////////////////////////////////////////////////////////////////////////////////////////////
 // Common cluster event data
 //////////////////////////////////////////////////////////////////////////////////////////////
 USTRUCT(BlueprintType)
-struct FDisplayClusterClusterEventBase
+struct DISPLAYCLUSTER_API FDisplayClusterClusterEventBase
 {
 	GENERATED_BODY()
 
 public:
+
 	FDisplayClusterClusterEventBase()
 		: bIsSystemEvent(false)
 		, bShouldDiscardOnRepeat(true)
 	{ }
 
+	virtual ~FDisplayClusterClusterEventBase() = default;
+
 public:
-	// Is nDisplay internal event (should never be true for end users)
+
+	/** Is nDisplay internal event(should never be true for end users) */
 	UPROPERTY(EditAnywhere, BlueprintReadWrite, meta = (DisplayName = "Is Sytem Event. 'True' is reserved for nDisplay internals."), Category = "NDisplay")
 	bool bIsSystemEvent = false;
 
-	// Should older events with the same Name/Type/Category (for JSON) or ID (for binary) be discarded if a new one received
+	/** Should older events with the same Name / Type / Category(for JSON) or ID(for binary) be discarded if a new one received */
 	UPROPERTY(EditAnywhere, BlueprintReadWrite, Category = "NDisplay")
 	bool bShouldDiscardOnRepeat = true;
+
+public:
+
+	/** Serialization */
+	virtual void Serialize(FArchive& Ar);
 };
 
 
@@ -38,36 +49,59 @@ public:
 // Cluster event JSON
 //////////////////////////////////////////////////////////////////////////////////////////////
 USTRUCT(BlueprintType)
-struct FDisplayClusterClusterEventJson
+struct DISPLAYCLUSTER_API FDisplayClusterClusterEventJson
 	: public FDisplayClusterClusterEventBase
 {
 	GENERATED_BODY()
 
 public:
-	// Event name (used for discarding outdated events)
+
+	/** Event name(used for discarding outdated events) */
 	UPROPERTY(EditAnywhere, BlueprintReadWrite, Category = "NDisplay")
 	FString Name;
 
-	// Event type (used for discarding outdated events)
+	/** Event type(used for discarding outdated events) */
 	UPROPERTY(EditAnywhere, BlueprintReadWrite, Category = "NDisplay")
 	FString Type;
 
-	// Event category (used for discarding outdated events)
+	/** Event category(used for discarding outdated events) */
 	UPROPERTY(EditAnywhere, BlueprintReadWrite, Category = "NDisplay")
 	FString Category;
 
 public:
-	// Event parameters
+
+	/** Event parameters */
 	UPROPERTY(EditAnywhere, BlueprintReadWrite, Category = "NDisplay")
 	TMap<FString, FString> Parameters;
 
 public:
-	FString SerializeToString() const;
-	bool    DeserializeFromString(const FString& Arch);
 
-private:
-	FString SerializeParametersToString() const;
-	bool    DeserializeParametersFromString(const FString& Arch);
+	/** Serialization */
+	virtual void Serialize(FArchive& Ar) override;
+
+	/** Serialization */
+	friend FArchive& operator<<(FArchive& Ar, FDisplayClusterClusterEventJson& Instance)
+	{
+		Instance.Serialize(Ar);
+		return Ar;
+	}
+
+	/** String based representation */
+	FString ToString(bool bWithParams = false) const;
+
+public:
+
+	UE_DEPRECATED(5.6, "This function has been deprecated and will be removed soon. Please use 'Serialize' for JSON event serialization.")
+	FString SerializeToString() const
+	{
+		return FString();
+	}
+
+	UE_DEPRECATED(5.6, "This function has been deprecated and will be removed soon. Please use 'Serialize' for JSON event deserialization.")
+	bool DeserializeFromString(const FString& Arch)
+	{
+		return false;
+	}
 };
 
 
@@ -75,21 +109,36 @@ private:
 // Cluster event BINARY
 //////////////////////////////////////////////////////////////////////////////////////////////
 USTRUCT(BlueprintType)
-struct FDisplayClusterClusterEventBinary
+struct DISPLAYCLUSTER_API FDisplayClusterClusterEventBinary
 	: public FDisplayClusterClusterEventBase
 {
 	GENERATED_BODY()
 
 public:
-	// Event ID (used for discarding outdated events)
+
+	/** Event ID(used for discarding outdated events) */
 	UPROPERTY(EditAnywhere, BlueprintReadWrite, Category = "NDisplay")
 	int32 EventId = -1;
 
-	// Binary event data
+	/** Binary event data */
 	UPROPERTY(EditAnywhere, BlueprintReadWrite, Category = "NDisplay")
 	TArray<uint8> EventData;
 
 public:
+
+	/** Raw serialization */
 	void SerializeToByteArray(TArray<uint8>& Arch) const;
+
+	/** Raw deserialization */
 	bool DeserializeFromByteArray(const TArray<uint8>& Arch);
+
+	/** Serialization */
+	virtual void Serialize(FArchive& Ar) override;
+
+	/** Serialization */
+	friend FArchive& operator<<(FArchive& Ar, FDisplayClusterClusterEventBinary& Instance)
+	{
+		Instance.Serialize(Ar);
+		return Ar;
+	}
 };
diff --git a/Engine/Plugins/Runtime/nDisplay/Source/DisplayCluster/Public/Cluster/IDisplayClusterClusterManager.h b/Engine/Plugins/Runtime/nDisplay/Source/DisplayCluster/Public/Cluster/IDisplayClusterClusterManager.h
--- a/Engine/Plugins/Runtime/nDisplay/Source/DisplayCluster/Public/Cluster/IDisplayClusterClusterManager.h
+++ b/Engine/Plugins/Runtime/nDisplay/Source/DisplayCluster/Public/Cluster/IDisplayClusterClusterManager.h
@@ -41,6 +41,11 @@ public:
 	virtual bool IsBackup() const = 0;
 	/** Returns cluster node role. */
 	virtual EDisplayClusterNodeRole GetClusterRole() const = 0;
+	/** Checks if the node has a specific role. */
+	virtual bool HasClusterRole(EDisplayClusterNodeRole Role) const = 0;
+
+	/** Returns current primary node ID. */
+	virtual FString GetPrimaryNodeId() const = 0;
 
 	/** Returns current cluster node ID. */
 	virtual FString GetNodeId() const = 0;
@@ -48,6 +53,8 @@ public:
 	virtual uint32 GetNodesAmount() const = 0;
 	/** Returns IDs of available cluster nodes. */
 	virtual void GetNodeIds(TArray<FString>& OutNodeIds) const = 0;
+	/** Returns IDs of available cluster nodes. */
+	virtual void GetNodeIds(TSet<FString>& OutNodeIds) const = 0;
 
 	/** Drop specific cluster node */
 	virtual bool DropClusterNode(const FString& NodeId) = 0;
@@ -57,8 +64,9 @@ public:
 	/** Unregisters synchronization object. */
 	virtual void UnregisterSyncObject(IDisplayClusterClusterSyncObject* SyncObj) = 0;
 
-	/** Spawns a client for generic barriers service. */
-	virtual TUniquePtr<IDisplayClusterGenericBarriersClient, FDisplayClusterGenericBarriersClientDeleter> CreateGenericBarriersClient(const FString& ClientName) = 0;
+	/** Creates a general purpose barriers interface. */
+	virtual TSharedRef<IDisplayClusterGenericBarriersClient> CreateGenericBarriersClient() = 0;
+
 
 	//////////////////////////////////////////////////////////////////////////////////////////////
 	// Cluster events
@@ -93,4 +101,12 @@ public:
 
 	/** Sends binary cluster event to a specific target (outside of the cluster). */
 	virtual void SendClusterEventTo(const FString& Address, const uint16 Port, const FDisplayClusterClusterEventBinary& Event, bool bPrimaryOnly) = 0;
+
+public:
+	
+	UE_DEPRECATED(5.6, "This API method has been deprecated. Please use 'GetGenericBarriersClient'.")
+	virtual TUniquePtr<IDisplayClusterGenericBarriersClient> CreateGenericBarriersClient(const FString& ClientName)
+	{
+		return nullptr;
+	}
 };
diff --git a/Engine/Plugins/Runtime/nDisplay/Source/DisplayCluster/Public/Cluster/IDisplayClusterGenericBarriersClient.h b/Engine/Plugins/Runtime/nDisplay/Source/DisplayCluster/Public/Cluster/IDisplayClusterGenericBarriersClient.h
--- a/Engine/Plugins/Runtime/nDisplay/Source/DisplayCluster/Public/Cluster/IDisplayClusterGenericBarriersClient.h
+++ b/Engine/Plugins/Runtime/nDisplay/Source/DisplayCluster/Public/Cluster/IDisplayClusterGenericBarriersClient.h
@@ -13,13 +13,13 @@ struct FGenericBarrierSynchronizationDelegateData
 	/** Barrier ID */
 	const FString& BarrierId;
 
-	/** Associates thread markers with the owning cluster nodes (thread marker - to - cluster node) */
+	/** Associates the caller IDs with the owning cluster nodes (Caller ID - to - cluster node) */
 	const TMap<FString, FString>& ThreadToNodeMap;
 
-	/** Binary data provided on sync request (thread marker - to - data mapping) */
+	/** Binary data provided on sync request (Caller ID - to - data mapping) */
 	const TMap<FString, TArray<uint8>>& RequestData;
 
-	/** Binary data to respond (thread marker - to - data mapping) */
+	/** Binary data to respond (Caller ID - to - data mapping) */
 	TMap<FString, TArray<uint8>>& ResponseData;
 };
 
@@ -30,66 +30,134 @@ struct FGenericBarrierSynchronizationDelegateData
 class IDisplayClusterGenericBarriersClient
 {
 public:
+
 	virtual ~IDisplayClusterGenericBarriersClient() = default;
 
-// Networking API
 public:
-	/** Connects to a server */
-	virtual bool Connect() = 0;
-
-	/** Terminates current connection */
-	virtual void Disconnect() = 0;
 
-	/** Returns connection status */
-	virtual bool IsConnected() const = 0;
+	/** Synchronization delegate. It's called on the primary node only. */
+	DECLARE_DELEGATE_OneParam(FOnGenericBarrierSynchronizationDelegate, FGenericBarrierSynchronizationDelegateData&);
 
-	/** Returns client name */
-	virtual FString GetName() const = 0;
+	/** Barrier timeout delegate. It's called on the primary node only. */
+	DECLARE_DELEGATE_OneParam(FOnGenericBarrierTimeoutDelegate, const TSet<FString>&);
 
-// Barrier API
 public:
-	/** Creates new barrier */
-	virtual bool CreateBarrier(const FString& BarrierId, const TArray<FString>& UniqueThreadMarkers, uint32 Timeout) = 0;
 
-	/** Wait until a barrier with specific ID is created and ready to go */
+	/**
+	 * Creates new barrier
+	 * 
+	 * @param BarrierId         - ID of the new barrier
+	 * @param NodeToSyncCallers - NodeId-to-CallerId association map
+	 * @param Timeout           - Barrier synchronization timeout
+	 * 
+	 * @return true if the barrier was created successfully, or already exists
+	 */
+	virtual bool CreateBarrier(const FString& BarrierId, const TMap<FString, TSet<FString>>& NodeToSyncCallers, uint32 Timeout) = 0;
+
+	/**
+	 * Wait until a barrier with specific ID is created and ready to go
+	 * 
+	 * @param BarrierId - ID of the barrier to wait
+	 * 
+	 * @return true if created or existed
+	 */
 	virtual bool WaitUntilBarrierIsCreated(const FString& BarrierId) = 0;
 
-	/** Checks if a specific barrier exists */
+	/**
+	 * Checks if a specific barrier exists
+	 * 
+	 * @param BarrierId - ID of the barrier to check
+	 *
+	 * @return true if the barrier exists
+	 */
 	virtual bool IsBarrierAvailable(const FString& BarrierId) = 0;
 
-	/** Returns a synchronization delegate to a specific barrier or nullptr if not available */
-	DECLARE_DELEGATE_OneParam(FOnGenericBarrierSynchronizationDelegate, FGenericBarrierSynchronizationDelegateData&);
+	/**
+	 * Releases specific barrier
+	 *
+	 * @param BarrierId - ID of a barrier to release
+	 *
+	 * @return true if the barrier was successfully released
+	 */
+	virtual bool ReleaseBarrier(const FString& BarrierId) = 0;
+
+	/**
+	 * Synchronize the calling thread on a specific barrier
+	 *
+	 * @param BarrierId - ID of a barrier to use for synchronization
+	 * @param CallerId  - ID of a synchronization caller (thread)
+	 *
+	 * @return true if synchronization succeeded
+	 */
+	virtual bool Synchronize(const FString& BarrierId, const FString& CallerId) = 0;
+
+	/**
+	 * Synchronize the calling thread on a specific barrier with custom data
+	 *
+	 * @param BarrierId   - ID of a barrier to use for synchronization
+	 * @param CallerId    - ID of a synchronization caller (thread)
+	 * @param RequestData - Synchronization request data of the calling thread (caller)
+	 * @param OutResponseData - [out] Synchronization response data
+	 *
+	 * @return true if synchronization succeeded
+	 */
+	virtual bool Synchronize(const FString& BarrierId, const FString& CallerId, const TArray<uint8>& RequestData, TArray<uint8>& OutResponseData) = 0;
+
+public:
+
+	/**
+	 * Returns synchronization delegate of a specific barrier. This delegate is called on the primary node only.
+	 *
+	 * @param BarrierId - ID of the barrier
+	 *
+	 * @return Synchronization delegate or nullptr if barrier is not available
+	 */
 	virtual FOnGenericBarrierSynchronizationDelegate* GetBarrierSyncDelegate(const FString& BarrierId) = 0;
 
-	/** Releases specific barrier */
-	virtual bool ReleaseBarrier(const FString& BarrierId) = 0;
+	/**
+	 * Returns timeout delegate of a specific barrier. This delegate is called on the primary node only.
+	 *
+	 * @param BarrierId - ID of the barrier
+	 *
+	 * @return Timeout delegate or nullptr if barrier is not available
+	 */
+	virtual FOnGenericBarrierTimeoutDelegate* GetBarrierTimeoutDelegate(const FString& BarrierId) = 0;
 
-	/** Synchronize calling thread on a specific barrier */
-	virtual bool Synchronize(const FString& BarrierId, const FString& UniqueThreadMarker) = 0;
+public:
 
-	/** Synchronize calling thread on a specific barrier with custom data */
-	virtual bool Synchronize(const FString& BarrierId, const FString& UniqueThreadMarker, const TArray<uint8>& RequestData, TArray<uint8>& OutResponseData) = 0;
-};
+	UE_DEPRECATED(5.6, "This method has been deprecated. There is no need to connect/disconnect anymore.")
+	virtual bool Connect()
+	{
+		return true;
+	}
 
-/**
- * Custom deleter to avoid calling the virtual Disconnect function in the destructor.
- * This is necessary since the destructor of the base class is calling it but expecting
- * the derived class implementation to be called. This also avoids the PVS warning:
- * V1053: Calling the 'Disconnect' virtual function in the destructor may lead to unexpected result at runtime.
- */
-template<typename ClientType>
-struct FDisplayClusterGenericClientDeleter
-{
-	void operator()(ClientType* Client)
+	UE_DEPRECATED(5.6, "This method has been deprecated. There is no need to connect/disconnect anymore.")
+	virtual void Disconnect()
+	{
+	}
+
+	UE_DEPRECATED(5.6, "This method has been deprecated. There is no need to connect/disconnect anymore.")
+	virtual bool IsConnected() const
+	{
+		return true;
+	}
+
+	UE_DEPRECATED(5.6, "This method has been deprecated.")
+	virtual FString GetName() const
 	{
-		if (Client && Client->IsConnected())
-		{
-			Client->Disconnect();
-		}
+		return FString();
+	}
 
-		delete Client;
+	UE_DEPRECATED(5.6, "This method has been deprecated. Please use a TMap based version of CreateBarrier.")
+	virtual bool CreateBarrier(const FString& BarrierId, const TArray<FString>& UniqueThreadMarkers, uint32 Timeout)
+	{
+		return false;
 	}
 };
 
-// Custom deleter for IDisplayClusterGenericBarriersClient.
-using FDisplayClusterGenericBarriersClientDeleter = FDisplayClusterGenericClientDeleter<IDisplayClusterGenericBarriersClient>;
+
+struct UE_DEPRECATED(5.6, "No longer used.") FDisplayClusterGenericBarriersClientDeleter
+{
+	void operator()(void*)
+	{ }
+};
diff --git a/Engine/Plugins/Runtime/nDisplay/Source/DisplayCluster/Public/Components/DisplayClusterSceneComponentSync.h b/Engine/Plugins/Runtime/nDisplay/Source/DisplayCluster/Public/Components/DisplayClusterSceneComponentSync.h
--- a/Engine/Plugins/Runtime/nDisplay/Source/DisplayCluster/Public/Components/DisplayClusterSceneComponentSync.h
+++ b/Engine/Plugins/Runtime/nDisplay/Source/DisplayCluster/Public/Components/DisplayClusterSceneComponentSync.h
@@ -6,9 +6,6 @@
 #include "Cluster/IDisplayClusterClusterSyncObject.h"
 #include "DisplayClusterSceneComponentSync.generated.h"
 
-class IPDisplayClusterGameManager;
-class IPDisplayClusterClusterManager;
-
 
 /**
  * Abstract synchronization component
@@ -65,10 +62,6 @@ protected:
 	virtual void SetSyncTransform(const FTransform& t)
 	{ }
 
-protected:
-	IPDisplayClusterGameManager*    GameMgr = nullptr;
-	IPDisplayClusterClusterManager* ClusterMgr = nullptr;
-
 protected:
 	// Caching state
 	FVector  LastSyncLoc;
diff --git a/Engine/Plugins/Runtime/nDisplay/Source/DisplayCluster/Public/DisplayClusterGameEngine.h b/Engine/Plugins/Runtime/nDisplay/Source/DisplayCluster/Public/DisplayClusterGameEngine.h
--- a/Engine/Plugins/Runtime/nDisplay/Source/DisplayCluster/Public/DisplayClusterGameEngine.h
+++ b/Engine/Plugins/Runtime/nDisplay/Source/DisplayCluster/Public/DisplayClusterGameEngine.h
@@ -11,7 +11,6 @@
 #include "DisplayClusterGameEngine.generated.h"
 
 
-class IPDisplayClusterClusterManager;
 class IDisplayClusterClusterNodeController;
 class IDisplayClusterClusterSyncObject;
 class UDisplayClusterConfigurationData;
@@ -87,8 +86,8 @@ private:
 	/** Custom logic to update PendingLevel in order to support listening server
 	 */
 	void PendingLevelUpdate(FWorldContext& Context, float DeltaSeconds);
+
 private:
-	IPDisplayClusterClusterManager* ClusterMgr = nullptr;
 
 	FDisplayClusterConfigurationDiagnostics Diagnostics;
 
diff --git a/Engine/Plugins/Runtime/nDisplay/Source/DisplayCluster/Public/DisplayClusterPlayerInput.h b/Engine/Plugins/Runtime/nDisplay/Source/DisplayCluster/Public/DisplayClusterPlayerInput.h
--- a/Engine/Plugins/Runtime/nDisplay/Source/DisplayCluster/Public/DisplayClusterPlayerInput.h
+++ b/Engine/Plugins/Runtime/nDisplay/Source/DisplayCluster/Public/DisplayClusterPlayerInput.h
@@ -19,25 +19,31 @@ class DISPLAYCLUSTER_API UDisplayClusterPlayerInput : public UEnhancedPlayerInpu
 	GENERATED_BODY()
 
 public:
+
 	UDisplayClusterPlayerInput();
 
 public:
+
 	/** Process the frame's input events given the current input component stack. */
 	virtual void ProcessInputStack(const TArray<UInputComponent*>& InputComponentStack, const float DeltaTime, const bool bGamePaused) override;
 
 protected:
-	// Performs sync logic where all input data from the primary node is replicated to all other nodes
+
+	/** Performs sync logic where all input data from the primary node is replicated to all other nodes */
 	virtual void ProcessPolicy_ReplicatePrimary(const TArray<UInputComponent*>& InputComponentStack, const float DeltaTime, const bool bGamePaused);
 
 protected:
-	// Aux function that serializes whole input data into a map
+
+	/** Auxiliary function that serializes whole input data into a byte stream */
 	bool SerializeKeyStateMap(TMap<FString, FString>& OutKeyStateMap);
-	// Aux function that deserializes whole input data from a map
+
+	/** Aux function that deserializes whole input data from a map */
 	bool DeserializeKeyStateMap(const TMap<FString, FString>& InKeyStateMap);
 
 private:
+
 	// Currently we have this bool only wich provides 2 cases: sync or not sync. It covers all
-	// functionality currently available . However in the future, this might be refactored in a way
-	// that we can have multiple input sync policies including custom ones.
+	// functionality currently available . However, this might be refactored one day in order
+	// to support multiple input sync policies including custom ones.
 	bool bReplicatePrimary;
 };
diff --git a/Engine/Plugins/Runtime/nDisplay/Source/DisplayCluster/Public/IDisplayClusterCallbacks.h b/Engine/Plugins/Runtime/nDisplay/Source/DisplayCluster/Public/IDisplayClusterCallbacks.h
--- a/Engine/Plugins/Runtime/nDisplay/Source/DisplayCluster/Public/IDisplayClusterCallbacks.h
+++ b/Engine/Plugins/Runtime/nDisplay/Source/DisplayCluster/Public/IDisplayClusterCallbacks.h
@@ -87,10 +87,14 @@ public:
 	DECLARE_EVENT_OneParam(IDisplayClusterCallbacks, FDisplayClusterFramePresentated_RHIThread, bool bNativePresent);
 	virtual FDisplayClusterFramePresentated_RHIThread& OnDisplayClusterFramePresented_RHIThread() = 0;
 
-	/** Failover notification **/
+	/** Notifies about cluster node failure. Called from AnyThread. **/
 	DECLARE_EVENT_OneParam(IDisplayClusterCallbacks, FDisplayClusterFailoverNodeDown, const FString&);
 	virtual FDisplayClusterFailoverNodeDown& OnDisplayClusterFailoverNodeDown() = 0;
 
+	/** Notifies about new primary node. Called from AnyThread. **/
+	DECLARE_EVENT_OneParam(IDisplayClusterCallbacks, FDisplayClusterFailoverPrimaryNodeChanged, const FString&);
+	virtual FDisplayClusterFailoverPrimaryNodeChanged& OnDisplayClusterFailoverPrimaryNodeChanged() = 0;
+
 	/** Called after tonemap pass (before applying OCIO transformation) **/
 	DECLARE_EVENT_FiveParams(IDisplayClusterCallbacks, FDisplayClusterPostTonemapPass_RenderThread, FRDGBuilder&, const IDisplayClusterViewportProxy*, const FSceneView&, const FPostProcessMaterialInputs&, const uint32);
 	virtual FDisplayClusterPostTonemapPass_RenderThread& OnDisplayClusterPostTonemapPass_RenderThread() = 0;
diff --git a/Engine/Plugins/Runtime/nDisplay/Source/DisplayCluster/Public/Misc/DisplayClusterHelpers.h b/Engine/Plugins/Runtime/nDisplay/Source/DisplayCluster/Public/Misc/DisplayClusterHelpers.h
--- a/Engine/Plugins/Runtime/nDisplay/Source/DisplayCluster/Public/Misc/DisplayClusterHelpers.h
+++ b/Engine/Plugins/Runtime/nDisplay/Source/DisplayCluster/Public/Misc/DisplayClusterHelpers.h
@@ -91,8 +91,8 @@ namespace DisplayClusterHelpers
 
 		// Exports array data to a string
 		// Example: {item1, item2, item3, item4} => "item1,item2,item3,item4"
-		template<typename T>
-		static FString ArrayToStr(const TArray<T>& InData, const FString& InSeparator = FString(DisplayClusterStrings::common::ArrayValSeparator), bool bAddQuotes = true)
+		template<typename TVal>
+		static FString ArrayToStr(const TArray<TVal>& InData, const FString& InSeparator = FString(DisplayClusterStrings::common::ArrayValSeparator), bool bAddQuotes = true)
 		{
 			const FString Quotes("\"");
 
@@ -122,6 +122,25 @@ namespace DisplayClusterHelpers
 			return ResultStr;
 		}
 
+		// Parses string items separated by specified separator into a set
+		// Example: item1, item2, item3, item2, item1  ,  item4 => {item1, item2, item3, item4}
+		template<typename TVal>
+		static void StrToSet(const FString& InData, const FString& InSeparator, TSet<TVal>& OutData, bool bCullEmpty = true)
+		{
+			TArray<FString> TempData;
+			StrToArray<TVal>(InData, InSeparator, TempData, bCullEmpty);
+			OutData.Reset();
+			OutData.Append(TempData);
+		}
+
+		// Exports set to a string
+		// Example: {item1, item2, item3, item4} => "item1,item2,item3,item4"
+		template<typename TVal>
+		static FString SetToStr(const TSet<TVal>& InData, const FString& InSeparator = FString(DisplayClusterStrings::common::ArrayValSeparator), bool bAddQuotes = true)
+		{
+			return ArrayToStr<TVal>(InData.Array(), InSeparator, bAddQuotes);
+		}
+
 		// Parses string of key-value pairs separated by specified separator into map
 		// Example: "key1=val1 key2=val2 key3=val3" => {{key1, val2}, {key2, val2}, {key3, val3}}
 		template<typename TKey, typename TVal>
diff --git a/Engine/Plugins/Runtime/nDisplay/Source/DisplayClusterConfiguration/Private/DisplayClusterConfigurationTypes.cpp b/Engine/Plugins/Runtime/nDisplay/Source/DisplayClusterConfiguration/Private/DisplayClusterConfigurationTypes.cpp
--- a/Engine/Plugins/Runtime/nDisplay/Source/DisplayClusterConfiguration/Private/DisplayClusterConfigurationTypes.cpp
+++ b/Engine/Plugins/Runtime/nDisplay/Source/DisplayClusterConfiguration/Private/DisplayClusterConfigurationTypes.cpp
@@ -45,6 +45,26 @@ UDisplayClusterConfigurationData::UDisplayClusterConfigurationData()
 {
 }
 
+UDisplayClusterConfigurationClusterNode* UDisplayClusterConfigurationData::GetPrimaryNode() const
+{
+	if (Cluster)
+	{
+		return Cluster->GetNode(Cluster->PrimaryNode.Id);
+	}
+
+	return nullptr;
+}
+
+UDisplayClusterConfigurationClusterNode* UDisplayClusterConfigurationData::GetNode(const FString& NodeId) const
+{
+	if (Cluster)
+	{
+		return Cluster->GetNode(NodeId);
+	}
+
+	return nullptr;
+}
+
 UDisplayClusterConfigurationViewport* UDisplayClusterConfigurationData::GetViewport(const FString& NodeId, const FString& ViewportId) const
 {
 	UDisplayClusterConfigurationClusterNode* Node = Cluster->GetNode(NodeId);
@@ -380,6 +400,20 @@ void UDisplayClusterConfigurationClusterNode::GetObjectsToExport(TArray<UObject*
 	SAVE_MAP(Viewports);
 }
 
+void UDisplayClusterConfigurationCluster::PostLoad()
+{
+	Super::PostLoad();
+
+#if WITH_EDITOR
+	PRAGMA_DISABLE_DEPRECATION_WARNINGS
+	if (Failover.FailoverPolicy == EDisplayClusterConfigurationFailoverPolicy::DropSecondaryNodesOnly)
+	{
+		Failover.bEnabled = true;
+	}
+	PRAGMA_ENABLE_DEPRECATION_WARNINGS
+#endif // WITH_EDITOR
+}
+
 void UDisplayClusterConfigurationCluster::GetObjectsToExport(TArray<UObject*>& OutObjects)
 {
 	Super::GetObjectsToExport(OutObjects);
@@ -433,7 +467,7 @@ void UDisplayClusterConfigurationScene::GetObjectsToExport(TArray<UObject*>& Out
 }
 
 FDisplayClusterConfigurationPrimaryNodePorts::FDisplayClusterConfigurationPrimaryNodePorts()
-	: ClusterSync(41001)
+	: ClusterSync        (41001)
 	, ClusterEventsJson  (41003)
 	, ClusterEventsBinary(41004)
 {
@@ -446,6 +480,21 @@ FDisplayClusterConfigurationClusterSync::FDisplayClusterConfigurationClusterSync
 	InputSyncPolicy.Type  = cluster::input_sync::InputSyncPolicyReplicatePrimary;
 }
 
+FDisplayClusterConfigurationNetworkSettings::FDisplayClusterConfigurationNetworkSettings()
+	: ConnectRetriesAmount     (300)
+	, ConnectRetryDelay        (1000) // 1 second between attempts
+	, GameStartBarrierTimeout  (1000 * 60 * 60) // 1 hour
+	, FrameStartBarrierTimeout (1000 * 60 * 30) // 30 minutes
+	, FrameEndBarrierTimeout   (1000 * 60 * 30) // 30 minutes
+	, RenderSyncBarrierTimeout (1000 * 60 * 30) // 30 minutes
+{
+}
+
+FDisplayClusterConfigurationFailoverSettings::FDisplayClusterConfigurationFailoverSettings()
+	: bEnabled(false)
+{
+}
+
 void UDisplayClusterConfigurationViewport::GetReferencedMeshNames(TArray<FString>& OutMeshNames) const
 {
 	// Collect all mesh references from projection policies
diff --git a/Engine/Plugins/Runtime/nDisplay/Source/DisplayClusterConfiguration/Private/Formats/JSON500/DisplayClusterConfigurationJsonParser_500.cpp b/Engine/Plugins/Runtime/nDisplay/Source/DisplayClusterConfiguration/Private/Formats/JSON500/DisplayClusterConfigurationJsonParser_500.cpp
--- a/Engine/Plugins/Runtime/nDisplay/Source/DisplayClusterConfiguration/Private/Formats/JSON500/DisplayClusterConfigurationJsonParser_500.cpp
+++ b/Engine/Plugins/Runtime/nDisplay/Source/DisplayClusterConfiguration/Private/Formats/JSON500/DisplayClusterConfigurationJsonParser_500.cpp
@@ -209,7 +209,9 @@ namespace JSON500
 
 			// Failover
 			{
-				Config->Cluster->Failover.FailoverPolicy = CfgJson.Cluster.Failover.FailoverPolicy;
+				Config->Cluster->Failover.bEnabled = CfgJson.Cluster.Failover.bEnabled;
+				Config->Cluster->Failover.PrimaryBackups.ItemNames = CfgJson.Cluster.Failover.PrimaryBackups;
+				Config->Cluster->Failover.CriticalNodes.ItemNames  = CfgJson.Cluster.Failover.CriticalNodes;
 			}
 
 			// Cluster nodes
@@ -405,7 +407,9 @@ namespace JSON500
 
 			// Failover
 			{
-				Json.Cluster.Failover.FailoverPolicy = Config->Cluster->Failover.FailoverPolicy;
+				Json.Cluster.Failover.bEnabled = Config->Cluster->Failover.bEnabled;
+				Json.Cluster.Failover.PrimaryBackups = Config->Cluster->Failover.PrimaryBackups.ItemNames;
+				Json.Cluster.Failover.CriticalNodes  = Config->Cluster->Failover.CriticalNodes.ItemNames;
 			}
 
 			// Cluster nodes
diff --git a/Engine/Plugins/Runtime/nDisplay/Source/DisplayClusterConfiguration/Private/Formats/JSON500/DisplayClusterConfigurationJsonTypes_500.h b/Engine/Plugins/Runtime/nDisplay/Source/DisplayClusterConfiguration/Private/Formats/JSON500/DisplayClusterConfigurationJsonTypes_500.h
--- a/Engine/Plugins/Runtime/nDisplay/Source/DisplayClusterConfiguration/Private/Formats/JSON500/DisplayClusterConfigurationJsonTypes_500.h
+++ b/Engine/Plugins/Runtime/nDisplay/Source/DisplayClusterConfiguration/Private/Formats/JSON500/DisplayClusterConfigurationJsonTypes_500.h
@@ -380,13 +380,18 @@ struct DISPLAYCLUSTERCONFIGURATION_API FDisplayClusterConfigurationFailoverSetti
 	GENERATED_BODY()
 
 public:
-	FDisplayClusterConfigurationFailoverSettings_500()
-		: FailoverPolicy(EDisplayClusterConfigurationFailoverPolicy::Disabled)
-	{ }
 
-public:
+	/** Enable/disable failover subsystem */
+	UPROPERTY()
+	bool bEnabled = false;
+
+	/** Prioritized nodes for primary node replacement. Used in the same order as listed. */
+	UPROPERTY()
+	TArray<FString> PrimaryBackups;
+
+	/** Critical nodes not allowed to fail. */
 	UPROPERTY()
-	EDisplayClusterConfigurationFailoverPolicy FailoverPolicy;
+	TArray<FString> CriticalNodes;
 };
 
 
diff --git a/Engine/Plugins/Runtime/nDisplay/Source/DisplayClusterConfiguration/Public/DisplayClusterConfigurationTypes.h b/Engine/Plugins/Runtime/nDisplay/Source/DisplayClusterConfiguration/Public/DisplayClusterConfigurationTypes.h
--- a/Engine/Plugins/Runtime/nDisplay/Source/DisplayClusterConfiguration/Public/DisplayClusterConfigurationTypes.h
+++ b/Engine/Plugins/Runtime/nDisplay/Source/DisplayClusterConfiguration/Public/DisplayClusterConfigurationTypes.h
@@ -233,30 +233,34 @@ struct DISPLAYCLUSTERCONFIGURATION_API FDisplayClusterConfigurationNetworkSettin
 {
 	GENERATED_BODY()
 
+public:
+
+	FDisplayClusterConfigurationNetworkSettings();
+
 public:
 	/** Advanced: amount of times nDisplay tries to reconnect before dropping */
 	UPROPERTY(EditAnywhere, BlueprintReadWrite, Category = NDisplay, meta = (ClampMin = "1", ClampMax = "500", UIMin = "1", UIMax = "500"))
-	int32 ConnectRetriesAmount = 300;
+	int32 ConnectRetriesAmount;
 
 	/** Advanced: delay in between connection retries */
 	UPROPERTY(EditAnywhere, BlueprintReadWrite, Category = NDisplay, meta = (ClampMin = "1", ClampMax = "5000", UIMin = "1", UIMax = "5000"))
-	int32 ConnectRetryDelay = 1000; // 1 second between attempts
+	int32 ConnectRetryDelay;
 
 	/** Advanced: timeout for Game Thread Barrier */
 	UPROPERTY(EditAnywhere, BlueprintReadWrite, Category = NDisplay, meta = (ClampMin = "5000", UIMin = "5000"))
-	int32 GameStartBarrierTimeout = 1000 * 3600 * 5; // 5 hours
+	int32 GameStartBarrierTimeout;
 
 	/** Advanced: timeout value for Start Frame Barrier */
 	UPROPERTY(EditAnywhere, BlueprintReadWrite, Category = NDisplay, meta = (ClampMin = "1", UIMin = "1"))
-	int32 FrameStartBarrierTimeout = 1000 * 60 * 30; // 30 minutes
+	int32 FrameStartBarrierTimeout;
 
 	/** Advanced: timeout value for End Frame Barrier */
 	UPROPERTY(EditAnywhere, BlueprintReadWrite, Category = NDisplay, meta = (ClampMin = "1", UIMin = "1"))
-	int32 FrameEndBarrierTimeout = 1000 * 60 * 30; // 30 minutes
+	int32 FrameEndBarrierTimeout;
 
 	/** Advanced: timeout value for Render Sync */
 	UPROPERTY(EditAnywhere, BlueprintReadWrite, Category = NDisplay, meta = (ClampMin = "1", UIMin = "1"))
-	int32 RenderSyncBarrierTimeout = 1000 * 60 * 30; // 30 minutes
+	int32 RenderSyncBarrierTimeout;
 };
 
 USTRUCT(Blueprintable)
@@ -265,9 +269,32 @@ struct DISPLAYCLUSTERCONFIGURATION_API FDisplayClusterConfigurationFailoverSetti
 	GENERATED_BODY()
 
 public:
-	/** Failover policy */
-	UPROPERTY(EditAnywhere, BlueprintReadWrite, Category = NDisplay)
+
+	FDisplayClusterConfigurationFailoverSettings();
+
+public:
+
+	/** Enable/disable failover subsystem */
+	UPROPERTY(EditAnywhere, BlueprintReadWrite, Category = Failover, meta = (DisplayName = "Enable Failover"))
+	bool bEnabled;
+
+	/** Prioritized nodes for primary node replacement. Used in the same order as listed. */
+	UPROPERTY(EditAnywhere, BlueprintReadWrite, Category = Failover, meta = (DisplayName = "Prioritized Primary Backups", EditCondition = "bEnabled", ClusterItemType = ClusterNodes))
+	FDisplayClusterConfigurationClusterItemReferenceList PrimaryBackups;
+
+	/** Critical cluster nodes. None of these are allowed to fail. If even one fails, the cluster will self-terminate. */
+	UPROPERTY(EditAnywhere, BlueprintReadWrite, Category = Failover, meta = (DisplayName = "Critical Nodes", EditCondition = "bEnabled", ClusterItemType = ClusterNodes))
+	FDisplayClusterConfigurationClusterItemReferenceList CriticalNodes;
+
+public:
+
+#if WITH_EDITORONLY_DATA
+	PRAGMA_DISABLE_DEPRECATION_WARNINGS
+	UE_DEPRECATED(5.6, "This property has been deprecated.")
+	UPROPERTY(meta = (DeprecatedProperty, DeprecationMessage = "This property has been deprecated."))
 	EDisplayClusterConfigurationFailoverPolicy FailoverPolicy = EDisplayClusterConfigurationFailoverPolicy::Disabled;
+	PRAGMA_ENABLE_DEPRECATION_WARNINGS
+#endif // WITH_EDITORONLY_DATA
 };
 
 USTRUCT()
@@ -469,6 +496,12 @@ public:
 	TMap<FString, TObjectPtr<UDisplayClusterConfigurationHostDisplayData>> HostDisplayData;
 #endif
 
+public:
+
+	//~ Begin UObject interface
+	virtual void PostLoad() override;
+	//~ End UObject interface
+
 protected:
 	virtual void GetObjectsToExport(TArray<UObject*>& OutObjects) override;
 
@@ -519,6 +552,15 @@ public:
 	UDisplayClusterConfigurationData();
 
 public:
+
+	/** Returns primary node configuration. Otherwise nullptr. */
+	UFUNCTION(BlueprintPure, Category = "NDisplay|Configuration")
+	UDisplayClusterConfigurationClusterNode* GetPrimaryNode() const;
+
+	/** Returns specific node configuration. Otherwise nullptr. */
+	UFUNCTION(BlueprintPure, Category = "NDisplay|Configuration")
+	UDisplayClusterConfigurationClusterNode* GetNode(const FString& NodeId) const;
+
 	// Viewports API
 	UFUNCTION(BlueprintPure, Category = "NDisplay|Configuration")
 	UDisplayClusterConfigurationViewport* GetViewport(const FString& NodeId, const FString& ViewportId) const;
