[nDisplay] Base failover implementation (v1.5, P-node replacement)

diff --git a/Engine/Plugins/Runtime/nDisplay/Source/DisplayClusterConfigurator/Private/Views/Details/Cluster/DisplayClusterConfiguratorNodeSelectionCustomization.cpp b/Engine/Plugins/Runtime/nDisplay/Source/DisplayClusterConfigurator/Private/Views/Details/Cluster/DisplayClusterConfiguratorNodeSelectionCustomization.cpp
--- a/Engine/Plugins/Runtime/nDisplay/Source/DisplayClusterConfigurator/Private/Views/Details/Cluster/DisplayClusterConfiguratorNodeSelectionCustomization.cpp
+++ b/Engine/Plugins/Runtime/nDisplay/Source/DisplayClusterConfigurator/Private/Views/Details/Cluster/DisplayClusterConfiguratorNodeSelectionCustomization.cpp
@@ -68,7 +68,6 @@ UDisplayClusterConfigurationData* FDisplayClusterConfiguratorNodeSelection::GetC
 		ConfigData = RootActorPtr->GetConfigData();
 	}
 
-	check(ConfigData);
 	return ConfigData;
 }
 
diff --git a/Engine/Plugins/Runtime/nDisplay/Source/DisplayClusterMedia/Private/DisplayClusterMediaHelpers.cpp b/Engine/Plugins/Runtime/nDisplay/Source/DisplayClusterMedia/Private/DisplayClusterMediaHelpers.cpp
--- a/Engine/Plugins/Runtime/nDisplay/Source/DisplayClusterMedia/Private/DisplayClusterMediaHelpers.cpp
+++ b/Engine/Plugins/Runtime/nDisplay/Source/DisplayClusterMedia/Private/DisplayClusterMediaHelpers.cpp
@@ -24,12 +24,12 @@ namespace DisplayClusterMediaHelpers
 
 				case EMediaOwnerType::Viewport:
 					return InTilePos
-						? FString::Printf(TEXT("%s_%s_%s_tile_%d_%d_viewport_input"), *NodeId, *DCRAName, *OwnerName, InTilePos->X, InTilePos->Y)
+						? FString::Printf(TEXT("%s_%s_%s_viewport_tile_%d_%d_input"), *NodeId, *DCRAName, *OwnerName, InTilePos->X, InTilePos->Y)
 						: FString::Printf(TEXT("%s_%s_%s_viewport_input"), *NodeId, *DCRAName, *OwnerName);
 
 				case EMediaOwnerType::ICVFXCamera:
 					return InTilePos
-						? FString::Printf(TEXT("%s_%s_%s_tile_%d_%d_icvfx_input"), *NodeId, *DCRAName, *OwnerName, InTilePos->X, InTilePos->Y)
+						? FString::Printf(TEXT("%s_%s_%s_icvfx_tile_%d_%d_input"), *NodeId, *DCRAName, *OwnerName, InTilePos->X, InTilePos->Y)
 						: FString::Printf(TEXT("%s_%s_%s_icvfx_input"), *NodeId, *DCRAName, *OwnerName);
 
 				default:
@@ -42,17 +42,17 @@ namespace DisplayClusterMediaHelpers
 				{
 				case EMediaOwnerType::Backbuffer:
 					return InTilePos
-						? FString::Printf(TEXT("%s_%s_tile_%d_%d_backbuffer_capture_%u"), *NodeId, *DCRAName, InTilePos->X, InTilePos->Y, Index)
+						? FString::Printf(TEXT("%s_%s_backbuffer_tile_%d_%d_capture_%u"), *NodeId, *DCRAName, InTilePos->X, InTilePos->Y, Index)
 						: FString::Printf(TEXT("%s_%s_backbuffer_capture_%u"), *NodeId, *DCRAName, Index);
 
 				case EMediaOwnerType::Viewport:
 					return InTilePos
-						? FString::Printf(TEXT("%s_%s_%s_tile_%d_%d_viewport_capture_%u"), *NodeId, *DCRAName, *OwnerName, InTilePos->X, InTilePos->Y, Index)
+						? FString::Printf(TEXT("%s_%s_%s_viewport_tile_%d_%d_capture_%u"), *NodeId, *DCRAName, *OwnerName, InTilePos->X, InTilePos->Y, Index)
 						: FString::Printf(TEXT("%s_%s_%s_viewport_capture_%u"), *NodeId, *DCRAName, *OwnerName, Index);
 
 				case EMediaOwnerType::ICVFXCamera:
 					return InTilePos
-						? FString::Printf(TEXT("%s_%s_%s_tile_%d_%d_icvfx_capture_%u"), *NodeId, *DCRAName, *OwnerName, InTilePos->X, InTilePos->Y, Index)
+						? FString::Printf(TEXT("%s_%s_%s_icvfx_tile_%d_%d_capture_%u"), *NodeId, *DCRAName, *OwnerName, InTilePos->X, InTilePos->Y, Index)
 						: FString::Printf(TEXT("%s_%s_%s_icvfx_capture_%u"), *NodeId, *DCRAName, *OwnerName, Index);;
 
 				default:
diff --git a/Engine/Plugins/Runtime/nDisplay/Source/DisplayClusterMedia/Private/DisplayClusterMediaModule.cpp b/Engine/Plugins/Runtime/nDisplay/Source/DisplayClusterMedia/Private/DisplayClusterMediaModule.cpp
--- a/Engine/Plugins/Runtime/nDisplay/Source/DisplayClusterMedia/Private/DisplayClusterMediaModule.cpp
+++ b/Engine/Plugins/Runtime/nDisplay/Source/DisplayClusterMedia/Private/DisplayClusterMediaModule.cpp
@@ -213,9 +213,9 @@ void FDisplayClusterMediaModule::InitializeBackbufferFullFrameOutput(const UDisp
 						MediaOutputItem.OutputSyncPolicy);
 
 					CaptureDevices.Emplace(MediaCaptureId, MoveTemp(NewNodeCapture));
-				}
 
-				++CaptureIdx;
+					++CaptureIdx;
+				}
 			}
 		}
 	}
@@ -263,9 +263,9 @@ void FDisplayClusterMediaModule::InitializeBackbufferUniformTilesOutput(const UD
 						MediaOutputTile.OutputSyncPolicy);
 
 					CaptureDevices.Emplace(MediaCaptureId, NewNodeCapture);
-				}
 
-				++CaptureIdx;
+					++CaptureIdx;
+				}
 			}
 		}
 	}
@@ -327,9 +327,9 @@ void FDisplayClusterMediaModule::InitializeViewportOutput(const UDisplayClusterC
 					MediaOutputItem.OutputSyncPolicy);
 
 				CaptureDevices.Emplace(MediaCaptureId, MoveTemp(NewViewportCapture));
-			}
 
-			++CaptureIdx;
+				++CaptureIdx;
+			}
 		}
 	}
 }
@@ -397,9 +397,9 @@ void FDisplayClusterMediaModule::InitializeICVFXCameraFullFrameOutput(const UDis
 						MediaOutputItem.OutputSyncPolicy);
 
 					CaptureDevices.Emplace(MediaCaptureId, MoveTemp(NewICVFXCapture));
-				}
 
-				++CaptureIdx;
+					++CaptureIdx;
+				}
 			}
 		}
 	}
@@ -443,9 +443,9 @@ void FDisplayClusterMediaModule::InitializeICVFXCameraUniformTilesInput(const UD
 							MediaInputTile.MediaSource);
 
 						InputDevices.Emplace(MediaInputId, MoveTemp(NewICVFXTileInput));
-					}
 
-					++Index;
+						++Index;
+					}
 				}
 			}
 		}
@@ -488,12 +488,12 @@ void FDisplayClusterMediaModule::InitializeICVFXCameraUniformTilesOutput(const U
 							ICVFXCameraName,
 							MediaOutputTile.Position,
 							MediaOutputTile.MediaOutput,
-							nullptr); // no sync for tiles as it's basically for in-cluster use
+							MediaOutputTile.OutputSyncPolicy);
 
 						CaptureDevices.Emplace(MediaOutputId, MoveTemp(NewICVFXTileOutput));
-					}
 
-					++Index;
+						++Index;
+					}
 				}
 			}
 		}
diff --git a/Engine/Plugins/Runtime/nDisplay/Source/DisplayClusterMedia/Private/Synchronization/SyncPolicies/DisplayClusterMediaOutputSynchronizationPolicyEthernetBarrierBase.cpp b/Engine/Plugins/Runtime/nDisplay/Source/DisplayClusterMedia/Private/Synchronization/SyncPolicies/DisplayClusterMediaOutputSynchronizationPolicyEthernetBarrierBase.cpp
--- a/Engine/Plugins/Runtime/nDisplay/Source/DisplayClusterMedia/Private/Synchronization/SyncPolicies/DisplayClusterMediaOutputSynchronizationPolicyEthernetBarrierBase.cpp
+++ b/Engine/Plugins/Runtime/nDisplay/Source/DisplayClusterMedia/Private/Synchronization/SyncPolicies/DisplayClusterMediaOutputSynchronizationPolicyEthernetBarrierBase.cpp
@@ -2,21 +2,19 @@
 
 #include "Synchronization/DisplayClusterMediaOutputSynchronizationPolicyEthernetBarrierBase.h"
 
-#include "DisplayClusterMediaHelpers.h"
-#include "DisplayClusterMediaLog.h"
-
-#include "MediaCapture.h"
-
 #include "DisplayClusterConfigurationTypes.h"
 #include "DisplayClusterEnums.h"
+#include "DisplayClusterMediaHelpers.h"
+#include "DisplayClusterMediaLog.h"
 #include "DisplayClusterRootActor.h"
 
-#include "Components/DisplayClusterICVFXCameraComponent.h"
-
 #include "IDisplayCluster.h"
+#include "MediaCapture.h"
+
 #include "Cluster/IDisplayClusterClusterManager.h"
+#include "Cluster/IDisplayClusterGenericBarriersClient.h"
+#include "Components/DisplayClusterICVFXCameraComponent.h"
 #include "Game/IDisplayClusterGameManager.h"
-
 #include "Templates/SharedPointer.h"
 
 
@@ -107,9 +105,9 @@ FString FDisplayClusterMediaOutputSynchronizationPolicyEthernetBarrierBaseHandle
 	return MediaDeviceId;
 }
 
-IDisplayClusterGenericBarriersClient* const FDisplayClusterMediaOutputSynchronizationPolicyEthernetBarrierBaseHandler::GetBarrierClient() const
+TSharedPtr<IDisplayClusterGenericBarriersClient> FDisplayClusterMediaOutputSynchronizationPolicyEthernetBarrierBaseHandler::GetBarrierClient()
 {
-	return EthernetBarrierClient.Get();
+	return EthernetBarrierClient;
 }
 
 const FString& FDisplayClusterMediaOutputSynchronizationPolicyEthernetBarrierBaseHandler::GetBarrierId() const
@@ -125,7 +123,7 @@ const FString& FDisplayClusterMediaOutputSynchronizationPolicyEthernetBarrierBas
 void FDisplayClusterMediaOutputSynchronizationPolicyEthernetBarrierBaseHandler::SyncThreadOnBarrier()
 {
 	// Sync on the barrier if everything is good
-	if (bIsRunning && EthernetBarrierClient && EthernetBarrierClient->IsConnected())
+	if (bIsRunning && EthernetBarrierClient)
 	{
 		UE_LOG(LogDisplayClusterMediaSync, Verbose, TEXT("'%s': Synchronizing caller '%s' at the barrier '%s'"), *GetMediaDeviceId(), *ThreadMarker, *BarrierId);
 		EthernetBarrierClient->Synchronize(BarrierId, ThreadMarker);
@@ -140,39 +138,25 @@ bool FDisplayClusterMediaOutputSynchronizationPolicyEthernetBarrierBaseHandler::
 		return false;
 	}
 
-	IDisplayClusterClusterManager& ClusterMgr = *IDisplayCluster::Get().GetClusterMgr();
-
-	// Barriers client ID
-	const FString ClientId = FString::Printf(TEXT("CLN_GB_%s"), *MediaId);
-	UE_LOG(LogDisplayClusterMediaSync, VeryVerbose, TEXT("'%s': Requesting barrier client '%s'"), *MediaId, *ClientId);
-
 	// Instantiate barrier client
-	EthernetBarrierClient = IDisplayCluster::Get().GetClusterMgr()->CreateGenericBarriersClient(ClientId);
+	EthernetBarrierClient = IDisplayCluster::Get().GetClusterMgr()->CreateGenericBarriersClient();
 	if (!EthernetBarrierClient)
 	{
-		UE_LOG(LogDisplayClusterMediaSync, Warning, TEXT("'%s': Couldn't instantiate barrier client '%s'"), *MediaId, *ClientId);
-		return false;
-	}
-
-	UE_LOG(LogDisplayClusterMediaSync, Verbose, TEXT("'%s': Instantiated barrier client '%s'"), *MediaId, *ClientId);
-
-	if (!EthernetBarrierClient->Connect())
-	{
-		UE_LOG(LogDisplayClusterMediaSync, Warning, TEXT("'%s': Barrier client '%s' couldn't establish connection"), *MediaId, *ClientId);
+		UE_LOG(LogDisplayClusterMediaSync, Warning, TEXT("'%s': Couldn't get generic barriers API"), *MediaId);
 		return false;
 	}
 
 	BarrierId     = GenerateBarrierName();
 	ThreadMarker  = MediaId;
 
-	// Thread markers
-	TArray<FString> ThreadMarkers;
-	GenerateListOfThreadMarkers(ThreadMarkers);
+	// Sync callers
+	TMap<FString, TSet<FString>> NodeToSyncCallers;
+	GenerateSyncCallersMapping(NodeToSyncCallers);
 
 	// Create sync barrier
-	if (!EthernetBarrierClient->CreateBarrier(BarrierId, ThreadMarkers, BarrierTimeoutMs))
+	if (!EthernetBarrierClient->CreateBarrier(BarrierId, NodeToSyncCallers, BarrierTimeoutMs))
 	{
-		UE_LOG(LogDisplayClusterMediaSync, Warning, TEXT("'%s': Barrier client '%s' couldn't create barrier '%s'"), *MediaId, *ClientId, *BarrierId);
+		UE_LOG(LogDisplayClusterMediaSync, Warning, TEXT("'%s': Couldn't create barrier '%s'"), *MediaId, *BarrierId);
 		EthernetBarrierClient.Reset();
 		return false;
 	}
@@ -184,7 +168,7 @@ void FDisplayClusterMediaOutputSynchronizationPolicyEthernetBarrierBaseHandler::
 {
 	if (EthernetBarrierClient)
 	{
-		if (EthernetBarrierClient->IsConnected() && !BarrierId.IsEmpty())
+		if (!BarrierId.IsEmpty())
 		{
 			EthernetBarrierClient->ReleaseBarrier(BarrierId);
 		}
@@ -207,11 +191,11 @@ FString FDisplayClusterMediaOutputSynchronizationPolicyEthernetBarrierBaseHandle
 	return GetPolicyClass()->GetName();
 }
 
-void FDisplayClusterMediaOutputSynchronizationPolicyEthernetBarrierBaseHandler::GenerateListOfThreadMarkers(TArray<FString>& OutMarkers) const
+void FDisplayClusterMediaOutputSynchronizationPolicyEthernetBarrierBaseHandler::GenerateSyncCallersMapping(TMap<FString, TSet<FString>>& OutNodeToSyncCallers) const
 {
-	OutMarkers.Empty();
+	OutNodeToSyncCallers.Empty();
 
-	UE_LOG(LogDisplayClusterMediaSync, Verbose, TEXT("'%s': Generating thread markers for barrier '%s'..."), *GetMediaDeviceId(), *BarrierId);
+	UE_LOG(LogDisplayClusterMediaSync, Verbose, TEXT("'%s': Generating sync caller mappings for barrier '%s'..."), *GetMediaDeviceId(), *BarrierId);
 
 	// Get active DCRA
 	if (const ADisplayClusterRootActor* const RootActor = IDisplayCluster::Get().GetGameMgr()->GetRootActor())
@@ -229,24 +213,49 @@ void FDisplayClusterMediaOutputSynchronizationPolicyEthernetBarrierBaseHandler::
 
 					if (MediaSettings.bEnable)
 					{
-						uint8 CaptureIdx = 0;
-						for (const FDisplayClusterConfigurationMediaOutput& MediaOutputItem : MediaSettings.MediaOutputs)
+						// Iterate over full frame outputs
 						{
-							if (IsValid(MediaOutputItem.MediaOutput) && IsValid(MediaOutputItem.OutputSyncPolicy))
+							uint8 CaptureIdx = 0;
+							for (const FDisplayClusterConfigurationMediaOutput& MediaOutputItem : MediaSettings.MediaOutputs)
 							{
-								// Pick the same sync policy only
-								if (MediaOutputItem.OutputSyncPolicy->GetClass() == GetPolicyClass())
+								if (IsValid(MediaOutputItem.MediaOutput) && IsValid(MediaOutputItem.OutputSyncPolicy))
 								{
-									const FString BackbufferCaptureId = DisplayClusterMediaHelpers::MediaId::GenerateMediaId(
-										DisplayClusterMediaHelpers::MediaId::EMediaDeviceType::Output,
-										DisplayClusterMediaHelpers::MediaId::EMediaOwnerType::Backbuffer,
-										NodeIt.Key, RootActor->GetName(), FString(), CaptureIdx);
+									// Pick the same sync policy only
+									if (MediaOutputItem.OutputSyncPolicy->GetClass() == GetPolicyClass())
+									{
+										const FString BackbufferCaptureId = DisplayClusterMediaHelpers::MediaId::GenerateMediaId(
+											DisplayClusterMediaHelpers::MediaId::EMediaDeviceType::Output,
+											DisplayClusterMediaHelpers::MediaId::EMediaOwnerType::Backbuffer,
+											NodeIt.Key, RootActor->GetName(), FString(), CaptureIdx++);
 
-									OutMarkers.Add(BackbufferCaptureId);
+										OutNodeToSyncCallers.FindOrAdd(NodeIt.Key).Add(BackbufferCaptureId);
+									}
 								}
 							}
+						}
 
-							++CaptureIdx;
+						// Iterate over tiles
+						{
+							if (MediaSettings.TiledSplitLayout.X > 1 || MediaSettings.TiledSplitLayout.Y > 1)
+							{
+								uint8 CaptureIdx = 0;
+								for (const FDisplayClusterConfigurationMediaUniformTileOutput& OutputTile : MediaSettings.TiledMediaOutputs)
+								{
+									if (IsValid(OutputTile.MediaOutput) && IsValid(OutputTile.OutputSyncPolicy))
+									{
+										// Pick the same sync policy only
+										if (OutputTile.OutputSyncPolicy->GetClass() == GetPolicyClass())
+										{
+											const FString BackbufferCaptureId = DisplayClusterMediaHelpers::MediaId::GenerateMediaId(
+												DisplayClusterMediaHelpers::MediaId::EMediaDeviceType::Output,
+												DisplayClusterMediaHelpers::MediaId::EMediaOwnerType::Backbuffer,
+												NodeIt.Key, RootActor->GetName(), FString(), CaptureIdx++, &OutputTile.Position);
+
+											OutNodeToSyncCallers.FindOrAdd(NodeIt.Key).Add(BackbufferCaptureId);
+										}
+									}
+								}
+							}
 						}
 					}
 				}
@@ -272,13 +281,11 @@ void FDisplayClusterMediaOutputSynchronizationPolicyEthernetBarrierBaseHandler::
 										const FString ViewportCaptureId = DisplayClusterMediaHelpers::MediaId::GenerateMediaId(
 											DisplayClusterMediaHelpers::MediaId::EMediaDeviceType::Output,
 											DisplayClusterMediaHelpers::MediaId::EMediaOwnerType::Viewport,
-											NodeIt.Key, RootActor->GetName(), ViewportIt.Key, CaptureIdx);
+											NodeIt.Key, RootActor->GetName(), ViewportIt.Key, CaptureIdx++);
 
-										OutMarkers.Add(ViewportCaptureId);
+										OutNodeToSyncCallers.FindOrAdd(NodeIt.Key).Add(ViewportCaptureId);
 									}
 								}
-
-								++CaptureIdx;
 							}
 						}
 					}
@@ -300,37 +307,83 @@ void FDisplayClusterMediaOutputSynchronizationPolicyEthernetBarrierBaseHandler::
 
 				if (MediaSettings.bEnable)
 				{
-					uint8 CaptureIdx = 0;
-					for (const FDisplayClusterConfigurationMediaOutputGroup& MediaOutputGroup : MediaSettings.MediaOutputGroups)
+					// Full-frame camera capture
+					if (MediaSettings.SplitType == EDisplayClusterConfigurationMediaSplitType::FullFrame)
 					{
-						// Pick the same sync policy only
-						if (MediaOutputGroup.MediaOutput && MediaOutputGroup.OutputSyncPolicy)
+						uint8 CaptureIdx = 0;
+						for (const FDisplayClusterConfigurationMediaOutputGroup& MediaOutputGroup : MediaSettings.MediaOutputGroups)
 						{
-							if (MediaOutputGroup.OutputSyncPolicy->GetClass() == GetPolicyClass())
+							// Pick the same sync policy only
+							if (IsValid(MediaOutputGroup.MediaOutput) && IsValid(MediaOutputGroup.OutputSyncPolicy))
 							{
-								for (const FString& NodeId : MediaOutputGroup.ClusterNodes.ItemNames)
+								if (MediaOutputGroup.OutputSyncPolicy->GetClass() == GetPolicyClass())
 								{
-									const FString ICVFXCaptureId = DisplayClusterMediaHelpers::MediaId::GenerateMediaId(
-										DisplayClusterMediaHelpers::MediaId::EMediaDeviceType::Output,
-										DisplayClusterMediaHelpers::MediaId::EMediaOwnerType::ICVFXCamera,
-										NodeId, RootActor->GetName(), ICVFXCameraComponent->GetName(), CaptureIdx);
+									for (const FString& NodeId : MediaOutputGroup.ClusterNodes.ItemNames)
+									{
+										const FString ICVFXCaptureId = DisplayClusterMediaHelpers::MediaId::GenerateMediaId(
+											DisplayClusterMediaHelpers::MediaId::EMediaDeviceType::Output,
+											DisplayClusterMediaHelpers::MediaId::EMediaOwnerType::ICVFXCamera,
+											NodeId, RootActor->GetName(), ICVFXCameraComponent->GetName(), CaptureIdx++);
 
-									OutMarkers.Add(ICVFXCaptureId);
+										OutNodeToSyncCallers.FindOrAdd(NodeId).Add(ICVFXCaptureId);
+									}
+								}
+							}
+						}
+					}
+					else if (MediaSettings.SplitType == EDisplayClusterConfigurationMediaSplitType::UniformTiles)
+					{
+						uint8 CaptureIdx = 0;
+						for (const FDisplayClusterConfigurationMediaTiledOutputGroup& OutputGroup : MediaSettings.TiledMediaOutputGroups)
+						{
+							for (const FDisplayClusterConfigurationMediaUniformTileOutput& OutputTile : OutputGroup.Tiles)
+							{
+								// Pick the same sync policy only
+								if (IsValid(OutputTile.MediaOutput) && IsValid(OutputTile.OutputSyncPolicy))
+								{
+									if (OutputTile.OutputSyncPolicy->GetClass() == GetPolicyClass())
+									{
+										for (const FString& NodeId : OutputGroup.ClusterNodes.ItemNames)
+										{
+											const FString ICVFXCaptureId = DisplayClusterMediaHelpers::MediaId::GenerateMediaId(
+												DisplayClusterMediaHelpers::MediaId::EMediaDeviceType::Output,
+												DisplayClusterMediaHelpers::MediaId::EMediaOwnerType::ICVFXCamera,
+												NodeId, RootActor->GetName(), ICVFXCameraComponent->GetName(), CaptureIdx++, &OutputTile.Position);
+
+											OutNodeToSyncCallers.FindOrAdd(NodeId).Add(ICVFXCaptureId);
+										}
+									}
 								}
 							}
 						}
-
-						++CaptureIdx;
 					}
 				}
 			}
 		}
 	}
 
-	UE_LOG(LogDisplayClusterMediaSync, Verbose, TEXT("'%s': Generated %d thread markers for barrier '%s' "), *GetMediaDeviceId(), OutMarkers.Num(), *BarrierId);
-	for (int32 Idx = 0; Idx < OutMarkers.Num(); ++Idx)
+	LogSyncCallersMapping(OutNodeToSyncCallers);
+}
+
+void FDisplayClusterMediaOutputSynchronizationPolicyEthernetBarrierBaseHandler::LogSyncCallersMapping(const TMap<FString, TSet<FString>>& NodeToSyncCallers) const
+{
+	if (UE_GET_LOG_VERBOSITY(LogDisplayClusterMediaSync) >= ELogVerbosity::Type::Verbose)
 	{
-		UE_LOG(LogDisplayClusterMediaSync, Verbose, TEXT("'%s': Barrier '%s', marker %d: %s"), *GetMediaDeviceId(), *BarrierId, Idx, *OutMarkers[Idx]);
+		FString LogMsg;
+		LogMsg.Reserve(2048);
+
+		LogMsg = FString::Printf(TEXT("'%s': Generated the following NodeToCallers mapping:\n"), *GetMediaDeviceId());
+		for (const TPair<FString, TSet<FString>>& NodeMapping : NodeToSyncCallers)
+		{
+			LogMsg += FString::Printf(TEXT(" [%s] "), *NodeMapping.Key);
+			for (const FString& CallerId : NodeMapping.Value)
+			{
+				LogMsg += FString::Printf(TEXT("%s, "), *CallerId);
+			}
+			LogMsg += TEXT("\n");
+		}
+
+		UE_LOG(LogDisplayClusterMediaSync, Verbose, TEXT("%s"), *LogMsg);
 	}
 }
 
diff --git a/Engine/Plugins/Runtime/nDisplay/Source/DisplayClusterMedia/Public/Synchronization/DisplayClusterMediaOutputSynchronizationPolicyEthernetBarrierBase.h b/Engine/Plugins/Runtime/nDisplay/Source/DisplayClusterMedia/Public/Synchronization/DisplayClusterMediaOutputSynchronizationPolicyEthernetBarrierBase.h
--- a/Engine/Plugins/Runtime/nDisplay/Source/DisplayClusterMedia/Public/Synchronization/DisplayClusterMediaOutputSynchronizationPolicyEthernetBarrierBase.h
+++ b/Engine/Plugins/Runtime/nDisplay/Source/DisplayClusterMedia/Public/Synchronization/DisplayClusterMediaOutputSynchronizationPolicyEthernetBarrierBase.h
@@ -4,10 +4,11 @@
 
 #include "DisplayClusterConfigurationTypes_MediaSync.h"
 
-#include "Cluster/IDisplayClusterGenericBarriersClient.h"
-
 #include "DisplayClusterMediaOutputSynchronizationPolicyEthernetBarrierBase.generated.h"
 
+class IDisplayClusterGenericBarriersClient;
+
+
 /*
  * Synchronization logic handler class for UDisplayClusterMediaOutputSynchronizationPolicyEthernetBarrierBase.
  */
@@ -18,6 +19,8 @@ public:
 
 	FDisplayClusterMediaOutputSynchronizationPolicyEthernetBarrierBaseHandler(UDisplayClusterMediaOutputSynchronizationPolicyEthernetBarrierBase* InPolicyObject);
 
+public:
+
 	//~ Begin IDisplayClusterMediaOutputSynchronizationPolicyHandler interface
 	virtual bool StartSynchronization(UMediaCapture* MediaCapture, const FString& MediaId) override;
 	virtual void StopSynchronization() override;
@@ -25,6 +28,8 @@ public:
 	virtual bool IsCaptureTypeSupported(UMediaCapture* MediaCapture) const override;
 	//~ End IDisplayClusterMediaOutputSynchronizationPolicyHandler interface
 
+public:
+
 	/** Children implement their own sync approaches. */
 	virtual void Synchronize() = 0;
 
@@ -40,7 +45,14 @@ protected:
 	FString GetMediaDeviceId() const;
 
 	/** Returns barrier client created for this sync policy */
-	IDisplayClusterGenericBarriersClient* const GetBarrierClient() const;
+	UE_DEPRECATED(5.6, "This method has been deprecated. Please use another 'GetBarrierClient' that returns TSharedPtr.")
+	IDisplayClusterGenericBarriersClient* const GetBarrierClient() const
+	{
+		return EthernetBarrierClient.Get();
+	}
+
+	/** Returns barrier client created for this sync policy */
+	TSharedPtr<IDisplayClusterGenericBarriersClient> GetBarrierClient();
 
 	/** Get Barrier ID for this sync policy */
 	const FString& GetBarrierId() const;
@@ -56,8 +68,15 @@ private:
 	/** Generates name of the dynamic barrier. */
 	FString GenerateBarrierName() const;
 
-	/** Generates array of thread markers that are going to use a barrier. */
-	void GenerateListOfThreadMarkers(TArray<FString>& OutMarkers) const;
+	/** Generates per-node sync callers mapping. */
+	void GenerateSyncCallersMapping(TMap<FString, TSet<FString>>& OutMarkers) const;
+
+	/** Auxiliary function to log custom barrier callers mapping. */
+	void LogSyncCallersMapping(const TMap<FString, TSet<FString>>& OutNodeToSyncCallers) const;
+
+	UE_DEPRECATED(5.6, "This method has been deprecated. Please use 'GenerateSyncCallersMapping'.")
+	void GenerateListOfThreadMarkers(TArray<FString>& OutMarkers) const
+	{ }
 
 	/** Handles media capture sync callbacks. */
 	void ProcessMediaSynchronizationCallback();
@@ -83,8 +102,8 @@ private:
 	/** Unique thread (caller) marker to be used on the barrier. */
 	FString ThreadMarker;
 
-	/** Barrier sync client. */
-	TUniquePtr<IDisplayClusterGenericBarriersClient, FDisplayClusterGenericBarriersClientDeleter> EthernetBarrierClient;
+	/** Generic barriers API. */
+	TSharedPtr<IDisplayClusterGenericBarriersClient> EthernetBarrierClient;
 };
 
 /*
diff --git a/Engine/Plugins/Runtime/nDisplay/Source/DisplayClusterMessageInterceptor/Private/DisplayClusterMessageInterceptionModule.cpp b/Engine/Plugins/Runtime/nDisplay/Source/DisplayClusterMessageInterceptor/Private/DisplayClusterMessageInterceptionModule.cpp
--- a/Engine/Plugins/Runtime/nDisplay/Source/DisplayClusterMessageInterceptor/Private/DisplayClusterMessageInterceptionModule.cpp
+++ b/Engine/Plugins/Runtime/nDisplay/Source/DisplayClusterMessageInterceptor/Private/DisplayClusterMessageInterceptionModule.cpp
@@ -5,12 +5,12 @@
 
 #include "Cluster/DisplayClusterClusterEvent.h"
 #include "Cluster/IDisplayClusterClusterManager.h"
-#include "Cluster/IDisplayClusterClusterSyncObject.h"
 #include "DisplayClusterGameEngine.h"
 #include "DisplayClusterMessageInterceptor.h"
 #include "DisplayClusterMessageInterceptionSettings.h"
 #include "Engine/Engine.h"
 #include "HAL/IConsoleManager.h"
+#include "Templates/Atomic.h"
 #include "IDisplayCluster.h"
 #include "IDisplayClusterCallbacks.h"
 #include "IMessageBus.h"
@@ -27,7 +27,9 @@
 #if WITH_EDITOR
 	#include "ISettingsModule.h"
 	#include "ISettingsSection.h"
-#endif 
+#endif
+
+#include <atomic>
 
 #define LOCTEXT_NAMESPACE "DisplayClusterInterception"
 
@@ -45,11 +47,13 @@ class FSetBoolOnPreTick
 public:
 	operator bool() const
 	{
+		FScopeLock Lock(&InternalsCS);
 		return bValue;
 	}
 
 	FSetBoolOnPreTick& operator=(bool bInValue)
 	{
+		FScopeLock Lock(&InternalsCS);
 		if (bInValue != bValue)
 		{
 			bPendingValue = bInValue;
@@ -59,12 +63,16 @@ public:
 
 	void ResetWithValue(bool bInValue)
 	{
+		FScopeLock Lock(&InternalsCS);
+
 		bValue = bInValue;
 		bPendingValue = {};
 	}
 
 	void UpdateOnPreTick()
 	{
+		FScopeLock Lock(&InternalsCS);
+
 		if (bPendingValue)
 		{
 			bValue = *bPendingValue;
@@ -75,6 +83,7 @@ public:
 private:
 	bool bValue = false;
 	TOptional<bool> bPendingValue = {};
+	mutable FCriticalSection InternalsCS;
 };
 /**
  * Display Cluster Message Interceptor module
@@ -155,6 +164,7 @@ public:
 			IDisplayCluster::Get().GetCallbacks().OnDisplayClusterStartScene().RemoveAll(this);
 			IDisplayCluster::Get().GetCallbacks().OnDisplayClusterEndSession().RemoveAll(this);
 			IDisplayCluster::Get().GetCallbacks().OnDisplayClusterPreTick().RemoveAll(this);
+			IDisplayCluster::Get().GetCallbacks().OnDisplayClusterFailoverNodeDown().RemoveAll(this);
 		}
 		Interceptor.Reset();
 		StartMessageSyncCommand.Reset();
@@ -355,6 +365,7 @@ private:
 			// Register cluster session events
 			IDisplayCluster::Get().GetCallbacks().OnDisplayClusterEndSession().AddRaw(this, &FDisplayClusterMessageInterceptionModule::StopInterception);
 			IDisplayCluster::Get().GetCallbacks().OnDisplayClusterPreTick().AddRaw(this, &FDisplayClusterMessageInterceptionModule::HandleClusterPreTick);
+			IDisplayCluster::Get().GetCallbacks().OnDisplayClusterFailoverNodeDown().AddRaw(this, &FDisplayClusterMessageInterceptionModule::HandleClusterNodeFailed);
 
 			SetupForMultiUser();
 		}
@@ -422,9 +433,25 @@ private:
 		}
 	}
 
+	void HandleClusterNodeFailed(const FString& NodeId)
+	{
+		UE_LOG(LogDisplayClusterInterception, Log, TEXT("Handling '%s' node failure..."), *NodeId);
+
+		HandleWorkspaceNodeFailure(NodeId);
+		HandlePackageNodeFailure(NodeId);
+
+		if (Interceptor)
+		{
+			Interceptor->HandleClusterNodeFailure(NodeId);
+		}
+	}
+
 	void HandlePackageEvent(const FDisplayClusterClusterEventJson& InEvent)
 	{
 		UE_LOG(LogDisplayClusterInterception, Display, TEXT("Handle multi-user package event sync with id %s -> %s."), *InEvent.Type, *InEvent.Name);
+
+		FScopeLock Lock(&PackageActivitiesCS);
+
 		TSet<FString>& NodesReceivedPackage = PackageActivities.FindOrAdd(InEvent.Type);
 		NodesReceivedPackage.Add(InEvent.Name);
 
@@ -435,17 +462,56 @@ private:
 			PackageActivities.Remove(InEvent.Type);
 			if (PackageActivities.Num() == 0)
 			{
-				UE_LOG(LogDisplayClusterInterception, Display, TEXT("Releasing the package lock."), *InEvent.Type);
+				UE_LOG(LogDisplayClusterInterception, Display, TEXT("Releasing the package lock."));
 				bAllowPackages = true;
 			}
 		}
 	}
 
+
+	void HandlePackageNodeFailure(const FString& NodeId)
+	{
+		FScopeLock Lock(&PackageActivitiesCS);
+
+		TSet<FString> ActivitiesToClean;
+
+		const int32 NodesAmount = (int32)IDisplayCluster::Get().GetClusterMgr()->GetNodesAmount();
+
+		// Remove this node from the awaiting list of all activities
+		for (TPair<FString, TSet<FString>>& PackageActivity : PackageActivities)
+		{
+			if (PackageActivity.Value.Remove(NodeId) > 0)
+			{
+				if (PackageActivity.Value.Num() >= NodesAmount)
+				{
+					UE_LOG(LogDisplayClusterInterception, Display, TEXT("All nodes have received package event %s."), *PackageActivity.Key);
+					ActivitiesToClean.Add(PackageActivity.Key);
+				}
+			}
+		}
+
+		// Forget activities that have been received already by all nodes
+		for (const FString& ActivityId : ActivitiesToClean)
+		{
+			PackageActivities.Remove(ActivityId);
+		}
+
+		// Release lock
+		if (PackageActivities.Num() == 0)
+		{
+			UE_LOG(LogDisplayClusterInterception, Display, TEXT("Releasing the package lock."));
+			bAllowPackages = true;
+		}
+	}
+
 	void HandleWorkspaceSyncEvent(const FDisplayClusterClusterEventJson& InEvent)
 	{
 		UE_LOG(LogDisplayClusterInterception, Display, TEXT("Handle multi-user workspace sync -> %s."), *InEvent.Name);
 
 		IDisplayClusterClusterManager* ClusterManager = IDisplayCluster::Get().GetClusterMgr();
+
+		FScopeLock Lock(&NodesReadyCS);
+
 		NodesReady.Add(InEvent.Name);
 		if (NodesReady.Num() >= (int32)ClusterManager->GetNodesAmount())
 		{
@@ -454,6 +520,20 @@ private:
 		}
 	}
 
+	void HandleWorkspaceNodeFailure(const FString& NodeId)
+	{
+		FScopeLock Lock(&NodesReadyCS);
+
+		if (NodesReady.Remove(NodeId) > 0)
+		{
+			if (NodesReady.Num() >= (int32)IDisplayCluster::Get().GetClusterMgr()->GetNodesAmount())
+			{
+				UE_LOG(LogDisplayClusterInterception, Display, TEXT("Allowing multi-user workspace sync on next pre-tick."));
+				CanFinalizeWorkspace = true;
+			}
+		}
+	}
+
 	void HandleMessageInterceptorSetupEvent(const FDisplayClusterClusterEventJson& InEvent)
 	{
 		if (Interceptor)
@@ -486,9 +566,11 @@ private:
 	FSetBoolOnPreTick CanFinalizeWorkspace;
 
 	TSet<FString> NodesReady;
+	FCriticalSection NodesReadyCS;
 
-	bool bAllowPackages = true;
+	std::atomic<bool> bAllowPackages = true;
 	TMap<FString, TSet<FString>> PackageActivities;
+	FCriticalSection PackageActivitiesCS;
 
 	/** MessageBus interceptor */
 	TSharedPtr<FDisplayClusterMessageInterceptor, ESPMode::ThreadSafe> Interceptor;
diff --git a/Engine/Plugins/Runtime/nDisplay/Source/DisplayClusterMessageInterceptor/Private/DisplayClusterMessageInterceptor.cpp b/Engine/Plugins/Runtime/nDisplay/Source/DisplayClusterMessageInterceptor/Private/DisplayClusterMessageInterceptor.cpp
--- a/Engine/Plugins/Runtime/nDisplay/Source/DisplayClusterMessageInterceptor/Private/DisplayClusterMessageInterceptor.cpp
+++ b/Engine/Plugins/Runtime/nDisplay/Source/DisplayClusterMessageInterceptor/Private/DisplayClusterMessageInterceptor.cpp
@@ -177,6 +177,50 @@ void FDisplayClusterMessageInterceptor::HandleClusterEvent(const FDisplayCluster
 	}
 }
 
+void FDisplayClusterMessageInterceptor::HandleClusterNodeFailure(const FString& NodeId)
+{
+	TArray<TSharedPtr<IMessageContext, ESPMode::ThreadSafe>> ContextToForward;
+
+	{
+		FScopeLock Lock(&ContextQueueCS);
+
+		const uint32 NodesInCluster = ClusterManager->GetNodesAmount();
+
+		// Will contain the message IDs that we can forward to process and forget about
+		TSet<FString> ContextsToClear;
+
+		// Remove this node ID from all the message contexts
+		for (TPair<FString, FContextSync>& MessageContext : ContextMap)
+		{
+			MessageContext.Value.NodesReceived.Remove(NodeId);
+
+			// If this node ID was the only one missing, then forward this message to process
+			const uint32 NodesReceivedNum = static_cast<uint32>(MessageContext.Value.NodesReceived.Num());
+			if (NodesReceivedNum >= NodesInCluster)
+			{
+				UE_LOG(LogDisplayClusterInterception, VeryVerbose, TEXT("Fowarding message for message id %s on frame %d"), *MessageContext.Key, GFrameCounter);
+				ContextToForward.Add(MessageContext.Value.ContextPtr);
+				ContextsToClear.Add(MessageContext.Key);
+			}
+		}
+
+		// Now remove the contexts we don't need anymore
+		for (const FString& MessageId : ContextsToClear)
+		{
+			ContextMap.Remove(MessageId);
+		}
+	}
+
+	// Forward all selected messages to process
+	if (InterceptedBus)
+	{
+		for (const auto& Context : ContextToForward)
+		{
+			InterceptedBus->Forward(Context.ToSharedRef(), Context->GetRecipients(), FTimespan::Zero(), AsShared());
+		}
+	}
+}
+
 FName FDisplayClusterMessageInterceptor::GetDebugName() const
 {
 	return FName("DisplayClusterInterceptor");
diff --git a/Engine/Plugins/Runtime/nDisplay/Source/DisplayClusterMessageInterceptor/Private/DisplayClusterMessageInterceptor.h b/Engine/Plugins/Runtime/nDisplay/Source/DisplayClusterMessageInterceptor/Private/DisplayClusterMessageInterceptor.h
--- a/Engine/Plugins/Runtime/nDisplay/Source/DisplayClusterMessageInterceptor/Private/DisplayClusterMessageInterceptor.h
+++ b/Engine/Plugins/Runtime/nDisplay/Source/DisplayClusterMessageInterceptor/Private/DisplayClusterMessageInterceptor.h
@@ -44,6 +44,7 @@ public:
 	
 	void SyncMessages();
 	void HandleClusterEvent(const FDisplayClusterClusterEventJson& InEvent);
+	void HandleClusterNodeFailure(const FString& NodeId);
 
 	// IMessageInterceptor interface
 	virtual FName GetDebugName() const override;
@@ -92,7 +93,7 @@ private:
 		TSharedPtr<IMessageContext, ESPMode::ThreadSafe> ContextPtr;
 	};
 
-	/** The intercepted message map. */
+	/** The intercepted message map (MessageId-to-MessageContext). */
 	TMap<FString, FContextSync> ContextMap;
 };
 
diff --git a/Engine/Plugins/VirtualProduction/Rivermax/RivermaxSync/Source/RivermaxSync/Private/MediaOutputSynchronizationPolicyRivermax.cpp b/Engine/Plugins/VirtualProduction/Rivermax/RivermaxSync/Source/RivermaxSync/Private/MediaOutputSynchronizationPolicyRivermax.cpp
--- a/Engine/Plugins/VirtualProduction/Rivermax/RivermaxSync/Source/RivermaxSync/Private/MediaOutputSynchronizationPolicyRivermax.cpp
+++ b/Engine/Plugins/VirtualProduction/Rivermax/RivermaxSync/Source/RivermaxSync/Private/MediaOutputSynchronizationPolicyRivermax.cpp
@@ -96,10 +96,10 @@ void FMediaOutputSynchronizationPolicyRivermaxHandler::Synchronize()
 			return;
 		}
 
-		IDisplayClusterGenericBarriersClient* const BarrierClient = GetBarrierClient();
-		if (!BarrierClient || !BarrierClient->IsConnected())
+		TSharedPtr<IDisplayClusterGenericBarriersClient> BarrierClient = GetBarrierClient();
+		if (!BarrierClient)
 		{
-			UE_LOG(LogRivermaxSync, Warning, TEXT("'%s': Barrier client is not connected or nullptr"), *GetMediaDeviceId());
+			UE_LOG(LogRivermaxSync, Warning, TEXT("'%s': Barrier client is nullptr"), *GetMediaDeviceId());
 			return;
 		}
 
@@ -213,15 +213,9 @@ bool FMediaOutputSynchronizationPolicyRivermaxHandler::InitializeBarrier(const F
 		return false;
 	}
 
-	// Only setup barrier delegate for primary node to handle barrier data verification
-	if (!IDisplayCluster::Get().GetClusterMgr()->IsPrimary())
-	{
-		return true;
-	}
-
 	// Get barrier client
-	IDisplayClusterGenericBarriersClient* const BarrierClient = GetBarrierClient();
-	if (!BarrierClient || !BarrierClient->IsConnected())
+	TSharedPtr<IDisplayClusterGenericBarriersClient> BarrierClient = GetBarrierClient();
+	if (!BarrierClient)
 	{
 		UE_LOG(LogRivermaxSync, Warning, TEXT("Couldn't access a barrier client for '%s'"), *GetMediaDeviceId());
 		return false;
diff --git a/Engine/Plugins/VirtualProduction/Rivermax/RivermaxSync/Source/RivermaxSync/Public/MediaOutputSynchronizationPolicyRivermax.h b/Engine/Plugins/VirtualProduction/Rivermax/RivermaxSync/Source/RivermaxSync/Public/MediaOutputSynchronizationPolicyRivermax.h
--- a/Engine/Plugins/VirtualProduction/Rivermax/RivermaxSync/Source/RivermaxSync/Public/MediaOutputSynchronizationPolicyRivermax.h
+++ b/Engine/Plugins/VirtualProduction/Rivermax/RivermaxSync/Source/RivermaxSync/Public/MediaOutputSynchronizationPolicyRivermax.h
@@ -9,6 +9,8 @@
 
 #include "MediaOutputSynchronizationPolicyRivermax.generated.h"
 
+struct FGenericBarrierSynchronizationDelegateData;
+
 
 /*
  * Synchronization logic handler class for UMediaOutputSynchronizationPolicyRivermax.
