[nDisplay] Fixed CIS issues

#rnx

diff --git a/Engine/Plugins/Runtime/nDisplay/Source/DisplayCluster/Private/Network/Protocol/IDisplayClusterProtocolInternalComm.h b/Engine/Plugins/Runtime/nDisplay/Source/DisplayCluster/Private/Network/Protocol/IDisplayClusterProtocolInternalComm.h
--- a/Engine/Plugins/Runtime/nDisplay/Source/DisplayCluster/Private/Network/Protocol/IDisplayClusterProtocolInternalComm.h
+++ b/Engine/Plugins/Runtime/nDisplay/Source/DisplayCluster/Private/Network/Protocol/IDisplayClusterProtocolInternalComm.h
@@ -1,4 +1,4 @@
-// Copyright Epic Games, Inc//. All Rights Reserved.
+// Copyright Epic Games, Inc. All Rights Reserved.
 
 #pragma once
 
diff --git a/Engine/Plugins/Runtime/nDisplay/Source/DisplayClusterConfiguration/Public/DisplayClusterConfigurationTypes_Enums.h b/Engine/Plugins/Runtime/nDisplay/Source/DisplayClusterConfiguration/Public/DisplayClusterConfigurationTypes_Enums.h
--- a/Engine/Plugins/Runtime/nDisplay/Source/DisplayClusterConfiguration/Public/DisplayClusterConfigurationTypes_Enums.h
+++ b/Engine/Plugins/Runtime/nDisplay/Source/DisplayClusterConfiguration/Public/DisplayClusterConfigurationTypes_Enums.h
@@ -17,8 +17,10 @@ enum class EDisplayClusterConfigurationDataSource : uint8
 	Json UMETA(DisplayName = "JSON file")
 };
 
+// UE_DEPRECATED(5.6, "Removed as unnecessary in the new failover implementation")
+// The depracation topping has not been set in order to avoid the other warnings in the backward compatibility code (Linux only)
 UENUM()
-enum class UE_DEPRECATED(5.6, "Removed as unnecessary in the new failover implementation") EDisplayClusterConfigurationFailoverPolicy : uint8
+enum class EDisplayClusterConfigurationFailoverPolicy : uint8
 {
 	// No failover operations performed. The whole cluster gets terminated in case of any error
 	Disabled                          UMETA(DisplayName = "Disabled"),
