[nDisplay] Fixed CIS issues

#rnx

diff --git a/Engine/Plugins/Runtime/nDisplay/Source/DisplayCluster/Private/Cluster/Controller/DisplayClusterClusterNodeCtrlMain.cpp b/Engine/Plugins/Runtime/nDisplay/Source/DisplayCluster/Private/Cluster/Controller/DisplayClusterClusterNodeCtrlMain.cpp
--- a/Engine/Plugins/Runtime/nDisplay/Source/DisplayCluster/Private/Cluster/Controller/DisplayClusterClusterNodeCtrlMain.cpp
+++ b/Engine/Plugins/Runtime/nDisplay/Source/DisplayCluster/Private/Cluster/Controller/DisplayClusterClusterNodeCtrlMain.cpp
@@ -1100,7 +1100,7 @@ const TSharedPtr<FDisplayClusterGenericBarrierClient> FDisplayClusterClusterNode
 	return OutcomeClient;
 }
 
-bool FDisplayClusterClusterNodeCtrlMain::FNodeGeneralPurposeBarrierClientSet::Initialize(const IDisplayClusterProtocolInternalComm::FClusterServicesHostingInfo& HostingInfo)
+bool FDisplayClusterClusterNodeCtrlMain::FNodeGeneralPurposeBarrierClientSet::Initialize(const IDisplayClusterProtocolInternalComm::FClusterServicesHostingInfo& InHostingInfo)
 {
 	// Get all active nodes
 	TSet<FString> ActiveNodes;
@@ -1128,7 +1128,7 @@ bool FDisplayClusterClusterNodeCtrlMain::FNodeGeneralPurposeBarrierClientSet::In
 					TSharedRef<FDisplayClusterGenericBarrierClient>& NewClient = Clients.Emplace(*NodeId, MakeShared<FDisplayClusterGenericBarrierClient>(*ClientName));
 
 					// Find hosting info of this node
-					const IDisplayClusterProtocolInternalComm::FNodeServicesHostingInfo* NodeInfo = HostingInfo.ClusterHostingInfo.Find(*NodeId);
+					const IDisplayClusterProtocolInternalComm::FNodeServicesHostingInfo* NodeInfo = InHostingInfo.ClusterHostingInfo.Find(*NodeId);
 					if (!NodeInfo)
 					{
 						UE_LOG(LogDisplayClusterCluster, Warning, TEXT("Couldn't initialize GPB client for node '%s'. No corresponding config data found."), *NodeId);
diff --git a/Engine/Plugins/Runtime/nDisplay/Source/DisplayCluster/Private/Cluster/Controller/DisplayClusterClusterNodeCtrlMain.h b/Engine/Plugins/Runtime/nDisplay/Source/DisplayCluster/Private/Cluster/Controller/DisplayClusterClusterNodeCtrlMain.h
--- a/Engine/Plugins/Runtime/nDisplay/Source/DisplayCluster/Private/Cluster/Controller/DisplayClusterClusterNodeCtrlMain.h
+++ b/Engine/Plugins/Runtime/nDisplay/Source/DisplayCluster/Private/Cluster/Controller/DisplayClusterClusterNodeCtrlMain.h
@@ -200,7 +200,7 @@ private: // GP barrier clients
 	struct FNodeGeneralPurposeBarrierClientSet
 	{
 		/** Instantiates the client set, and establishes all neccessary connections. Returns true if everything is Ok. */
-		bool Initialize(const IDisplayClusterProtocolInternalComm::FClusterServicesHostingInfo& HostingInfo);
+		bool Initialize(const IDisplayClusterProtocolInternalComm::FClusterServicesHostingInfo& InHostingInfo);
 
 		/** Per-node clients (used on 'Any' thread) */
 		TMap<FName, TSharedRef<FDisplayClusterGenericBarrierClient>> Clients;
