nDisplay: Added viewport freezing options to new editor viewport toolbar

diff --git a/Engine/Plugins/Runtime/nDisplay/Source/DisplayClusterConfigurator/Private/DisplayClusterConfiguratorModule.cpp b/Engine/Plugins/Runtime/nDisplay/Source/DisplayClusterConfigurator/Private/DisplayClusterConfiguratorModule.cpp
--- a/Engine/Plugins/Runtime/nDisplay/Source/DisplayClusterConfigurator/Private/DisplayClusterConfiguratorModule.cpp
+++ b/Engine/Plugins/Runtime/nDisplay/Source/DisplayClusterConfigurator/Private/DisplayClusterConfiguratorModule.cpp
@@ -48,6 +48,7 @@
 #include "AssetRegistry/AssetRegistryModule.h"
 #include "AssetTypeCategories.h"
 #include "DisplayClusterConfiguratorBlueprintEditor.h"
+#include "DisplayClusterConfiguratorToolbarExtensions.h"
 #include "HAL/IConsoleManager.h"
 #include "Modules/ModuleManager.h"
 #include "ActorFactories/ActorFactoryBlueprint.h"
@@ -114,6 +115,7 @@ void FDisplayClusterConfiguratorModule::StartupModule()
 	FDisplayClusterConfiguratorCommands::Register();
 	FDisplayClusterConfiguratorOutputMappingCommands::Register();
 
+	ToolbarExtensions = MakeShared<FDisplayClusterConfiguratorToolbarExtensions>();
 	MenuExtensibilityManager = MakeShareable(new FExtensibilityManager);
 	ToolBarExtensibilityManager = MakeShareable(new FExtensibilityManager);
 
@@ -148,6 +150,9 @@ void FDisplayClusterConfiguratorModule::ShutdownModule()
 	UnregisterCustomLayouts();
 	UnregisterSectionMappings();
 
+	ToolbarExtensions->UnregisterToolbarExtensions();
+	ToolbarExtensions.Reset();
+	
 	MenuExtensibilityManager.Reset();
 	ToolBarExtensibilityManager.Reset();
 
@@ -166,6 +171,7 @@ void FDisplayClusterConfiguratorModule::ShutdownModule()
 void FDisplayClusterConfiguratorModule::OnPostEngineInit()
 {
 	RegisterPanelExtensions();
+	ToolbarExtensions->RegisterToolbarExtensions();
 }
 
 const FDisplayClusterConfiguratorCommands& FDisplayClusterConfiguratorModule::GetCommands() const
diff --git a/Engine/Plugins/Runtime/nDisplay/Source/DisplayClusterConfigurator/Private/DisplayClusterConfiguratorModule.h b/Engine/Plugins/Runtime/nDisplay/Source/DisplayClusterConfigurator/Private/DisplayClusterConfiguratorModule.h
--- a/Engine/Plugins/Runtime/nDisplay/Source/DisplayClusterConfigurator/Private/DisplayClusterConfiguratorModule.h
+++ b/Engine/Plugins/Runtime/nDisplay/Source/DisplayClusterConfigurator/Private/DisplayClusterConfiguratorModule.h
@@ -6,6 +6,7 @@
 #include "IAssetTypeActions.h"
 #include "IDisplayClusterConfigurator.h"
 
+class FDisplayClusterConfiguratorToolbarExtensions;
 class FPlacementModeID;
 class IAssetTools;
 class FDisplayClusterConfiguratorAssetTypeActions;
@@ -63,6 +64,7 @@ private:
 	TArray<TSharedPtr<IAssetTypeActions>> CreatedAssetTypeActions;
 	TSharedPtr<FExtensibilityManager> MenuExtensibilityManager;
 	TSharedPtr<FExtensibilityManager> ToolBarExtensibilityManager;
+	TSharedPtr<FDisplayClusterConfiguratorToolbarExtensions> ToolbarExtensions;
 	FDisplayClusterConfiguratorKismetCompiler  BlueprintCompiler;
 	TArray<FName> RegisteredClassLayoutNames;
 	TArray<FName> RegisteredPropertyLayoutNames;
diff --git a/Engine/Plugins/Runtime/nDisplay/Source/DisplayClusterConfigurator/Private/DisplayClusterConfiguratorToolbarExtensions.cpp b/Engine/Plugins/Runtime/nDisplay/Source/DisplayClusterConfigurator/Private/DisplayClusterConfiguratorToolbarExtensions.cpp
new file mode 100644
--- /dev/null
+++ b/Engine/Plugins/Runtime/nDisplay/Source/DisplayClusterConfigurator/Private/DisplayClusterConfiguratorToolbarExtensions.cpp
@@ -0,0 +1,144 @@
+﻿// Copyright Epic Games, Inc. All Rights Reserved.
+
+#include "DisplayClusterConfiguratorToolbarExtensions.h"
+
+#include "DisplayClusterConfiguratorStyle.h"
+#include "DisplayClusterRootActor.h"
+#include "ScopedTransaction.h"
+#include "ToolMenus.h"
+#include "Styling/SlateIconFinder.h"
+
+#define LOCTEXT_NAMESPACE "FDisplayClusterConfiguratorToolbarExtensions"
+
+void FDisplayClusterConfiguratorToolbarExtensions::RegisterToolbarExtensions()
+{
+	FToolMenuOwnerScoped ToolMenuOwnerScoped(this);
+
+	UToolMenu* Menu = UToolMenus::Get()->ExtendMenu("LevelEditor.ViewportToolBar.PerformanceAndScalability");
+	FToolMenuSection& Section = Menu->FindOrAddSection("FreezeNDisplayViewports", LOCTEXT("FreezeNDisplayViewportsSectionLabel", "Freeze nDisplay Viewports"));
+	Section.AddDynamicEntry(NAME_None, FNewToolMenuSectionDelegate::CreateSP(this, &FDisplayClusterConfiguratorToolbarExtensions::CreateFreezeViewportsMenu));
+}
+
+void FDisplayClusterConfiguratorToolbarExtensions::UnregisterToolbarExtensions()
+{
+}
+
+void FDisplayClusterConfiguratorToolbarExtensions::CreateFreezeViewportsMenu(FToolMenuSection& InSection)
+{
+	bool bAddedUnfreezeAll = false;
+	if (GEditor)
+	{
+		if (const UWorld* World = GEditor->GetEditorWorldContext().World())
+		{
+			for (TActorIterator<ADisplayClusterRootActor> It(World); It; ++It)
+			{
+				// If root actors have been found in the level, before the first one is added, add the "Unfreeze All Viewports" option
+				if (!bAddedUnfreezeAll)
+				{
+					FToolMenuEntry UnfreezeAllEntry = FToolMenuEntry::InitMenuEntry(
+						FName("UnfreezeAllViewports"),
+						LOCTEXT("UnfreezeAllViewportsLabel", "Unfreeze All Viewports"),
+						LOCTEXT("ViewportsFrozenWarningToolTip", "nDisplay viewports are frozen. Click to unfreeze all frozen viewports."),
+						FSlateIcon(FDisplayClusterConfiguratorStyle::Get().GetStyleSetName(), "ClassIcon.DisplayClusterRootActor"),
+						FUIAction(
+							FExecuteAction::CreateSP(this, &FDisplayClusterConfiguratorToolbarExtensions::UnfreezeAllViewports),
+							FCanExecuteAction::CreateSP(this, &FDisplayClusterConfiguratorToolbarExtensions::AreAnyViewportsFrozen)),
+						EUserInterfaceActionType::Button);
+
+					UnfreezeAllEntry.SetShowInToolbarTopLevel(TAttribute<bool>::CreateLambda(
+						[this]() -> bool
+						{
+							return AreAnyViewportsFrozen();
+						}
+					));
+
+					UnfreezeAllEntry.StyleNameOverride = "ViewportToolbarWarning";
+					InSection.AddEntry(UnfreezeAllEntry);
+					bAddedUnfreezeAll = true;
+				}
+				
+				TWeakObjectPtr<ADisplayClusterRootActor> WeakPtr = *It;
+				FString RootActorLabel = It->GetActorNameOrLabel();
+				InSection.AddEntry(FToolMenuEntry::InitMenuEntry(
+					FName("ToggleFreezeViewports_" + RootActorLabel),
+					FText::FromString(RootActorLabel),
+					LOCTEXT("ToggleFreezeViewportsTooltip", "Toggles whether this stage's viewports are frozen or not"),
+					FSlateIconFinder::FindIconForClass(ADisplayClusterRootActor::StaticClass()),
+					FUIAction(
+						FExecuteAction::CreateSP(this, &FDisplayClusterConfiguratorToolbarExtensions::ToggleFreezeViewports, WeakPtr),
+						FCanExecuteAction(),
+						FIsActionChecked::CreateSP(this, &FDisplayClusterConfiguratorToolbarExtensions::AreViewportsFrozen, WeakPtr)),
+					EUserInterfaceActionType::ToggleButton
+				));
+			}
+		}
+	}
+}
+
+void FDisplayClusterConfiguratorToolbarExtensions::UnfreezeAllViewports()
+{
+	if (GEditor)
+	{
+		FScopedTransaction Transaction(LOCTEXT("UnfreezeViewports", "Unfreeze viewports"));
+		
+		for (TActorIterator<ADisplayClusterRootActor> It(GEditor->GetEditorWorldContext().World()); It; ++It)
+		{
+			It->SetFreezeOuterViewports(false);
+		}
+	}
+}
+
+bool FDisplayClusterConfiguratorToolbarExtensions::AreAnyViewportsFrozen() const
+{
+	if (GEditor)
+	{
+		if (const UWorld* World = GEditor->GetEditorWorldContext().World())
+		{
+			for (TActorIterator<ADisplayClusterRootActor> It(World); It; ++It)
+			{
+				const UDisplayClusterConfigurationData* ConfigData = It->GetConfigData();
+				if (ConfigData && ConfigData->StageSettings.bFreezeRenderOuterViewports)
+				{
+					return true;
+				}
+			}
+		}
+	}
+
+	return false;
+}
+
+void FDisplayClusterConfiguratorToolbarExtensions::ToggleFreezeViewports(TWeakObjectPtr<ADisplayClusterRootActor> InRootActor)
+{
+	if (!InRootActor.IsValid())
+	{
+		return;
+	}
+
+	const UDisplayClusterConfigurationData* ConfigData = InRootActor->GetConfigData();
+	if (!ConfigData)
+	{
+		return;
+	}
+
+	FScopedTransaction Transaction(FText::Format(LOCTEXT("ToggleFreezeViewportsTransaction", "Toggle Freeze Viewports for stage '{0}'"), FText::FromString(InRootActor->GetActorNameOrLabel())));
+	InRootActor->SetFreezeOuterViewports(!ConfigData->StageSettings.bFreezeRenderOuterViewports);
+}
+
+bool FDisplayClusterConfiguratorToolbarExtensions::AreViewportsFrozen(TWeakObjectPtr<ADisplayClusterRootActor> InRootActor) const
+{
+	if (!InRootActor.IsValid())
+	{
+		return false;
+	}
+
+	const UDisplayClusterConfigurationData* ConfigData = InRootActor->GetConfigData();
+	if (!ConfigData)
+	{
+		return false;
+	}
+
+	return ConfigData->StageSettings.bFreezeRenderOuterViewports;
+}
+
+#undef LOCTEXT_NAMESPACE
diff --git a/Engine/Plugins/Runtime/nDisplay/Source/DisplayClusterConfigurator/Private/DisplayClusterConfiguratorToolbarExtensions.h b/Engine/Plugins/Runtime/nDisplay/Source/DisplayClusterConfigurator/Private/DisplayClusterConfiguratorToolbarExtensions.h
new file mode 100644
--- /dev/null
+++ b/Engine/Plugins/Runtime/nDisplay/Source/DisplayClusterConfigurator/Private/DisplayClusterConfiguratorToolbarExtensions.h
@@ -0,0 +1,30 @@
+﻿// Copyright Epic Games, Inc. All Rights Reserved.
+
+#pragma once
+
+#include "CoreMinimal.h"
+#include "Templates/SharedPointer.h"
+#include "UObject/WeakObjectPtrTemplates.h"
+
+struct FToolMenuSection;
+class ADisplayClusterRootActor;
+
+/** Manages any toolbar extensions the nDisplay plugin needs to register with the editor */
+class FDisplayClusterConfiguratorToolbarExtensions : public TSharedFromThis<FDisplayClusterConfiguratorToolbarExtensions>
+{
+public:
+	/** Registers any toolbar extensions with the editor */
+	void RegisterToolbarExtensions();
+
+	/** Unregisters all registered toolbar extensions */
+	void UnregisterToolbarExtensions();
+
+private:
+	void CreateFreezeViewportsMenu(FToolMenuSection& InSection);
+
+	void UnfreezeAllViewports();
+	bool AreAnyViewportsFrozen() const;
+	
+	void ToggleFreezeViewports(TWeakObjectPtr<ADisplayClusterRootActor> InRootActor);
+	bool AreViewportsFrozen(TWeakObjectPtr<ADisplayClusterRootActor> InRootActor) const;
+};
