nDisplay: Fixed issue where color grading and white balance settings would reset to zero instead of desirable default values when the color grading or white balace structs were elements of an array.

diff --git a/Engine/Plugins/Editor/ColorGrading/Source/ColorGradingEditor/Private/SColorGradingColorWheelPanel.cpp b/Engine/Plugins/Editor/ColorGrading/Source/ColorGradingEditor/Private/SColorGradingColorWheelPanel.cpp
--- a/Engine/Plugins/Editor/ColorGrading/Source/ColorGradingEditor/Private/SColorGradingColorWheelPanel.cpp
+++ b/Engine/Plugins/Editor/ColorGrading/Source/ColorGradingEditor/Private/SColorGradingColorWheelPanel.cpp
@@ -349,20 +349,22 @@ void SColorGradingColorWheelPanel::ClearColorGradingElementsToolBar()
 
 void SColorGradingColorWheelPanel::FillColorWheels(const FColorGradingEditorDataModel::FColorGradingElement& ColorGradingElement)
 {
-	auto FillColorWheel = [this](const TSharedPtr<SColorGradingColorWheel>& ColorWheel, const TSharedPtr<IPropertyHandle>& PropertyHandle)
+	auto FillColorWheel = [this](const TSharedPtr<SColorGradingColorWheel>& ColorWheel,
+		const TSharedPtr<IPropertyHandle>& PropertyHandle,
+		const TOptional<FResetToDefaultOverride>& ResetToDefaultOverride)
 	{
 		if (ColorWheel.IsValid())
 		{
 			ColorWheel->SetColorPropertyHandle(PropertyHandle);
-			ColorWheel->SetHeaderContent(CreateColorWheelHeaderWidget(PropertyHandle));
+			ColorWheel->SetHeaderContent(CreateColorWheelHeaderWidget(PropertyHandle, ResetToDefaultOverride));
 		}
 	};
 
-	FillColorWheel(ColorWheels[0], ColorGradingElement.SaturationPropertyHandle);
-	FillColorWheel(ColorWheels[1], ColorGradingElement.ContrastPropertyHandle);
-	FillColorWheel(ColorWheels[2], ColorGradingElement.GammaPropertyHandle);
-	FillColorWheel(ColorWheels[3], ColorGradingElement.GainPropertyHandle);
-	FillColorWheel(ColorWheels[4], ColorGradingElement.OffsetPropertyHandle);
+	FillColorWheel(ColorWheels[0], ColorGradingElement.SaturationPropertyHandle, ColorGradingElement.SaturationResetToDefaultOverride);
+	FillColorWheel(ColorWheels[1], ColorGradingElement.ContrastPropertyHandle, ColorGradingElement.ContrastResetToDefaultOverride);
+	FillColorWheel(ColorWheels[2], ColorGradingElement.GammaPropertyHandle, ColorGradingElement.GammaResetToDefaultOverride);
+	FillColorWheel(ColorWheels[3], ColorGradingElement.GainPropertyHandle, ColorGradingElement.GainResetToDefaultOverride);
+	FillColorWheel(ColorWheels[4], ColorGradingElement.OffsetPropertyHandle, ColorGradingElement.OffsetResetToDefaultOverride);
 }
 
 void SColorGradingColorWheelPanel::ClearColorWheels()
@@ -377,7 +379,7 @@ void SColorGradingColorWheelPanel::ClearColorWheels()
 	};
 }
 
-TSharedRef<SWidget> SColorGradingColorWheelPanel::CreateColorWheelHeaderWidget(const TSharedPtr<IPropertyHandle>& ColorPropertyHandle)
+TSharedRef<SWidget> SColorGradingColorWheelPanel::CreateColorWheelHeaderWidget(const TSharedPtr<IPropertyHandle>& ColorPropertyHandle, const TOptional<FResetToDefaultOverride>& ResetToDefaultOverride)
 {
 	if (TSharedPtr<IDetailTreeNode> TreeNode = ColorGradingDataModel->GetPropertyRowGenerator()->FindTreeNode(ColorPropertyHandle))
 	{
@@ -415,7 +417,7 @@ TSharedRef<SWidget> SColorGradingColorWheelPanel::CreateColorWheelHeaderWidget(c
 					SNew(SBox)
 					.WidthOverride(22.0f)
 					[
-						CreateColorPropertyExtensions(ColorPropertyHandle, TreeNode)
+						CreateColorPropertyExtensions(ColorPropertyHandle, TreeNode, ResetToDefaultOverride)
 					]
 				];
 		}
@@ -426,7 +428,9 @@ TSharedRef<SWidget> SColorGradingColorWheelPanel::CreateColorWheelHeaderWidget(c
 	return SNullWidget::NullWidget;
 }
 
-TSharedRef<SWidget> SColorGradingColorWheelPanel::CreateColorPropertyExtensions(const TSharedPtr<IPropertyHandle>& ColorPropertyHandle, const TSharedPtr<IDetailTreeNode>& DetailTreeNode)
+TSharedRef<SWidget> SColorGradingColorWheelPanel::CreateColorPropertyExtensions(const TSharedPtr<IPropertyHandle>& ColorPropertyHandle,
+	const TSharedPtr<IDetailTreeNode>& DetailTreeNode,
+	const TOptional<FResetToDefaultOverride>& ResetToDefaultOverride)
 {
 	TArray<FPropertyRowExtensionButton> ExtensionButtons;
 
@@ -436,11 +440,18 @@ TSharedRef<SWidget> SColorGradingColorWheelPanel::CreateColorPropertyExtensions(
 	FPropertyRowExtensionButton& ResetToDefaultButton = ExtensionButtons.AddDefaulted_GetRef();
 	ResetToDefaultButton.Label = NSLOCTEXT("PropertyEditor", "ResetToDefault", "Reset to Default");
 	ResetToDefaultButton.UIAction = FUIAction(
-		FExecuteAction::CreateLambda([WeakColorPropertyHandle]
+		FExecuteAction::CreateLambda([WeakColorPropertyHandle, ResetToDefaultOverride]
 		{
 			if (WeakColorPropertyHandle.IsValid())
 			{
-				WeakColorPropertyHandle.Pin()->ResetToDefault();
+				if (ResetToDefaultOverride.IsSet())
+				{
+					WeakColorPropertyHandle.Pin()->ExecuteCustomResetToDefault(ResetToDefaultOverride.GetValue());
+				}
+				else
+				{
+					WeakColorPropertyHandle.Pin()->ResetToDefault();
+				}
 			}
 		}),
 		FCanExecuteAction::CreateLambda([WeakColorPropertyHandle]
diff --git a/Engine/Plugins/Editor/ColorGrading/Source/ColorGradingEditor/Private/SColorGradingColorWheelPanel.h b/Engine/Plugins/Editor/ColorGrading/Source/ColorGradingEditor/Private/SColorGradingColorWheelPanel.h
--- a/Engine/Plugins/Editor/ColorGrading/Source/ColorGradingEditor/Private/SColorGradingColorWheelPanel.h
+++ b/Engine/Plugins/Editor/ColorGrading/Source/ColorGradingEditor/Private/SColorGradingColorWheelPanel.h
@@ -54,8 +54,8 @@ private:
 	void FillColorWheels(const FColorGradingEditorDataModel::FColorGradingElement& ColorGradingElement);
 	void ClearColorWheels();
 
-	TSharedRef<SWidget> CreateColorWheelHeaderWidget(const TSharedPtr<IPropertyHandle>& ColorPropertyHandle);
-	TSharedRef<SWidget> CreateColorPropertyExtensions(const TSharedPtr<IPropertyHandle>& ColorPropertyHandle, const TSharedPtr<IDetailTreeNode>& DetailTreeNode);
+	TSharedRef<SWidget> CreateColorWheelHeaderWidget(const TSharedPtr<IPropertyHandle>& ColorPropertyHandle, const TOptional<FResetToDefaultOverride>& ResetToDefaultOverride);
+	TSharedRef<SWidget> CreateColorPropertyExtensions(const TSharedPtr<IPropertyHandle>& ColorPropertyHandle, const TSharedPtr<IDetailTreeNode>& DetailTreeNode, const TOptional<FResetToDefaultOverride>& ResetToDefaultOverride);
 
 	bool FilterDetailTreeNode(const TSharedRef<IDetailTreeNode>& InDetailTreeNode);
 
diff --git a/Engine/Plugins/Editor/ColorGrading/Source/ColorGradingEditor/Public/ColorGradingEditorDataModel.h b/Engine/Plugins/Editor/ColorGrading/Source/ColorGradingEditor/Public/ColorGradingEditorDataModel.h
--- a/Engine/Plugins/Editor/ColorGrading/Source/ColorGradingEditor/Public/ColorGradingEditorDataModel.h
+++ b/Engine/Plugins/Editor/ColorGrading/Source/ColorGradingEditor/Public/ColorGradingEditorDataModel.h
@@ -3,8 +3,10 @@
 #pragma once
 
 #include "CoreMinimal.h"
+#include "IDetailPropertyRow.h"
 #include "PropertyEditorDelegates.h"
 
+class FResetToDefaultOverride;
 class IPropertyHandle;
 class IPropertyRowGenerator;
 class SWidget;
@@ -54,17 +56,32 @@ public:
 		/** Property handle of the color property used for saturation */
 		TSharedPtr<IPropertyHandle> SaturationPropertyHandle;
 
+		/** Optional reset to default handler for the saturation property */
+		TOptional<FResetToDefaultOverride> SaturationResetToDefaultOverride;
+		
 		/** Property handle of the color property used for contrast */
 		TSharedPtr<IPropertyHandle> ContrastPropertyHandle;
 
+		/** Optional reset to default handler for the contrast property */
+		TOptional<FResetToDefaultOverride> ContrastResetToDefaultOverride;
+		
 		/** Property handle of the color property used for gamma */
 		TSharedPtr<IPropertyHandle> GammaPropertyHandle;
 
+		/** Optional reset to default handler for the gamma property */
+		TOptional<FResetToDefaultOverride> GammaResetToDefaultOverride;
+		
 		/** Property handle of the color property used for gain */
 		TSharedPtr<IPropertyHandle> GainPropertyHandle;
 
+		/** Optional reset to default handler for the gain property */
+		TOptional<FResetToDefaultOverride> GainResetToDefaultOverride;
+		
 		/** Property handle of the color property used for offset */
 		TSharedPtr<IPropertyHandle> OffsetPropertyHandle;
+
+		/** Optional reset to default handler for the offset property */
+		TOptional<FResetToDefaultOverride> OffsetResetToDefaultOverride;
 	};
 
 	/** Stores a list of color grading elements that form a single color grading group */
diff --git a/Engine/Plugins/Runtime/nDisplay/Source/DisplayClusterColorGrading/Private/DataModelGenerators/DisplayClusterColorGradingGenerator_RootActor.cpp b/Engine/Plugins/Runtime/nDisplay/Source/DisplayClusterColorGrading/Private/DataModelGenerators/DisplayClusterColorGradingGenerator_RootActor.cpp
--- a/Engine/Plugins/Runtime/nDisplay/Source/DisplayClusterColorGrading/Private/DataModelGenerators/DisplayClusterColorGradingGenerator_RootActor.cpp
+++ b/Engine/Plugins/Runtime/nDisplay/Source/DisplayClusterColorGrading/Private/DataModelGenerators/DisplayClusterColorGradingGenerator_RootActor.cpp
@@ -33,16 +33,16 @@
 
 #define GET_MEMBER_NAME_ARRAY_CHECKED(ClassName, MemberName, Index) FName(FString(GET_MEMBER_NAME_STRING_CHECKED(ClassName, MemberName)).Replace(TEXT(#Index), *FString::FromInt(Index), ESearchCase::CaseSensitive))
 
-FColorGradingEditorDataModel::FColorGradingGroup FDisplayClusterColorGradingGenerator_ColorGradingRenderingSettings::CreateColorGradingGroup(const TSharedPtr<IPropertyHandle>& GroupPropertyHandle)
+FColorGradingEditorDataModel::FColorGradingGroup FDisplayClusterColorGradingGenerator_ColorGradingRenderingSettings::CreateColorGradingGroup(const TSharedPtr<IPropertyHandle>& GroupPropertyHandle, bool bOverrideResetToDefault)
 {
 	FColorGradingEditorDataModel::FColorGradingGroup ColorGradingGroup;
 	ColorGradingGroup.DisplayName = GroupPropertyHandle->GetPropertyDisplayName();
 	ColorGradingGroup.GroupPropertyHandle = GroupPropertyHandle;
 
-	ColorGradingGroup.ColorGradingElements.Add(CreateColorGradingElement(GroupPropertyHandle, TEXT("Global"), LOCTEXT("ColorGrading_GlobalLabel", "Global")));
-	ColorGradingGroup.ColorGradingElements.Add(CreateColorGradingElement(GroupPropertyHandle, TEXT("Shadows"), LOCTEXT("ColorGrading_ShadowsLabel", "Shadows")));
-	ColorGradingGroup.ColorGradingElements.Add(CreateColorGradingElement(GroupPropertyHandle, TEXT("Midtones"), LOCTEXT("ColorGrading_MidtonesLabel", "Midtones")));
-	ColorGradingGroup.ColorGradingElements.Add(CreateColorGradingElement(GroupPropertyHandle, TEXT("Highlights"), LOCTEXT("ColorGrading_HighlightsLabel", "Highlights")));
+	ColorGradingGroup.ColorGradingElements.Add(CreateColorGradingElement(GroupPropertyHandle, TEXT("Global"), LOCTEXT("ColorGrading_GlobalLabel", "Global"), bOverrideResetToDefault));
+	ColorGradingGroup.ColorGradingElements.Add(CreateColorGradingElement(GroupPropertyHandle, TEXT("Shadows"), LOCTEXT("ColorGrading_ShadowsLabel", "Shadows"), bOverrideResetToDefault));
+	ColorGradingGroup.ColorGradingElements.Add(CreateColorGradingElement(GroupPropertyHandle, TEXT("Midtones"), LOCTEXT("ColorGrading_MidtonesLabel", "Midtones"), bOverrideResetToDefault));
+	ColorGradingGroup.ColorGradingElements.Add(CreateColorGradingElement(GroupPropertyHandle, TEXT("Highlights"), LOCTEXT("ColorGrading_HighlightsLabel", "Highlights"), bOverrideResetToDefault));
 
 	ColorGradingGroup.DetailsViewCategories =
 	{
@@ -58,7 +58,8 @@ FColorGradingEditorDataModel::FColorGradingGroup FDisplayClusterColorGradingGene
 FColorGradingEditorDataModel::FColorGradingElement FDisplayClusterColorGradingGenerator_ColorGradingRenderingSettings::CreateColorGradingElement(
 	const TSharedPtr<IPropertyHandle>& GroupPropertyHandle,
 	FName ElementPropertyName,
-	FText ElementLabel)
+	FText ElementLabel,
+	bool bOverrideResetToDefault)
 {
 	FColorGradingEditorDataModel::FColorGradingElement ColorGradingElement;
 	ColorGradingElement.DisplayName = ElementLabel;
@@ -71,6 +72,27 @@ FColorGradingEditorDataModel::FColorGradingElement FDisplayClusterColorGradingGe
 		ColorGradingElement.GammaPropertyHandle = ElementPropertyHandle->GetChildHandle(GET_MEMBER_NAME_CHECKED(FDisplayClusterConfigurationViewport_ColorGradingSettings, Gamma));
 		ColorGradingElement.GainPropertyHandle = ElementPropertyHandle->GetChildHandle(GET_MEMBER_NAME_CHECKED(FDisplayClusterConfigurationViewport_ColorGradingSettings, Gain));
 		ColorGradingElement.OffsetPropertyHandle = ElementPropertyHandle->GetChildHandle(GET_MEMBER_NAME_CHECKED(FDisplayClusterConfigurationViewport_ColorGradingSettings, Offset));
+
+		if (bOverrideResetToDefault)
+		{
+			auto ResetPropertyColorValue = [](TSharedPtr<IPropertyHandle> InPropertyHandle, FVector4 InDefaultValue)
+			{
+				// Set the value with an interactive, non-transactable change first to avoid invoking any post edit change events on each component
+				// (Vector properties use property handles to set each component's value), in case the object owning the property becomes invalid
+				// on post edit change events (such as construction script created objects)
+				InPropertyHandle->SetValue(InDefaultValue, EPropertyValueSetFlags::InteractiveChange | EPropertyValueSetFlags::NotTransactable);
+
+				// Now update the value normally to invoke the usual notification/transaction pipeline
+				InPropertyHandle->SetValue(InDefaultValue, EPropertyValueSetFlags::ResetToDefault);
+			};
+			
+			// If we are overriding the reset to default, assume we want to reset the color values to the defaults defined in FDisplayClusterConfigurationViewport_ColorGradingSettings
+			ColorGradingElement.SaturationResetToDefaultOverride = FResetToDefaultOverride::Create(FResetToDefaultHandler::CreateLambda(ResetPropertyColorValue, FVector4::One()));
+			ColorGradingElement.ContrastResetToDefaultOverride = FResetToDefaultOverride::Create(FResetToDefaultHandler::CreateLambda(ResetPropertyColorValue, FVector4::One()));
+			ColorGradingElement.GammaResetToDefaultOverride = FResetToDefaultOverride::Create(FResetToDefaultHandler::CreateLambda(ResetPropertyColorValue, FVector4::One()));
+			ColorGradingElement.GainResetToDefaultOverride = FResetToDefaultOverride::Create(FResetToDefaultHandler::CreateLambda(ResetPropertyColorValue, FVector4::One()));
+			ColorGradingElement.OffsetResetToDefaultOverride = FResetToDefaultOverride::Create(FResetToDefaultHandler::CreateLambda(ResetPropertyColorValue, FVector4::Zero()));
+		}
 	}
 
 	return ColorGradingElement;
@@ -215,7 +237,7 @@ public:
 private:
 	void AddColorGradingDetailProperties(IDetailLayoutBuilder& DetailBuilder)
 	{
-		auto AddColorGradingSettings = [&DetailBuilder](const TSharedRef<IPropertyHandle>& ColorGradingSettingsHandle)
+		auto AddColorGradingSettings = [&DetailBuilder](const TSharedRef<IPropertyHandle>& ColorGradingSettingsHandle, bool bOverrideResetToDefault)
 		{
 			IDetailCategoryBuilder& DetailExposureCategoryBuilder = DetailBuilder.EditCategory(TEXT("DetailView_Exposure"), LOCTEXT("DetailView_ExposureDisplayName", "Exposure"));
 			DetailExposureCategoryBuilder.AddProperty(MakePropertyTransactional(ColorGradingSettingsHandle->GetChildHandle(GET_MEMBER_NAME_CHECKED(FDisplayClusterConfigurationViewport_ColorGradingRenderingSettings, AutoExposureBias))));
@@ -227,7 +249,22 @@ private:
 
 			IDetailCategoryBuilder& DetailWhiteBalanceCategoryBuilder = DetailBuilder.EditCategory(TEXT("DetailView_WhiteBalance"), LOCTEXT("DetailView_WhiteBalanceDisplayName", "White Balance"));
 			DetailWhiteBalanceCategoryBuilder.AddProperty(ColorGradingSettingsHandle->GetChildHandle(GET_MEMBER_NAME_CHECKED(FDisplayClusterConfigurationViewport_ColorGradingWhiteBalanceSettings, TemperatureType)));
-			DetailWhiteBalanceCategoryBuilder.AddProperty(MakePropertyTransactional(ColorGradingSettingsHandle->GetChildHandle(GET_MEMBER_NAME_CHECKED(FDisplayClusterConfigurationViewport_ColorGradingWhiteBalanceSettings, WhiteTemp))));
+			
+			if (bOverrideResetToDefault)
+			{
+				// Due to the way the editor handles resetting properties that are in structs inside an array, we must override the reset to default
+				// behavior to ensure the value we reset to is a sensible value and not just zero
+				DetailWhiteBalanceCategoryBuilder.AddProperty(MakePropertyTransactional(ColorGradingSettingsHandle->GetChildHandle(GET_MEMBER_NAME_CHECKED(FDisplayClusterConfigurationViewport_ColorGradingWhiteBalanceSettings, WhiteTemp))))
+					.OverrideResetToDefault(FResetToDefaultOverride::Create(FResetToDefaultHandler::CreateLambda([](TSharedPtr<IPropertyHandle> InPropertyHandle)
+					{
+						InPropertyHandle->SetValue(6500.0f);
+					})));
+			}
+			else
+			{
+				DetailWhiteBalanceCategoryBuilder.AddProperty(MakePropertyTransactional(ColorGradingSettingsHandle->GetChildHandle(GET_MEMBER_NAME_CHECKED(FDisplayClusterConfigurationViewport_ColorGradingWhiteBalanceSettings, WhiteTemp))));
+			}
+			
 			DetailWhiteBalanceCategoryBuilder.AddProperty(MakePropertyTransactional(ColorGradingSettingsHandle->GetChildHandle(GET_MEMBER_NAME_CHECKED(FDisplayClusterConfigurationViewport_ColorGradingWhiteBalanceSettings, WhiteTint))));
 
 			IDetailCategoryBuilder& DetailMiscCategoryBuilder = DetailBuilder.EditCategory(TEXT("DetailView_Misc"), LOCTEXT("DetailView_MiscDisplayName", "Misc"));
@@ -250,11 +287,13 @@ private:
 			IDetailCategoryBuilder& PerNodeSettingsCategoryBuilder = DetailBuilder.EditCategory(TEXT("DetailView_PerViewport"), LOCTEXT("DetailView_PerViewportDisplayName", "Per-Viewport Settings"));
 			PerNodeSettingsCategoryBuilder.AddProperty(DetailBuilder.GetProperty(GET_MEMBER_NAME_ARRAY_CHECKED(UDisplayClusterConfigurationData, StageSettings.PerViewportColorGrading[Index].bIsEntireClusterEnabled, Index), ConfigDataClass));
 
-			AddColorGradingSettings(DetailBuilder.GetProperty(GET_MEMBER_NAME_ARRAY_CHECKED(UDisplayClusterConfigurationData, StageSettings.PerViewportColorGrading[Index].ColorGradingSettings, Index), ConfigDataClass));
+			const bool bOverrideResetToDefault = true;
+			AddColorGradingSettings(DetailBuilder.GetProperty(GET_MEMBER_NAME_ARRAY_CHECKED(UDisplayClusterConfigurationData, StageSettings.PerViewportColorGrading[Index].ColorGradingSettings, Index), ConfigDataClass), bOverrideResetToDefault);
 		}
 		else
 		{
-			AddColorGradingSettings(DetailBuilder.GetProperty(GET_MEMBER_NAME_CHECKED(UDisplayClusterConfigurationData, StageSettings.EntireClusterColorGrading.ColorGradingSettings), ConfigDataClass));
+			const bool bOverrideResetToDefault = false;
+			AddColorGradingSettings(DetailBuilder.GetProperty(GET_MEMBER_NAME_CHECKED(UDisplayClusterConfigurationData, StageSettings.EntireClusterColorGrading.ColorGradingSettings), ConfigDataClass), bOverrideResetToDefault);
 		}
 	}
 
@@ -435,7 +474,10 @@ void FDisplayClusterColorGradingGenerator_RootActor::GenerateDataModel(IProperty
 			{
 				TSharedRef<IPropertyHandle> PerViewportElementHandle = PerViewportColorGradingHandle->AsArray()->GetElement(Index);
 
-				FColorGradingEditorDataModel::FColorGradingGroup PerViewportGroup = CreateColorGradingGroup(PerViewportElementHandle);
+				// Due to the way the editor handles resetting properties that are in structs inside an array, we must override the reset to default
+				// behavior to ensure the value we reset to is a sensible value and not just zero
+				const bool bOverrideResetToDefault = true;
+				FColorGradingEditorDataModel::FColorGradingGroup PerViewportGroup = CreateColorGradingGroup(PerViewportElementHandle, bOverrideResetToDefault);
 				PerViewportGroup.bCanBeDeleted = true;
 				PerViewportGroup.bCanBeRenamed = true;
 				PerViewportGroup.EditConditionPropertyHandle = PerViewportElementHandle->GetChildHandle(GET_MEMBER_NAME_CHECKED(FDisplayClusterConfigurationViewport_PerViewportColorGrading, bIsEnabled));
@@ -825,7 +867,7 @@ public:
 private:
 	void AddColorGradingDetailProperties(IDetailLayoutBuilder& DetailBuilder)
 	{
-		auto AddColorGradingSettings = [&DetailBuilder](const TSharedRef<IPropertyHandle>& ColorGradingSettingsHandle)
+		auto AddColorGradingSettings = [&DetailBuilder](const TSharedRef<IPropertyHandle>& ColorGradingSettingsHandle, bool bOverrideResetToDefault)
 		{
 			IDetailCategoryBuilder& DetailExposureCategoryBuilder = DetailBuilder.EditCategory(TEXT("DetailView_Exposure"), LOCTEXT("DetailView_ExposureDisplayName", "Exposure"));
 			DetailExposureCategoryBuilder.AddProperty(MakePropertyTransactional(ColorGradingSettingsHandle->GetChildHandle(GET_MEMBER_NAME_CHECKED(FDisplayClusterConfigurationViewport_ColorGradingRenderingSettings, AutoExposureBias))));
@@ -837,7 +879,22 @@ private:
 
 			IDetailCategoryBuilder& DetailWhiteBalanceCategoryBuilder = DetailBuilder.EditCategory(TEXT("DetailView_WhiteBalance"), LOCTEXT("DetailView_WhiteBalanceDisplayName", "White Balance"));
 			DetailWhiteBalanceCategoryBuilder.AddProperty(ColorGradingSettingsHandle->GetChildHandle(GET_MEMBER_NAME_CHECKED(FDisplayClusterConfigurationViewport_ColorGradingWhiteBalanceSettings, TemperatureType)));
-			DetailWhiteBalanceCategoryBuilder.AddProperty(MakePropertyTransactional(ColorGradingSettingsHandle->GetChildHandle(GET_MEMBER_NAME_CHECKED(FDisplayClusterConfigurationViewport_ColorGradingWhiteBalanceSettings, WhiteTemp))));
+
+			if (bOverrideResetToDefault)
+			{
+				// Due to the way the editor handles resetting properties that are in structs inside an array, we must override the reset to default
+				// behavior to ensure the value we reset to is a sensible value and not just zero
+				DetailWhiteBalanceCategoryBuilder.AddProperty(MakePropertyTransactional(ColorGradingSettingsHandle->GetChildHandle(GET_MEMBER_NAME_CHECKED(FDisplayClusterConfigurationViewport_ColorGradingWhiteBalanceSettings, WhiteTemp))))
+					.OverrideResetToDefault(FResetToDefaultOverride::Create(FResetToDefaultHandler::CreateLambda([](TSharedPtr<IPropertyHandle> InPropertyHandle)
+					{
+						InPropertyHandle->SetValue(6500.0f);
+					})));
+			}
+			else
+			{
+				DetailWhiteBalanceCategoryBuilder.AddProperty(MakePropertyTransactional(ColorGradingSettingsHandle->GetChildHandle(GET_MEMBER_NAME_CHECKED(FDisplayClusterConfigurationViewport_ColorGradingWhiteBalanceSettings, WhiteTemp))));
+			}
+			
 			DetailWhiteBalanceCategoryBuilder.AddProperty(MakePropertyTransactional(ColorGradingSettingsHandle->GetChildHandle(GET_MEMBER_NAME_CHECKED(FDisplayClusterConfigurationViewport_ColorGradingWhiteBalanceSettings, WhiteTint))));
 
 			IDetailCategoryBuilder& DetailMiscCategoryBuilder = DetailBuilder.EditCategory(TEXT("DetailView_Misc"), LOCTEXT("DetailView_MiscDisplayName", "Misc"));
@@ -860,14 +917,16 @@ private:
 			PerNodeSettingsCategoryBuilder.AddProperty(DetailBuilder.GetProperty(GET_MEMBER_NAME_ARRAY_CHECKED(UDisplayClusterICVFXCameraComponent, CameraSettings.PerNodeColorGrading[Index].bEntireClusterColorGrading, Index)));
 			PerNodeSettingsCategoryBuilder.AddProperty(DetailBuilder.GetProperty(GET_MEMBER_NAME_ARRAY_CHECKED(UDisplayClusterICVFXCameraComponent, CameraSettings.PerNodeColorGrading[Index].bAllNodesColorGrading, Index)));
 
-			AddColorGradingSettings(DetailBuilder.GetProperty(GET_MEMBER_NAME_ARRAY_CHECKED(UDisplayClusterICVFXCameraComponent, CameraSettings.PerNodeColorGrading[Index].ColorGradingSettings, Index)));
+			const bool bOverrideResetToDefault = true;
+			AddColorGradingSettings(DetailBuilder.GetProperty(GET_MEMBER_NAME_ARRAY_CHECKED(UDisplayClusterICVFXCameraComponent, CameraSettings.PerNodeColorGrading[Index].ColorGradingSettings, Index)), bOverrideResetToDefault);
 		}
 		else
 		{
 			IDetailCategoryBuilder& AllNodeSettingsCategoryBuilder = DetailBuilder.EditCategory(TEXT("DetailView_AllNodes"), LOCTEXT("DetailView_AllNodesDisplayName", "All Nodes Settings"));
 			AllNodeSettingsCategoryBuilder.AddProperty(DetailBuilder.GetProperty(GET_MEMBER_NAME_CHECKED(UDisplayClusterICVFXCameraComponent, CameraSettings.AllNodesColorGrading.bEnableEntireClusterColorGrading)));
 
-			AddColorGradingSettings(DetailBuilder.GetProperty(GET_MEMBER_NAME_CHECKED(UDisplayClusterICVFXCameraComponent, CameraSettings.AllNodesColorGrading.ColorGradingSettings)));
+			const bool bOverrideResetToDefault = false;
+			AddColorGradingSettings(DetailBuilder.GetProperty(GET_MEMBER_NAME_CHECKED(UDisplayClusterICVFXCameraComponent, CameraSettings.AllNodesColorGrading.ColorGradingSettings)), bOverrideResetToDefault);
 		}
 	}
 
@@ -1058,7 +1117,10 @@ void FDisplayClusterColorGradingGenerator_ICVFXCamera::GenerateDataModel(IProper
 			{
 				TSharedRef<IPropertyHandle> PerNodeElementHandle = PerNodeColorGradingHandle->AsArray()->GetElement(Index);
 
-				FColorGradingEditorDataModel::FColorGradingGroup PerNodeGroup = CreateColorGradingGroup(PerNodeElementHandle);
+				// Due to the way the editor handles resetting properties that are in structs inside an array, we must override the reset to default
+				// behavior to ensure the value we reset to is a sensible value and not just zero
+				const bool bOverrideResetToDefault = true;
+				FColorGradingEditorDataModel::FColorGradingGroup PerNodeGroup = CreateColorGradingGroup(PerNodeElementHandle, bOverrideResetToDefault);
 				PerNodeGroup.bCanBeDeleted = true;
 				PerNodeGroup.bCanBeRenamed = true;
 				PerNodeGroup.EditConditionPropertyHandle = PerNodeElementHandle->GetChildHandle(GET_MEMBER_NAME_CHECKED(FDisplayClusterConfigurationViewport_PerNodeColorGrading, bIsEnabled));
diff --git a/Engine/Plugins/Runtime/nDisplay/Source/DisplayClusterColorGrading/Private/DataModelGenerators/DisplayClusterColorGradingGenerator_RootActor.h b/Engine/Plugins/Runtime/nDisplay/Source/DisplayClusterColorGrading/Private/DataModelGenerators/DisplayClusterColorGradingGenerator_RootActor.h
--- a/Engine/Plugins/Runtime/nDisplay/Source/DisplayClusterColorGrading/Private/DataModelGenerators/DisplayClusterColorGradingGenerator_RootActor.h
+++ b/Engine/Plugins/Runtime/nDisplay/Source/DisplayClusterColorGrading/Private/DataModelGenerators/DisplayClusterColorGradingGenerator_RootActor.h
@@ -20,10 +20,10 @@ class FDisplayClusterColorGradingGenerator_ColorGradingRenderingSettings : publi
 {
 protected:
 	/** Creates a new color grading group structure from the specified group property handle, which finds and connects the appropriate property handles for the color wheels and details view */
-	FColorGradingEditorDataModel::FColorGradingGroup CreateColorGradingGroup(const TSharedPtr<IPropertyHandle>& GroupPropertyHandle);
+	FColorGradingEditorDataModel::FColorGradingGroup CreateColorGradingGroup(const TSharedPtr<IPropertyHandle>& GroupPropertyHandle, bool bOverrideResetToDefault = false);
 
 	/** Creates a color grading element structure from the specified property handle, whose child properties are expected to be colors with the ColorGradingMode metadata set */
-	FColorGradingEditorDataModel::FColorGradingElement CreateColorGradingElement(const TSharedPtr<IPropertyHandle>& GroupPropertyHandle, FName ElementPropertyName, FText ElementLabel);
+	FColorGradingEditorDataModel::FColorGradingElement CreateColorGradingElement(const TSharedPtr<IPropertyHandle>& GroupPropertyHandle, FName ElementPropertyName, FText ElementLabel, bool bOverrideResetToDefault = false);
 
 	/** Recursively searches the detail tree hierarchy for a property detail tree node whose name matches the specified name */
 	TSharedPtr<IDetailTreeNode> FindPropertyTreeNode(const TSharedRef<IDetailTreeNode>& Node, const FCachedPropertyPath& PropertyPath);
diff --git a/Engine/Plugins/Runtime/nDisplay/Source/DisplayClusterConfigurator/Private/DisplayClusterConfiguratorModule.cpp b/Engine/Plugins/Runtime/nDisplay/Source/DisplayClusterConfigurator/Private/DisplayClusterConfiguratorModule.cpp
--- a/Engine/Plugins/Runtime/nDisplay/Source/DisplayClusterConfigurator/Private/DisplayClusterConfiguratorModule.cpp
+++ b/Engine/Plugins/Runtime/nDisplay/Source/DisplayClusterConfigurator/Private/DisplayClusterConfiguratorModule.cpp
@@ -55,6 +55,8 @@
 #include "ScopedTransaction.h"
 #include "SEditorViewportToolBarButton.h"
 #include "Subsystems/PanelExtensionSubsystem.h"
+#include "Views/Details/Cluster/DCConfiguratorColorGradingSettingsCustomization.h"
+#include "Views/Details/Cluster/DCConfiguratorWhiteBalanceCustomization.h"
 
 #define REGISTER_PROPERTY_LAYOUT(PropertyType, CustomizationType) { \
 	const FName LayoutName = PropertyType::StaticStruct()->GetFName(); \
@@ -247,6 +249,8 @@ void FDisplayClusterConfiguratorModule::RegisterCustomLayouts()
 	REGISTER_PROPERTY_LAYOUT(FDisplayClusterConfigurationMediaOutput, FDisplayClusterConfiguratorMediaFullFrameOutputCustomization);
 	REGISTER_PROPERTY_LAYOUT(FDisplayClusterConfigurationMediaInputGroup,  FDisplayClusterConfiguratorMediaFullFrameInputCustomization);
 	REGISTER_PROPERTY_LAYOUT(FDisplayClusterConfigurationMediaOutputGroup, FDisplayClusterConfiguratorMediaFullFrameOutputCustomization);
+	REGISTER_PROPERTY_LAYOUT(FDisplayClusterConfigurationViewport_ColorGradingSettings, FDCConfiguratorColorGradingSettingsCustomization);
+	REGISTER_PROPERTY_LAYOUT(FDisplayClusterConfigurationViewport_ColorGradingWhiteBalanceSettings, FDCConfiguratorWhiteBalanceCustomization);
 }
 
 void FDisplayClusterConfiguratorModule::UnregisterCustomLayouts()
diff --git a/Engine/Plugins/Runtime/nDisplay/Source/DisplayClusterConfigurator/Private/Views/Details/Cluster/DCConfiguratorColorGradingSettingsCustomization.cpp b/Engine/Plugins/Runtime/nDisplay/Source/DisplayClusterConfigurator/Private/Views/Details/Cluster/DCConfiguratorColorGradingSettingsCustomization.cpp
new file mode 100644
--- /dev/null
+++ b/Engine/Plugins/Runtime/nDisplay/Source/DisplayClusterConfigurator/Private/Views/Details/Cluster/DCConfiguratorColorGradingSettingsCustomization.cpp
@@ -0,0 +1,75 @@
+﻿// Copyright Epic Games, Inc. All Rights Reserved.
+
+#include "DCConfiguratorColorGradingSettingsCustomization.h"
+
+#include "DisplayClusterConfigurationTypes_Postprocess.h"
+#include "IDetailChildrenBuilder.h"
+#include "IDetailPropertyRow.h"
+
+void FDCConfiguratorColorGradingSettingsCustomization::Initialize(const TSharedRef<IPropertyHandle>& InPropertyHandle, IPropertyTypeCustomizationUtils& CustomizationUtils)
+{
+	FDisplayClusterConfiguratorBaseTypeCustomization::Initialize(InPropertyHandle, CustomizationUtils);
+
+	// Check if this color grading struct is an element of an array. If so, we will override its "reset to default" behavior because
+	// the default value of a struct in an array that is empty in the CDO is entirely zeroes instead of the default values
+	// of the struct, which is undesirable for some color grading properties
+	TSharedPtr<IPropertyHandle> CurrentPropertyHandle = InPropertyHandle->GetParentHandle();
+	while (CurrentPropertyHandle)
+	{
+		if (CurrentPropertyHandle->AsArray().IsValid())
+		{
+			bIsArrayMember = true;
+			break;
+		}
+		
+		CurrentPropertyHandle = CurrentPropertyHandle->GetParentHandle();
+	}
+}
+
+void FDCConfiguratorColorGradingSettingsCustomization::SetChildren(const TSharedRef<IPropertyHandle>& InPropertyHandle, IDetailChildrenBuilder& InChildBuilder, IPropertyTypeCustomizationUtils& CustomizationUtils)
+{
+	if (!bIsArrayMember)
+	{
+		FDisplayClusterConfiguratorBaseTypeCustomization::SetChildren(InPropertyHandle, InChildBuilder, CustomizationUtils);
+		return;
+	}
+
+	static const TMap<FName, FVector4> ColorPropertyDefaultValues =
+	{
+	{ GET_MEMBER_NAME_CHECKED(FDisplayClusterConfigurationViewport_ColorGradingSettings, Saturation), FVector4::One() },
+	{ GET_MEMBER_NAME_CHECKED(FDisplayClusterConfigurationViewport_ColorGradingSettings, Contrast), FVector4::One() },
+	{ GET_MEMBER_NAME_CHECKED(FDisplayClusterConfigurationViewport_ColorGradingSettings, Gamma), FVector4::One() },
+	{ GET_MEMBER_NAME_CHECKED(FDisplayClusterConfigurationViewport_ColorGradingSettings, Gain), FVector4::One() },
+	{ GET_MEMBER_NAME_CHECKED(FDisplayClusterConfigurationViewport_ColorGradingSettings, Offset), FVector4::Zero() },
+	};
+
+	auto ResetPropertyColorValue = [](TSharedPtr<IPropertyHandle> InPropertyHandle, FVector4 InDefaultValue)
+	{
+		// Set the value with an interactive, non-transactable change first to avoid invoking any post edit change events on each component
+		// (Vector properties use property handles to set each component's value), in case the object owning the property becomes invalid
+		// on post edit change events (such as construction script created objects)
+		InPropertyHandle->SetValue(InDefaultValue, EPropertyValueSetFlags::InteractiveChange | EPropertyValueSetFlags::NotTransactable);
+
+		// Now update the value normally to invoke the usual notification/transaction pipeline
+		InPropertyHandle->SetValue(InDefaultValue, EPropertyValueSetFlags::ResetToDefault);
+	};
+	
+	uint32 NumChildren = 0;
+	InPropertyHandle->GetNumChildren(NumChildren);
+	for (uint32 ChildIndex = 0; ChildIndex < NumChildren; ++ChildIndex)
+	{
+		TSharedPtr<IPropertyHandle> ChildHandle = InPropertyHandle->GetChildHandle(ChildIndex);
+		if (ChildHandle && ChildHandle->IsValidHandle() && !ChildHandle->IsCustomized())
+		{
+			FText ChildTooltip = ApplySubstitutions(ChildHandle->GetToolTipText());
+			ChildHandle->SetToolTipText(ChildTooltip);
+
+			IDetailPropertyRow& PropertyRow = InChildBuilder.AddProperty(ChildHandle.ToSharedRef());
+			if (ColorPropertyDefaultValues.Contains(ChildHandle->GetProperty()->GetFName()))
+			{
+				const FVector4& DefaultValue = ColorPropertyDefaultValues[ChildHandle->GetProperty()->GetFName()];
+				PropertyRow.OverrideResetToDefault(FResetToDefaultOverride::Create(FResetToDefaultHandler::CreateLambda(ResetPropertyColorValue, DefaultValue)));
+			}
+		}
+	}
+}
diff --git a/Engine/Plugins/Runtime/nDisplay/Source/DisplayClusterConfigurator/Private/Views/Details/Cluster/DCConfiguratorColorGradingSettingsCustomization.h b/Engine/Plugins/Runtime/nDisplay/Source/DisplayClusterConfigurator/Private/Views/Details/Cluster/DCConfiguratorColorGradingSettingsCustomization.h
new file mode 100644
--- /dev/null
+++ b/Engine/Plugins/Runtime/nDisplay/Source/DisplayClusterConfigurator/Private/Views/Details/Cluster/DCConfiguratorColorGradingSettingsCustomization.h
@@ -0,0 +1,24 @@
+﻿// Copyright Epic Games, Inc. All Rights Reserved.
+
+#pragma once
+
+#include "CoreMinimal.h"
+#include "Views/Details/DisplayClusterConfiguratorBaseTypeCustomization.h"
+
+/** Customization that ensures proper reset to default values for the color properties in the nDisplay color grading struct when it is an element of an array */
+class FDCConfiguratorColorGradingSettingsCustomization : public FDisplayClusterConfiguratorBaseTypeCustomization
+{
+public:
+	static TSharedRef<IPropertyTypeCustomization> MakeInstance()
+	{
+		return MakeShared<FDCConfiguratorColorGradingSettingsCustomization>();
+	}
+
+protected:
+	virtual void Initialize(const TSharedRef<IPropertyHandle>& InPropertyHandle, IPropertyTypeCustomizationUtils& CustomizationUtils) override;
+	virtual void SetChildren(const TSharedRef<IPropertyHandle>& InPropertyHandle, IDetailChildrenBuilder& InChildBuilder, IPropertyTypeCustomizationUtils& CustomizationUtils) override;
+
+private:
+	/** Indicates if the struct being customized by this customization is a member of an array */
+	bool bIsArrayMember = false;
+};
diff --git a/Engine/Plugins/Runtime/nDisplay/Source/DisplayClusterConfigurator/Private/Views/Details/Cluster/DCConfiguratorWhiteBalanceCustomization.cpp b/Engine/Plugins/Runtime/nDisplay/Source/DisplayClusterConfigurator/Private/Views/Details/Cluster/DCConfiguratorWhiteBalanceCustomization.cpp
new file mode 100644
--- /dev/null
+++ b/Engine/Plugins/Runtime/nDisplay/Source/DisplayClusterConfigurator/Private/Views/Details/Cluster/DCConfiguratorWhiteBalanceCustomization.cpp
@@ -0,0 +1,65 @@
+﻿// Copyright Epic Games, Inc. All Rights Reserved.
+
+#include "DCConfiguratorWhiteBalanceCustomization.h"
+
+#include "DisplayClusterConfigurationTypes_Postprocess.h"
+#include "IDetailChildrenBuilder.h"
+#include "IDetailPropertyRow.h"
+
+void FDCConfiguratorWhiteBalanceCustomization::Initialize(const TSharedRef<IPropertyHandle>& InPropertyHandle, IPropertyTypeCustomizationUtils& CustomizationUtils)
+{
+	FDisplayClusterConfiguratorBaseTypeCustomization::Initialize(InPropertyHandle, CustomizationUtils);
+
+	// Check if this white balance struct is an element of an array. If so, we will override its "reset to default" behavior because
+	// the default value of a struct in an array that is empty in the CDO is entirely zeroes instead of the default values
+	// of the struct, which is undesirable for some white balance properties
+	TSharedPtr<IPropertyHandle> CurrentPropertyHandle = InPropertyHandle->GetParentHandle();
+	while (CurrentPropertyHandle)
+	{
+		if (CurrentPropertyHandle->AsArray().IsValid())
+		{
+			bIsArrayMember = true;
+			break;
+		}
+		
+		CurrentPropertyHandle = CurrentPropertyHandle->GetParentHandle();
+	}
+}
+
+void FDCConfiguratorWhiteBalanceCustomization::SetChildren(const TSharedRef<IPropertyHandle>& InPropertyHandle, IDetailChildrenBuilder& InChildBuilder, IPropertyTypeCustomizationUtils& CustomizationUtils)
+{
+	if (!bIsArrayMember)
+	{
+		FDisplayClusterConfiguratorBaseTypeCustomization::SetChildren(InPropertyHandle, InChildBuilder, CustomizationUtils);
+		return;
+	}
+
+	static const TMap<FName, float> FloatPropertyDefaultValues =
+	{
+	{ GET_MEMBER_NAME_CHECKED(FDisplayClusterConfigurationViewport_ColorGradingWhiteBalanceSettings, WhiteTemp), 6500.0f }
+	};
+
+	auto ResetPropertyFloatValue = [](TSharedPtr<IPropertyHandle> InPropertyHandle, float InDefaultValue)
+	{
+		InPropertyHandle->SetValue(InDefaultValue, EPropertyValueSetFlags::ResetToDefault);
+	};
+	
+	uint32 NumChildren = 0;
+	InPropertyHandle->GetNumChildren(NumChildren);
+	for (uint32 ChildIndex = 0; ChildIndex < NumChildren; ++ChildIndex)
+	{
+		TSharedPtr<IPropertyHandle> ChildHandle = InPropertyHandle->GetChildHandle(ChildIndex);
+		if (ChildHandle && ChildHandle->IsValidHandle() && !ChildHandle->IsCustomized())
+		{
+			FText ChildTooltip = ApplySubstitutions(ChildHandle->GetToolTipText());
+			ChildHandle->SetToolTipText(ChildTooltip);
+
+			IDetailPropertyRow& PropertyRow = InChildBuilder.AddProperty(ChildHandle.ToSharedRef());
+			if (FloatPropertyDefaultValues.Contains(ChildHandle->GetProperty()->GetFName()))
+			{
+				const float DefaultValue = FloatPropertyDefaultValues[ChildHandle->GetProperty()->GetFName()];
+				PropertyRow.OverrideResetToDefault(FResetToDefaultOverride::Create(FResetToDefaultHandler::CreateLambda(ResetPropertyFloatValue, DefaultValue)));
+			}
+		}
+	}
+}
diff --git a/Engine/Plugins/Runtime/nDisplay/Source/DisplayClusterConfigurator/Private/Views/Details/Cluster/DCConfiguratorWhiteBalanceCustomization.h b/Engine/Plugins/Runtime/nDisplay/Source/DisplayClusterConfigurator/Private/Views/Details/Cluster/DCConfiguratorWhiteBalanceCustomization.h
new file mode 100644
--- /dev/null
+++ b/Engine/Plugins/Runtime/nDisplay/Source/DisplayClusterConfigurator/Private/Views/Details/Cluster/DCConfiguratorWhiteBalanceCustomization.h
@@ -0,0 +1,24 @@
+﻿// Copyright Epic Games, Inc. All Rights Reserved.
+
+#pragma once
+
+#include "CoreMinimal.h"
+#include "Views/Details/DisplayClusterConfiguratorBaseTypeCustomization.h"
+
+/** Customization that ensures proper reset to default values for the temperature properties in the nDisplay white balance struct when it is an element of an array */
+class FDCConfiguratorWhiteBalanceCustomization : public FDisplayClusterConfiguratorBaseTypeCustomization
+{
+public:
+	static TSharedRef<IPropertyTypeCustomization> MakeInstance()
+	{
+		return MakeShared<FDCConfiguratorWhiteBalanceCustomization>();
+	}
+
+protected:
+	virtual void Initialize(const TSharedRef<IPropertyHandle>& InPropertyHandle, IPropertyTypeCustomizationUtils& CustomizationUtils) override;
+	virtual void SetChildren(const TSharedRef<IPropertyHandle>& InPropertyHandle, IDetailChildrenBuilder& InChildBuilder, IPropertyTypeCustomizationUtils& CustomizationUtils) override;
+
+private:
+	/** Indicates if the struct being customized by this customization is a member of an array */
+	bool bIsArrayMember = false;
+};
