nDisplay: Updated "Viewports Frozen" toolbar icon with new icon

diff --git a/Engine/Plugins/Runtime/nDisplay/Source/DisplayClusterConfigurator/Private/DisplayClusterConfiguratorStyle.cpp b/Engine/Plugins/Runtime/nDisplay/Source/DisplayClusterConfigurator/Private/DisplayClusterConfiguratorStyle.cpp
--- a/Engine/Plugins/Runtime/nDisplay/Source/DisplayClusterConfigurator/Private/DisplayClusterConfiguratorStyle.cpp
+++ b/Engine/Plugins/Runtime/nDisplay/Source/DisplayClusterConfigurator/Private/DisplayClusterConfiguratorStyle.cpp
@@ -97,6 +97,11 @@ void FDisplayClusterConfiguratorStyle::Initialize()
 		Set("DisplayClusterConfigurator.Toolbar.SaveToFile", new EDITOR_IMAGE_BRUSH_SVG("Starship/Common/Export", Icon20x20));
 	}
 
+	// Level Editor Toolbar
+	{
+		Set("DisplayClusterConfigurator.LevelEditor.ViewportsFrozen", new IMAGE_BRUSH_SVG("Icons/Viewport/nDisplayFrozen_16", Icon16x16));
+	}
+
 	// Cluster
 	{
 		Set("DisplayClusterConfigurator.TreeItems.Cluster", new IMAGE_BRUSH_SVG("Icons/Cluster/Cluster", Icon16x16));
diff --git a/Engine/Plugins/Runtime/nDisplay/Source/DisplayClusterConfigurator/Private/DisplayClusterConfiguratorToolbarExtensions.cpp b/Engine/Plugins/Runtime/nDisplay/Source/DisplayClusterConfigurator/Private/DisplayClusterConfiguratorToolbarExtensions.cpp
--- a/Engine/Plugins/Runtime/nDisplay/Source/DisplayClusterConfigurator/Private/DisplayClusterConfiguratorToolbarExtensions.cpp
+++ b/Engine/Plugins/Runtime/nDisplay/Source/DisplayClusterConfigurator/Private/DisplayClusterConfiguratorToolbarExtensions.cpp
@@ -39,7 +39,7 @@ void FDisplayClusterConfiguratorToolbarExtensions::CreateFreezeViewportsMenu(FTo
 						FName("UnfreezeAllViewports"),
 						LOCTEXT("UnfreezeAllViewportsLabel", "Unfreeze All Viewports"),
 						LOCTEXT("ViewportsFrozenWarningToolTip", "nDisplay viewports are frozen. Click to unfreeze all frozen viewports."),
-						FSlateIcon(FDisplayClusterConfiguratorStyle::Get().GetStyleSetName(), "ClassIcon.DisplayClusterRootActor"),
+						FSlateIcon(FDisplayClusterConfiguratorStyle::Get().GetStyleSetName(), "DisplayClusterConfigurator.LevelEditor.ViewportsFrozen"),
 						FUIAction(
 							FExecuteAction::CreateSP(this, &FDisplayClusterConfiguratorToolbarExtensions::UnfreezeAllViewports),
 							FCanExecuteAction::CreateSP(this, &FDisplayClusterConfiguratorToolbarExtensions::AreAnyViewportsFrozen)),
