Deprecate RHITransitionResource since D3D12_RESOURCE_STATE_TBD is not a valid state to use as Before state in a DX12 transition.
A transition using the high level RHI is needed to use Unknown as Before state.

diff --git a/Engine/Plugins/Runtime/nDisplay/Source/DisplayClusterProjection/Private/Policy/EasyBlend/Windows/DX12/DisplayClusterProjectionEasyBlendPolicyViewDataDX12.cpp b/Engine/Plugins/Runtime/nDisplay/Source/DisplayClusterProjection/Private/Policy/EasyBlend/Windows/DX12/DisplayClusterProjectionEasyBlendPolicyViewDataDX12.cpp
--- a/Engine/Plugins/Runtime/nDisplay/Source/DisplayClusterProjection/Private/Policy/EasyBlend/Windows/DX12/DisplayClusterProjectionEasyBlendPolicyViewDataDX12.cpp
+++ b/Engine/Plugins/Runtime/nDisplay/Source/DisplayClusterProjection/Private/Policy/EasyBlend/Windows/DX12/DisplayClusterProjectionEasyBlendPolicyViewDataDX12.cpp
@@ -148,7 +148,9 @@ bool FDisplayClusterProjectionEasyBlendPolicyViewDataDX12::ApplyWarpBlend_Render
 	{
 		return false;
 	}
-
+	
+	// Setup In/Out EasyBlend textures
+	RHICmdList.Transition(FRHITransitionInfo(OutputTexture, ERHIAccess::Unknown, ERHIAccess::RTV));
 	// Insert an outer query that encloses the whole batch
 	RHICmdList.EnqueueLambda(
 	[ViewData = SharedThis(this), InputTexture = InputTexture, OutputTexture = OutputTexture, InViewLocation = InViewInfo.ViewLocation](FRHICommandList& ExecutingCmdList)
@@ -189,15 +191,11 @@ bool FDisplayClusterProjectionEasyBlendPolicyViewDataDX12::ApplyWarpBlend_Render
 
 		if (SrcTexture && DestTexture && D3D12CommandQueue)
 		{
-			// Setup In/Out EasyBlend textures
-			D3D12DynamicRHI->RHITransitionResource(ExecutingCmdList, OutputTexture, D3D12_RESOURCE_STATE_RENDER_TARGET, D3D12_RESOURCE_BARRIER_ALL_SUBRESOURCES);
-
 			// Render EasyBlend on proxy
 			const EasyBlendSDKError Result2 = EasyBlendAPI->EasyBlendSDK_TransformInputToOutputDX12_CommandQueue(ViewData->EasyBlendMeshData.Get(), SrcTexture, DestTexture, RTVHandle);
-
-			D3D12DynamicRHI->RHITransitionResource(ExecutingCmdList, OutputTexture, D3D12_RESOURCE_STATE_PRESENT, D3D12_RESOURCE_BARRIER_ALL_SUBRESOURCES);
 		}
 	});
+	RHICmdList.Transition(FRHITransitionInfo(OutputTexture, ERHIAccess::RTV, ERHIAccess::Present));
 
 	return true;
 }
diff --git a/Engine/Plugins/Runtime/nDisplay/Source/DisplayClusterProjection/Private/Policy/VIOSO/Windows/DisplayClusterProjectionVIOSOPolicyViewData.cpp b/Engine/Plugins/Runtime/nDisplay/Source/DisplayClusterProjection/Private/Policy/VIOSO/Windows/DisplayClusterProjectionVIOSOPolicyViewData.cpp
--- a/Engine/Plugins/Runtime/nDisplay/Source/DisplayClusterProjection/Private/Policy/VIOSO/Windows/DisplayClusterProjectionVIOSOPolicyViewData.cpp
+++ b/Engine/Plugins/Runtime/nDisplay/Source/DisplayClusterProjection/Private/Policy/VIOSO/Windows/DisplayClusterProjectionVIOSOPolicyViewData.cpp
@@ -219,6 +219,8 @@ bool FDisplayClusterProjectionVIOSOPolicyViewData::RenderVIOSO_RenderThread(FRHI
 		case ERenderDevice::D3D12:
 		{
 			// Insert an outer query that encloses the whole batch
+			//experimental: add resource barrier
+			RHICmdList.Transition(FRHITransitionInfo(RenderTargetTexture, ERHIAccess::Unknown, ERHIAccess::RTV));
 			RHICmdList.EnqueueLambda([ViosoViewData = SharedThis(this), RenderTargetTexture = RenderTargetTexture, ShaderResourceTexture= ShaderResourceTexture](FRHICommandList& ExecutingCmdList)
 				{
 					ID3D12DynamicRHI* D3D12RHI = GetID3D12DynamicRHI();
@@ -229,15 +231,10 @@ bool FDisplayClusterProjectionVIOSOPolicyViewData::RenderVIOSO_RenderThread(FRHI
 					RenderInput.renderTarget = D3D12RHI->RHIGetResource(RenderTargetTexture);
 					RenderInput.rtvHandlePtr = RTVHandle.ptr;
 
-					//experimental: add resource barrier
-					D3D12RHI->RHITransitionResource(ExecutingCmdList, RenderTargetTexture, D3D12_RESOURCE_STATE_RENDER_TARGET, D3D12_RESOURCE_BARRIER_ALL_SUBRESOURCES);
-
-					if (ViosoViewData->WarperInterface->Render(&RenderInput, VWB_STATEMASK_ALL))
-					{
-						//experimental: add resource barrier
-						D3D12RHI->RHITransitionResource(ExecutingCmdList, RenderTargetTexture, D3D12_RESOURCE_STATE_PRESENT, D3D12_RESOURCE_BARRIER_ALL_SUBRESOURCES);
-					}
+					ViosoViewData->WarperInterface->Render(&RenderInput, VWB_STATEMASK_ALL);
 				});
+			//experimental: add resource barrier
+			RHICmdList.Transition(FRHITransitionInfo(RenderTargetTexture, ERHIAccess::RTV, ERHIAccess::Present));
 
 			return true;
 
diff --git a/Engine/Source/Runtime/D3D12RHI/Private/D3D12RHI.cpp b/Engine/Source/Runtime/D3D12RHI/Private/D3D12RHI.cpp
--- a/Engine/Source/Runtime/D3D12RHI/Private/D3D12RHI.cpp
+++ b/Engine/Source/Runtime/D3D12RHI/Private/D3D12RHI.cpp
@@ -741,11 +741,9 @@ void FD3D12DynamicRHI::RHIFinishExternalComputeWork(FRHICommandListBase& Executi
 
 void FD3D12DynamicRHI::RHITransitionResource(FRHICommandList& RHICmdList, FRHITexture* InTexture, D3D12_RESOURCE_STATES InState, uint32 InSubResource)
 {
-	FD3D12Texture* D3D12Texture = GetD3D12TextureFromRHITexture(InTexture);
-	for (uint32 GPUIndex : RHICmdList.GetGPUMask())
-	{
-		FD3D12CommandContext::Get(RHICmdList, GPUIndex).TransitionResource(D3D12Texture->GetResource(), D3D12_RESOURCE_STATE_TBD, InState, InSubResource);
-	}
+	FString StateString = ConvertToResourceStateString(InState);
+
+	UE_LOG(LogD3D12RHI, Error, TEXT("RHITransitionResource cannot transition from Unknown state to %s state on Texture %s subresource %d. Use RHICmdList.Transition() instead."), *StateString, *InTexture->GetName().GetPlainNameString(), InSubResource);
 }
 
 void FD3D12DynamicRHI::RHISignalManualFence(FRHICommandList& RHICmdList, ID3D12Fence* Fence, uint64 Value)
diff --git a/Engine/Source/Runtime/D3D12RHI/Public/ID3D12DynamicRHI.h b/Engine/Source/Runtime/D3D12RHI/Public/ID3D12DynamicRHI.h
--- a/Engine/Source/Runtime/D3D12RHI/Public/ID3D12DynamicRHI.h
+++ b/Engine/Source/Runtime/D3D12RHI/Public/ID3D12DynamicRHI.h
@@ -47,6 +47,8 @@ struct ID3D12DynamicRHI : public FDynamicRHI
 	virtual D3D12_CPU_DESCRIPTOR_HANDLE RHIGetRenderTargetView(FRHITexture* InTexture, int32 InMipIndex = 0, int32 InArraySliceIndex = 0) const = 0;
 
 	virtual void                  RHIFinishExternalComputeWork(FRHICommandListBase& ExecutingCmdList, uint32 InDeviceIndex, ID3D12GraphicsCommandList* InCommandList) = 0;
+	
+	UE_DEPRECATED				  (5.6, "RHITransitionResource without explicit Before and After states is deprecated. Please call RHICmdList.Transition() using ERHIAccess states to issue the transition, ERHIAccess::Unknown as Before state is supported.")
 	virtual void                  RHITransitionResource(FRHICommandList& RHICmdList, FRHITexture* InTexture, D3D12_RESOURCE_STATES InState, uint32 InSubResource) = 0;
 
 	virtual void                  RHISignalManualFence(FRHICommandList& RHICmdList, ID3D12Fence* Fence, uint64 Value) = 0;
