[nDisplay] [Virtual Production] UV Lightcards not getting OCIO applied in preview.

Reworked the logic to always call FDisplayClusterViewport::SetupSceneView() for viewports, even if they are not rendered (uvLC).

diff --git a/Engine/Plugins/Runtime/nDisplay/Source/DisplayCluster/Private/Render/Viewport/Preview/DisplayClusterViewportManagerPreview_Render.cpp b/Engine/Plugins/Runtime/nDisplay/Source/DisplayCluster/Private/Render/Viewport/Preview/DisplayClusterViewportManagerPreview_Render.cpp
--- a/Engine/Plugins/Runtime/nDisplay/Source/DisplayCluster/Private/Render/Viewport/Preview/DisplayClusterViewportManagerPreview_Render.cpp
+++ b/Engine/Plugins/Runtime/nDisplay/Source/DisplayCluster/Private/Render/Viewport/Preview/DisplayClusterViewportManagerPreview_Render.cpp
@@ -171,18 +171,22 @@ bool FDisplayClusterViewportManagerPreview::InitializeClusterNodePreview(const E
 
 							check(ViewportPtr->GetContexts().IsValidIndex(ViewIt.ContextNum));
 
-							if (ViewIt.IsViewportContextCanBeRendered() && ViewFamily->RenderTarget)
+							// Always call CalcSceneView() because this function also starts VE and after that calls FDisplayClusterViewport::SetupSceneView() -> OCIO
+							// uvLC viewport is not rendered, but late OCIO is used in FDisplayClusterViewportProxy::ApplyOCIO_RenderThread()
+							FSceneView* View = ViewportPtr->ViewportPreview->CalcSceneView(*ViewFamily, ViewIt.ContextNum);
+
+							// Remove the view that is not used for rendering.
+							if (View && (!ViewIt.IsViewportContextCanBeRendered() || ViewFamily->RenderTarget == nullptr))
 							{
-								// Calculate the player's view information.
-								if (FSceneView* View = ViewportPtr->ViewportPreview->CalcSceneView(*ViewFamily, ViewIt.ContextNum))
-								{
-									Views.Add(View);
-								}
+								ViewFamily->Views.Remove(View);
+
+								delete View;
+								View = nullptr;
 							}
-							else
+
+							if (View)
 							{
-								// This viewport is not rendered, but we need to calculate the math for it.
-								ViewportPtr->ViewportPreview->CalculateViewportContext(ViewIt.ContextNum);
+								Views.Add(View);
 							}
 						}
 
diff --git a/Engine/Plugins/Runtime/nDisplay/Source/DisplayCluster/Private/Render/Viewport/Preview/DisplayClusterViewportPreview.h b/Engine/Plugins/Runtime/nDisplay/Source/DisplayCluster/Private/Render/Viewport/Preview/DisplayClusterViewportPreview.h
--- a/Engine/Plugins/Runtime/nDisplay/Source/DisplayCluster/Private/Render/Viewport/Preview/DisplayClusterViewportPreview.h
+++ b/Engine/Plugins/Runtime/nDisplay/Source/DisplayCluster/Private/Render/Viewport/Preview/DisplayClusterViewportPreview.h
@@ -94,15 +94,6 @@ public:
 	*/
 	FSceneView* CalcSceneView(class FSceneViewFamilyContext& InOutViewFamily, uint32 ContextNum);
 
-	/**
-	* Calculate viewport math
-	*
-	* @param ContextNum - context index for rendering
-	*
-	* @return - true if viewport can be rendered
-	*/
-	bool CalculateViewportContext(uint32 ContextNum);
-
 	/** Returns true once if this type of log message can be displayed for the first time.*/
 	bool CanShowLogMsgOnce(const EDisplayClusterViewportPreviewShowLogMsgOnce& InLogState)
 	{
diff --git a/Engine/Plugins/Runtime/nDisplay/Source/DisplayCluster/Private/Render/Viewport/Preview/DisplayClusterViewportPreview_Render.cpp b/Engine/Plugins/Runtime/nDisplay/Source/DisplayCluster/Private/Render/Viewport/Preview/DisplayClusterViewportPreview_Render.cpp
--- a/Engine/Plugins/Runtime/nDisplay/Source/DisplayCluster/Private/Render/Viewport/Preview/DisplayClusterViewportPreview_Render.cpp
+++ b/Engine/Plugins/Runtime/nDisplay/Source/DisplayCluster/Private/Render/Viewport/Preview/DisplayClusterViewportPreview_Render.cpp
@@ -22,23 +22,6 @@
 ////////////////////////////////////////////////////////////////////////////////////////
 // FDisplayClusterViewportPreview
 ////////////////////////////////////////////////////////////////////////////////////////
-bool FDisplayClusterViewportPreview::CalculateViewportContext(uint32 InContextNum)
-{
-	if (FDisplayClusterViewport* InViewport = GetViewportImpl())
-	{
-		FVector ViewLocation;
-		FVector StereoViewLocation;
-		FRotator ViewRotation;
-		FMatrix ProjectionMatrix;
-
-		return Configuration->IsSceneOpened()
-			&& CalculateStereoViewOffset(*InViewport, InContextNum, ViewRotation, ViewLocation, StereoViewLocation)
-			&& InViewport->GetProjectionMatrix(InContextNum, ProjectionMatrix);
-	}
-
-	return false;
-}
-
 FSceneView* FDisplayClusterViewportPreview::CalcSceneView(FSceneViewFamilyContext& InOutViewFamily, uint32 InContextNum)
 {
 	FDisplayClusterViewport* InViewport = GetViewportImpl();
