[nDisplay] Fixed crash when deleting a geometry used by [REDACTED_PERSON].

Reworked geometry caching logic.
Removed public var FDisplayClusterWarpEye::bUpdateGeometryContext. This variable is no longer used.

diff --git a/Engine/Plugins/Runtime/nDisplay/Source/DisplayClusterWarp/Private/Policy/InFrustumFit/DisplayClusterWarpInFrustumFitPolicy.cpp b/Engine/Plugins/Runtime/nDisplay/Source/DisplayClusterWarp/Private/Policy/InFrustumFit/DisplayClusterWarpInFrustumFitPolicy.cpp
--- a/Engine/Plugins/Runtime/nDisplay/Source/DisplayClusterWarp/Private/Policy/InFrustumFit/DisplayClusterWarpInFrustumFitPolicy.cpp
+++ b/Engine/Plugins/Runtime/nDisplay/Source/DisplayClusterWarp/Private/Policy/InFrustumFit/DisplayClusterWarpInFrustumFitPolicy.cpp
@@ -45,9 +45,6 @@ void FDisplayClusterWarpInFrustumFitPolicy::HandleNewFrame(const TArray<TSharedP
 	// In the new frame, we have to find a new solution because the viewer's position or geometry may have changed.
 	// The number of viewports can also be changed at runtime. This changes the shape of the united geometry.
 	{
-		// Invalidate geometry cache for the new frame.
-		bGeometryContextsUpdated = false;
-
 		// Reset united geometry AABB
 		OptUnitedGeometryWorldAABB.Reset();
 
@@ -94,7 +91,6 @@ void FDisplayClusterWarpInFrustumFitPolicy::HandleNewFrame(const TArray<TSharedP
 			{
 				return;
 			}
-			bGeometryContextsUpdated = true;
 			OptUnitedGeometryWorldAABB = UnitedGeometryWorldAABBox;
 
 			// in the first iteration, use the center of the united AABB geometry as the view target.
@@ -149,8 +145,6 @@ void FDisplayClusterWarpInFrustumFitPolicy::HandleNewFrame(const TArray<TSharedP
 		FDisplayClusterWarpProjection UnitedWarpProjection;
 		if (CalcUnitedGeometryFrustum(InViewports, ContextNum, WorldScale, UnitedWarpProjection))
 		{
-			bGeometryContextsUpdated = true;
-
 			// If the view target is set to a fixed value instead of being computed by the group AABB, we do not want to alter the view direction,
 			// but make the frustum symmetric around that fixed direction. This involves expanding the asymmetric frustum so that its left and right,
 			// top and bottom angles are equal
diff --git a/Engine/Plugins/Runtime/nDisplay/Source/DisplayClusterWarp/Private/Policy/InFrustumFit/DisplayClusterWarpInFrustumFitPolicy.h b/Engine/Plugins/Runtime/nDisplay/Source/DisplayClusterWarp/Private/Policy/InFrustumFit/DisplayClusterWarpInFrustumFitPolicy.h
--- a/Engine/Plugins/Runtime/nDisplay/Source/DisplayClusterWarp/Private/Policy/InFrustumFit/DisplayClusterWarpInFrustumFitPolicy.h
+++ b/Engine/Plugins/Runtime/nDisplay/Source/DisplayClusterWarp/Private/Policy/InFrustumFit/DisplayClusterWarpInFrustumFitPolicy.h
@@ -118,8 +118,4 @@ private:
 	// This value is used to obtain the direction of view and to build the "projection plane".
 	// (The world space is used because each viewport can use its own origin.)
 	TOptional<FVector> OptOverrideWorldViewTarget;
-
-	// true if the geometry cache inside the WarpBlend api has been updated.
-	// This is used to optimize performance to avoid multiple updates when we call the WarpBlend api multiple times per frame.
-	bool bGeometryContextsUpdated = false;
 };
diff --git a/Engine/Plugins/Runtime/nDisplay/Source/DisplayClusterWarp/Private/Policy/InFrustumFit/DisplayClusterWarpInFrustumFitPolicy_Math.cpp b/Engine/Plugins/Runtime/nDisplay/Source/DisplayClusterWarp/Private/Policy/InFrustumFit/DisplayClusterWarpInFrustumFitPolicy_Math.cpp
--- a/Engine/Plugins/Runtime/nDisplay/Source/DisplayClusterWarp/Private/Policy/InFrustumFit/DisplayClusterWarpInFrustumFitPolicy_Math.cpp
+++ b/Engine/Plugins/Runtime/nDisplay/Source/DisplayClusterWarp/Private/Policy/InFrustumFit/DisplayClusterWarpInFrustumFitPolicy_Math.cpp
@@ -101,8 +101,6 @@ void FDisplayClusterWarpInFrustumFitPolicy::BeginCalcFrustum(IDisplayClusterView
 	// Override view direction for entire group
 	if (WarpData.WarpEye.IsValid())
 	{
-		WarpData.WarpEye->bUpdateGeometryContext = !bGeometryContextsUpdated;
-
 		// Overrides the view direction vector to build a custom "projection plane".
 		if (GDisplayClusterWarpInFrustumFitPolicyUseStaticViewDirectionForMPCDIProfile2D && WarpBlend->GetWarpProfileType() == EDisplayClusterWarpProfileType::warp_2D)
 		{
diff --git a/Engine/Plugins/Runtime/nDisplay/Source/DisplayClusterWarp/Private/WarpBlend/DisplayClusterWarpBlend.cpp b/Engine/Plugins/Runtime/nDisplay/Source/DisplayClusterWarp/Private/WarpBlend/DisplayClusterWarpBlend.cpp
--- a/Engine/Plugins/Runtime/nDisplay/Source/DisplayClusterWarp/Private/WarpBlend/DisplayClusterWarpBlend.cpp
+++ b/Engine/Plugins/Runtime/nDisplay/Source/DisplayClusterWarp/Private/WarpBlend/DisplayClusterWarpBlend.cpp
@@ -92,18 +92,15 @@ bool FDisplayClusterWarpBlend::CalcFrustumContext(const TSharedPtr<FDisplayClust
 
 	FDisplayClusterWarpBlendMath_Frustum Frustum(CurrentWarpData, GeometryContext);
 
-	// Override settings from warp policy
-	BeginCalcFrustum(InWarpEye);
-
-	if (InWarpEye->bUpdateGeometryContext)
+	if (!GeometryContext.UpdateGeometryContext(InWarpEye->WorldScale))
 	{
-		if (!GeometryContext.UpdateGeometryContext(InWarpEye->WorldScale))
-		{
-			// wrong geometry
-			return false;
-		}
+		// wrong geometry
+		return false;
 	}
 
+	// Override settings from warp policy
+	BeginCalcFrustum(InWarpEye);
+
 	if (!Frustum.CalcFrustum())
 	{
 		return false;
diff --git a/Engine/Plugins/Runtime/nDisplay/Source/DisplayClusterWarp/Private/WarpBlend/DisplayClusterWarpBlend_GeometryProxy.cpp b/Engine/Plugins/Runtime/nDisplay/Source/DisplayClusterWarp/Private/WarpBlend/DisplayClusterWarpBlend_GeometryProxy.cpp
--- a/Engine/Plugins/Runtime/nDisplay/Source/DisplayClusterWarp/Private/WarpBlend/DisplayClusterWarpBlend_GeometryProxy.cpp
+++ b/Engine/Plugins/Runtime/nDisplay/Source/DisplayClusterWarp/Private/WarpBlend/DisplayClusterWarpBlend_GeometryProxy.cpp
@@ -15,26 +15,36 @@
 /////////////////////////////////////////////////////////////////////////////////
 bool FDisplayClusterWarpBlend_GeometryProxy::UpdateFrustumGeometry()
 {
+	bIsGeometryValid = false;
 	switch (FrustumGeometryType)
 	{
 	case EDisplayClusterWarpFrustumGeometryType::WarpMesh:
-		return ImplUpdateFrustumGeometry_WarpMesh();
+		bIsGeometryValid = ImplUpdateFrustumGeometry_WarpMesh();
+		break;
 
 	case EDisplayClusterWarpFrustumGeometryType::WarpProceduralMesh:
-		return ImplUpdateFrustumGeometry_WarpProceduralMesh();
+		bIsGeometryValid = ImplUpdateFrustumGeometry_WarpProceduralMesh();
+		break;
 
 	case EDisplayClusterWarpFrustumGeometryType::WarpMap:
-		return ImplUpdateFrustumGeometry_WarpMap();
+		bIsGeometryValid = ImplUpdateFrustumGeometry_WarpMap();
+		break;
 
 	case EDisplayClusterWarpFrustumGeometryType::MPCDIAttributes:
-		return ImplUpdateFrustumGeometry_MPCDIAttributes();
+		bIsGeometryValid = ImplUpdateFrustumGeometry_MPCDIAttributes();
+		break;
 
 	default:
 		break;
 	}
 
-	bIsGeometryValid = false;
-	return false;
+	if (!bIsGeometryValid)
+	{
+		// In case of an error, the cached data is invalidated
+		bIsGeometryCacheValid = false;
+	}
+	
+	return bIsGeometryValid;
 }
 
 const IDisplayClusterRender_MeshComponentProxy* FDisplayClusterWarpBlend_GeometryProxy::GetWarpMeshProxy_RenderThread() const
@@ -78,41 +88,20 @@ bool FDisplayClusterWarpBlend_GeometryProxy::MarkWarpFrustumGeometryComponentDir
 
 bool FDisplayClusterWarpBlend_GeometryProxy::ImplUpdateFrustumGeometry_MPCDIAttributes()
 {
-	bIsGeometryValid = false;
-
-	// Update caches
-	if ((bIsGeometryValid = bIsGeometryCacheValid) == false)
-	{
-		bIsGeometryValid = ImplUpdateFrustumGeometryCache_MPCDIAttributes();
-	}
-
 	GeometryCache.GeometryToOrigin = FTransform::Identity;
 
-	return bIsGeometryValid;
+	return ImplUpdateFrustumGeometryCache_MPCDIAttributes();
 }
 
 bool FDisplayClusterWarpBlend_GeometryProxy::ImplUpdateFrustumGeometry_WarpMap()
 {
-	bIsGeometryValid = false;
-
-	if (WarpMapTexture.IsValid() && WarpMapTexture->IsEnabled())
-	{
-		// Update caches
-		if ((bIsGeometryValid = bIsGeometryCacheValid) == false)
-		{
-			bIsGeometryValid = ImplUpdateFrustumGeometryCache_WarpMap();
-		}
-
-		GeometryCache.GeometryToOrigin = FTransform::Identity;
-	}
+	GeometryCache.GeometryToOrigin = FTransform::Identity;
 
-	return bIsGeometryValid;
+	return ImplUpdateFrustumGeometryCache_WarpMap();
 }
 
 bool FDisplayClusterWarpBlend_GeometryProxy::ImplUpdateFrustumGeometry_WarpMesh()
 {
-	bIsGeometryValid = false;
-
 	if (!WarpMeshComponent.IsValid())
 	{
 		return false;
@@ -127,6 +116,7 @@ bool FDisplayClusterWarpBlend_GeometryProxy::ImplUpdateFrustumGeometry_WarpMesh(
 		// mesh deleted?
 		WarpMeshComponent->ReleaseProxyGeometry();
 		bIsMeshComponentLost = true;
+
 		return false;
 	};
 
@@ -138,9 +128,9 @@ bool FDisplayClusterWarpBlend_GeometryProxy::ImplUpdateFrustumGeometry_WarpMesh(
 	}
 	
 	// Update caches
-	if ((bIsGeometryValid = bIsGeometryCacheValid) == false)
+	if (!ImplUpdateFrustumGeometryCache_WarpMesh())
 	{
-		bIsGeometryValid = ImplUpdateFrustumGeometryCache_WarpMesh();
+		return false;
 	}
 
 	if (OriginComponent)
@@ -155,13 +145,11 @@ bool FDisplayClusterWarpBlend_GeometryProxy::ImplUpdateFrustumGeometry_WarpMesh(
 		GeometryCache.GeometryToOrigin = StaticMeshComponent->GetRelativeTransform();
 	}
 
-	return bIsGeometryValid;
+	return true;
 }
 
 bool FDisplayClusterWarpBlend_GeometryProxy::ImplUpdateFrustumGeometry_WarpProceduralMesh()
 {
-	bIsGeometryValid = false;
-
 	if (!WarpMeshComponent.IsValid())
 	{
 		return false;
@@ -187,24 +175,12 @@ bool FDisplayClusterWarpBlend_GeometryProxy::ImplUpdateFrustumGeometry_WarpProce
 	}
 
 	// Update caches
-	if ((bIsGeometryValid = bIsGeometryCacheValid) == false)
+	if (!ImplUpdateFrustumGeometryCache_WarpProceduralMesh())
 	{
-		bIsGeometryValid = ImplUpdateFrustumGeometryCache_WarpProceduralMesh();
+		return false;
 	}
 
-	// These matrices were copied from LocalPlayer.cpp.
-	// They change the coordinate system from the Unreal "Game" coordinate system to the Unreal "Render" coordinate system
-	static const FMatrix Game2Render(
-		FPlane(0, 0, 1, 0),
-		FPlane(1, 0, 0, 0),
-		FPlane(0, 1, 0, 0),
-		FPlane(0, 0, 0, 1));
-
-	static const FMatrix Render2Game = Game2Render.Inverse();
-
-	// Procedural mesh coord is render, so convert to game at first
-	FMatrix GeometryToOriginMatrix;// = Render2Game;
-
+	FMatrix GeometryToOriginMatrix;
 	if (OriginComponent)
 	{
 		FMatrix MeshToWorldMatrix = ProceduralMeshComponent->GetComponentTransform().ToMatrixWithScale();
@@ -219,62 +195,77 @@ bool FDisplayClusterWarpBlend_GeometryProxy::ImplUpdateFrustumGeometry_WarpProce
 
 	GeometryCache.GeometryToOrigin.SetFromMatrix(GeometryToOriginMatrix);
 
-	return bIsGeometryValid;
+	return true;
 }
 
 bool FDisplayClusterWarpBlend_GeometryProxy::ImplUpdateFrustumGeometryCache_WarpMesh()
 {
-	bIsGeometryValid = false;
-
 	if (WarpMeshComponent.IsValid())
 	{
 		const FStaticMeshLODResources* StaticMeshLODResources = WarpMeshComponent->GetStaticMeshComponentLODResources(StaticMeshComponentLODIndex);
 		if (StaticMeshLODResources != nullptr)
 		{
+			if (bIsGeometryCacheValid)
+			{
+				return true;
+			}
+
 			FDisplayClusterWarpBlendMath_WarpMesh MeshHelper(*StaticMeshLODResources);
 
 			GeometryCache.AABBox = MeshHelper.CalcAABBox();
 			MeshHelper.CalcSurfaceVectors(GeometryCache.SurfaceViewNormal, GeometryCache.SurfaceViewPlane);
 
 			bIsGeometryCacheValid = true;
-			bIsGeometryValid = true;
+
+			return true;
 		}
 	}
 
-	return bIsGeometryValid;
+	bIsGeometryCacheValid = false;
+
+	return false;
 }
 
 bool FDisplayClusterWarpBlend_GeometryProxy::ImplUpdateFrustumGeometryCache_WarpProceduralMesh()
 {
-	bIsGeometryValid = false;
-
 	if (WarpMeshComponent.IsValid())
 	{
 		const FProcMeshSection* ProcMeshSection = WarpMeshComponent->GetProceduralMeshComponentSection(ProceduralMeshComponentSectionIndex);
 		if (ProcMeshSection != nullptr)
 		{
+			if (bIsGeometryCacheValid)
+			{
+				return true;
+			}
+
 			FDisplayClusterWarpBlendMath_WarpProceduralMesh ProceduralMeshHelper(*ProcMeshSection);
 
 			GeometryCache.AABBox = ProceduralMeshHelper.CalcAABBox();
 			ProceduralMeshHelper.CalcSurfaceVectors(GeometryCache.SurfaceViewNormal, GeometryCache.SurfaceViewPlane);
 
 			bIsGeometryCacheValid = true;
-			bIsGeometryValid = true;
 
 			return true;
 		}
 	}
 
+	bIsGeometryCacheValid = false;
+
 	return false;
 }
 
 bool FDisplayClusterWarpBlend_GeometryProxy::ImplUpdateFrustumGeometryCache_WarpMap()
 {
-	// WarpMap must be initialized before, outside
-	bIsGeometryValid = false;
-	
+
 	if (WarpMapTexture.IsValid() && WarpMapTexture->IsEnabled())
 	{
+		if (bIsGeometryCacheValid)
+		{
+			// use cached values
+			return true;
+		}
+
+		// Update cache
 		FDisplayClusterWarpBlendMath_WarpMap DataHelper(*(WarpMapTexture.Get()));
 
 		GeometryCache.AABBox = DataHelper.GetAABBox();
@@ -282,22 +273,27 @@ bool FDisplayClusterWarpBlend_GeometryProxy::ImplUpdateFrustumGeometryCache_Warp
 		GeometryCache.SurfaceViewPlane  = DataHelper.GetSurfaceViewPlane();
 
 		bIsGeometryCacheValid = true;
-		bIsGeometryValid = true;
 
 		return true;
 	}
 
+	bIsGeometryCacheValid = false;
+
 	return false;
 }
 
 bool FDisplayClusterWarpBlend_GeometryProxy::ImplUpdateFrustumGeometryCache_MPCDIAttributes()
 {
-	// WarpMap must be initialized before, outside
-	bIsGeometryValid = false;
-
 	switch (MPCDIAttributes.ProfileType)
 	{
 	case EDisplayClusterWarpProfileType::warp_2D:
+		if (bIsGeometryCacheValid)
+		{
+			// use cached value
+			return true;
+		}
+
+		// Update cached value
 		GeometryCache.SurfaceViewNormal = FVector(1, 0, 0);
 		GeometryCache.SurfaceViewPlane = FVector(1, 0, 0);
 
@@ -311,16 +307,18 @@ bool FDisplayClusterWarpBlend_GeometryProxy::ImplUpdateFrustumGeometryCache_MPCD
 
 			GeometryCache.AABBox = WarpAABB;
 		}
-		break;
+
+		bIsGeometryCacheValid = true;
+
+		return true;
 
 	default:
-		return false;
+		break;
 	}
 
-	bIsGeometryCacheValid = true;
-	bIsGeometryValid = true;
+	bIsGeometryCacheValid = false;
 
-	return true;
+	return false;
 }
 
 bool FDisplayClusterWarpBlend_GeometryProxy::UpdateFrustumGeometryLOD(const FIntPoint& InSizeLOD)
diff --git a/Engine/Plugins/Runtime/nDisplay/Source/DisplayClusterWarp/Public/Containers/DisplayClusterWarpEye.h b/Engine/Plugins/Runtime/nDisplay/Source/DisplayClusterWarp/Public/Containers/DisplayClusterWarpEye.h
--- a/Engine/Plugins/Runtime/nDisplay/Source/DisplayClusterWarp/Public/Containers/DisplayClusterWarpEye.h
+++ b/Engine/Plugins/Runtime/nDisplay/Source/DisplayClusterWarp/Public/Containers/DisplayClusterWarpEye.h
@@ -49,9 +49,6 @@ public:
 	// Warp policy
 	TSharedPtr<IDisplayClusterWarpPolicy, ESPMode::ThreadSafe> WarpPolicy;
 
-	// Is geometry context data must should be updated
-	bool bUpdateGeometryContext = true;
-
 	// (optional) this a view target point. By default used AABB center.
 	TOptional<FVector> OverrideViewTarget;
 
