nDisplay: ICVFX panel change to new style toolbar.

Pending:
* Final icons for dome and orthographic projection
* Yellow highlight when draw lightcards is acive.
* Extra TRS menu at the far left.

#rnx

diff --git a/Engine/Plugins/Runtime/nDisplay/Source/DisplayClusterLightCardEditor/Private/DisplayClusterLightCardEditorStyle.cpp b/Engine/Plugins/Runtime/nDisplay/Source/DisplayClusterLightCardEditor/Private/DisplayClusterLightCardEditorStyle.cpp
--- a/Engine/Plugins/Runtime/nDisplay/Source/DisplayClusterLightCardEditor/Private/DisplayClusterLightCardEditorStyle.cpp
+++ b/Engine/Plugins/Runtime/nDisplay/Source/DisplayClusterLightCardEditor/Private/DisplayClusterLightCardEditorStyle.cpp
@@ -3,6 +3,7 @@
 #include "DisplayClusterLightCardEditorStyle.h"
 
 #include "Styling/SlateStyleMacros.h"
+#include "Styling/ToolBarStyle.h"
 
 FDisplayClusterLightCardEditorStyle::FDisplayClusterLightCardEditorStyle(): FSlateStyleSet("DisplayClusterLightCardEditorStyle")
 {
@@ -10,29 +11,31 @@ FDisplayClusterLightCardEditorStyle::FDisplayClusterLightCardEditorStyle(): FSla
 	const FVector2D Icon16x16(16.0f, 16.0f);
 
 	// Set miscellaneous icons
-	SetContentRoot(FPaths::EnginePluginsDir() / TEXT("Runtime/nDisplay/Content/Icons/OperatorPanel/"));
+	SetContentRoot(FPaths::EnginePluginsDir() / TEXT("Runtime/nDisplay/Content/"));
 	SetCoreContentRoot(FPaths::EngineContentDir() / TEXT("Editor/Slate"));
 
-	Set("DrawPoly", new IMAGE_BRUSH("PolyPath_40x", Icon40x40));
+	// Core Content
+	{
+		Set("DisplayClusterLightCardEditor.Labels", new CORE_IMAGE_BRUSH_SVG("Starship/Common/IssueTracker", Icon16x16));
+		Set("DisplayClusterLightCardEditor.IconSymbol", new CORE_IMAGE_BRUSH_SVG("Starship/Common/IssueTracker", Icon16x16));
+	}
+
+	// Content
+	{
+		Set("DisplayClusterLightCardEditor.LabelSymbol", new IMAGE_BRUSH_SVG("Icons/OperatorPanel/LabelSymbol", Icon16x16));
+
+		Set("DisplayClusterLightCardEditor.UV", new IMAGE_BRUSH_SVG("Icons/LightCard/LightCardUV", Icon16x16));
+		Set("DisplayClusterLightCardEditor.Dome", new IMAGE_BRUSH_SVG("Icons/OperatorPanel/Dome", Icon16x16));
+		Set("DisplayClusterLightCardEditor.Orthographic", new IMAGE_BRUSH_SVG("Icons/OperatorPanel/Orthographic", Icon16x16));
+
+		Set("DisplayClusterLightCardEditor.DrawPoly", new IMAGE_BRUSH("Icons/OperatorPanel/PolyPath_40x", Icon40x40));
+		Set("DisplayClusterLightCardEditor.ActorHidden", new IMAGE_BRUSH_SVG("Icons/OperatorPanel/ActorHidden", Icon16x16));
+		Set("DisplayClusterLightCardEditor.ActorNotHidden", new IMAGE_BRUSH_SVG("Icons/OperatorPanel/ActorNotHidden", Icon16x16));
+		Set("DisplayClusterLightCardEditor.Template", new IMAGE_BRUSH_SVG("Icons/OperatorPanel/Template", Icon16x16));
+		Set("DisplayClusterLightCardEditor.FrustumOnTop", new IMAGE_BRUSH_SVG("Icons/OperatorPanel/FrustumOnTop", Icon16x16));
+		Set("DisplayClusterLightCardEditor.FrustumUnderneath", new IMAGE_BRUSH_SVG("Icons/OperatorPanel/FrustumUnderneath", Icon16x16));
+		Set("DisplayClusterLightCardEditor.ViewportsFrozen", new IMAGE_BRUSH_SVG("Icons/Viewport/nDisplayFrozen_16", Icon16x16));
+	}
 
-	Set("DisplayClusterLightCardEditor.PasteHere", new CORE_IMAGE_BRUSH_SVG("Starship/Actors/paste-here", Icon16x16));
-
-	Set("DisplayClusterLightCardEditor.ViewOrientationTop", new CORE_IMAGE_BRUSH_SVG("Starship/Common/ViewTop", Icon16x16));
-	Set("DisplayClusterLightCardEditor.ViewOrientationBottom", new CORE_IMAGE_BRUSH_SVG("Starship/Common/ViewBottom", Icon16x16));
-	Set("DisplayClusterLightCardEditor.ViewOrientationLeft", new CORE_IMAGE_BRUSH_SVG("Starship/Common/ViewLeft", Icon16x16));
-	Set("DisplayClusterLightCardEditor.ViewOrientationRight", new CORE_IMAGE_BRUSH_SVG("Starship/Common/ViewRight", Icon16x16));
-	Set("DisplayClusterLightCardEditor.ViewOrientationFront", new CORE_IMAGE_BRUSH_SVG("Starship/Common/ViewFront", Icon16x16));
-	Set("DisplayClusterLightCardEditor.ViewOrientationBack", new CORE_IMAGE_BRUSH_SVG("Starship/Common/ViewBack", Icon16x16));
-
-	Set("DisplayClusterLightCardEditor.Labels", new CORE_IMAGE_BRUSH_SVG("Starship/Common/IssueTracker", Icon16x16));
-
-	Set("DisplayClusterLightCardEditor.ActorHidden", new IMAGE_BRUSH_SVG("ActorHidden", Icon16x16));
-	Set("DisplayClusterLightCardEditor.ActorNotHidden", new IMAGE_BRUSH_SVG("ActorNotHidden", Icon16x16));
-
-	Set("DisplayClusterLightCardEditor.Template", new IMAGE_BRUSH_SVG("Template", Icon16x16));
-
-	Set("DisplayClusterLightCardEditor.FrustumOnTop", new IMAGE_BRUSH_SVG("FrustumOnTop", Icon16x16));
-	Set("DisplayClusterLightCardEditor.FrustumUnderneath", new IMAGE_BRUSH_SVG("FrustumUnderneath", Icon16x16));
-		
 	FSlateStyleRegistry::RegisterSlateStyle(*this);
 }
diff --git a/Engine/Plugins/Runtime/nDisplay/Source/DisplayClusterLightCardEditor/Private/Viewport/DisplayClusterLightcardEditorViewport.cpp b/Engine/Plugins/Runtime/nDisplay/Source/DisplayClusterLightCardEditor/Private/Viewport/DisplayClusterLightcardEditorViewport.cpp
--- a/Engine/Plugins/Runtime/nDisplay/Source/DisplayClusterLightCardEditor/Private/Viewport/DisplayClusterLightcardEditorViewport.cpp
+++ b/Engine/Plugins/Runtime/nDisplay/Source/DisplayClusterLightCardEditor/Private/Viewport/DisplayClusterLightcardEditorViewport.cpp
@@ -16,6 +16,7 @@
 
 #include "Algo/Transform.h"
 #include "EditorViewportCommands.h"
+#include "Internationalization/Text.h"
 #include "ScopedTransaction.h"
 #include "SEditorViewportToolBarButton.h"
 #include "SEditorViewportToolBarMenu.h"
@@ -26,13 +27,17 @@
 #include "Slate/SceneViewport.h"
 #include "Styling/AppStyle.h"
 #include "Styling/SlateIconFinder.h"
+#include "ToolMenu.h"
+#include "ToolMenus.h"
+#include "ToolMenuSection.h"
 
 #define LOCTEXT_NAMESPACE "DisplayClusterLightcardEditorViewport"
 
+
 class SDisplayClusterLightCardEditorViewportToolBar : public SViewportToolBar
 {
 public:
-	SLATE_BEGIN_ARGS(SDisplayClusterLightCardEditorViewportToolBar){}
+	SLATE_BEGIN_ARGS(SDisplayClusterLightCardEditorViewportToolBar) {}
 		SLATE_ARGUMENT(TWeakPtr<SDisplayClusterLightCardEditorViewport>, EditorViewport)
 	SLATE_END_ARGS()
 
@@ -40,202 +45,412 @@ public:
 	void Construct(const FArguments& InArgs)
 	{
 		EditorViewport = InArgs._EditorViewport;
-		static const FName DefaultForegroundName("DefaultForeground");
 
 		this->ChildSlot
-		[
-			SNew(SBorder)
-			.BorderImage(FAppStyle::GetBrush("NoBorder"))
-			.ForegroundColor(FAppStyle::GetSlateColor(DefaultForegroundName))
 			[
-				SNew(SHorizontalBox)
-				+SHorizontalBox::Slot()
-				.AutoWidth()
-				.Padding(2.0f, 2.0f)
-				[
-					SNew(SEditorViewportToolbarMenu)
-					.ParentToolBar(SharedThis(this))
-					.Cursor(EMouseCursor::Default)
-					.Image("EditorViewportToolBar.OptionsDropdown")
-					.OnGetMenuContent(this, &SDisplayClusterLightCardEditorViewportToolBar::GeneratePreviewMenu)
-				]
-				+ SHorizontalBox::Slot()
-				.AutoWidth()
-				.Padding(2.0f, 2.0f)
-				[
-					SNew(SEditorViewportToolbarMenu)
-					.ParentToolBar(SharedThis(this))
-					.Cursor(EMouseCursor::Default)
-					.Label(this, &SDisplayClusterLightCardEditorViewportToolBar::GetProjectionMenuLabel)
-					.OnGetMenuContent(this, &SDisplayClusterLightCardEditorViewportToolBar::GenerateProjectionMenu)
-				]
-				+ SHorizontalBox::Slot()
-				.AutoWidth()
-				.Padding(2.0f, 2.0f)
-				[
-					SNew(SEditorViewportToolbarMenu)
-					.ParentToolBar(SharedThis(this))
-					.Cursor(EMouseCursor::Default)
-					.Label(LOCTEXT("LightCardEditorViewMenuLabel", "View"))
-					.OnGetMenuContent(this, &SDisplayClusterLightCardEditorViewportToolBar::GenerateViewMenu)
-				]
-				+ SHorizontalBox::Slot()
-				.AutoWidth()
-				.Padding(2.0f, 2.0f)
-				[
-					SNew(SEditorViewportToolbarMenu)
-					.ParentToolBar(SharedThis(this))
-					.Cursor(EMouseCursor::Default)
-					.Label(LOCTEXT("LightCardEditorShowMenuLabel", "Show"))
-					.OnGetMenuContent(this, &SDisplayClusterLightCardEditorViewportToolBar::GenerateShowMenu)
-				]
-				+ SHorizontalBox::Slot()
-				.AutoWidth()
-				.Padding(2.0f, 2.0f)
-				[
-					SNew(SEditorViewportToolBarButton)	
-					.ButtonType(EUserInterfaceActionType::Button)
-					.ButtonStyle(&FAppStyle::Get().GetWidgetStyle<FButtonStyle>("EditorViewportToolBar.WarningButton"))
-					.OnClicked(this, &SDisplayClusterLightCardEditorViewportToolBar::OnViewportsFrozenWarningClicked)
-					.Visibility(this, &SDisplayClusterLightCardEditorViewportToolBar::GetViewportsFrozenWarningVisibility)
-					.ToolTipText(LOCTEXT("DisplayClusterViewportsFrozenOff_ToolTip", "The nDisplay viewports are frozen. Click to unfreeze the viewports."))
-					.Content()
-					[
-						SNew(STextBlock)
-						.TextStyle(&FAppStyle::Get().GetWidgetStyle<FTextBlockStyle>("SmallText"))
-						.Text(LOCTEXT("DisplayClusterViewportsFrozen", "nDisplay Viewports Frozen"))
-					]
-				]
-				+ SHorizontalBox::Slot()
-				.Padding(3.0f, 1.0f)
-				.HAlign(HAlign_Right)
-				[
-					MakeDrawingToolBar()
-				]
-				+ SHorizontalBox::Slot()
-				.Padding( 3.0f, 1.0f )
-				.HAlign( HAlign_Right )
-				[
-					MakeTransformToolBar()
-				]
-			]
-		];
+				BuildToolbar()
+			];
 
 		SViewportToolBar::Construct(SViewportToolBar::FArguments());
 	}
 
-	/** Creates the preview menu */
-	TSharedRef<SWidget> GeneratePreviewMenu() const
+	/** Constructs the unified toolbar using UToolMenus. */
+	TSharedRef<SWidget> BuildToolbar()
 	{
-		TSharedPtr<const FUICommandList> CommandList = EditorViewport.IsValid()? EditorViewport.Pin()->GetCommandList(): nullptr;
+		const FName ToolbarName = FName("DisplayClusterLightCardEditor.ViewportToolbar");
+
+		// Register the toolbar menu if it doesn't already exist.
+		
+		UToolMenu* ToolbarMenu = nullptr;
+
+		if (!UToolMenus::Get()->IsMenuRegistered(ToolbarName))
+		{
+			ToolbarMenu = UToolMenus::Get()->RegisterMenu(ToolbarName, NAME_None, EMultiBoxType::SlimHorizontalToolBar);
+		}
+		else
+		{
+			ToolbarMenu = UToolMenus::Get()->FindMenu(ToolbarName);
+		}
 
-		const bool bInShouldCloseWindowAfterMenuSelection = true;
+		ToolbarMenu->StyleName = "ViewportToolbar";
 
-		FMenuBuilder PreviewOptionsMenuBuilder(bInShouldCloseWindowAfterMenuSelection, CommandList);
+		// Left section: TRS buttons.
 		{
-			// TODO: Any advanced viewport options can be added here
+			FToolMenuSection& LeftSection = ToolbarMenu->FindOrAddSection("Left");
+
+			LeftSection.AddEntry(FToolMenuEntry::InitToolBarButton(
+				FEditorViewportCommands::Get().TranslateMode,
+				TAttribute<FText>(FText::GetEmpty()) // empty label
+			));
+
+			LeftSection.AddEntry(FToolMenuEntry::InitToolBarButton(
+				FEditorViewportCommands::Get().RotateMode,
+				TAttribute<FText>(FText::GetEmpty())
+			));
+
+			LeftSection.AddEntry(FToolMenuEntry::InitToolBarButton(
+				FEditorViewportCommands::Get().ScaleMode,
+				TAttribute<FText>(FText::GetEmpty())
+			));
+
+			LeftSection.AddSeparator("CoordinateSystemSeparator");
+
+			// Coordinate system toggle button
+			LeftSection.AddEntry(CreateCoordinateSystemSubMenuEntry());
 		}
 
-		return PreviewOptionsMenuBuilder.MakeWidget();
+		// Right section: Frozen viewports | projection | Show | DrawLightCard
+		{
+			FToolMenuSection& RightSection = ToolbarMenu->FindOrAddSection("Right");
+			RightSection.Alignment = EToolMenuSectionAlign::Last;
+
+			// Frozen viewports Warning button
+			{
+				FToolMenuEntry UnfreezeAllEntry = FToolMenuEntry::InitMenuEntry(
+					FName("UnfreezeAllViewports"),
+					LOCTEXT("UnfreezeAllViewportsLabel", "Unfreeze All Viewports"),
+					LOCTEXT("ViewportsFrozenWarningToolTip", "Outer viewports are frozen. Click to unfreeze them."),
+					FSlateIcon(FDisplayClusterLightCardEditorStyle::Get().GetStyleSetName(), "DisplayClusterLightCardEditor.ViewportsFrozen"),
+					FUIAction(
+						FExecuteAction::CreateSP(this, &SDisplayClusterLightCardEditorViewportToolBar::UnfreezeAllViewports),
+						FCanExecuteAction::CreateSP(this, &SDisplayClusterLightCardEditorViewportToolBar::AreViewportsFrozen)
+					),
+					EUserInterfaceActionType::Button
+				);
+
+				UnfreezeAllEntry.ToolBarData.LabelOverride = FText();
+				UnfreezeAllEntry.Visibility = [this]() -> bool { return AreViewportsFrozen(); };
+
+				UnfreezeAllEntry.SubMenuData.Style.StyleSet = &FDisplayClusterLightCardEditorStyle::Get();
+				UnfreezeAllEntry.StyleNameOverride = "ViewportToolbarWarning.Raised";
+
+				RightSection.AddEntry(UnfreezeAllEntry);
+			}
+
+			// Projection dropdown.
+			{
+				RightSection.AddSubMenu(
+					"ProjectionView",
+					TAttribute<FText>::Create(TAttribute<FText>::FGetter::CreateSP(this, &SDisplayClusterLightCardEditorViewportToolBar::GetProjectionMenuLabel)),
+					LOCTEXT("ProjectionViewTooltip", "Select Projection/View Options"),
+					FNewToolMenuDelegate::CreateSP(this, &SDisplayClusterLightCardEditorViewportToolBar::GenerateProjectionViewMenu),
+					true /* bInOpenSubMenuOnClick */,
+					TAttribute<FSlateIcon>::Create(TAttribute<FSlateIcon>::FGetter::CreateSP(this, &SDisplayClusterLightCardEditorViewportToolBar::GetProjectionMenuIcon))
+				);
+			}
+
+			// Show (eye) dropdown.
+			{
+				FToolMenuEntry& ShowEntry = RightSection.AddSubMenu(
+					"Show",
+					LOCTEXT("ShowSubmenuLabel", "Show"),
+					LOCTEXT("ShowSubmenuTooltip", "Show flags related to the current viewport"),
+					FNewToolMenuDelegate::CreateLambda(
+						[](UToolMenu* InMenu)
+						{
+							FToolMenuSection& SubSection = InMenu->FindOrAddSection("ShowFlags", LOCTEXT("ShowFlagsSection", "Show Flags"));
+
+							SubSection.AddMenuEntry(
+								FDisplayClusterLightCardEditorCommands::Get().ToggleAllLabels,
+								TAttribute<FText>(),
+								TAttribute<FText>(),
+								TAttribute<FSlateIcon>(FSlateIcon(FDisplayClusterLightCardEditorStyle::Get().GetStyleSetName(), "DisplayClusterLightCardEditor.LabelSymbol"))
+							);
+
+							SubSection.AddMenuEntry(
+								FDisplayClusterLightCardEditorCommands::Get().ToggleIconVisibility,
+								TAttribute<FText>(),
+								TAttribute<FText>(),
+								TAttribute<FSlateIcon>(FSlateIcon(FDisplayClusterLightCardEditorStyle::Get().GetStyleSetName(), "DisplayClusterLightCardEditor.IconSymbol"))
+							);
+						}
+					),
+					false,
+					FSlateIcon(FAppStyle::GetAppStyleSetName(), "Level.VisibleHighlightIcon16x")
+				);
+				ShowEntry.ToolBarData.LabelOverride = FText();
+				ShowEntry.ToolBarData.ResizeParams.ClippingPriority = 800;
+			}
+
+			// DrawLightCard button
+			{
+				FToolMenuEntry DrawEntry = FToolMenuEntry::InitToolBarButton(
+					FDisplayClusterLightCardEditorCommands::Get().DrawLightCard,
+					TAttribute<FText>(FText::GetEmpty()),
+					TAttribute<FText>(LOCTEXT("DrawLCTooltip",
+						"Draw a new custom light card by adding polygon points clicking the left mouse button on the viewport. "
+						"End the shape with right mouse button, or abort by pressing this button again.")),
+					TAttribute<FSlateIcon>(FSlateIcon(FDisplayClusterLightCardEditorStyle::Get().GetStyleSetName(), "DisplayClusterLightCardEditor.DrawPoly"))
+				);
+
+				DrawEntry.StyleNameOverride = "ViewportToolbarWarning";
+				RightSection.AddEntry(DrawEntry);
+			}
+		}
+
+		// Set up a menu context with the command list.
+
+		FToolMenuContext Context;
+
+		if (EditorViewport.IsValid())
+		{
+			Context.AppendCommandList(EditorViewport.Pin()->GetCommandList());
+		}
+
+		return UToolMenus::Get()->GenerateWidget(ToolbarName, Context);
 	}
 
-	FText GetProjectionMenuLabel() const
+	/** Returns the icon for the current projection mode. */
+	FSlateIcon GetProjectionMenuIcon() const
 	{
-		FText Label = LOCTEXT("ProjectionMenuTitle_Default", "Projection");
-
 		if (EditorViewport.IsValid())
 		{
 			TSharedRef<FDisplayClusterLightCardEditorViewportClient> ViewportClient = EditorViewport.Pin()->GetLightCardEditorViewportClient();
+
 			switch (ViewportClient->GetProjectionMode())
 			{
 			case EDisplayClusterMeshProjectionType::Linear:
 				if (ViewportClient->GetRenderViewportType() == ELevelViewportType::LVT_Perspective)
 				{
-					Label = LOCTEXT("ProjectionwMenuTitle_Perspective", "Perspective");
+					return FSlateIcon(FAppStyle::GetAppStyleSetName(), "EditorViewport.Perspective");
 				}
 				else
 				{
-					Label = LOCTEXT("ProjectionwMenuTitle_Orthographic", "Orthographic");
+					return FSlateIcon(FDisplayClusterLightCardEditorStyle::Get().GetStyleSetName(), "DisplayClusterLightCardEditor.Orthographic");
 				}
-				break;
 
 			case EDisplayClusterMeshProjectionType::Azimuthal:
-				Label = LOCTEXT("ProjectionwMenuTitle_Azimuthal", "Dome");
-				break;
+				return FSlateIcon(FDisplayClusterLightCardEditorStyle::Get().GetStyleSetName(), "DisplayClusterLightCardEditor.Dome");
 
 			case EDisplayClusterMeshProjectionType::UV:
-				Label = LOCTEXT("ProjectionwMenuTitle_UV", "UV");
-				break;
+				return FSlateIcon(FDisplayClusterLightCardEditorStyle::Get().GetStyleSetName(), "DisplayClusterLightCardEditor.UV");
 			}
 		}
 
-		return Label;
+		// Fallback icon
+		return FSlateIcon(FAppStyle::GetAppStyleSetName(), "EditorViewport.Perspective");
 	}
 
-	TSharedRef<SWidget> GenerateProjectionMenu() const
+	/** Creates the menu entry for the coordinate system (spherical/cartesian). Doesn't refer to global/local */
+	FToolMenuEntry CreateCoordinateSystemSubMenuEntry()
 	{
-		TSharedPtr<const FUICommandList> CommandList = EditorViewport.IsValid() ? EditorViewport.Pin()->GetCommandList() : nullptr;
+		FToolMenuEntry CoordinateSystemSubmenu = FToolMenuEntry::InitSubMenu(
+			"CoordinateSystem",
+			LOCTEXT("CoordinateSystemLabel", "Coordinate System"),
+			LOCTEXT("CoordinateSystemTooltip", "Select between coordinate systems"),
+			FNewToolMenuDelegate::CreateLambda(
+				[](UToolMenu* InSubmenu)
+				{
+					FToolMenuSection& UnnamedSection = InSubmenu->FindOrAddSection(NAME_None);
 
-		const bool bInShouldCloseWindowAfterMenuSelection = true;
-		FMenuBuilder ViewMenuBuilder(bInShouldCloseWindowAfterMenuSelection, CommandList);
+					UnnamedSection.AddMenuEntry(FDisplayClusterLightCardEditorCommands::Get().SphericalCoordinateSystem);
+					UnnamedSection.AddMenuEntry(FDisplayClusterLightCardEditorCommands::Get().CartesianCoordinateSystem);
+				}
+			)
+		);
 
-		ViewMenuBuilder.AddMenuEntry(FDisplayClusterLightCardEditorCommands::Get().PerspectiveProjection);
-		ViewMenuBuilder.AddMenuEntry(FDisplayClusterLightCardEditorCommands::Get().OrthographicProjection);
-		ViewMenuBuilder.AddMenuEntry(FDisplayClusterLightCardEditorCommands::Get().AzimuthalProjection);
-		ViewMenuBuilder.AddMenuEntry(FDisplayClusterLightCardEditorCommands::Get().UVProjection);
+		// Set the icon
+		{
+			TWeakPtr<SDisplayClusterLightCardEditorViewportToolBar> WeakSelf = SharedThis(this);
+			CoordinateSystemSubmenu.Icon = TAttribute<FSlateIcon>::CreateLambda(
+				[WeakSelf]() -> FSlateIcon
+				{
+					TSharedPtr<SDisplayClusterLightCardEditorViewportToolBar> PinnedSelf = WeakSelf.Pin();
 
-		return ViewMenuBuilder.MakeWidget();
+					if (PinnedSelf.IsValid()
+						&& PinnedSelf->EditorViewport.IsValid()
+						&& PinnedSelf->EditorViewport.Pin()->GetLightCardEditorViewportClient()->GetCoordinateSystem() == FDisplayClusterLightCardEditorHelper::ECoordinateSystem::Cartesian)
+					{
+						return FSlateIcon(FAppStyle::GetAppStyleSetName(), TEXT("EditorViewport.RelativeCoordinateSystem_World"));
+					}
+
+					return FSlateIcon(FAppStyle::GetAppStyleSetName(), TEXT("Icons.Transform"));
+				}
+			);
+		}
+
+		// Tooltip to show the hotkeys for cycling.
+		{
+			CoordinateSystemSubmenu.ToolTip = TAttribute<FText>::CreateLambda(
+				[]() -> FText
+				{
+					const FInputChord& PrimaryChord = *FDisplayClusterLightCardEditorCommands::Get().CycleEditorWidgetCoordinateSystem->GetActiveChord(EMultipleKeyBindingIndex::Primary);
+					const FInputChord& SecondaryChord = *FDisplayClusterLightCardEditorCommands::Get().CycleEditorWidgetCoordinateSystem->GetActiveChord(EMultipleKeyBindingIndex::Secondary);
+
+					if (PrimaryChord.IsValidChord() && SecondaryChord.IsValidChord())
+					{
+						FFormatNamedArguments Args;
+						Args.Add(TEXT("PrimaryChord"), PrimaryChord.GetInputText());
+						Args.Add(TEXT("SecondaryChord"), SecondaryChord.GetInputText());
+						return FText::Format(
+							LOCTEXT("CoordinateSystemTooltipWithBothChords", "Select between coordinate systems. \n{PrimaryChord} or {SecondaryChord} to cycle between them."),
+							Args
+						);
+					}
+					else if (PrimaryChord.IsValidChord() || SecondaryChord.IsValidChord())
+					{
+						FText ChordText = PrimaryChord.IsValidChord() ? PrimaryChord.GetInputText() : SecondaryChord.GetInputText();
+						return FText::Format(
+							LOCTEXT("CoordinateSystemTooltipSingleChord", "Select between coordinate systems. \n{0} to cycle between them."),
+							ChordText
+						);
+					}
+					return LOCTEXT("CoordinateSystemTooltipNoChords", "Select between coordinate systems");
+				}
+			);
+		}
+
+		// Action override so that clicking the button executes the cycle command.
+		FToolUIAction CycleCoordSystemAction;
+		{
+			CycleCoordSystemAction.ExecuteAction = FToolMenuExecuteAction::CreateLambda(
+				[](const FToolMenuContext& InContext)
+				{
+					if (const FUIAction* Action = InContext.GetActionForCommand(FDisplayClusterLightCardEditorCommands::Get().CycleEditorWidgetCoordinateSystem))
+					{
+						Action->Execute();
+					}
+				}
+			);
+		}
+
+		CoordinateSystemSubmenu.ToolBarData.LabelOverride = FText::GetEmpty();
+		CoordinateSystemSubmenu.ToolBarData.ActionOverride = CycleCoordSystemAction;
+		CoordinateSystemSubmenu.SetShowInToolbarTopLevel(true);
+
+		return CoordinateSystemSubmenu;
 	}
 
-	TSharedRef<SWidget> GenerateViewMenu() const
+	/** Combined Projection and View Menu */
+	void GenerateProjectionViewMenu(UToolMenu* InMenu)
 	{
-		TSharedPtr<const FUICommandList> CommandList = EditorViewport.IsValid() ? EditorViewport.Pin()->GetCommandList() : nullptr;
+		// Projection section.
+		{
+			FToolMenuSection& ProjectionSection = InMenu->AddSection("Projection", LOCTEXT("ProjectionMenuHeader", "Projection"));
+
+			// Perspective
+			ProjectionSection.AddMenuEntry(
+				FDisplayClusterLightCardEditorCommands::Get().PerspectiveProjection,
+				TAttribute<FText>(),
+				TAttribute<FText>(),
+				FSlateIcon(FAppStyle::Get().GetStyleSetName(), "EditorViewport.Perspective")
+			);
+
+			// Orthographic
+			ProjectionSection.AddMenuEntry(
+				FDisplayClusterLightCardEditorCommands::Get().OrthographicProjection,
+				TAttribute<FText>(),
+				TAttribute<FText>(),
+				FSlateIcon(FDisplayClusterLightCardEditorStyle::Get().GetStyleSetName(), "DisplayClusterLightCardEditor.Orthographic")
+			);
 
-		const bool bInShouldCloseWindowAfterMenuSelection = true;
+			// Azimuthal
+			ProjectionSection.AddMenuEntry(
+				FDisplayClusterLightCardEditorCommands::Get().AzimuthalProjection,
+				TAttribute<FText>(),
+				TAttribute<FText>(),
+				FSlateIcon(FDisplayClusterLightCardEditorStyle::Get().GetStyleSetName(), "DisplayClusterLightCardEditor.Dome")
+			);
+
+			// UV
+			ProjectionSection.AddMenuEntry(
+				FDisplayClusterLightCardEditorCommands::Get().UVProjection,
+				TAttribute<FText>(),
+				TAttribute<FText>(),
+				FSlateIcon(FDisplayClusterLightCardEditorStyle::Get().GetStyleSetName(), "DisplayClusterLightCardEditor.UV")
+			);
+		}
 
-		FMenuBuilder MenuBuilder(bInShouldCloseWindowAfterMenuSelection, CommandList);
+		// View Orientation section.
 		{
-			MenuBuilder.BeginSection("LightCardEditorViewOrientation", LOCTEXT("ViewOrientationMenuHeader", "View Orientation"));
-			{
-				MenuBuilder.AddMenuEntry(FDisplayClusterLightCardEditorCommands::Get().ViewOrientationTop);
-				MenuBuilder.AddMenuEntry(FDisplayClusterLightCardEditorCommands::Get().ViewOrientationBottom);
-				MenuBuilder.AddMenuEntry(FDisplayClusterLightCardEditorCommands::Get().ViewOrientationLeft);
-				MenuBuilder.AddMenuEntry(FDisplayClusterLightCardEditorCommands::Get().ViewOrientationRight);
-				MenuBuilder.AddMenuEntry(FDisplayClusterLightCardEditorCommands::Get().ViewOrientationFront);
-				MenuBuilder.AddMenuEntry(FDisplayClusterLightCardEditorCommands::Get().ViewOrientationBack);
-			}
-			MenuBuilder.EndSection();
+			FToolMenuSection& ViewSection = InMenu->AddSection("ViewOrientation", LOCTEXT("ViewOrientationMenuHeader", "View Orientation"));
 
-			MenuBuilder.BeginSection("LightCardEditorViewOptions", LOCTEXT("ViewOptionsMenuHeader", "View Options"));
-			{
-				MenuBuilder.AddMenuEntry(FDisplayClusterLightCardEditorCommands::Get().ResetCamera);
-			}
-			MenuBuilder.EndSection();
+			ViewSection.AddMenuEntry(
+				FDisplayClusterLightCardEditorCommands::Get().ViewOrientationTop,
+				TAttribute<FText>(),
+				TAttribute<FText>(),
+				TAttribute<FSlateIcon>(FSlateIcon(FAppStyle::Get().GetStyleSetName(), "EditorViewport.Top"))
+			);
+
+			ViewSection.AddMenuEntry(
+				FDisplayClusterLightCardEditorCommands::Get().ViewOrientationBottom,
+				TAttribute<FText>(),
+				TAttribute<FText>(),
+				TAttribute<FSlateIcon>(FSlateIcon(FAppStyle::Get().GetStyleSetName(), "EditorViewport.Bottom"))
+			);
+
+			ViewSection.AddMenuEntry(
+				FDisplayClusterLightCardEditorCommands::Get().ViewOrientationLeft,
+				TAttribute<FText>(),
+				TAttribute<FText>(),
+				TAttribute<FSlateIcon>(FSlateIcon(FAppStyle::Get().GetStyleSetName(), "EditorViewport.Left"))
+			);
+
+			ViewSection.AddMenuEntry(
+				FDisplayClusterLightCardEditorCommands::Get().ViewOrientationRight,
+				TAttribute<FText>(),
+				TAttribute<FText>(),
+				TAttribute<FSlateIcon>(FSlateIcon(FAppStyle::Get().GetStyleSetName(), "EditorViewport.Right"))
+			);
+
+			ViewSection.AddMenuEntry(
+				FDisplayClusterLightCardEditorCommands::Get().ViewOrientationFront,
+				TAttribute<FText>(),
+				TAttribute<FText>(),
+				TAttribute<FSlateIcon>(FSlateIcon(FAppStyle::Get().GetStyleSetName(), "EditorViewport.Front"))
+			);
+
+			ViewSection.AddMenuEntry(
+				FDisplayClusterLightCardEditorCommands::Get().ViewOrientationBack,
+				TAttribute<FText>(),
+				TAttribute<FText>(),
+				TAttribute<FSlateIcon>(FSlateIcon(FAppStyle::Get().GetStyleSetName(), "EditorViewport.Back"))
+			);
 		}
 
-		return MenuBuilder.MakeWidget();
+		// Additional View Options section.
+		{
+			FToolMenuSection& OptionsSection = InMenu->AddSection("ViewOptions", LOCTEXT("ViewOptionsMenuHeader", "View Options"));
+
+			OptionsSection.AddMenuEntry(
+				FDisplayClusterLightCardEditorCommands::Get().ResetCamera,
+				TAttribute<FText>(),
+				TAttribute<FText>(),
+				FSlateIcon(FAppStyle::Get().GetStyleSetName(), "StaticMeshEditor.ResetCamera")
+			);
+		}
 	}
-	
-	TSharedRef<SWidget> GenerateShowMenu() const
-	{
-		const TSharedPtr<const FUICommandList> CommandList = EditorViewport.IsValid() ? EditorViewport.Pin()->GetCommandList() : nullptr;
 
-		const bool bInShouldCloseWindowAfterMenuSelection = true;
+	/** Returns the label for the Projection/View dropdown based on the current projection mode. */
+	FText GetProjectionMenuLabel() const
+	{
+		FText Label = LOCTEXT("ProjectionMenuTitle_Default", "Projection");
 
-		FMenuBuilder MenuBuilder(bInShouldCloseWindowAfterMenuSelection, CommandList);
+		if (EditorViewport.IsValid())
 		{
-			MenuBuilder.BeginSection("LightCardEditorShow", LOCTEXT("ShowMenuHeader", "Show Flags"));
+			TSharedRef<FDisplayClusterLightCardEditorViewportClient> ViewportClient = EditorViewport.Pin()->GetLightCardEditorViewportClient();
+			switch (ViewportClient->GetProjectionMode())
 			{
-				MenuBuilder.AddMenuEntry(FDisplayClusterLightCardEditorCommands::Get().ToggleAllLabels);
-				MenuBuilder.AddMenuEntry(FDisplayClusterLightCardEditorCommands::Get().ToggleIconVisibility);
+			case EDisplayClusterMeshProjectionType::Linear:
+				Label = (ViewportClient->GetRenderViewportType() == ELevelViewportType::LVT_Perspective)
+					? LOCTEXT("ProjectionMenuTitle_Perspective", "Perspective")
+					: LOCTEXT("ProjectionMenuTitle_Orthographic", "Orthographic");
+				break;
+
+			case EDisplayClusterMeshProjectionType::Azimuthal:
+				Label = LOCTEXT("ProjectionMenuTitle_Azimuthal", "Dome");
+				break;
+
+			case EDisplayClusterMeshProjectionType::UV:
+				Label = LOCTEXT("ProjectionMenuTitle_UV", "UV");
+				break;
 			}
-			MenuBuilder.EndSection();
 		}
 
-		return MenuBuilder.MakeWidget();
+		return Label;
 	}
 
-	FReply OnViewportsFrozenWarningClicked()
+	void UnfreezeAllViewports()
 	{
 		if (EditorViewport.IsValid())
 		{
@@ -246,11 +461,9 @@ public:
 				RootActor->SetFreezeOuterViewports(false);
 			}
 		}
-
-		return FReply::Handled();
 	}
-	
-	EVisibility GetViewportsFrozenWarningVisibility() const
+
+	bool AreViewportsFrozen() const
 	{
 		if (EditorViewport.IsValid())
 		{
@@ -260,105 +473,11 @@ public:
 				const UDisplayClusterConfigurationData* ConfigData = RootActor->GetConfigData();
 				if (ConfigData && ConfigData->StageSettings.bFreezeRenderOuterViewports)
 				{
-					return EVisibility::Visible;
+					return true;
 				}
 			}
 		}
-
-		return EVisibility::Collapsed;
-	}
-	
-	TSharedRef<SWidget> MakeTransformToolBar()
-	{
-		FSlimHorizontalToolBarBuilder ToolbarBuilder(EditorViewport.Pin()->GetCommandList(), FMultiBoxCustomization::None);
-
-		const FName ToolBarStyle = TEXT("EditorViewportToolBar");
-		ToolbarBuilder.SetStyle(&FAppStyle::Get(), ToolBarStyle);
-		ToolbarBuilder.SetLabelVisibility(EVisibility::Collapsed);
-
-		ToolbarBuilder.SetIsFocusable(false);
-
-		ToolbarBuilder.BeginSection("Transform");
-		{
-			ToolbarBuilder.BeginBlockGroup();
-
-			//static FName SelectModeName = FName(TEXT("SelectMode"));
-			//ToolbarBuilder.AddToolBarButton(FEditorViewportCommands::Get().SelectMode, NAME_None, TAttribute<FText>(), TAttribute<FText>(), TAttribute<FSlateIcon>(), SelectModeName);
-
-			static FName TranslateModeName = FName(TEXT("TranslateMode"));
-			ToolbarBuilder.AddToolBarButton(FEditorViewportCommands::Get().TranslateMode, NAME_None, TAttribute<FText>(), TAttribute<FText>(), TAttribute<FSlateIcon>(), TranslateModeName);
-
-			static FName RotateModeName = FName(TEXT("RotateMode"));
-			ToolbarBuilder.AddToolBarButton(FEditorViewportCommands::Get().RotateMode, NAME_None, TAttribute<FText>(), TAttribute<FText>(), TAttribute<FSlateIcon>(), RotateModeName);
-
-			static FName ScaleModeName = FName(TEXT("ScaleMode"));
-			ToolbarBuilder.AddToolBarButton(FEditorViewportCommands::Get().ScaleMode, NAME_None, TAttribute<FText>(), TAttribute<FText>(), TAttribute<FSlateIcon>(), ScaleModeName);
-
-			ToolbarBuilder.EndBlockGroup();
-		}
-
-		ToolbarBuilder.EndSection();
-		ToolbarBuilder.AddSeparator();
-
-		ToolbarBuilder.AddToolBarButton(FDisplayClusterLightCardEditorCommands::Get().CycleEditorWidgetCoordinateSystem,
-			NAME_None,
-			TAttribute<FText>(),
-			TAttribute<FText>(),
-			TAttribute<FSlateIcon>(this, &SDisplayClusterLightCardEditorViewportToolBar::GetLocalToWorldIcon),
-			FName(TEXT("CycleTransformGizmoCoordSystem")),
-
-			FNewMenuDelegate::CreateLambda([](FMenuBuilder& InMenuBuilder)
-			{
-				InMenuBuilder.AddMenuEntry(FDisplayClusterLightCardEditorCommands::Get().SphericalCoordinateSystem);
-				InMenuBuilder.AddMenuEntry(FDisplayClusterLightCardEditorCommands::Get().CartesianCoordinateSystem);
-			}
-		));
-
-		return ToolbarBuilder.MakeWidget();
-	}
-
-	FSlateIcon GetLocalToWorldIcon() const
-	{
-		if (EditorViewport.IsValid() && EditorViewport.Pin()->GetLightCardEditorViewportClient()->GetCoordinateSystem() == FDisplayClusterLightCardEditorHelper::ECoordinateSystem::Cartesian)
-		{
-			return FSlateIcon(FAppStyle::GetAppStyleSetName(), TEXT("EditorViewport.RelativeCoordinateSystem_World"));
-		}
-
-		return FSlateIcon(FAppStyle::GetAppStyleSetName(), TEXT("Icons.Transform"));
-	}
-
-	TSharedRef<SWidget> MakeDrawingToolBar()
-	{
-		FSlimHorizontalToolBarBuilder ToolbarBuilder(EditorViewport.Pin()->GetCommandList(), FMultiBoxCustomization::None);
-
-		const FName ToolBarStyle = TEXT("EditorViewportToolBar");
-		ToolbarBuilder.SetStyle(&FAppStyle::Get(), ToolBarStyle);
-		ToolbarBuilder.SetLabelVisibility(EVisibility::Collapsed);
-
-		ToolbarBuilder.SetIsFocusable(false);
-
-		ToolbarBuilder.BeginSection("Drawing");
-		{
-			ToolbarBuilder.BeginBlockGroup();
-
-			static FName DrawLightCardName = FName(TEXT("DrawLightCard"));
-			ToolbarBuilder.AddToolBarButton(
-				FDisplayClusterLightCardEditorCommands::Get().DrawLightCard,
-				NAME_None,
-				LOCTEXT("DrawLC", "Draw Light Card    "),
-				LOCTEXT("DrawLCTooltip", 
-					"Draw a new custom light card by adding polygon points clicking the left mouse button on the viewport. "
-					"End the shape with right mouse button, or abort by pressing this button again."),
-				FSlateIcon(FSlateIcon(FDisplayClusterLightCardEditorStyle::Get().GetStyleSetName(), "DrawPoly")),
-				DrawLightCardName
-			);
-
-			ToolbarBuilder.EndBlockGroup();
-		}
-
-		ToolbarBuilder.EndSection();
-
-		return ToolbarBuilder.MakeWidget();
+		return false;
 	}
 
 private:
@@ -376,9 +495,9 @@ const FVector SDisplayClusterLightCardEditorViewport::ViewDirectionBack = FVecto
 void SDisplayClusterLightCardEditorViewport::Construct(const FArguments& InArgs, TSharedPtr<FDisplayClusterLightCardEditor> InLightCardEditor, TSharedPtr<class FUICommandList> InCommandList)
 {
 	LightCardEditorPtr = InLightCardEditor;
-	
+
 	PreviewScene = MakeShared<FPreviewScene>(FPreviewScene::ConstructionValues());
-		
+
 	SEditorViewport::Construct(SEditorViewport::FArguments());
 
 	if (InCommandList.IsValid())
@@ -401,10 +520,7 @@ SDisplayClusterLightCardEditorViewport::~SDisplayClusterLightCardEditorViewport(
 	{
 		if (UWorld* PreviewWorld = PreviewScene->GetWorld())
 		{
-			// Call DestroyWorld first so OnWorldCleanup is broadcasted while actors still have a valid world.
-			// Then mark as garbage to prevent AddReferencedObjects from continuing to be called.
 			PreviewWorld->DestroyWorld(true);
-			
 			PreviewWorld->MarkObjectsPendingKill();
 			PreviewWorld->MarkAsGarbage();
 		}
@@ -429,7 +545,6 @@ void SDisplayClusterLightCardEditorViewport::OnFloatingButtonClicked()
 
 FReply SDisplayClusterLightCardEditorViewport::OnKeyDown(const FGeometry& MyGeometry, const FKeyEvent& InKeyEvent)
 {
-	// Cache the current mouse position in case the key event invokes the paste here command
 	const FVector2D MousePos = FSlateApplication::Get().GetCursorPos();
 	const float DPIScale = FPlatformApplicationMisc::GetDPIScaleFactorAtPoint(MousePos.X, MousePos.Y);
 	PasteHerePos = MyGeometry.AbsoluteToLocal(MousePos) * DPIScale;
@@ -451,7 +566,6 @@ TWeakObjectPtr<ADisplayClusterRootActor> SDisplayClusterLightCardEditorViewport:
 	{
 		return LightCardEditorPtr.Pin()->GetActiveRootActor();
 	}
-
 	return nullptr;
 }
 
@@ -473,23 +587,20 @@ void SDisplayClusterLightCardEditorViewport::SummonContextMenu()
 TSharedRef<FEditorViewportClient> SDisplayClusterLightCardEditorViewport::MakeEditorViewportClient()
 {
 	check(PreviewScene.IsValid());
-	
-	ViewportClient = MakeShareable(new FDisplayClusterLightCardEditorViewportClient(*PreviewScene.Get(),
-		SharedThis(this)));
+
+	ViewportClient = MakeShareable(new FDisplayClusterLightCardEditorViewportClient(*PreviewScene.Get(), SharedThis(this)));
 
 	if (LightCardEditorPtr.IsValid())
 	{
-		// Call after construction as UpdatePreviewActor will create a weak reference to itself
 		ViewportClient->UpdatePreviewActor(LightCardEditorPtr.Pin()->GetActiveRootActor().Get());
 	}
-	
+
 	return ViewportClient.ToSharedRef();
 }
 
-TSharedPtr<SWidget> SDisplayClusterLightCardEditorViewport::MakeViewportToolbar()
+TSharedPtr<SWidget> SDisplayClusterLightCardEditorViewport::BuildViewportToolbar()
 {
-	return
-		SNew(SDisplayClusterLightCardEditorViewportToolBar)
+	return SNew(SDisplayClusterLightCardEditorViewportToolBar)
 		.EditorViewport(SharedThis(this))
 		.IsEnabled(FSlateApplication::Get().GetNormalExecutionAttribute());
 }
@@ -503,7 +614,6 @@ void SDisplayClusterLightCardEditorViewport::BindCommands()
 {
 	SEditorViewport::BindCommands();
 
-	// Override some of the editor viewport commands to work with the light card editor viewport
 	{
 		const FEditorViewportCommands& Commands = FEditorViewportCommands::Get();
 
@@ -513,7 +623,7 @@ void SDisplayClusterLightCardEditorViewport::BindCommands()
 			FCanExecuteAction(),
 			FIsActionChecked::CreateSP(this, &SDisplayClusterLightCardEditorViewport::IsEditorWidgetModeSelected, FDisplayClusterLightCardEditorWidget::EWidgetMode::WM_Translate)
 		);
-		
+
 		CommandList->MapAction(
 			Commands.RotateMode,
 			FExecuteAction::CreateSP(this, &SDisplayClusterLightCardEditorViewport::SetEditorWidgetMode, FDisplayClusterLightCardEditorWidget::EWidgetMode::WM_RotateZ),
@@ -534,7 +644,6 @@ void SDisplayClusterLightCardEditorViewport::BindCommands()
 			FCanExecuteAction()
 		);
 
-		// Remove commands that are being replaced by new specific commands for the nDisplay operator
 		CommandList->UnmapAction(Commands.Top);
 		CommandList->UnmapAction(Commands.Bottom);
 		CommandList->UnmapAction(Commands.Left);
@@ -549,56 +658,67 @@ void SDisplayClusterLightCardEditorViewport::BindCommands()
 		const FDisplayClusterLightCardEditorCommands& Commands = FDisplayClusterLightCardEditorCommands::Get();
 
 		CommandList->MapAction(
-			FDisplayClusterLightCardEditorCommands::Get().PerspectiveProjection,
+			Commands.PerspectiveProjection,
 			FExecuteAction::CreateSP(this, &SDisplayClusterLightCardEditorViewport::SetProjectionMode, EDisplayClusterMeshProjectionType::Linear, ELevelViewportType::LVT_Perspective),
 			FCanExecuteAction(),
-			FIsActionChecked::CreateSP(this, &SDisplayClusterLightCardEditorViewport::IsProjectionModeSelected, EDisplayClusterMeshProjectionType::Linear, ELevelViewportType::LVT_Perspective));
+			FIsActionChecked::CreateSP(this, &SDisplayClusterLightCardEditorViewport::IsProjectionModeSelected, EDisplayClusterMeshProjectionType::Linear, ELevelViewportType::LVT_Perspective)
+		);
 
 		CommandList->MapAction(
-			FDisplayClusterLightCardEditorCommands::Get().OrthographicProjection,
+			Commands.OrthographicProjection,
 			FExecuteAction::CreateSP(this, &SDisplayClusterLightCardEditorViewport::SetProjectionMode, EDisplayClusterMeshProjectionType::Linear, ELevelViewportType::LVT_OrthoFreelook),
 			FCanExecuteAction(),
-			FIsActionChecked::CreateSP(this, &SDisplayClusterLightCardEditorViewport::IsProjectionModeSelected, EDisplayClusterMeshProjectionType::Linear, ELevelViewportType::LVT_OrthoFreelook));
+			FIsActionChecked::CreateSP(this, &SDisplayClusterLightCardEditorViewport::IsProjectionModeSelected, EDisplayClusterMeshProjectionType::Linear, ELevelViewportType::LVT_OrthoFreelook)
+		);
 
 		CommandList->MapAction(
-			FDisplayClusterLightCardEditorCommands::Get().AzimuthalProjection,
+			Commands.AzimuthalProjection,
 			FExecuteAction::CreateSP(this, &SDisplayClusterLightCardEditorViewport::SetProjectionMode, EDisplayClusterMeshProjectionType::Azimuthal, ELevelViewportType::LVT_Perspective),
 			FCanExecuteAction(),
-			FIsActionChecked::CreateSP(this, &SDisplayClusterLightCardEditorViewport::IsProjectionModeSelected, EDisplayClusterMeshProjectionType::Azimuthal, ELevelViewportType::LVT_Perspective));
+			FIsActionChecked::CreateSP(this, &SDisplayClusterLightCardEditorViewport::IsProjectionModeSelected, EDisplayClusterMeshProjectionType::Azimuthal, ELevelViewportType::LVT_Perspective)
+		);
 
 		CommandList->MapAction(
-			FDisplayClusterLightCardEditorCommands::Get().UVProjection,
+			Commands.UVProjection,
 			FExecuteAction::CreateSP(this, &SDisplayClusterLightCardEditorViewport::SetProjectionMode, EDisplayClusterMeshProjectionType::UV, ELevelViewportType::LVT_OrthoFreelook),
 			FCanExecuteAction(),
-			FIsActionChecked::CreateSP(this, &SDisplayClusterLightCardEditorViewport::IsProjectionModeSelected, EDisplayClusterMeshProjectionType::UV, ELevelViewportType::LVT_OrthoFreelook));
+			FIsActionChecked::CreateSP(this, &SDisplayClusterLightCardEditorViewport::IsProjectionModeSelected, EDisplayClusterMeshProjectionType::UV, ELevelViewportType::LVT_OrthoFreelook)
+		);
 
 		CommandList->MapAction(
 			Commands.ViewOrientationTop,
-			FExecuteAction::CreateSP(this, &SDisplayClusterLightCardEditorViewport::SetViewDirection, ViewDirectionTop));
+			FExecuteAction::CreateSP(this, &SDisplayClusterLightCardEditorViewport::SetViewDirection, ViewDirectionTop)
+		);
 
 		CommandList->MapAction(
 			Commands.ViewOrientationBottom,
-			FExecuteAction::CreateSP(this, &SDisplayClusterLightCardEditorViewport::SetViewDirection, ViewDirectionBottom));
+			FExecuteAction::CreateSP(this, &SDisplayClusterLightCardEditorViewport::SetViewDirection, ViewDirectionBottom)
+		);
 
 		CommandList->MapAction(
 			Commands.ViewOrientationLeft,
-			FExecuteAction::CreateSP(this, &SDisplayClusterLightCardEditorViewport::SetViewDirection, ViewDirectionLeft));
+			FExecuteAction::CreateSP(this, &SDisplayClusterLightCardEditorViewport::SetViewDirection, ViewDirectionLeft)
+		);
 
 		CommandList->MapAction(
 			Commands.ViewOrientationRight,
-			FExecuteAction::CreateSP(this, &SDisplayClusterLightCardEditorViewport::SetViewDirection, ViewDirectionRight));
+			FExecuteAction::CreateSP(this, &SDisplayClusterLightCardEditorViewport::SetViewDirection, ViewDirectionRight)
+		);
 
 		CommandList->MapAction(
 			Commands.ViewOrientationFront,
-			FExecuteAction::CreateSP(this, &SDisplayClusterLightCardEditorViewport::SetViewDirection, ViewDirectionFront));
+			FExecuteAction::CreateSP(this, &SDisplayClusterLightCardEditorViewport::SetViewDirection, ViewDirectionFront)
+		);
 
 		CommandList->MapAction(
 			Commands.ViewOrientationBack,
-			FExecuteAction::CreateSP(this, &SDisplayClusterLightCardEditorViewport::SetViewDirection, ViewDirectionBack));
+			FExecuteAction::CreateSP(this, &SDisplayClusterLightCardEditorViewport::SetViewDirection, ViewDirectionBack)
+		);
 
 		CommandList->MapAction(
 			Commands.ResetCamera,
-			FExecuteAction::CreateSP(ViewportClient.Get(), &FDisplayClusterLightCardEditorViewportClient::ResetCamera, false));
+			FExecuteAction::CreateSP(ViewportClient.Get(), &FDisplayClusterLightCardEditorViewportClient::ResetCamera, false)
+		);
 
 		CommandList->MapAction(
 			Commands.FrameSelection,
@@ -608,48 +728,53 @@ void SDisplayClusterLightCardEditorViewport::BindCommands()
 
 		CommandList->MapAction(
 			Commands.CycleEditorWidgetCoordinateSystem,
-			FExecuteAction::CreateSP(ViewportClient.Get(), &FDisplayClusterLightCardEditorViewportClient::CycleCoordinateSystem));
+			FExecuteAction::CreateSP(ViewportClient.Get(), &FDisplayClusterLightCardEditorViewportClient::CycleCoordinateSystem)
+		);
 
 		CommandList->MapAction(
 			Commands.SphericalCoordinateSystem,
 			FExecuteAction::CreateSP(ViewportClient.Get(), &FDisplayClusterLightCardEditorViewportClient::SetCoordinateSystem, FDisplayClusterLightCardEditorHelper::ECoordinateSystem::Spherical),
 			FCanExecuteAction(),
-			FIsActionChecked::CreateLambda([this]() { return ViewportClient->GetCoordinateSystem() == FDisplayClusterLightCardEditorHelper::ECoordinateSystem::Spherical; }));
+			FIsActionChecked::CreateLambda([this]() { return ViewportClient->GetCoordinateSystem() == FDisplayClusterLightCardEditorHelper::ECoordinateSystem::Spherical; })
+		);
 
 		CommandList->MapAction(
 			Commands.CartesianCoordinateSystem,
 			FExecuteAction::CreateSP(ViewportClient.Get(), &FDisplayClusterLightCardEditorViewportClient::SetCoordinateSystem, FDisplayClusterLightCardEditorHelper::ECoordinateSystem::Cartesian),
 			FCanExecuteAction(),
-			FIsActionChecked::CreateLambda([this]() { return ViewportClient->GetCoordinateSystem() == FDisplayClusterLightCardEditorHelper::ECoordinateSystem::Cartesian; }));
-
+			FIsActionChecked::CreateLambda([this]() { return ViewportClient->GetCoordinateSystem() == FDisplayClusterLightCardEditorHelper::ECoordinateSystem::Cartesian; })
+		);
 
 		CommandList->MapAction(
 			FDisplayClusterLightCardEditorCommands::Get().DrawLightCard,
 			FExecuteAction::CreateSP(this, &SDisplayClusterLightCardEditorViewport::DrawLightCard),
 			FCanExecuteAction(),
-			FIsActionChecked::CreateSP(this, &SDisplayClusterLightCardEditorViewport::IsDrawingLightCard));
+			FIsActionChecked::CreateSP(this, &SDisplayClusterLightCardEditorViewport::IsDrawingLightCard)
+		);
 
 		CommandList->MapAction(
 			Commands.PasteHere,
 			FExecuteAction::CreateSP(this, &SDisplayClusterLightCardEditorViewport::PasteLightCardsHere),
-			FCanExecuteAction::CreateSP(this, &SDisplayClusterLightCardEditorViewport::CanPasteLightCardsHere));
+			FCanExecuteAction::CreateSP(this, &SDisplayClusterLightCardEditorViewport::CanPasteLightCardsHere)
+		);
 
 		CommandList->MapAction(
 			Commands.ToggleAllLabels,
 			FExecuteAction::CreateSP(this, &SDisplayClusterLightCardEditorViewport::ToggleLabels),
 			FCanExecuteAction(),
-			FIsActionChecked::CreateSP(this, &SDisplayClusterLightCardEditorViewport::AreLabelsToggled));
+			FIsActionChecked::CreateSP(this, &SDisplayClusterLightCardEditorViewport::AreLabelsToggled)
+		);
 
 		CommandList->MapAction(
 			Commands.ToggleIconVisibility,
 			FExecuteAction::CreateSP(this, &SDisplayClusterLightCardEditorViewport::ToggleIcons),
 			FCanExecuteAction(),
-			FIsActionChecked::CreateSP(this, &SDisplayClusterLightCardEditorViewport::AreIconsToggled));
+			FIsActionChecked::CreateSP(this, &SDisplayClusterLightCardEditorViewport::AreIconsToggled)
+		);
 	}
 }
 
-FReply SDisplayClusterLightCardEditorViewport::OnDragOver(const FGeometry& MyGeometry,
-	const FDragDropEvent& DragDropEvent)
+FReply SDisplayClusterLightCardEditorViewport::OnDragOver(const FGeometry& MyGeometry, const FDragDropEvent& DragDropEvent)
 {
 	const TSharedPtr<FDisplayClusterLightCardTemplateDragDropOp> TemplateDragDropOp = DragDropEvent.GetOperationAs<FDisplayClusterLightCardTemplateDragDropOp>();
 	if (TemplateDragDropOp.IsValid() && LightCardEditorPtr.IsValid() &&
@@ -663,17 +788,16 @@ FReply SDisplayClusterLightCardEditorViewport::OnDragOver(const FGeometry& MyGeo
 		const FVector2D MousePos = MyGeometry.AbsoluteToLocal(DragDropEvent.GetScreenSpacePosition()) * MyGeometry.Scale - ViewportOrigin;
 
 		const TArray<UObject*> DroppedObjects;
-		bool bDroppedObjectsVisible = true; 
+		bool bDroppedObjectsVisible = true;
 		ViewportClient->UpdateDropPreviewActors(MousePos.X, MousePos.Y, DroppedObjects, bDroppedObjectsVisible, nullptr);
-			
+
 		return FReply::Handled();
 	}
 
 	return SEditorViewport::OnDragOver(MyGeometry, DragDropEvent);
 }
 
-void SDisplayClusterLightCardEditorViewport::OnDragEnter(const FGeometry& MyGeometry,
-	const FDragDropEvent& DragDropEvent)
+void SDisplayClusterLightCardEditorViewport::OnDragEnter(const FGeometry& MyGeometry, const FDragDropEvent& DragDropEvent)
 {
 	const TSharedPtr<FDisplayClusterLightCardTemplateDragDropOp> TemplateDragDropOp = DragDropEvent.GetOperationAs<FDisplayClusterLightCardTemplateDragDropOp>();
 	if (TemplateDragDropOp.IsValid() && LightCardEditorPtr.IsValid())
@@ -688,30 +812,30 @@ void SDisplayClusterLightCardEditorViewport::OnDragEnter(const FGeometry& MyGeom
 		if (LightCardTemplate.IsValid())
 		{
 			const TArray<UObject*> DroppedObjects{ LightCardTemplate.Get() };
-			
+
 			TArray<AActor*> TemporaryActors;
-			
+
 			const bool bIsPreview = true;
 			ViewportClient->DropObjectsAtCoordinates(MousePos.X, MousePos.Y, DroppedObjects, TemporaryActors, false, bIsPreview, false, nullptr);
 
 			return;
 		}
 	}
-	
+
 	SEditorViewport::OnDragEnter(MyGeometry, DragDropEvent);
 }
 
 void SDisplayClusterLightCardEditorViewport::OnDragLeave(const FDragDropEvent& DragDropEvent)
 {
 	ViewportClient->DestroyDropPreviewActors();
-	
+
 	const TSharedPtr<FDisplayClusterLightCardTemplateDragDropOp> TemplateDragDropOp = DragDropEvent.GetOperationAs<FDisplayClusterLightCardTemplateDragDropOp>();
 	if (TemplateDragDropOp.IsValid())
 	{
 		TemplateDragDropOp->SetDropAsInvalid();
 		return;
 	}
-	
+
 	SEditorViewport::OnDragLeave(DragDropEvent);
 }
 
@@ -730,9 +854,8 @@ FReply SDisplayClusterLightCardEditorViewport::OnDrop(const FGeometry& MyGeometr
 		if (LightCardTemplate.IsValid())
 		{
 			const TArray<UObject*> DroppedObjects{ LightCardTemplate.Get() };
-			
+
 			TArray<AActor*> TemporaryActors;
-			// Only select actor on drop
 			const bool bSelectActor = true;
 
 			const FScopedTransaction Transaction(LOCTEXT("CreateLightCardFromTemplate", "Create Light Card from Template"));
@@ -741,7 +864,7 @@ FReply SDisplayClusterLightCardEditorViewport::OnDrop(const FGeometry& MyGeometr
 			return FReply::Handled();
 		}
 	}
-	
+
 	return SEditorViewport::OnDrop(MyGeometry, DragDropEvent);
 }
 
@@ -755,13 +878,14 @@ TSharedRef<SWidget> SDisplayClusterLightCardEditorViewport::MakeContextMenu()
 		MenuBuilder.AddSubMenu(
 			LOCTEXT("PlaceActorsSubMenuLabel", "Place Actor"),
 			LOCTEXT("PlaceActorsSubMenuToolTip", "Add new actors to the stage"),
-			FNewMenuDelegate::CreateSP(this, &SDisplayClusterLightCardEditorViewport::MakePlaceActorsSubMenu));
+			FNewMenuDelegate::CreateSP(this, &SDisplayClusterLightCardEditorViewport::MakePlaceActorsSubMenu)
+		);
 
 		MenuBuilder.AddMenuEntry(FDisplayClusterLightCardEditorCommands::Get().RemoveLightCard);
 		MenuBuilder.AddMenuEntry(FDisplayClusterLightCardEditorCommands::Get().SaveLightCardTemplate);
 	}
 	MenuBuilder.EndSection();
-	
+
 	MenuBuilder.BeginSection("View", LOCTEXT("ViewSection", "View"));
 	{
 		MenuBuilder.AddMenuEntry(FDisplayClusterLightCardEditorCommands::Get().FrameSelection);
@@ -794,28 +918,27 @@ void SDisplayClusterLightCardEditorViewport::MakePlaceActorsSubMenu(FMenuBuilder
 		bIsUVMode = ViewportClient->GetProjectionMode() == EDisplayClusterMeshProjectionType::UV;
 	}
 
-	// Add custom UI actions here for the AddNew commands so that the newly added actors can be moved to the user's mouse position after the
-	// add operation is performed
 	MenuBuilder.AddMenuEntry(
 		FDisplayClusterLightCardEditorCommands::Get().AddNewFlag->GetLabel(),
 		FDisplayClusterLightCardEditorCommands::Get().AddNewFlag->GetDescription(),
 		FlagIcon,
 		FUIAction(FExecuteAction::CreateSP(this, &SDisplayClusterLightCardEditorViewport::AddFlagHere),
-			FCanExecuteAction::CreateSP(this, &SDisplayClusterLightCardEditorViewport::CanPlaceActorHere)));
+			FCanExecuteAction::CreateSP(this, &SDisplayClusterLightCardEditorViewport::CanPlaceActorHere))
+	);
 
 	MenuBuilder.AddMenuEntry(
 		FDisplayClusterLightCardEditorCommands::Get().AddNewLightCard->GetLabel(),
 		FDisplayClusterLightCardEditorCommands::Get().AddNewLightCard->GetDescription(),
 		bIsUVMode ? UVLightCardIcon : LightCardIcon,
 		FUIAction(FExecuteAction::CreateSP(this, &SDisplayClusterLightCardEditorViewport::AddLightCardHere),
-			FCanExecuteAction::CreateSP(this, &SDisplayClusterLightCardEditorViewport::CanPlaceActorHere)));
+			FCanExecuteAction::CreateSP(this, &SDisplayClusterLightCardEditorViewport::CanPlaceActorHere))
+	);
 
 	TSet<UClass*> StageActorClasses = UE::DisplayClusterLightCardEditorUtils::GetAllStageActorClasses();
 	for (UClass* Class : StageActorClasses)
 	{
 		if (Class == ADisplayClusterLightCardActor::StaticClass())
 		{
-			// Added manually already
 			continue;
 		}
 
@@ -823,10 +946,11 @@ void SDisplayClusterLightCardEditorViewport::MakePlaceActorsSubMenu(FMenuBuilder
 		FSlateIcon StageActorIcon = FSlateIconFinder::FindIconForClass(Class);
 		MenuBuilder.AddMenuEntry(
 			Label,
-			LOCTEXT("AddStageActorHeader", "Add a stage actor to the scene"), 
+			LOCTEXT("AddStageActorHeader", "Add a stage actor to the scene"),
 			StageActorIcon,
 			FUIAction(FExecuteAction::CreateSP(this, &SDisplayClusterLightCardEditorViewport::AddStageActorHere, Class),
-				FCanExecuteAction::CreateSP(this, &SDisplayClusterLightCardEditorViewport::CanPlaceActorHere, Class)));
+				FCanExecuteAction::CreateSP(this, &SDisplayClusterLightCardEditorViewport::CanPlaceActorHere, Class))
+		);
 	}
 }
 
@@ -865,13 +989,10 @@ void SDisplayClusterLightCardEditorViewport::DrawLightCard()
 	}
 }
 
-
 void SDisplayClusterLightCardEditorViewport::CycleEditorWidgetMode()
 {
 	int32 WidgetModeAsInt = ViewportClient->GetEditorWidgetMode();
-
 	WidgetModeAsInt = (WidgetModeAsInt + 1) % FDisplayClusterLightCardEditorWidget::EWidgetMode::WM_Max;
-
 	SetEditorWidgetMode((FDisplayClusterLightCardEditorWidget::EWidgetMode)WidgetModeAsInt);
 }
 
@@ -916,30 +1037,24 @@ void SDisplayClusterLightCardEditorViewport::PasteLightCardsHere()
 	if (LightCardEditorPtr.IsValid())
 	{
 		const TSharedPtr<FScopedTransaction> Transaction = MakeShared<FScopedTransaction>(LOCTEXT("PasteActorsHereTransactionMessage", "Paste Actors Here"));
-		
 		TArray<AActor*> PastedActors = LightCardEditorPtr.Pin()->PasteActors();
-
 		if (ViewportClient.IsValid() && PastedActors.Num() > 0)
 		{
-			// Perform the positioning of the pasted light cards on the next scene refresh, since the light card proxies will have not been 
-			// regenerated until then
-
 			AddActorHereTransactions.Add(Transaction);
-			ViewportClient->GetOnNextSceneRefresh().AddLambda([this, Transaction, PastedActors=MoveTemp(PastedActors)]()
-			{
-				TArray<FDisplayClusterWeakStageActorPtr> PastedStageActors;
-				Algo::TransformIf(PastedActors, PastedStageActors, [](const AActor* InItem)
+			ViewportClient->GetOnNextSceneRefresh().AddLambda([this, Transaction, PastedActors = MoveTemp(PastedActors)]()
 				{
-					return IsValid(InItem);
-				},
-				[](const AActor* InItem) -> FDisplayClusterWeakStageActorPtr
-				{
-					return InItem;
+					TArray<FDisplayClusterWeakStageActorPtr> PastedStageActors;
+					Algo::TransformIf(PastedActors, PastedStageActors, [](const AActor* InItem)
+						{
+							return IsValid(InItem);
+						},
+						[](const AActor* InItem) -> FDisplayClusterWeakStageActorPtr
+						{
+							return InItem;
+						});
+					ViewportClient->MoveActorsToPixel(FIntPoint(PasteHerePos.X, PasteHerePos.Y), PastedStageActors);
+					AddActorHereTransactions.Remove(Transaction);
 				});
-				
-				ViewportClient->MoveActorsToPixel(FIntPoint(PasteHerePos.X, PasteHerePos.Y), PastedStageActors);
-				AddActorHereTransactions.Remove(Transaction);
-			});
 		}
 	}
 }
@@ -959,19 +1074,15 @@ void SDisplayClusterLightCardEditorViewport::AddLightCardHere()
 	if (LightCardEditorPtr.IsValid())
 	{
 		const TSharedPtr<FScopedTransaction> Transaction = MakeShared<FScopedTransaction>(LOCTEXT("AddLightCardHereTransactionMessage", "Add Light Card Here"));
-		
 		ADisplayClusterLightCardActor* NewLightCard = LightCardEditorPtr.Pin()->AddNewLightCard();
-
 		if (ViewportClient.IsValid() && NewLightCard)
 		{
-			// Perform the positioning of the pasted light cards on the next scene refresh, since the light card proxies will have not been 
-			// regenerated until then
 			AddActorHereTransactions.Add(Transaction);
 			ViewportClient->GetOnNextSceneRefresh().AddLambda([this, Transaction, NewLightCard]()
-			{
-				ViewportClient->MoveActorsToPixel(FIntPoint(PasteHerePos.X, PasteHerePos.Y), { NewLightCard });
-				AddActorHereTransactions.Remove(Transaction);
-			});
+				{
+					ViewportClient->MoveActorsToPixel(FIntPoint(PasteHerePos.X, PasteHerePos.Y), { NewLightCard });
+					AddActorHereTransactions.Remove(Transaction);
+				});
 		}
 	}
 }
@@ -981,19 +1092,15 @@ void SDisplayClusterLightCardEditorViewport::AddFlagHere()
 	if (LightCardEditorPtr.IsValid())
 	{
 		const TSharedPtr<FScopedTransaction> Transaction = MakeShared<FScopedTransaction>(LOCTEXT("AddFlagHereTransactionMessage", "Add Flag Here"));
-		
 		ADisplayClusterLightCardActor* NewFlag = LightCardEditorPtr.Pin()->AddNewFlag();
-
 		if (ViewportClient.IsValid() && NewFlag)
 		{
-			// Perform the positioning of the pasted light cards on the next scene refresh, since the light card proxies will have not been 
-			// regenerated until then
 			AddActorHereTransactions.Add(Transaction);
 			ViewportClient->GetOnNextSceneRefresh().AddLambda([this, Transaction, NewFlag]()
-			{
-				ViewportClient->MoveActorsToPixel(FIntPoint(PasteHerePos.X, PasteHerePos.Y), { NewFlag });
-				AddActorHereTransactions.Remove(Transaction);
-			});
+				{
+					ViewportClient->MoveActorsToPixel(FIntPoint(PasteHerePos.X, PasteHerePos.Y), { NewFlag });
+					AddActorHereTransactions.Remove(Transaction);
+				});
 		}
 	}
 }
@@ -1003,19 +1110,15 @@ void SDisplayClusterLightCardEditorViewport::AddStageActorHere(UClass* InClass)
 	if (LightCardEditorPtr.IsValid())
 	{
 		const TSharedPtr<FScopedTransaction> Transaction = MakeShared<FScopedTransaction>(LOCTEXT("AddStageActorHereTransactionMessage", "Add Stage Actor Here"));
-		
 		AActor* NewActor = LightCardEditorPtr.Pin()->AddNewDynamic(InClass);
-
 		if (ViewportClient.IsValid() && NewActor)
 		{
-			// Perform the positioning of the pasted light cards on the next scene refresh, since the light card proxies will have not been 
-			// regenerated until then
 			AddActorHereTransactions.Add(Transaction);
 			ViewportClient->GetOnNextSceneRefresh().AddLambda([this, Transaction, NewActor]()
-			{
-				ViewportClient->MoveActorsToPixel(FIntPoint(PasteHerePos.X, PasteHerePos.Y), { NewActor });
-				AddActorHereTransactions.Remove(Transaction);
-			});
+				{
+					ViewportClient->MoveActorsToPixel(FIntPoint(PasteHerePos.X, PasteHerePos.Y), { NewActor });
+					AddActorHereTransactions.Remove(Transaction);
+				});
 		}
 	}
 }
@@ -1032,7 +1135,6 @@ bool SDisplayClusterLightCardEditorViewport::CanPlaceActorHere(UClass* Class) co
 
 void SDisplayClusterLightCardEditorViewport::ToggleLabels()
 {
-	// TODO: Handle CCR
 	if (LightCardEditorPtr.IsValid())
 	{
 		return LightCardEditorPtr.Pin()->ToggleLightCardLabels();
@@ -1041,7 +1143,6 @@ void SDisplayClusterLightCardEditorViewport::ToggleLabels()
 
 bool SDisplayClusterLightCardEditorViewport::AreLabelsToggled() const
 {
-	// TODO: Handle CCR
 	return LightCardEditorPtr.IsValid() && LightCardEditorPtr.Pin()->ShouldShowLightCardLabels();
 }
 
diff --git a/Engine/Plugins/Runtime/nDisplay/Source/DisplayClusterLightCardEditor/Private/Viewport/DisplayClusterLightcardEditorViewport.h b/Engine/Plugins/Runtime/nDisplay/Source/DisplayClusterLightCardEditor/Private/Viewport/DisplayClusterLightcardEditorViewport.h
--- a/Engine/Plugins/Runtime/nDisplay/Source/DisplayClusterLightCardEditor/Private/Viewport/DisplayClusterLightcardEditorViewport.h
+++ b/Engine/Plugins/Runtime/nDisplay/Source/DisplayClusterLightCardEditor/Private/Viewport/DisplayClusterLightcardEditorViewport.h
@@ -60,7 +60,7 @@ public:
 private:
 	// SEditorViewport
 	virtual TSharedRef<FEditorViewportClient> MakeEditorViewportClient() override;
-	virtual TSharedPtr<SWidget> MakeViewportToolbar() override;
+	virtual TSharedPtr<SWidget> BuildViewportToolbar() override;
 	virtual void PopulateViewportOverlays(TSharedRef<SOverlay> Overlay) override;
 	virtual void BindCommands() override;
 	virtual FReply OnDragOver(const FGeometry& MyGeometry, const FDragDropEvent& DragDropEvent) override;
