nDisplay: ICVFX Panel migration to new toolbar.

Draw LC toggle now uses warning colors.

#rnx

diff --git a/Engine/Plugins/Runtime/nDisplay/Source/DisplayClusterLightCardEditor/Private/DisplayClusterLightCardEditorStyle.cpp b/Engine/Plugins/Runtime/nDisplay/Source/DisplayClusterLightCardEditor/Private/DisplayClusterLightCardEditorStyle.cpp
--- a/Engine/Plugins/Runtime/nDisplay/Source/DisplayClusterLightCardEditor/Private/DisplayClusterLightCardEditorStyle.cpp
+++ b/Engine/Plugins/Runtime/nDisplay/Source/DisplayClusterLightCardEditor/Private/DisplayClusterLightCardEditorStyle.cpp
@@ -2,9 +2,12 @@
 
 #include "DisplayClusterLightCardEditorStyle.h"
 
+#include "Styling/AppStyle.h"
 #include "Styling/SlateStyleMacros.h"
+#include "Styling/StyleColors.h"
 #include "Styling/ToolBarStyle.h"
 
+
 FDisplayClusterLightCardEditorStyle::FDisplayClusterLightCardEditorStyle(): FSlateStyleSet("DisplayClusterLightCardEditorStyle")
 {
 	const FVector2D Icon40x40(40.0f, 40.0f);
@@ -37,5 +40,42 @@ FDisplayClusterLightCardEditorStyle::FDisplayClusterLightCardEditorStyle(): FSla
 		Set("DisplayClusterLightCardEditor.ViewportsFrozen", new IMAGE_BRUSH_SVG("Icons/Viewport/nDisplayFrozen_16", Icon16x16));
 	}
 
+	{
+		// Build custom style by copying the base ToggleButtonCheckbox
+		const FCheckBoxStyle& ToggleButtonCheckboxStyle = FAppStyle::Get().GetWidgetStyle<FCheckBoxStyle>("ToggleButtonCheckbox");
+		FCheckBoxStyle DrawLightcardsToggleButtonStyle = ToggleButtonCheckboxStyle;
+
+		// Rounder corners
+		constexpr float CornerRadius = 5.0f;
+
+		// Match ViewportToolbarWarning.Raised colors (from StarshipStyle.cpp)
+
+		FLinearColor WarningColor = FStyleColors::Warning.GetSpecifiedColor();
+		const FSlateRoundedBoxBrush WarningBrush(WarningColor, CornerRadius, WarningColor, 1.0f);
+
+		FLinearColor WarningHoveredColor = FStyleColors::Warning.GetSpecifiedColor();
+		WarningHoveredColor.R = FMath::Min(1.0f, WarningHoveredColor.R * 1.5f);
+		WarningHoveredColor.G = FMath::Min(1.0f, WarningHoveredColor.G * 1.5f);
+		WarningHoveredColor.B = FMath::Min(1.0f, WarningHoveredColor.B * 1.5f);
+		const FSlateRoundedBoxBrush WarningHoveredBrush(WarningHoveredColor, CornerRadius, WarningHoveredColor, 1.0f);
+
+		FLinearColor WarningPressedColor = FStyleColors::Warning.GetSpecifiedColor();
+		WarningPressedColor.A = .50f;
+		const FSlateRoundedBoxBrush WarningPressedBrush(WarningPressedColor, CornerRadius, WarningPressedColor, 1.0f);
+
+		// Override the checked/hover/pressed background color
+		DrawLightcardsToggleButtonStyle.SetCheckedImage(WarningBrush);
+		DrawLightcardsToggleButtonStyle.SetCheckedHoveredImage(WarningHoveredBrush);
+		DrawLightcardsToggleButtonStyle.SetCheckedPressedImage(WarningPressedBrush);
+
+		// Override the checked/hover/pressed foreground color
+		DrawLightcardsToggleButtonStyle
+			.SetCheckedForegroundColor(FStyleColors::ForegroundInverted)
+			.SetCheckedHoveredForegroundColor(FStyleColors::ForegroundInverted)
+			.SetCheckedPressedForegroundColor(FStyleColors::ForegroundInverted);
+
+		Set("DisplayClusterLightCardEditor.DrawLightcardsToggleButton", DrawLightcardsToggleButtonStyle);
+	}
+
 	FSlateStyleRegistry::RegisterSlateStyle(*this);
 }
diff --git a/Engine/Plugins/Runtime/nDisplay/Source/DisplayClusterLightCardEditor/Private/Viewport/DisplayClusterLightcardEditorViewport.cpp b/Engine/Plugins/Runtime/nDisplay/Source/DisplayClusterLightCardEditor/Private/Viewport/DisplayClusterLightcardEditorViewport.cpp
--- a/Engine/Plugins/Runtime/nDisplay/Source/DisplayClusterLightCardEditor/Private/Viewport/DisplayClusterLightcardEditorViewport.cpp
+++ b/Engine/Plugins/Runtime/nDisplay/Source/DisplayClusterLightCardEditor/Private/Viewport/DisplayClusterLightcardEditorViewport.cpp
@@ -30,6 +30,8 @@
 #include "ToolMenu.h"
 #include "ToolMenus.h"
 #include "ToolMenuSection.h"
+#include "Widgets/Input/SCheckBox.h"
+
 
 #define LOCTEXT_NAMESPACE "DisplayClusterLightcardEditorViewport"
 
@@ -172,18 +174,18 @@ public:
 				ShowEntry.ToolBarData.ResizeParams.ClippingPriority = 800;
 			}
 
-			// DrawLightCard button
+			// DrawLightCard button as a widget (it is more straight forward to customize the style of a Widget than a ToolBarButton).
 			{
-				FToolMenuEntry DrawEntry = FToolMenuEntry::InitToolBarButton(
-					FDisplayClusterLightCardEditorCommands::Get().DrawLightCard,
-					TAttribute<FText>(FText::GetEmpty()),
-					TAttribute<FText>(LOCTEXT("DrawLCTooltip",
-						"Draw a new custom light card by adding polygon points clicking the left mouse button on the viewport. "
-						"End the shape with right mouse button, or abort by pressing this button again.")),
-					TAttribute<FSlateIcon>(FSlateIcon(FDisplayClusterLightCardEditorStyle::Get().GetStyleSetName(), "DisplayClusterLightCardEditor.DrawPoly"))
+				FToolMenuEntry DrawEntry = FToolMenuEntry::InitWidget(
+					FName("DrawLightCardWidget"),
+					MakeDrawLightCardWidget(),
+					LOCTEXT("DrawLC", "Draw Light Card"),
+					false, /* bNoIndent */
+					false, /* bSearchable */
+					false, /* bNoPadding */
+					FText()
 				);
 
-				DrawEntry.StyleNameOverride = "ViewportToolbarWarning";
 				RightSection.AddEntry(DrawEntry);
 			}
 		}
@@ -200,6 +202,67 @@ public:
 		return UToolMenus::Get()->GenerateWidget(ToolbarName, Context);
 	}
 
+	/** Make the DrawLightCard toggle button widget. */
+	TSharedRef<SWidget> MakeDrawLightCardWidget()
+	{
+		const TSharedPtr<const FUICommandList> Commands = EditorViewport.Pin()->GetCommandList();
+		const TSharedPtr<const FUICommandInfo> DrawCmd = FDisplayClusterLightCardEditorCommands::Get().DrawLightCard;
+
+		return SNew(SCheckBox)
+			.Style(&FDisplayClusterLightCardEditorStyle::Get().GetWidgetStyle<FCheckBoxStyle>("DisplayClusterLightCardEditor.DrawLightcardsToggleButton"))
+			.Padding(FMargin(2))
+			.IsChecked_Lambda([Commands, DrawCmd]() -> ECheckBoxState
+				{
+					if (Commands.IsValid() && DrawCmd.IsValid())
+					{
+						if (const FUIAction* Action = Commands->GetActionForCommand(DrawCmd))
+						{
+							return Action->GetCheckState();
+						}
+					}
+					return ECheckBoxState::Unchecked;
+				})
+			.OnCheckStateChanged_Lambda([Commands, DrawCmd](ECheckBoxState)
+				{
+					if (Commands.IsValid() && DrawCmd.IsValid())
+					{
+						if (const FUIAction* Action = Commands->GetActionForCommand(DrawCmd))
+						{
+							Action->Execute();
+						}
+					}
+				})
+			.ToolTipText(DrawCmd->GetDescription())
+			[
+				SNew(SBorder)
+					.Padding(FMargin(4)) // Pad to make the icon not fully fill the button area.
+					.BorderImage(FCoreStyle::Get().GetBrush("NoBorder"))
+					[
+						SNew(SBox)
+							.WidthOverride(16) // Constrain the size to match the rest of the toolbar.
+							.HeightOverride(16)
+							[
+								SNew(SImage)
+									.Image(FDisplayClusterLightCardEditorStyle::Get().GetBrush("DisplayClusterLightCardEditor.DrawPoly"))
+									// Tint the icon according to its state.
+									.ColorAndOpacity_Lambda([Commands, DrawCmd]() -> FSlateColor 
+										{
+											if (Commands.IsValid() && DrawCmd.IsValid())
+											{
+												if (const FUIAction* Action = Commands->GetActionForCommand(DrawCmd))
+												{
+													return Action->GetCheckState() == ECheckBoxState::Checked
+														? FSlateColor(FLinearColor::Black)
+														: FSlateColor(FLinearColor::White);
+												}
+											}
+											return FSlateColor(FLinearColor::White);
+										})
+							]
+					]
+			];
+	}
+
 	/** Returns the icon for the current projection mode. */
 	FSlateIcon GetProjectionMenuIcon() const
 	{
