[Truncation Warnings] Update EdGraphNode and EdGraph API to use FVector2f

diff --git a/Engine/Plugins/Animation/DeformerGraph/Source/OptimusCore/OptimusCore.Build.cs b/Engine/Plugins/Animation/DeformerGraph/Source/OptimusCore/OptimusCore.Build.cs
--- a/Engine/Plugins/Animation/DeformerGraph/Source/OptimusCore/OptimusCore.Build.cs
+++ b/Engine/Plugins/Animation/DeformerGraph/Source/OptimusCore/OptimusCore.Build.cs
@@ -11,6 +11,7 @@ namespace UnrealBuildTool.Rules
             PublicDependencyModuleNames.AddRange(
 				new string[]
 				{
+					"SlateCore"
 				}
 			);
 
diff --git a/Engine/Plugins/Animation/DeformerGraph/Source/OptimusCore/Private/Actions/OptimusNodeActions.cpp b/Engine/Plugins/Animation/DeformerGraph/Source/OptimusCore/Private/Actions/OptimusNodeActions.cpp
--- a/Engine/Plugins/Animation/DeformerGraph/Source/OptimusCore/Private/Actions/OptimusNodeActions.cpp
+++ b/Engine/Plugins/Animation/DeformerGraph/Source/OptimusCore/Private/Actions/OptimusNodeActions.cpp
@@ -91,7 +91,7 @@ bool FOptimusNodeAction_MoveNode::Undo(IOptimusPathResolver* InRoot)
 
 FOptimusCommentNodeAction_ResizeNode::FOptimusCommentNodeAction_ResizeNode(
 	UOptimusNode_Comment* InNode,
-	const FVector2D& InNewSize)
+	const UE::Slate::FDeprecateVector2DParameter& InNewSize)
 {
 	NodePath = InNode->GetNodePath();
 	OldSize = InNode->GetSize();	
diff --git a/Engine/Plugins/Animation/DeformerGraph/Source/OptimusCore/Private/Actions/OptimusNodeActions.h b/Engine/Plugins/Animation/DeformerGraph/Source/OptimusCore/Private/Actions/OptimusNodeActions.h
--- a/Engine/Plugins/Animation/DeformerGraph/Source/OptimusCore/Private/Actions/OptimusNodeActions.h
+++ b/Engine/Plugins/Animation/DeformerGraph/Source/OptimusCore/Private/Actions/OptimusNodeActions.h
@@ -4,6 +4,7 @@
 
 #include "OptimusAction.h"
 #include "OptimusNodePin.h"
+#include "Types/SlateVector2.h"
 
 #include "OptimusNodeActions.generated.h"
 
@@ -65,10 +66,10 @@ private:
 	FString NodePath;
 
 	// The node's new position
-	FVector2D NewPosition;
+	FDeprecateSlateVector2D NewPosition;
 
 	// The node's old position
-	FVector2D OldPosition;
+	FDeprecateSlateVector2D OldPosition;
 };
 
 USTRUCT()
@@ -82,7 +83,7 @@ public:
 
 	FOptimusCommentNodeAction_ResizeNode(
 		UOptimusNode_Comment* InNode,
-		const FVector2D& InNewSize
+		const UE::Slate::FDeprecateVector2DParameter& InNewSize
 	);
 
 protected:
@@ -93,9 +94,9 @@ private:
 	// The path of the node to be resized
 	FString NodePath;
 
-	FVector2D OldSize;
+	FDeprecateSlateVector2D OldSize;
 
-	FVector2D NewSize;
+	FDeprecateSlateVector2D NewSize;
 };
 
 USTRUCT()
diff --git a/Engine/Plugins/Animation/DeformerGraph/Source/OptimusCore/Private/OptimusNodeGraph.cpp b/Engine/Plugins/Animation/DeformerGraph/Source/OptimusCore/Private/OptimusNodeGraph.cpp
--- a/Engine/Plugins/Animation/DeformerGraph/Source/OptimusCore/Private/OptimusNodeGraph.cpp
+++ b/Engine/Plugins/Animation/DeformerGraph/Source/OptimusCore/Private/OptimusNodeGraph.cpp
@@ -604,7 +604,7 @@ UOptimusNode* UOptimusNodeGraph::AddCommentNode(const FVector2D& InPosition, con
 		[InColor, InSize, bCreatedFromUI](UOptimusNode *InNode)
 		{
 			UOptimusNode_Comment* CommentNode = CastChecked<UOptimusNode_Comment>(InNode);
-			CommentNode->SetSizeDirect(InSize);
+			CommentNode->SetSizeDirect(UE::Slate::FDeprecateVector2DParameter(InSize));
 			CommentNode->CommentColor = InColor;
 			CommentNode->FontSize = 18;
 			CommentNode->Comment = TEXT("Comment");
diff --git a/Engine/Plugins/Animation/DeformerGraph/Source/OptimusCore/Private/OptimusNode_Comment.cpp b/Engine/Plugins/Animation/DeformerGraph/Source/OptimusCore/Private/OptimusNode_Comment.cpp
--- a/Engine/Plugins/Animation/DeformerGraph/Source/OptimusCore/Private/OptimusNode_Comment.cpp
+++ b/Engine/Plugins/Animation/DeformerGraph/Source/OptimusCore/Private/OptimusNode_Comment.cpp
@@ -19,12 +19,12 @@ void UOptimusNode_Comment::PostEditChangeProperty(FPropertyChangedEvent& Propert
 }
 #endif
 
-bool UOptimusNode_Comment::SetSize(const FVector2D& InSize)
+bool UOptimusNode_Comment::SetSize(const UE::Slate::FDeprecateVector2DParameter& InSize)
 {
-	return GetActionStack()->RunAction<FOptimusCommentNodeAction_ResizeNode>(this, InSize);
+	return GetActionStack()->RunAction<FOptimusCommentNodeAction_ResizeNode>(this, FVector2D(InSize));
 }
 
-bool UOptimusNode_Comment::SetSizeDirect(const FVector2D& InNewSize)
+bool UOptimusNode_Comment::SetSizeDirect(const FVector2f& InNewSize)
 {
 	if (Size == InNewSize || InNewSize.ContainsNaN())
 	{
diff --git a/Engine/Plugins/Animation/DeformerGraph/Source/OptimusCore/Public/OptimusNode_Comment.h b/Engine/Plugins/Animation/DeformerGraph/Source/OptimusCore/Public/OptimusNode_Comment.h
--- a/Engine/Plugins/Animation/DeformerGraph/Source/OptimusCore/Public/OptimusNode_Comment.h
+++ b/Engine/Plugins/Animation/DeformerGraph/Source/OptimusCore/Public/OptimusNode_Comment.h
@@ -3,6 +3,7 @@
 #pragma once
 
 #include "OptimusNode.h"
+#include "Types/SlateVector2.h"
 
 #include "OptimusNode_Comment.generated.h"
 
@@ -25,8 +26,8 @@ public:
 
 	FName GetNodeCategory() const override {return NAME_None;}
 
-	bool SetSize(const FVector2D& InSize);
-	const FVector2D& GetSize() {return Size;}
+	bool SetSize(const UE::Slate::FDeprecateVector2DParameter& InSize);
+	const UE::Slate::FDeprecateVector2DResult& GetSize() {return Size;}
 	void SetComment(const FString& InNewComment);
 	const FString& GetComment() {return Comment;}
 	
@@ -55,10 +56,10 @@ public:
 private:
 	friend class UOptimusNodeGraph;
 	friend struct FOptimusCommentNodeAction_ResizeNode;
-	bool SetSizeDirect(const FVector2D& InNewSize);	
+	bool SetSizeDirect(const FVector2f& InNewSize);	
 
 	UPROPERTY(NonTransactional)
-	FVector2D Size = FVector2D(400, 100);	
+	FDeprecateSlateVector2D Size = FVector2f(400, 100);
 	
 	FOnOptimusCommentNodePropertyChanged OnPropertyChangedDelegate;
 };
diff --git a/Engine/Plugins/Animation/DeformerGraph/Source/OptimusEditor/Private/OptimusEditorGraphNode_Comment.cpp b/Engine/Plugins/Animation/DeformerGraph/Source/OptimusEditor/Private/OptimusEditorGraphNode_Comment.cpp
--- a/Engine/Plugins/Animation/DeformerGraph/Source/OptimusEditor/Private/OptimusEditorGraphNode_Comment.cpp
+++ b/Engine/Plugins/Animation/DeformerGraph/Source/OptimusEditor/Private/OptimusEditorGraphNode_Comment.cpp
@@ -48,7 +48,7 @@ void UOptimusEditorGraphNode_Comment::OnPositionChanged()
 	OnPositionChangedDelegate.ExecuteIfBound();
 }
 
-void UOptimusEditorGraphNode_Comment::ResizeNode(const FVector2D& NewSize)
+void UOptimusEditorGraphNode_Comment::ResizeNode(const FVector2f& NewSize)
 {
 	Super::ResizeNode(NewSize);
 
diff --git a/Engine/Plugins/Animation/DeformerGraph/Source/OptimusEditor/Private/OptimusEditorGraphNode_Comment.h b/Engine/Plugins/Animation/DeformerGraph/Source/OptimusEditor/Private/OptimusEditorGraphNode_Comment.h
--- a/Engine/Plugins/Animation/DeformerGraph/Source/OptimusEditor/Private/OptimusEditorGraphNode_Comment.h
+++ b/Engine/Plugins/Animation/DeformerGraph/Source/OptimusEditor/Private/OptimusEditorGraphNode_Comment.h
@@ -22,7 +22,7 @@ public:
 
 	void OnPositionChanged();
 
-	void ResizeNode(const FVector2D& NewSize) override;
+	void ResizeNode(const FVector2f& NewSize) override;
 	void PostPlacedNewNode() override;
 	
 	FOptimusCommentNodeSizeChanged& GetOnSizeChanged() {return OnSizeChangedDelegate; }
diff --git a/Engine/Plugins/Cameras/GameplayCameras/Source/GameplayCamerasEditor/Private/Editors/ObjectTreeGraphCommentNode.cpp b/Engine/Plugins/Cameras/GameplayCameras/Source/GameplayCamerasEditor/Private/Editors/ObjectTreeGraphCommentNode.cpp
--- a/Engine/Plugins/Cameras/GameplayCameras/Source/GameplayCamerasEditor/Private/Editors/ObjectTreeGraphCommentNode.cpp
+++ b/Engine/Plugins/Cameras/GameplayCameras/Source/GameplayCamerasEditor/Private/Editors/ObjectTreeGraphCommentNode.cpp
@@ -99,7 +99,7 @@ void UObjectTreeGraphCommentNode::PostPlacedNewNode()
 	}
 }
 
-void UObjectTreeGraphCommentNode::ResizeNode(const FVector2D& NewSize)
+void UObjectTreeGraphCommentNode::ResizeNode(const FVector2f& NewSize)
 {
 	Super::ResizeNode(NewSize);
 
diff --git a/Engine/Plugins/Cameras/GameplayCameras/Source/GameplayCamerasEditor/Public/Editors/ObjectTreeGraphCommentNode.h b/Engine/Plugins/Cameras/GameplayCameras/Source/GameplayCamerasEditor/Public/Editors/ObjectTreeGraphCommentNode.h
--- a/Engine/Plugins/Cameras/GameplayCameras/Source/GameplayCamerasEditor/Public/Editors/ObjectTreeGraphCommentNode.h
+++ b/Engine/Plugins/Cameras/GameplayCameras/Source/GameplayCamerasEditor/Public/Editors/ObjectTreeGraphCommentNode.h
@@ -28,7 +28,7 @@ public:
 	virtual TSharedPtr<SGraphNode> CreateVisualWidget() override;
 	virtual void GetNodeContextMenuActions(UToolMenu* Menu, UGraphNodeContextMenuContext* Context) const override;
 	virtual void PostPlacedNewNode() override;
-	virtual void ResizeNode(const FVector2D& NewSize) override;
+	virtual void ResizeNode(const FVector2f& NewSize) override;
 	virtual void OnRenameNode(const FString& NewName) override;
 
 public:
diff --git a/Engine/Plugins/Enterprise/DataprepEditor/Source/DataprepCore/Public/DataprepActionAsset.h b/Engine/Plugins/Enterprise/DataprepEditor/Source/DataprepCore/Public/DataprepActionAsset.h
--- a/Engine/Plugins/Enterprise/DataprepEditor/Source/DataprepCore/Public/DataprepActionAsset.h
+++ b/Engine/Plugins/Enterprise/DataprepEditor/Source/DataprepCore/Public/DataprepActionAsset.h
@@ -8,6 +8,7 @@
 #include "Delegates/DelegateCombinations.h"
 #include "Delegates/IDelegateInstance.h"
 #include "Engine/World.h"
+#include "Types/SlateVector2.h"
 #include "UObject/Object.h"
 #include "UObject/ObjectMacros.h"
 
@@ -184,7 +185,7 @@ public:
 	bool bIsExpanded;
 
 	UPROPERTY()
-	FVector2D NodeSize;
+	FDeprecateSlateVector2D NodeSize;
 
 	UPROPERTY()
 	int32 GroupId;
diff --git a/Engine/Plugins/Enterprise/DataprepEditor/Source/DataprepEditor/Private/DataprepGraph/DataprepGraphActionNode.cpp b/Engine/Plugins/Enterprise/DataprepEditor/Source/DataprepEditor/Private/DataprepGraph/DataprepGraphActionNode.cpp
--- a/Engine/Plugins/Enterprise/DataprepEditor/Source/DataprepEditor/Private/DataprepGraph/DataprepGraphActionNode.cpp
+++ b/Engine/Plugins/Enterprise/DataprepEditor/Source/DataprepEditor/Private/DataprepGraph/DataprepGraphActionNode.cpp
@@ -75,7 +75,7 @@ TSharedPtr<class INameValidatorInterface> UDataprepGraphActionNode::MakeNameVali
 	return MakeShareable(new FDummyNameValidator(EValidatorResult::Ok));
 }
 
-void UDataprepGraphActionNode::ResizeNode(const FVector2D& NewSize)
+void UDataprepGraphActionNode::ResizeNode(const FVector2f& NewSize)
 {
 	if ( DataprepActionAsset )
 	{
@@ -157,7 +157,7 @@ FText UDataprepGraphActionGroupNode::GetNodeTitle(ENodeTitleType::Type TitleType
 	return FText::FromString( NodeTitle );
 }
 
-void UDataprepGraphActionGroupNode::ResizeNode(const FVector2D& NewSize)
+void UDataprepGraphActionGroupNode::ResizeNode(const FVector2f& NewSize)
 {
 	DataprepAssetPtr->Modify();
 
diff --git a/Engine/Plugins/Enterprise/DataprepEditor/Source/DataprepEditor/Private/DataprepGraph/DataprepGraphActionNode.h b/Engine/Plugins/Enterprise/DataprepEditor/Source/DataprepEditor/Private/DataprepGraph/DataprepGraphActionNode.h
--- a/Engine/Plugins/Enterprise/DataprepEditor/Source/DataprepEditor/Private/DataprepGraph/DataprepGraphActionNode.h
+++ b/Engine/Plugins/Enterprise/DataprepEditor/Source/DataprepEditor/Private/DataprepGraph/DataprepGraphActionNode.h
@@ -76,7 +76,7 @@ public:
 	virtual void OnRenameNode(const FString& NewName) override;
 	virtual void DestroyNode() override;
 	TSharedPtr<class INameValidatorInterface> MakeNameValidator() const override;
-	virtual void ResizeNode(const FVector2D& NewSize) override;
+	virtual void ResizeNode(const FVector2f& NewSize) override;
 	// End EdGraphNode interface
 
 	void Initialize(TWeakObjectPtr<UDataprepAsset> InDataprepAssetPtr, UDataprepActionAsset* InDataprepActionAsset, int32 InExecutionOrder);
@@ -112,7 +112,7 @@ public:
 	virtual bool ShowPaletteIconOnNode() const override { return false; }
 	virtual FText GetNodeTitle(ENodeTitleType::Type TitleType) const override;
 	TSharedPtr<class INameValidatorInterface> MakeNameValidator() const override;
-	virtual void ResizeNode(const FVector2D& NewSize) override;
+	virtual void ResizeNode(const FVector2f& NewSize) override;
 	// End EdGraphNode interface
 
 	void Initialize(TWeakObjectPtr<UDataprepAsset> InDataprepAssetPtr, TArray<UDataprepActionAsset*>& InActions, int32 InExecutionOrder);
diff --git a/Engine/Plugins/Runtime/Metasound/Source/MetasoundEditor/Private/MetasoundEditorGraphCommentNode.cpp b/Engine/Plugins/Runtime/Metasound/Source/MetasoundEditor/Private/MetasoundEditorGraphCommentNode.cpp
--- a/Engine/Plugins/Runtime/Metasound/Source/MetasoundEditor/Private/MetasoundEditorGraphCommentNode.cpp
+++ b/Engine/Plugins/Runtime/Metasound/Source/MetasoundEditor/Private/MetasoundEditorGraphCommentNode.cpp
@@ -73,7 +73,7 @@ void UMetasoundEditorGraphCommentNode::PostEditChangeProperty(FPropertyChangedEv
 	}
 }
 
-void UMetasoundEditorGraphCommentNode::ResizeNode(const FVector2D& NewSize)
+void UMetasoundEditorGraphCommentNode::ResizeNode(const FVector2f& NewSize)
 {
 	using namespace Metasound::Engine;
 
@@ -82,7 +82,7 @@ void UMetasoundEditorGraphCommentNode::ResizeNode(const FVector2D& NewSize)
 	{
 		UMetasoundEditorGraph& Graph = *CastChecked<UMetasoundEditorGraph>(GetGraph());
 		UMetaSoundBuilderBase& Builder = FDocumentBuilderRegistry::GetChecked().FindOrBeginBuilding(Graph.GetMetasoundChecked());
-		Builder.FindOrAddGraphComment(CommentID).Size = FVector2D(NodeWidth, NodeHeight);
+		Builder.FindOrAddGraphComment(CommentID).Size = FVector2f(NodeWidth, NodeHeight);
 	}
 }
 
diff --git a/Engine/Plugins/Runtime/Metasound/Source/MetasoundEditor/Public/MetasoundEditorGraphCommentNode.h b/Engine/Plugins/Runtime/Metasound/Source/MetasoundEditor/Public/MetasoundEditorGraphCommentNode.h
--- a/Engine/Plugins/Runtime/Metasound/Source/MetasoundEditor/Public/MetasoundEditorGraphCommentNode.h
+++ b/Engine/Plugins/Runtime/Metasound/Source/MetasoundEditor/Public/MetasoundEditorGraphCommentNode.h
@@ -43,7 +43,7 @@ public:
 	//~ Begin UEdGraphNode Interface
 	METASOUNDEDITOR_API virtual bool CanUserDeleteNode() const override;
 	METASOUNDEDITOR_API virtual void OnRenameNode(const FString& NewName) override;
-	METASOUNDEDITOR_API virtual void ResizeNode(const FVector2D& NewSize) override;
+	METASOUNDEDITOR_API virtual void ResizeNode(const FVector2f& NewSize) override;
 	//~ End UEdGraphNode Interface
 
 	/** Set the Bounds for the comment node */
diff --git a/Engine/Plugins/Runtime/Metasound/Source/MetasoundFrontend/MetasoundFrontend.Build.cs b/Engine/Plugins/Runtime/Metasound/Source/MetasoundFrontend/MetasoundFrontend.Build.cs
--- a/Engine/Plugins/Runtime/Metasound/Source/MetasoundFrontend/MetasoundFrontend.Build.cs
+++ b/Engine/Plugins/Runtime/Metasound/Source/MetasoundFrontend/MetasoundFrontend.Build.cs
@@ -22,7 +22,8 @@ namespace UnrealBuildTool.Rules
 					"JsonUtilities",
 					"Serialization",
 					"SignalProcessing",
-					"MetasoundGraphCore"
+					"MetasoundGraphCore",
+					"SlateCore"
 				}
 			);
 
diff --git a/Engine/Plugins/Runtime/Metasound/Source/MetasoundFrontend/Public/MetasoundFrontendDocument.h b/Engine/Plugins/Runtime/Metasound/Source/MetasoundFrontend/Public/MetasoundFrontendDocument.h
--- a/Engine/Plugins/Runtime/Metasound/Source/MetasoundFrontend/Public/MetasoundFrontendDocument.h
+++ b/Engine/Plugins/Runtime/Metasound/Source/MetasoundFrontend/Public/MetasoundFrontendDocument.h
@@ -17,6 +17,7 @@
 #include "Templates/Function.h"
 #include "Templates/Invoke.h"
 #include "Templates/TypeHash.h"
+#include "Types/SlateVector2.h"
 #include "UObject/NoExportTypes.h"
 
 #include "MetasoundFrontendDocument.generated.h"
@@ -196,10 +197,10 @@ struct METASOUNDFRONTEND_API FMetaSoundFrontendGraphComment
 	int32 FontSize = 0;
 
 	UPROPERTY()
-	FVector2D Position = FVector2D::Zero();
+	FDeprecateSlateVector2D Position = FVector2f::ZeroVector;
 
 	UPROPERTY()
-	FVector2D Size = FVector2D::Zero();
+	FDeprecateSlateVector2D Size = FVector2f::ZeroVector;
 
 	UPROPERTY()
 	EMetaSoundFrontendGraphCommentMoveMode MoveMode = EMetaSoundFrontendGraphCommentMoveMode::GroupMovement;
diff --git a/Engine/Plugins/Runtime/nDisplay/Source/DisplayClusterConfigurator/Private/Views/OutputMapping/EdNodes/DisplayClusterConfiguratorBaseNode.cpp b/Engine/Plugins/Runtime/nDisplay/Source/DisplayClusterConfigurator/Private/Views/OutputMapping/EdNodes/DisplayClusterConfiguratorBaseNode.cpp
--- a/Engine/Plugins/Runtime/nDisplay/Source/DisplayClusterConfigurator/Private/Views/OutputMapping/EdNodes/DisplayClusterConfiguratorBaseNode.cpp
+++ b/Engine/Plugins/Runtime/nDisplay/Source/DisplayClusterConfigurator/Private/Views/OutputMapping/EdNodes/DisplayClusterConfiguratorBaseNode.cpp
@@ -58,12 +58,12 @@ void UDisplayClusterConfiguratorBaseNode::PostPlacedNewNode()
 	ReadNodeStateFromObject();
 }
 
-void UDisplayClusterConfiguratorBaseNode::ResizeNode(const FVector2D& NewSize)
+void UDisplayClusterConfiguratorBaseNode::ResizeNode(const FVector2f& NewSize)
 {
 	Modify();
 
-	NodeWidth = NewSize.X;
-	NodeHeight = NewSize.Y;
+	NodeHeight = UE::LWC::FloatToIntCastChecked<int32>(NewSize.Y);
+	NodeWidth = UE::LWC::FloatToIntCastChecked<int32>(NewSize.X);
 
 	UpdateObject();
 }
diff --git a/Engine/Plugins/Runtime/nDisplay/Source/DisplayClusterConfigurator/Private/Views/OutputMapping/EdNodes/DisplayClusterConfiguratorBaseNode.h b/Engine/Plugins/Runtime/nDisplay/Source/DisplayClusterConfigurator/Private/Views/OutputMapping/EdNodes/DisplayClusterConfiguratorBaseNode.h
--- a/Engine/Plugins/Runtime/nDisplay/Source/DisplayClusterConfigurator/Private/Views/OutputMapping/EdNodes/DisplayClusterConfiguratorBaseNode.h
+++ b/Engine/Plugins/Runtime/nDisplay/Source/DisplayClusterConfigurator/Private/Views/OutputMapping/EdNodes/DisplayClusterConfiguratorBaseNode.h
@@ -66,7 +66,7 @@ public:
 
 	//~ Begin UEdGraphNode Interface
 	virtual void PostPlacedNewNode() override;
-	virtual void ResizeNode(const FVector2D& NewSize) override;
+	virtual void ResizeNode(const FVector2f& NewSize) override;
 	virtual bool CanDuplicateNode() const override { return false; }
 	virtual bool CanUserDeleteNode() const override { return false; }
 	//~ End UEdGraphNode Interface
diff --git a/Engine/Plugins/Runtime/nDisplay/Source/DisplayClusterConfigurator/Private/Views/OutputMapping/EdNodes/DisplayClusterConfiguratorViewportNode.cpp b/Engine/Plugins/Runtime/nDisplay/Source/DisplayClusterConfigurator/Private/Views/OutputMapping/EdNodes/DisplayClusterConfiguratorViewportNode.cpp
--- a/Engine/Plugins/Runtime/nDisplay/Source/DisplayClusterConfigurator/Private/Views/OutputMapping/EdNodes/DisplayClusterConfiguratorViewportNode.cpp
+++ b/Engine/Plugins/Runtime/nDisplay/Source/DisplayClusterConfigurator/Private/Views/OutputMapping/EdNodes/DisplayClusterConfiguratorViewportNode.cpp
@@ -42,27 +42,27 @@ TSharedPtr<SGraphNode> UDisplayClusterConfiguratorViewportNode::CreateVisualWidg
 	return SNew(SDisplayClusterConfiguratorViewportNode, this, ToolkitPtr.Pin().ToSharedRef());
 }
 
-void UDisplayClusterConfiguratorViewportNode::ResizeNode(const FVector2D& NewSize)
+void UDisplayClusterConfiguratorViewportNode::ResizeNode(const FVector2f& NewSize)
 {
 	// If the viewport is rotated, there are some sizes the node (which represents the bounding box of the rotated viewport)
-	// cannot take on, beucase such sizes cause the viewport to have negative width or height. Find an appropriate size for the
+	// cannot take on, because such sizes cause the viewport to have negative width or height. Find an appropriate size for the
 	// node by computing the possible viewport size, and clamping it to valid values, and recompute a corrected node size from the
 	// clamped viewport size, to ensure the viewport size never becomes negative
 
 	const FDisplayClusterConfigurationViewport_RemapData& CfgRemap = GetCfgViewportRemap();
 	const FMatrix2x2 RotMat = GetSizeRotationMatrix(CfgRemap.Angle);
 
-	FVector2D CorrectedNewSize = NewSize;
+	FVector2f CorrectedNewSize = NewSize;
 	if (IsRotationAngleDegenerate(CfgRemap.Angle) || RotMat.Determinant() == 0.0f)
 	{
 		// The bounding box for a 45 degree rotation is always a square, so correct the size to the smallest square.
-		CorrectedNewSize = FVector2D(NewSize.GetMin());
+		CorrectedNewSize = FVector2f(NewSize.GetMin());
 	}
 	else
 	{
 		const FMatrix2x2 InvRotMat = RotMat.Inverse();
-		const FVector2D UnrotatedSize = InvRotMat.TransformVector(NewSize);
-		const FVector2D ClampedSize = UnrotatedSize.ClampAxes(UDisplayClusterConfigurationViewport::ViewportMinimumSize, UDisplayClusterConfigurationViewport::ViewportMaximumSize);
+		const FVector2f UnrotatedSize = InvRotMat.TransformVector(NewSize);
+		const FVector2f ClampedSize = UnrotatedSize.ClampAxes(UDisplayClusterConfigurationViewport::ViewportMinimumSize, UDisplayClusterConfigurationViewport::ViewportMaximumSize);
 
 		CorrectedNewSize = RotMat.TransformVector(ClampedSize);
 	}
diff --git a/Engine/Plugins/Runtime/nDisplay/Source/DisplayClusterConfigurator/Private/Views/OutputMapping/EdNodes/DisplayClusterConfiguratorViewportNode.h b/Engine/Plugins/Runtime/nDisplay/Source/DisplayClusterConfigurator/Private/Views/OutputMapping/EdNodes/DisplayClusterConfiguratorViewportNode.h
--- a/Engine/Plugins/Runtime/nDisplay/Source/DisplayClusterConfigurator/Private/Views/OutputMapping/EdNodes/DisplayClusterConfiguratorViewportNode.h
+++ b/Engine/Plugins/Runtime/nDisplay/Source/DisplayClusterConfigurator/Private/Views/OutputMapping/EdNodes/DisplayClusterConfiguratorViewportNode.h
@@ -27,7 +27,7 @@ public:
 
 	//~ Begin EdGraphNode Interface
 	virtual TSharedPtr<SGraphNode> CreateVisualWidget() override;
-	virtual void ResizeNode(const FVector2D& NewSize) override;
+	virtual void ResizeNode(const FVector2f& NewSize) override;
 	virtual bool CanDuplicateNode() const override { return true; }
 	virtual bool CanUserDeleteNode() const override { return true; }
 	//~ End EdGraphNode Interface
diff --git a/Engine/Plugins/Runtime/nDisplay/Source/DisplayClusterConfigurator/Private/Views/OutputMapping/GraphNodes/SDisplayClusterConfiguratorBaseNode.cpp b/Engine/Plugins/Runtime/nDisplay/Source/DisplayClusterConfigurator/Private/Views/OutputMapping/GraphNodes/SDisplayClusterConfiguratorBaseNode.cpp
--- a/Engine/Plugins/Runtime/nDisplay/Source/DisplayClusterConfigurator/Private/Views/OutputMapping/GraphNodes/SDisplayClusterConfiguratorBaseNode.cpp
+++ b/Engine/Plugins/Runtime/nDisplay/Source/DisplayClusterConfigurator/Private/Views/OutputMapping/GraphNodes/SDisplayClusterConfiguratorBaseNode.cpp
@@ -445,7 +445,7 @@ void SDisplayClusterConfiguratorBaseNode::SetNodeSize(const FVector2D InLocalSiz
 		AlignmentOffset.Y = FMath::Max(AlignmentOffset.Y, -BestSize.Y);
 	}
 
-	GraphNode->ResizeNode(BestSize + AlignmentOffset);
+	GraphNode->ResizeNode(UE::Slate::FDeprecateVector2DParameter(BestSize + AlignmentOffset));
 
 	// If the parent node is being auto-positioned, add it to the undo stack here because we need to store its old position with the this node's old position
 	// so that if the move operation is undone, this node can appropriately reset the backing config object's position without requiring a full auto-positioning pass.
diff --git a/Engine/Source/Editor/GraphEditor/Private/SGraphNodeComment.cpp b/Engine/Source/Editor/GraphEditor/Private/SGraphNodeComment.cpp
--- a/Engine/Source/Editor/GraphEditor/Private/SGraphNodeComment.cpp
+++ b/Engine/Source/Editor/GraphEditor/Private/SGraphNodeComment.cpp
@@ -293,7 +293,7 @@ FReply SGraphNodeComment::OnMouseButtonUp( const FGeometry& MyGeometry, const FP
 		UserSize.X = FMath::RoundToFloat(UserSize.X);
 		UserSize.Y = FMath::RoundToFloat(UserSize.Y);
 
-		GetNodeObj()->ResizeNode(FVector2D(UserSize));
+		GetNodeObj()->ResizeNode(UserSize);
 
 		// End resize transaction
 		ResizeTransactionPtr.Reset();
diff --git a/Engine/Source/Editor/UnrealEd/Classes/MaterialGraph/MaterialGraphNode_Comment.h b/Engine/Source/Editor/UnrealEd/Classes/MaterialGraph/MaterialGraphNode_Comment.h
--- a/Engine/Source/Editor/UnrealEd/Classes/MaterialGraph/MaterialGraphNode_Comment.h
+++ b/Engine/Source/Editor/UnrealEd/Classes/MaterialGraph/MaterialGraphNode_Comment.h
@@ -36,7 +36,7 @@ class UMaterialGraphNode_Comment : public UEdGraphNode_Comment
 	virtual bool CanCreateUnderSpecifiedSchema(const UEdGraphSchema* Schema) const override;
 	virtual void PostPlacedNewNode() override;
 	virtual void OnRenameNode(const FString& NewName) override;
-	virtual void ResizeNode(const FVector2D& NewSize) override;
+	virtual void ResizeNode(const FVector2f& NewSize) override;
 	virtual int32 GetFontSize() const override;
 	//~ End UEdGraphNode Interface.
 
diff --git a/Engine/Source/Editor/UnrealEd/Private/EdGraphNode_Comment.cpp b/Engine/Source/Editor/UnrealEd/Private/EdGraphNode_Comment.cpp
--- a/Engine/Source/Editor/UnrealEd/Private/EdGraphNode_Comment.cpp
+++ b/Engine/Source/Editor/UnrealEd/Private/EdGraphNode_Comment.cpp
@@ -154,7 +154,7 @@ FLinearColor UEdGraphNode_Comment::GetNodeCommentColor() const
 		: FLinearColor::White;
 }
 
-void UEdGraphNode_Comment::ResizeNode(const FVector2D& NewSize)
+void UEdGraphNode_Comment::ResizeNode(const FVector2f& NewSize)
 {
 	if (bCanResizeNode) 
 	{
diff --git a/Engine/Source/Editor/UnrealEd/Private/MaterialGraphNode_Comment.cpp b/Engine/Source/Editor/UnrealEd/Private/MaterialGraphNode_Comment.cpp
--- a/Engine/Source/Editor/UnrealEd/Private/MaterialGraphNode_Comment.cpp
+++ b/Engine/Source/Editor/UnrealEd/Private/MaterialGraphNode_Comment.cpp
@@ -157,7 +157,7 @@ void UMaterialGraphNode_Comment::OnRenameNode(const FString& NewName)
 	}
 }
 
-void UMaterialGraphNode_Comment::ResizeNode(const FVector2D& NewSize)
+void UMaterialGraphNode_Comment::ResizeNode(const FVector2f& NewSize)
 {
 	Super::ResizeNode(NewSize);
 
diff --git a/Engine/Source/Editor/UnrealEd/Public/EdGraphNode_Comment.h b/Engine/Source/Editor/UnrealEd/Public/EdGraphNode_Comment.h
--- a/Engine/Source/Editor/UnrealEd/Public/EdGraphNode_Comment.h
+++ b/Engine/Source/Editor/UnrealEd/Public/EdGraphNode_Comment.h
@@ -93,7 +93,7 @@ public:
 	UNREALED_API virtual FText GetNodeTitle(ENodeTitleType::Type TitleType) const override;
 	UNREALED_API virtual bool ShouldOverridePinNames() const override { return true; }
 	UNREALED_API virtual FText GetPinNameOverride(const UEdGraphPin& Pin) const override;
-	UNREALED_API virtual void ResizeNode(const FVector2D& NewSize) override;
+	UNREALED_API virtual void ResizeNode(const FVector2f& NewSize) override;
 	UNREALED_API virtual void PostPlacedNewNode() override;
 	UNREALED_API virtual void OnRenameNode(const FString& NewName) override;
 	UNREALED_API virtual TSharedPtr<class INameValidatorInterface> MakeNameValidator() const override;
diff --git a/Engine/Source/Runtime/Engine/Classes/EdGraph/EdGraph.h b/Engine/Source/Runtime/Engine/Classes/EdGraph/EdGraph.h
--- a/Engine/Source/Runtime/Engine/Classes/EdGraph/EdGraph.h
+++ b/Engine/Source/Runtime/Engine/Classes/EdGraph/EdGraph.h
@@ -208,7 +208,7 @@ public:
 	static ENGINE_API UEdGraph* GetOuterGraph(UObject* Obj);
 
 	/** Util to find a good place for a new node */
-	ENGINE_API FVector2D GetGoodPlaceForNewNode();
+	ENGINE_API UE::Slate::FDeprecateVector2DResult GetGoodPlaceForNewNode();
 
 #if WITH_EDITOR
 	/** Notify the graph and its associated listeners that a property is about to change  */
diff --git a/Engine/Source/Runtime/Engine/Classes/EdGraph/EdGraphNode.h b/Engine/Source/Runtime/Engine/Classes/EdGraph/EdGraphNode.h
--- a/Engine/Source/Runtime/Engine/Classes/EdGraph/EdGraphNode.h
+++ b/Engine/Source/Runtime/Engine/Classes/EdGraph/EdGraphNode.h
@@ -7,6 +7,7 @@
 #include "UObject/Object.h"
 #include "Misc/Guid.h"
 #include "UObject/WeakObjectPtr.h"
+#include "Types/SlateVector2.h"
 #include "EdGraphNode.generated.h"
 
 class INameValidatorInterface;
@@ -843,7 +844,9 @@ public:
 	ENGINE_API virtual FString GetDescriptiveCompiledName() const;
 
 	/** Update node size to new value */
-	virtual void ResizeNode(const FVector2D& NewSize) {}
+	UE_DEPRECATED(5.6, "Use the version of the function accepting FVector2f")
+	ENGINE_API virtual void ResizeNode(const FVector2D& NewSize) UE_SLATE_DEPRECATED_VECTOR_VIRTUAL_FUNCTION {}
+	ENGINE_API virtual void ResizeNode(const FVector2f& NewSize);
 
 	/**
 	 * Returns whether or not this node has dependencies on an external structure
diff --git a/Engine/Source/Runtime/Engine/Classes/EdGraph/EdGraphNode_Documentation.h b/Engine/Source/Runtime/Engine/Classes/EdGraph/EdGraphNode_Documentation.h
--- a/Engine/Source/Runtime/Engine/Classes/EdGraph/EdGraphNode_Documentation.h
+++ b/Engine/Source/Runtime/Engine/Classes/EdGraph/EdGraphNode_Documentation.h
@@ -35,7 +35,7 @@ public:
 	ENGINE_API virtual FText GetTooltipText() const override;
 	ENGINE_API virtual FText GetNodeTitle(ENodeTitleType::Type TitleType) const override;
 	ENGINE_API virtual FText GetPinNameOverride(const UEdGraphPin& Pin) const override;
-	ENGINE_API virtual void ResizeNode(const FVector2D& NewSize) override;
+	ENGINE_API virtual void ResizeNode(const FVector2f& NewSize) override;
 	ENGINE_API virtual void PostPlacedNewNode() override;
 	ENGINE_API virtual void OnRenameNode(const FString& NewName) override;
 	ENGINE_API virtual TSharedPtr<class INameValidatorInterface> MakeNameValidator() const override;
diff --git a/Engine/Source/Runtime/Engine/Private/EdGraph/EdGraph.cpp b/Engine/Source/Runtime/Engine/Private/EdGraph/EdGraph.cpp
--- a/Engine/Source/Runtime/Engine/Private/EdGraph/EdGraph.cpp
+++ b/Engine/Source/Runtime/Engine/Private/EdGraph/EdGraph.cpp
@@ -353,29 +353,29 @@ UEdGraph* UEdGraph::GetOuterGraph(UObject* Obj)
 	return nullptr;
 }
 
-FVector2D UEdGraph::GetGoodPlaceForNewNode()
+UE::Slate::FDeprecateVector2DResult UEdGraph::GetGoodPlaceForNewNode()
 {
-	FVector2D BottomLeft(0,0);
+	FVector2f BottomLeft(0.0f, 0.0f);
 
 	if(Nodes.Num() > 0)
 	{
 		UEdGraphNode* Node = Nodes[0];
 		if (Node)
 		{
-			BottomLeft = FVector2D(Node->NodePosX, Node->NodePosY);
+			BottomLeft = FVector2f(Node->NodePosX, Node->NodePosY);
 			for (int32 i = 1; i < Nodes.Num(); i++)
 			{
 				Node = Nodes[i];
 				if (Node)
 				{
-					BottomLeft.X = FMath::Min<float>(BottomLeft.X, Node->NodePosX);
-					BottomLeft.Y = FMath::Max<float>(BottomLeft.Y, Node->NodePosY);
+					BottomLeft.X = FMath::Min(BottomLeft.X, Node->NodePosX);
+					BottomLeft.Y = FMath::Max(BottomLeft.Y, Node->NodePosY);
 				}
 			}
 		}
 	}
 
-	return BottomLeft + FVector2D(0, 256);
+	return BottomLeft + FVector2f(0.0f, 256.0f);
 }
 
 #if WITH_EDITOR
diff --git a/Engine/Source/Runtime/Engine/Private/EdGraph/EdGraphNode.cpp b/Engine/Source/Runtime/Engine/Private/EdGraph/EdGraphNode.cpp
--- a/Engine/Source/Runtime/Engine/Private/EdGraph/EdGraphNode.cpp
+++ b/Engine/Source/Runtime/Engine/Private/EdGraph/EdGraphNode.cpp
@@ -623,6 +623,13 @@ FString UEdGraphNode::GetDescriptiveCompiledName() const
 	return GetFName().GetPlainNameString();
 }
 
+void UEdGraphNode::ResizeNode(const FVector2f& NewSize)
+{
+PRAGMA_DISABLE_DEPRECATION_WARNINGS
+	ResizeNode(FVector2D(NewSize));
+PRAGMA_ENABLE_DEPRECATION_WARNINGS
+}
+
 bool UEdGraphNode::IsDeprecated() const
 {
 	return GetClass()->HasAnyClassFlags(CLASS_Deprecated);
diff --git a/Engine/Source/Runtime/Engine/Private/EdGraph/EdGraphNode_Documentation.cpp b/Engine/Source/Runtime/Engine/Private/EdGraph/EdGraphNode_Documentation.cpp
--- a/Engine/Source/Runtime/Engine/Private/EdGraph/EdGraphNode_Documentation.cpp
+++ b/Engine/Source/Runtime/Engine/Private/EdGraph/EdGraphNode_Documentation.cpp
@@ -65,12 +65,12 @@ FText UEdGraphNode_Documentation::GetPinNameOverride(const UEdGraphPin& Pin) con
 	return GetNodeTitle( ENodeTitleType::ListView );
 }
 
-void UEdGraphNode_Documentation::ResizeNode(const FVector2D& NewSize)
+void UEdGraphNode_Documentation::ResizeNode(const FVector2f& NewSize)
 {
 	if( bCanResizeNode ) 
 	{
-		NodeHeight = NewSize.Y;
-		NodeWidth = NewSize.X;
+		NodeHeight = UE::LWC::FloatToIntCastChecked<int32>(NewSize.Y);
+		NodeWidth = UE::LWC::FloatToIntCastChecked<int32>(NewSize.X);
 	}
 }
 
@@ -79,7 +79,7 @@ void UEdGraphNode_Documentation::SetBounds( const class FSlateRect& Rect )
 	NodePosX = Rect.Left;
 	NodePosY = Rect.Top;
 
-	FVector2D Size = Rect.GetSize();
+	FVector2f Size = Rect.GetSize();
 	NodeWidth = Size.X;
 	NodeHeight = Size.Y;
 }
