nDisplay: nDisplay: ICVFX Panel migration to new toolbar part 3/3.

Final projection icons still pending.

#rnx

diff --git a/Engine/Plugins/Runtime/nDisplay/Source/DisplayClusterLightCardEditor/Private/Viewport/DisplayClusterLightcardEditorViewport.cpp b/Engine/Plugins/Runtime/nDisplay/Source/DisplayClusterLightCardEditor/Private/Viewport/DisplayClusterLightcardEditorViewport.cpp
--- a/Engine/Plugins/Runtime/nDisplay/Source/DisplayClusterLightCardEditor/Private/Viewport/DisplayClusterLightcardEditorViewport.cpp
+++ b/Engine/Plugins/Runtime/nDisplay/Source/DisplayClusterLightCardEditor/Private/Viewport/DisplayClusterLightcardEditorViewport.cpp
@@ -76,25 +76,15 @@ public:
 
 		ToolbarMenu->StyleName = "ViewportToolbar";
 
-		// Left section: TRS buttons.
+		// Left section: TRS and coordinate system (spherical/cartesian)
 		{
 			FToolMenuSection& LeftSection = ToolbarMenu->FindOrAddSection("Left");
 
-			LeftSection.AddEntry(FToolMenuEntry::InitToolBarButton(
-				FEditorViewportCommands::Get().TranslateMode,
-				TAttribute<FText>(FText::GetEmpty()) // empty label
-			));
-
-			LeftSection.AddEntry(FToolMenuEntry::InitToolBarButton(
-				FEditorViewportCommands::Get().RotateMode,
-				TAttribute<FText>(FText::GetEmpty())
-			));
-
-			LeftSection.AddEntry(FToolMenuEntry::InitToolBarButton(
-				FEditorViewportCommands::Get().ScaleMode,
-				TAttribute<FText>(FText::GetEmpty())
-			));
-
+			// TRS selection options, includes dropdown and 
+			// Note: Not using UE::UnrealEd::CreateTransformsSubmenu because of the extra clutter of n/a options.
+			LeftSection.AddEntry(CreateTransformsSubmenu());
+			
+			// Separator, to match level editor.
 			LeftSection.AddSeparator("CoordinateSystemSeparator");
 
 			// Coordinate system toggle button
@@ -202,6 +192,47 @@ public:
 		return UToolMenus::Get()->GenerateWidget(ToolbarName, Context);
 	}
 
+	/** Creates TRS menu. Slimmed down version of UE::UnrealEd::CreateTransformsSubmenu */
+	FToolMenuEntry CreateTransformsSubmenu()
+	{
+		FToolMenuEntry Entry = FToolMenuEntry::InitSubMenu(
+			"Transform",
+			LOCTEXT("TransformsSubmenuLabel", "Transform"),
+			LOCTEXT("TransformsSubmenuTooltip", "Viewport-related transforms tools"),
+			FNewToolMenuDelegate::CreateLambda(
+				[](UToolMenu* Submenu) -> void
+				{
+					{
+						FToolMenuSection& TransformToolsSection =
+							Submenu->FindOrAddSection("TransformTools", LOCTEXT("TransformToolsLabel", "Transform Tools"));
+
+						FToolMenuEntry TranslateMode =
+							FToolMenuEntry::InitMenuEntry(FEditorViewportCommands::Get().TranslateMode);
+						TranslateMode.SetShowInToolbarTopLevel(true);
+						TranslateMode.ToolBarData.StyleNameOverride = "ViewportToolbar.TransformTools";
+						TransformToolsSection.AddEntry(TranslateMode);
+
+						FToolMenuEntry RotateMode = FToolMenuEntry::InitMenuEntry(FEditorViewportCommands::Get().RotateMode);
+						RotateMode.SetShowInToolbarTopLevel(true);
+						RotateMode.ToolBarData.StyleNameOverride = "ViewportToolbar.TransformTools";
+						TransformToolsSection.AddEntry(RotateMode);
+
+						FToolMenuEntry ScaleMode = FToolMenuEntry::InitMenuEntry(FEditorViewportCommands::Get().ScaleMode);
+						ScaleMode.SetShowInToolbarTopLevel(true);
+						ScaleMode.ToolBarData.StyleNameOverride = "ViewportToolbar.TransformTools";
+						TransformToolsSection.AddEntry(ScaleMode);
+					}
+				}
+			)
+		);
+
+		Entry.Icon = FSlateIcon(FAppStyle::GetAppStyleSetName(), "LevelEditor.SelectMode");
+		Entry.ToolBarData.LabelOverride = FText();
+		Entry.ToolBarData.ResizeParams.ClippingPriority = 1000;
+
+		return Entry;
+	}
+
 	/** Make the DrawLightCard toggle button widget. */
 	TSharedRef<SWidget> MakeDrawLightCardWidget()
 	{
