[nDisplay] Makes it possible to copy/duplicate ICVFX camera components

#rnx

diff --git a/Engine/Plugins/Runtime/nDisplay/Source/DisplayCluster/Public/Components/DisplayClusterICVFXCameraComponent.h b/Engine/Plugins/Runtime/nDisplay/Source/DisplayCluster/Public/Components/DisplayClusterICVFXCameraComponent.h
--- a/Engine/Plugins/Runtime/nDisplay/Source/DisplayCluster/Public/Components/DisplayClusterICVFXCameraComponent.h
+++ b/Engine/Plugins/Runtime/nDisplay/Source/DisplayCluster/Public/Components/DisplayClusterICVFXCameraComponent.h
@@ -24,7 +24,7 @@ class IDisplayClusterViewport;
 /**
  * nDisplay in-camera VFX camera representation
  */
-UCLASS(ClassGroup = (DisplayCluster), HideCategories = (AssetUserData, Collision, Cooking, ComponentReplication, Events, Physics, Sockets, Activation, Tags, ComponentTick), meta = (DisplayName="ICVFX Camera"))
+UCLASS(ClassGroup = (DisplayCluster), HideCategories = (AssetUserData, Collision, Cooking, ComponentReplication, Events, Physics, Sockets, Activation, Tags, ComponentTick), meta = (BlueprintSpawnableComponent, DisplayName="ICVFX Camera"))
 class DISPLAYCLUSTER_API UDisplayClusterICVFXCameraComponent
 	: public UCineCameraComponent
 {
