[nDisplay] Fixed crash in logic

The failure was caused by assert().
This happens when we have multiple viewports and we remove the mesh from one of them.
Logic should fall back to default math.

diff --git a/Engine/Plugins/Runtime/nDisplay/Source/DisplayClusterWarp/Private/Policy/InFrustumFit/DisplayClusterWarpInFrustumFitPolicy_Math.cpp b/Engine/Plugins/Runtime/nDisplay/Source/DisplayClusterWarp/Private/Policy/InFrustumFit/DisplayClusterWarpInFrustumFitPolicy_Math.cpp
--- a/Engine/Plugins/Runtime/nDisplay/Source/DisplayClusterWarp/Private/Policy/InFrustumFit/DisplayClusterWarpInFrustumFitPolicy_Math.cpp
+++ b/Engine/Plugins/Runtime/nDisplay/Source/DisplayClusterWarp/Private/Policy/InFrustumFit/DisplayClusterWarpInFrustumFitPolicy_Math.cpp
@@ -95,9 +95,6 @@ void FDisplayClusterWarpInFrustumFitPolicy::BeginCalcFrustum(IDisplayClusterView
 
 	FDisplayClusterWarpData& WarpData = WarpBlend->GetWarpData(ContextNum);
 
-	// Todo: This feature now not supported here. need to be fixed.
-	WarpData.bEnabledRotateFrustumToFitContextSize = false;
-
 	// Override view direction for entire group
 	if (WarpData.WarpEye.IsValid())
 	{
@@ -127,9 +124,14 @@ void FDisplayClusterWarpInFrustumFitPolicy::BeginCalcFrustum(IDisplayClusterView
 		}
 		else
 		{
-			check(false); // United frustum can't be builded properly in this case.
+			// United frustum can't be builded properly in this case.
+			// do nothing and render without customization (as Default ViewPoint)
+			return;
 		}
 	}
+
+	// Todo: This feature now not supported here. need to be fixed.
+	WarpData.bEnabledRotateFrustumToFitContextSize = false;
 }
 
 void FDisplayClusterWarpInFrustumFitPolicy::EndCalcFrustum(IDisplayClusterViewport* InViewport, const uint32 ContextNum)
