[nDisplay] Show a gray texture in the DRCA preview when the viewport is not rendering properly.

The issue was that when viewport stopped rendering for some reason, the last rendered result was displayed.
This may be unclear to the user.

diff --git a/Engine/Plugins/Runtime/nDisplay/Source/DisplayCluster/Private/Render/DisplayDevice/Components/DisplayClusterDisplayDeviceBaseComponent.cpp b/Engine/Plugins/Runtime/nDisplay/Source/DisplayCluster/Private/Render/DisplayDevice/Components/DisplayClusterDisplayDeviceBaseComponent.cpp
--- a/Engine/Plugins/Runtime/nDisplay/Source/DisplayCluster/Private/Render/DisplayDevice/Components/DisplayClusterDisplayDeviceBaseComponent.cpp
+++ b/Engine/Plugins/Runtime/nDisplay/Source/DisplayCluster/Private/Render/DisplayDevice/Components/DisplayClusterDisplayDeviceBaseComponent.cpp
@@ -142,7 +142,7 @@ void UDisplayClusterDisplayDeviceBaseComponent::OnUpdateDisplayDeviceMeshAndMate
 {
 	using namespace UE::DisplayCluster::DisplayDeviceBaseComponent;
 
-	if (!InMeshComponent || !InMeshMaterialInstance || !ShouldUseDisplayDevice(InViewportPreview.GetConfiguration()))
+	if (!ShouldUseDisplayDevice(InViewportPreview.GetConfiguration()))
 	{
 		return;
 	}
@@ -152,10 +152,16 @@ void UDisplayClusterDisplayDeviceBaseComponent::OnUpdateDisplayDeviceMeshAndMate
 	case EDisplayClusterDisplayDeviceMeshType::PreviewMesh:
 	case EDisplayClusterDisplayDeviceMeshType::PreviewEditableMesh:
 		// Only some of the preview meshes are supported by Techvis
-		ImplUpdatePreviewMeshTechvis(InViewportPreview, *InMeshComponent);
+		if (InMeshComponent)
+		{
+			ImplUpdatePreviewMeshTechvis(InViewportPreview, *InMeshComponent);
+		}
 
 		// Only some of the preview meshes are supported by preview
-		ImplUpdatePreviewMeshMaterialInstanceParameters(InViewportPreview, InMaterialType, *InMeshMaterialInstance);
+		if (InMeshMaterialInstance)
+		{
+			ImplUpdatePreviewMeshMaterialInstanceParameters(InViewportPreview, InMaterialType, *InMeshMaterialInstance);
+		}
 		break;
 
 	default:
diff --git a/Engine/Plugins/Runtime/nDisplay/Source/DisplayCluster/Private/Render/Viewport/DisplayClusterViewportManager.cpp b/Engine/Plugins/Runtime/nDisplay/Source/DisplayCluster/Private/Render/Viewport/DisplayClusterViewportManager.cpp
--- a/Engine/Plugins/Runtime/nDisplay/Source/DisplayCluster/Private/Render/Viewport/DisplayClusterViewportManager.cpp
+++ b/Engine/Plugins/Runtime/nDisplay/Source/DisplayCluster/Private/Render/Viewport/DisplayClusterViewportManager.cpp
@@ -679,12 +679,6 @@ bool FDisplayClusterViewportManager::BeginNewFrame(FViewport* InViewport, FDispl
 
 	PostProcessManager->HandleBeginNewFrame(OutRenderFrame);
 
-	// Update viewport preview instances if preview is used and DCRA supports previews:
-	if (Configuration->IsPreviewRendering())
-	{
-		ViewportManagerPreview->Update();
-	}
-
 	return true;
 }
 
diff --git a/Engine/Plugins/Runtime/nDisplay/Source/DisplayCluster/Private/Render/Viewport/Preview/DisplayClusterViewportManagerPreview_Render.cpp b/Engine/Plugins/Runtime/nDisplay/Source/DisplayCluster/Private/Render/Viewport/Preview/DisplayClusterViewportManagerPreview_Render.cpp
--- a/Engine/Plugins/Runtime/nDisplay/Source/DisplayCluster/Private/Render/Viewport/Preview/DisplayClusterViewportManagerPreview_Render.cpp
+++ b/Engine/Plugins/Runtime/nDisplay/Source/DisplayCluster/Private/Render/Viewport/Preview/DisplayClusterViewportManagerPreview_Render.cpp
@@ -175,6 +175,14 @@ bool FDisplayClusterViewportManagerPreview::InitializeClusterNodePreview(const E
 							// uvLC viewport is not rendered, but late OCIO is used in FDisplayClusterViewportProxy::ApplyOCIO_RenderThread()
 							FSceneView* View = ViewportPtr->ViewportPreview->CalcSceneView(*ViewFamily, ViewIt.ContextNum);
 
+							if (View == nullptr)
+							{
+								// This viewport is not rendering. Release all textures.
+								// The preview mesh materials associated with them will revert to default materials.
+								// In this case, all invalid viewports will be displayed gray or black.
+								ViewportPtr->ReleaseTextures();
+							}
+
 							// Remove the view that is not used for rendering.
 							if (View && (!ViewIt.IsViewportContextCanBeRendered() || ViewFamily->RenderTarget == nullptr))
 							{
@@ -198,6 +206,9 @@ bool FDisplayClusterViewportManagerPreview::InitializeClusterNodePreview(const E
 					}
 				}
 
+				// Update viewport preview instances if preview is used and DCRA supports previews:
+				Update();
+
 				return ViewportsViewFamily.Num() > 0;
 			}
 		}
diff --git a/Engine/Plugins/Runtime/nDisplay/Source/DisplayCluster/Private/Render/Viewport/Preview/DisplayClusterViewportPreview.cpp b/Engine/Plugins/Runtime/nDisplay/Source/DisplayCluster/Private/Render/Viewport/Preview/DisplayClusterViewportPreview.cpp
--- a/Engine/Plugins/Runtime/nDisplay/Source/DisplayCluster/Private/Render/Viewport/Preview/DisplayClusterViewportPreview.cpp
+++ b/Engine/Plugins/Runtime/nDisplay/Source/DisplayCluster/Private/Render/Viewport/Preview/DisplayClusterViewportPreview.cpp
@@ -80,8 +80,15 @@ void FDisplayClusterViewportPreview::Update(TSet<const UMeshComponent*>& MeshCom
 	auto UpdatePreviewMeshAndMaterial = [&](
 		FDisplayClusterViewportPreviewMesh& PreviewMeshObj,
 		EDisplayClusterViewportPreviewFlags ChangedMaterialFlag,
-		EDisplayClusterDisplayDeviceMeshType MeshType)
+		EDisplayClusterDisplayDeviceMeshType InMeshType)
 		{
+			EDisplayClusterDisplayDeviceMeshType MeshType = InMeshType;
+			if (!EnumHasAnyFlags(RuntimeFlags, EDisplayClusterViewportPreviewFlags::HasValidPreviewRTT))
+			{
+				// Use mesh without preview.
+				MeshType = EDisplayClusterDisplayDeviceMeshType::DefaultMesh;
+			}
+
 			PreviewMeshObj.Update(Viewport, DisplayDeviceComponent, ViewPointComponent);
 
 			// Update Runtime Flags
