[nDisplay] Fixed InFrustumFit math.
The problem was in the logic of searching for a symmetric frustum.
If the required precision was not found, it returned an error.
In the this CL, at each iteration step the best, closest to symmetric frustum is stored, which is then used if the accuracy is not achieved.

diff --git a/Engine/Plugins/Runtime/nDisplay/Source/DisplayClusterWarp/Private/Policy/InFrustumFit/DisplayClusterWarpInFrustumFitPolicy.cpp b/Engine/Plugins/Runtime/nDisplay/Source/DisplayClusterWarp/Private/Policy/InFrustumFit/DisplayClusterWarpInFrustumFitPolicy.cpp
--- a/Engine/Plugins/Runtime/nDisplay/Source/DisplayClusterWarp/Private/Policy/InFrustumFit/DisplayClusterWarpInFrustumFitPolicy.cpp
+++ b/Engine/Plugins/Runtime/nDisplay/Source/DisplayClusterWarp/Private/Policy/InFrustumFit/DisplayClusterWarpInFrustumFitPolicy.cpp
@@ -109,27 +109,22 @@ void FDisplayClusterWarpInFrustumFitPolicy::HandleNewFrame(const TArray<TSharedP
 					break;
 				}
 
-				if (FrustumData.IterationNum == INDEX_NONE)
-				{
-					// The maximum iteration number has been reached.
-					if (FrustumData.NewUnitedSymmetricWarpProjection.IsSet())
+				// If the maximum iteration number has been reached, or new frustum can't be calculated.
+				if (FrustumData.IterationNum == INDEX_NONE
+				|| !FrustumData.NewWorldViewTarget.IsSet())
+				{	
+					// Use the best saved symmetric projection over all iterations.
+					if (FrustumData.BestSymmetricWarpProjection.IsSet())
 					{
-						// Use last valid projection
-						OptUnitedGeometryWarpProjection = *FrustumData.NewUnitedSymmetricWarpProjection;
-						break;
-					}
+						check(FrustumData.BestWorldViewTarget.IsSet());
 
-					// no valid projection found.
-					break;
-				}
+						OptUnitedGeometryWarpProjection = FrustumData.BestSymmetricWarpProjection->Value;
+						OptOverrideWorldViewTarget = *FrustumData.BestWorldViewTarget;
+					}
 
-				if (!FrustumData.NewWorldViewTarget.IsSet())
-				{
-					// no valid projection found.
 					break;
 				}
 
-
 				// Set the new location of the view target and perform the next iteration.
 				OptOverrideWorldViewTarget = *FrustumData.NewWorldViewTarget;
 			}
diff --git a/Engine/Plugins/Runtime/nDisplay/Source/DisplayClusterWarp/Private/Policy/InFrustumFit/DisplayClusterWarpInFrustumFitPolicy.h b/Engine/Plugins/Runtime/nDisplay/Source/DisplayClusterWarp/Private/Policy/InFrustumFit/DisplayClusterWarpInFrustumFitPolicy.h
--- a/Engine/Plugins/Runtime/nDisplay/Source/DisplayClusterWarp/Private/Policy/InFrustumFit/DisplayClusterWarpInFrustumFitPolicy.h
+++ b/Engine/Plugins/Runtime/nDisplay/Source/DisplayClusterWarp/Private/Policy/InFrustumFit/DisplayClusterWarpInFrustumFitPolicy.h
@@ -86,6 +86,10 @@ private:
 
 		// new location of the viewing target in the world space.
 		TOptional<FVector> NewWorldViewTarget;
+
+		// Store the best symmetric projection, which is used when the iterative method has not found a projection within the expected precision.
+		TOptional<TPair<double, FDisplayClusterWarpProjection>> BestSymmetricWarpProjection;
+		TOptional<FVector> BestWorldViewTarget;
 	};
 
 private:
diff --git a/Engine/Plugins/Runtime/nDisplay/Source/DisplayClusterWarp/Private/Policy/InFrustumFit/DisplayClusterWarpInFrustumFitPolicy_Math.cpp b/Engine/Plugins/Runtime/nDisplay/Source/DisplayClusterWarp/Private/Policy/InFrustumFit/DisplayClusterWarpInFrustumFitPolicy_Math.cpp
--- a/Engine/Plugins/Runtime/nDisplay/Source/DisplayClusterWarp/Private/Policy/InFrustumFit/DisplayClusterWarpInFrustumFitPolicy_Math.cpp
+++ b/Engine/Plugins/Runtime/nDisplay/Source/DisplayClusterWarp/Private/Policy/InFrustumFit/DisplayClusterWarpInFrustumFitPolicy_Math.cpp
@@ -301,11 +301,11 @@ bool FDisplayClusterWarpInFrustumFitPolicy::CalcUnitedGeometryFrustum(const TArr
 {
 	// The united geometry frustum is built from geometric points projected onto a special plane.
 	// This plane is called the 'projection plane' and is created from two quantities: the view direction vector and the eye position.
-
 	check(!OptUnitedGeometryWarpProjection.IsSet()); // when we calculate the united geometry frustum, we must disable the frustum fit logic.
 
 	OutUnitedGeometryWarpProjection.ResetProjectionAngles();
 
+	bool bResult = true;
 	for (const TSharedPtr<IDisplayClusterViewport, ESPMode::ThreadSafe>& Viewport : InViewports)
 	{
 		TSharedPtr<IDisplayClusterWarpBlend, ESPMode::ThreadSafe> WarpBlend;
@@ -347,15 +347,15 @@ bool FDisplayClusterWarpInFrustumFitPolicy::CalcUnitedGeometryFrustum(const TArr
 
 		if (!WarpBlend->CalcFrustumContext(WarpEye))
 		{
-			return false;
+			bResult = false;
 		}
 
+		// Merge two projection
 		const FDisplayClusterWarpData& WarpData = WarpBlend->GetWarpData(InContextNum);
-
 		OutUnitedGeometryWarpProjection.ExpandProjectionAngles(WarpData.GeometryWarpProjection);
 	}
 
-	return true;
+	return bResult;
 }
 
 bool FDisplayClusterWarpInFrustumFitPolicy::CalcUnitedGeometrySymmetricFrustum(const TArray<TSharedPtr<IDisplayClusterViewport, ESPMode::ThreadSafe>>& InViewports, const int32 InContextNum, FSymmetricFrustumData& InOutSymmetricFrustumData)
@@ -376,25 +376,53 @@ bool FDisplayClusterWarpInFrustumFitPolicy::CalcUnitedGeometrySymmetricFrustum(c
 
 	// Calculate the united  frustum for the current view target.
 	FDisplayClusterWarpProjection UnitedGeometryWarpProjection;
-	if (!CalcUnitedGeometryFrustum(InViewports, InContextNum, InOutSymmetricFrustumData.WorldScale, UnitedGeometryWarpProjection))
-	{
-		// Frustum not exist
-		return false;
-	}
 
-	const double OffsetH = UnitedGeometryWarpProjection.Left + (UnitedGeometryWarpProjection.Right - UnitedGeometryWarpProjection.Left) * 0.5;
-	const double OffsetV = UnitedGeometryWarpProjection.Bottom + (UnitedGeometryWarpProjection.Top - UnitedGeometryWarpProjection.Bottom) * 0.5;
+	const double OffsetH = 0.5 * (UnitedGeometryWarpProjection.Left + UnitedGeometryWarpProjection.Right);
+	const double OffsetV = 0.5 * (UnitedGeometryWarpProjection.Bottom + UnitedGeometryWarpProjection.Top);
 
 	const double OffsetDegreesH = UnitedGeometryWarpProjection.ConvertProjectionToDegrees(OffsetH);
 	const double OffsetDegreesV = UnitedGeometryWarpProjection.ConvertProjectionToDegrees(OffsetV);
 
-	// Check the precision.
-	if (FMath::Max(FMath::Abs(OffsetDegreesH), FMath::Abs(OffsetDegreesV)) <= FMath::Abs(GDisplayClusterWarpInFrustumFitPolicySymmetricPrecision))
+	const bool bFrustumValid = CalcUnitedGeometryFrustum(InViewports, InContextNum, InOutSymmetricFrustumData.WorldScale, UnitedGeometryWarpProjection);
+	if (bFrustumValid)
 	{
-		// A symmetric frustum has been found:
-		InOutSymmetricFrustumData.NewUnitedSymmetricWarpProjection = UnitedGeometryWarpProjection;
+		const double MaxOffset = FMath::Max(FMath::Abs(OffsetH), FMath::Abs(OffsetV));
+
+		// Update best value
+		if (!InOutSymmetricFrustumData.BestSymmetricWarpProjection.IsSet()
+			|| InOutSymmetricFrustumData.BestSymmetricWarpProjection->Key > MaxOffset)
+		{
+			InOutSymmetricFrustumData.BestSymmetricWarpProjection = { MaxOffset , UnitedGeometryWarpProjection };
+			InOutSymmetricFrustumData.BestWorldViewTarget = *OptOverrideWorldViewTarget;
+		}
+
+		if (InOutSymmetricFrustumData.BestSymmetricWarpProjection->Key < MaxOffset)
+		{
+			// @todo Add a new method to get a better result (multiple rays in a sphere, etc.)
 
-		return true;
+			// current value if worst, stop iterator
+			InOutSymmetricFrustumData.IterationNum = INDEX_NONE;
+
+			return false;
+		}
+
+		// Check the precision.
+		if (FMath::Max(FMath::Abs(OffsetDegreesH), FMath::Abs(OffsetDegreesV)) <= FMath::Abs(GDisplayClusterWarpInFrustumFitPolicySymmetricPrecision))
+		{
+			// A symmetric frustum has been found:
+			InOutSymmetricFrustumData.NewUnitedSymmetricWarpProjection = UnitedGeometryWarpProjection;
+
+			return true;
+		}
+	}
+	else
+	{
+		// Frustum not exist
+		if (InOutSymmetricFrustumData.IterationNum > 1)
+		{
+			// on the first iteration dont exit and try to find something from the base AABB center-point
+			return false;
+		}
 	}
 
 	// Calculate the new location of the view target from the values of the united geometry frustum.
