[nDisplay] Fixed invalid logging that reported "invalid layout [1x1] was requested "

#rnx

diff --git a/Engine/Plugins/Runtime/nDisplay/Source/DisplayClusterMedia/Private/DisplayClusterMediaModule.cpp b/Engine/Plugins/Runtime/nDisplay/Source/DisplayClusterMedia/Private/DisplayClusterMediaModule.cpp
--- a/Engine/Plugins/Runtime/nDisplay/Source/DisplayClusterMedia/Private/DisplayClusterMediaModule.cpp
+++ b/Engine/Plugins/Runtime/nDisplay/Source/DisplayClusterMedia/Private/DisplayClusterMediaModule.cpp
@@ -227,20 +227,20 @@ void FDisplayClusterMediaModule::InitializeBackbufferUniformTilesOutput(const UD
 	{
 		const FDisplayClusterConfigurationMediaNodeBackbuffer& MediaSettings = ClusterNode->MediaSettings;
 
-		// Validate tile layout
-		const FIntPoint MaxLayout = FDisplayClusterMediaCaptureNodeTile::GetMaxTileLayout();
-		const bool bLayoutIsValid = DisplayClusterMediaHelpers::IsValidLayout(MediaSettings.TiledSplitLayout, MaxLayout);
-		if (!bLayoutIsValid)
+		if (MediaSettings.bEnable)
 		{
-			UE_LOG(LogDisplayClusterMedia, Warning, TEXT("Invalid layout [%dx%d] was requested for backbuffer capture. Max layout is [%dx%d]."),
-				MediaSettings.TiledSplitLayout.X, MediaSettings.TiledSplitLayout.Y,
-				MaxLayout.X, MaxLayout.Y);
+			const FIntPoint MaxLayout = FDisplayClusterMediaCaptureNodeTile::GetMaxTileLayout();
+			const bool bLayoutIsValid = DisplayClusterMediaHelpers::IsValidLayout(MediaSettings.TiledSplitLayout, MaxLayout);
 
-			return;
-		}
+			// Validate tile layout
+			if (!bLayoutIsValid)
+			{
+				UE_LOG(LogDisplayClusterMedia, Warning, TEXT("Invalid layout [%dx%d] was requested for backbuffer capture. Max layout is [%dx%d]."),
+					MediaSettings.TiledSplitLayout.X, MediaSettings.TiledSplitLayout.Y, MaxLayout.X, MaxLayout.Y);
+
+				return;
+			}
 
-		if (MediaSettings.bEnable && bLayoutIsValid)
-		{
 			uint8 CaptureIdx = 0;
 			for (const FDisplayClusterConfigurationMediaUniformTileOutput& MediaOutputTile : MediaSettings.TiledMediaOutputs)
 			{
