[nDisplay] Update default template value of cluster outer viewports screen percentage, and exposed it to .ndisplay config

#rnx

diff --git a/Engine/Plugins/Runtime/nDisplay/Source/DisplayClusterConfiguration/Private/Formats/JSON500/DisplayClusterConfigurationJsonParser_500.cpp b/Engine/Plugins/Runtime/nDisplay/Source/DisplayClusterConfiguration/Private/Formats/JSON500/DisplayClusterConfigurationJsonParser_500.cpp
--- a/Engine/Plugins/Runtime/nDisplay/Source/DisplayClusterConfiguration/Private/Formats/JSON500/DisplayClusterConfigurationJsonParser_500.cpp
+++ b/Engine/Plugins/Runtime/nDisplay/Source/DisplayClusterConfiguration/Private/Formats/JSON500/DisplayClusterConfigurationJsonParser_500.cpp
@@ -107,6 +107,7 @@ namespace JSON500
 		Config->bExitOnEsc = CfgJson.Misc.bExitOnEsc;
 		Config->bOverrideViewportsFromExternalConfig = CfgJson.Misc.bOverrideViewportsFromExternalConfig;
 		Config->bOverrideTransformsFromExternalConfig = CfgJson.Misc.bOverrideTransformsFromExternalConfig;
+		Config->RenderFrameSettings.ClusterICVFXOuterViewportBufferRatioMult = CfgJson.Misc.ClusterOuterViewportsBufferRatio;
 
 		// Scene
 		{
@@ -314,6 +315,7 @@ namespace JSON500
 		Json.Misc.bExitOnEsc = Config->bExitOnEsc;
 		Json.Misc.bOverrideViewportsFromExternalConfig = Config->bOverrideViewportsFromExternalConfig;
 		Json.Misc.bOverrideTransformsFromExternalConfig = Config->bOverrideTransformsFromExternalConfig;
+		Json.Misc.ClusterOuterViewportsBufferRatio = Config->RenderFrameSettings.ClusterICVFXOuterViewportBufferRatioMult;
 
 		// Scene
 		{
diff --git a/Engine/Plugins/Runtime/nDisplay/Source/DisplayClusterConfiguration/Private/Formats/JSON500/DisplayClusterConfigurationJsonTypes_500.h b/Engine/Plugins/Runtime/nDisplay/Source/DisplayClusterConfiguration/Private/Formats/JSON500/DisplayClusterConfigurationJsonTypes_500.h
--- a/Engine/Plugins/Runtime/nDisplay/Source/DisplayClusterConfiguration/Private/Formats/JSON500/DisplayClusterConfigurationJsonTypes_500.h
+++ b/Engine/Plugins/Runtime/nDisplay/Source/DisplayClusterConfiguration/Private/Formats/JSON500/DisplayClusterConfigurationJsonTypes_500.h
@@ -216,6 +216,7 @@ public:
 		, bExitOnEsc(true)
 		, bOverrideViewportsFromExternalConfig(false)
 		, bOverrideTransformsFromExternalConfig(false)
+		, ClusterOuterViewportsBufferRatio(0.66f)
 	{ }
 
 public:
@@ -230,6 +231,9 @@ public:
 
 	UPROPERTY()
 	bool bOverrideTransformsFromExternalConfig;
+
+	UPROPERTY()
+	float ClusterOuterViewportsBufferRatio;
 };
 
 
diff --git a/Engine/Plugins/Runtime/nDisplay/Source/DisplayClusterConfigurator/Private/DisplayClusterConfiguratorFactory.cpp b/Engine/Plugins/Runtime/nDisplay/Source/DisplayClusterConfigurator/Private/DisplayClusterConfiguratorFactory.cpp
--- a/Engine/Plugins/Runtime/nDisplay/Source/DisplayClusterConfigurator/Private/DisplayClusterConfiguratorFactory.cpp
+++ b/Engine/Plugins/Runtime/nDisplay/Source/DisplayClusterConfigurator/Private/DisplayClusterConfiguratorFactory.cpp
@@ -167,14 +167,17 @@ void UDisplayClusterConfiguratorFactory::SetupNewBlueprint(UDisplayClusterBluepr
 
 	UDisplayClusterConfigurationData* NewConfigData = UDisplayClusterConfigurationData::CreateNewConfigData();
 
-	// In Windows, the default sync policy should be Nvidia. Linux/other can be Ethernet by default, since Nvidia sync might not be supported.
-#if PLATFORM_WINDOWS
 	if (ensure(NewConfigData && NewConfigData->Cluster))
 	{
+#if PLATFORM_WINDOWS
+		// In Windows, the default sync policy should be Nvidia. Linux/other can be Ethernet by default, since Nvidia sync might not be supported.
 		NewConfigData->Cluster->Sync.RenderSyncPolicy.Type = DisplayClusterConfigurationStrings::config::cluster::render_sync::Nvidia;
-	}
 #endif // PLATFORM_WINDOWS
 
+		// Decrease outer viewports screen percentage by default for new configurations
+		NewConfigData->RenderFrameSettings.ClusterICVFXOuterViewportBufferRatioMult = 0.66f;
+	}
+
 	NewBlueprint->SetConfigData(NewConfigData);
 	
 	check(NewBlueprint->GeneratedClass);
