[nDisplay] [Virtual Production] DLSS plugin integration for nDisplay.

Added new module VirtualProduction:
- 'Engine/Source/Runtime/VirtualProduction/VirtualProduction'
- Added `IUpscalerModularFeature` interface.

The new IUpscalerModularFeature API is used to integrate upscalers.
- The goal is to have the project compile and work even if there are no plugins with the feature. For example, some third-party plugins such as DLSS are used by nDisplay but are not present in the Unreal Engine release.
- Abstract custom settings of a module feature can be exported to the UI as a PropertyBag, and after modifying it can be sent back to the module feature without knowing this structure.
- The rendering pipeline is configured on the implementation side of the upscaler.

Changes to the nDisplay plugin:
- Added nDisplay configuration for AA and upscaling settings.
- Implemented nDisplay configuration for each viewport for upscaler and AA.
- Added detail customization `FDisplayClusterConfigurationUpscalerSettingsDetailCustomization`

diff --git a/Engine/Plugins/Runtime/nDisplay/Source/DisplayCluster/DisplayCluster.Build.cs b/Engine/Plugins/Runtime/nDisplay/Source/DisplayCluster/DisplayCluster.Build.cs
--- a/Engine/Plugins/Runtime/nDisplay/Source/DisplayCluster/DisplayCluster.Build.cs
+++ b/Engine/Plugins/Runtime/nDisplay/Source/DisplayCluster/DisplayCluster.Build.cs
@@ -22,6 +22,7 @@ public class DisplayCluster : ModuleRules
 				"Engine",
 				"EngineSettings",
 				"EnhancedInput",
+				"VirtualProduction",
 			});
 
 		PrivateDependencyModuleNames.AddRange(
diff --git a/Engine/Plugins/Runtime/nDisplay/Source/DisplayCluster/Private/Render/Upscaler/DisplayClusterUpscaler.cpp b/Engine/Plugins/Runtime/nDisplay/Source/DisplayCluster/Private/Render/Upscaler/DisplayClusterUpscaler.cpp
new file mode 100644
--- /dev/null
+++ b/Engine/Plugins/Runtime/nDisplay/Source/DisplayCluster/Private/Render/Upscaler/DisplayClusterUpscaler.cpp
@@ -0,0 +1,312 @@
+// Copyright Epic Games, Inc. All Rights Reserved.
+
+#include "Render/Upscaler/DisplayClusterUpscaler.h"
+
+#include "IUpscalerModularFeature.h"
+
+#include "Features/IModularFeatures.h"
+
+#include "LegacyScreenPercentageDriver.h"
+#include "SceneView.h"
+
+using namespace UE::VirtualProduction;
+
+namespace UE::DisplayCluster::Upscaler
+{
+	int32 GDisplayClusterRenderUpscaling_HideDisabledMethods = 1;
+	static FAutoConsoleVariableRef CVarDisplayClusterRenderUpscaling_HideDisabledMethods(
+		TEXT("nDisplay.render.upscaling.HideDisabledMethods"),
+		GDisplayClusterRenderUpscaling_HideDisabledMethods,
+		TEXT("Hide disabled upscaling methods (default = 1).\n")
+		TEXT("0 - Show all upsccaling methods.\n"),
+		ECVF_Default
+	);
+
+	static inline bool ShouldShowOnlyEnabledUpscalingMethods()
+	{
+		return GDisplayClusterRenderUpscaling_HideDisabledMethods!=0;
+	}
+
+	/** Iterate over all upscaler modular features. */
+	static void ForeachUpscaler(TFunction<void(IUpscalerModularFeature&)>&& IteratorFunc)
+	{
+		static IModularFeatures& ModularFeatures = IModularFeatures::Get();
+
+		ModularFeatures.LockModularFeatureList();
+		{
+			const TArray<IUpscalerModularFeature*> Extenders = ModularFeatures.GetModularFeatureImplementations<IUpscalerModularFeature>(IUpscalerModularFeature::ModularFeatureName);
+			for (IUpscalerModularFeature* Extender : Extenders)
+			{
+				// Iterate over active upscalers
+				if (Extender && (!ShouldShowOnlyEnabledUpscalingMethods() || Extender->IsFeatureEnabled()))
+				{
+					IteratorFunc(*Extender);
+				}
+			}
+		}
+		ModularFeatures.UnlockModularFeatureList();
+	}
+
+	/** Get global fraction. */
+	static float GetGlobalScreenPercentage()
+	{
+		static const TConsoleVariableData<float>* CVarScreenPercentage = IConsoleManager::Get().FindTConsoleVariableDataFloat(TEXT("r.ScreenPercentage"));
+		if (CVarScreenPercentage)
+		{
+			const float GlobalScreenPercentage = CVarScreenPercentage->GetValueOnGameThread() / 100.0f;
+			if (GlobalScreenPercentage > 0)
+			{
+				return GlobalScreenPercentage;
+			}
+		}
+
+		return 1.0f;
+	}
+
+	/** Check temporal upscaler CVar*/
+	static bool IsTemporalUpscalerActive()
+	{
+		static const IConsoleVariable* CVarTemporalAAUpscaler = IConsoleManager::Get().FindConsoleVariable(TEXT("r.TemporalAA.Upscaler"));
+		return CVarTemporalAAUpscaler && (CVarTemporalAAUpscaler->GetInt() != 0);
+	}
+
+	/** Temporal upscaler interface requires special AA method. */
+	static bool SceneViewSupportsTemporalUpscaler(const FSceneView& InSceneView)
+	{
+		// check if current AA method legitable for temporal upscaler
+		switch (InSceneView.AntiAliasingMethod)
+		{
+		case EAntiAliasingMethod::AAM_TSR:
+		case EAntiAliasingMethod::AAM_TemporalAA: //@TODO check AAM_TemporalAA
+			return true;
+
+		default:
+			break;
+		}
+
+		return false;
+	}
+
+	/** Get secondary screen percentage fraction. */
+	static float GetSecondaryViewFraction(const float InDPIScale)
+	{
+		static IConsoleVariable* CVarCustomSecondaryScreenPercentage = IConsoleManager::Get().FindConsoleVariable(TEXT("r.SecondaryScreenPercentage.GameViewport"), false);
+		if (CVarCustomSecondaryScreenPercentage)
+		{
+			const float CustomSecondaryScreenPercentage = CVarCustomSecondaryScreenPercentage->GetFloat();
+		if (CustomSecondaryScreenPercentage > 0.0)
+		{
+			// Override secondary resolution fraction with CVar.
+			return FMath::Min(CustomSecondaryScreenPercentage / 100.0f, 1.0f);
+		}
+		}
+
+#if WITH_EDITOR
+			// Automatically compute secondary resolution fraction from DPI.
+			// When in high res screenshot do not modify screen percentage based on dpi scale
+			if (GIsEditor && !GIsHighResScreenshot)
+			{
+				static IConsoleVariable* CVarEditorViewportHighDPIPtr = IConsoleManager::Get().FindConsoleVariable(TEXT("r.Editor.Viewport.HighDPI"));
+				if (CVarEditorViewportHighDPIPtr && CVarEditorViewportHighDPIPtr->GetInt() == 0)
+				{
+
+					const float DPIDerivedResolutionFraction = FMath::Min(1.0f / InDPIScale, 1.0f);
+
+					return DPIDerivedResolutionFraction;
+				}
+			}
+#endif
+
+		return 1.0f;
+	}
+};
+
+void FDisplayClusterUpscaler::SetupSceneView(
+	const FDisplayClusterUpscalerSettings& InUpscalerSettings,
+	FSceneView& InOutView)
+{
+	using namespace UE::DisplayCluster::Upscaler;
+
+	// Use custom AA methods
+	switch (InUpscalerSettings.AntiAliasingMethod)
+	{
+	case EDisplayClusterUpscalerAntiAliasingMethod::FXAA:
+		InOutView.AntiAliasingMethod = EAntiAliasingMethod::AAM_FXAA;
+		break;
+
+	case EDisplayClusterUpscalerAntiAliasingMethod::MSAA:
+		InOutView.AntiAliasingMethod = EAntiAliasingMethod::AAM_MSAA;
+		break;
+
+	case EDisplayClusterUpscalerAntiAliasingMethod::None:
+		InOutView.AntiAliasingMethod = EAntiAliasingMethod::AAM_None;
+		break;
+
+	case EDisplayClusterUpscalerAntiAliasingMethod::TAA:
+		InOutView.AntiAliasingMethod = EAntiAliasingMethod::AAM_TemporalAA;
+		break;
+
+	case EDisplayClusterUpscalerAntiAliasingMethod::TSR:
+		InOutView.AntiAliasingMethod = EAntiAliasingMethod::AAM_TSR;
+		break;
+
+	default:
+		break;
+	}
+
+	// Custom upscaler should configure view for the settings
+	ForeachUpscaler(
+		[&](IUpscalerModularFeature& Upscaler)
+		{
+			if (Upscaler.GetName() == InUpscalerSettings.CustomUpscalerName)
+			{
+				Upscaler.SetupSceneView(InUpscalerSettings.CustomUpscalerSettings, InOutView);
+			}
+		});
+}
+
+FName FDisplayClusterUpscaler::PostConfigureViewFamily(
+	const FDisplayClusterUpscalerSettings& InUpscalerSettings,
+	const float InScreenPercentage,
+	const float InDPIScale,
+	FSceneViewFamilyContext& InOutViewFamily,
+	const TArray<FSceneView*>& InOutViews)
+{
+	using namespace UE::DisplayCluster::Upscaler;
+
+	// One view per viewfamily
+	check(InOutViewFamily.Views.Num() == 1);
+	check(InOutViewFamily.Views[0]);
+
+#if WITH_EDITOR
+	if (GIsEditor)
+	{
+		// Force enable view family show flag for HighDPI derived's screen percentage.
+		InOutViewFamily.EngineShowFlags.ScreenPercentage = true;
+	}
+#endif
+
+	// Force screen percentage show flag to be turned off if not supported.
+	if (!InOutViewFamily.SupportsScreenPercentage())
+	{
+		InOutViewFamily.EngineShowFlags.ScreenPercentage = false;
+	}
+
+	if (!InOutViewFamily.EngineShowFlags.ScreenPercentage)
+	{
+		return NAME_None;
+	}
+
+	// If a screen percentage interface was set by dynamic resolution, then do nothing
+	if (InOutViewFamily.GetScreenPercentageInterface() != nullptr)
+	{
+		return NAME_None;
+	}
+
+	FUpscalerModularFeatureParameters UpscalerParam;
+	UpscalerParam.UpscalerScreenPercentage  = InScreenPercentage;
+	UpscalerParam.SecondaryScreenPercentage = GetSecondaryViewFraction(InDPIScale);
+
+	check(UpscalerParam.UpscalerScreenPercentage > 0.0f);
+	check(UpscalerParam.SecondaryScreenPercentage > 0.0f);
+
+	// Find  upscaler modular feature
+	FName UpscalerName = NAME_None;
+	ForeachUpscaler([&](IUpscalerModularFeature& Upscaler)
+		{
+			if (Upscaler.GetName() == InUpscalerSettings.CustomUpscalerName)
+			{
+				if (Upscaler.PostConfigureViewFamily(
+					InUpscalerSettings.CustomUpscalerSettings, UpscalerParam,
+					InOutViewFamily, InOutViews))
+				{
+					// Return upscaler name on success
+					UpscalerName = Upscaler.GetName();
+				}
+			}
+		});
+
+	if (UpscalerName != NAME_None)
+	{
+		return UpscalerName;
+	}
+
+	// Can't use Custom upscalers, back to the AntiAliasing method.
+
+	//@todo: logic below is in process/testing << need to discuss
+	switch (InUpscalerSettings.AntiAliasingMethod)
+	{
+	case EDisplayClusterUpscalerAntiAliasingMethod::Default:
+	case EDisplayClusterUpscalerAntiAliasingMethod::TSR:
+		// We need to split the screen percentage if below 0.5 because TAA upscaling only works well up to 2x.
+		{
+			// Get global view fraction set by r.ScreenPercentage.
+			const float GlobalScreenPercentage = GetGlobalScreenPercentage();
+
+			// Adjust for global screen% in math bellow
+			const float FinalResolutionFraction = UpscalerParam.UpscalerScreenPercentage * GlobalScreenPercentage;
+			if (FinalResolutionFraction < 0.5f)
+			{
+				UpscalerParam.SecondaryScreenPercentage *= 2.0f * FinalResolutionFraction;
+
+				// Adjust value for the global screen%. Expectations that the actual TSR_Fraction=0.5 
+				UpscalerParam.UpscalerScreenPercentage = 0.5f / GlobalScreenPercentage;
+			}
+		}
+
+		InOutViewFamily.SetScreenPercentageInterface(new FLegacyScreenPercentageDriver(
+			InOutViewFamily, UpscalerParam.UpscalerScreenPercentage));
+
+		InOutViewFamily.SecondaryViewFraction = UpscalerParam.SecondaryScreenPercentage;
+		break;
+
+	// Don't use upscaling
+	default:
+		InOutViewFamily.SetScreenPercentageInterface(new FLegacyScreenPercentageDriver(
+			InOutViewFamily, 1.0f));
+
+		InOutViewFamily.SecondaryViewFraction = UpscalerParam.SecondaryScreenPercentage;
+		break;
+	}
+
+	// In any other cases use default upscaler
+	return NAME_None;
+}
+
+void FDisplayClusterUpscaler::SetupUpscalerViewExtensions(bool bEnabled)
+{
+	using namespace UE::DisplayCluster::Upscaler;
+
+	ForeachUpscaler(
+		[&](IUpscalerModularFeature& Upscaler)
+	{
+			// By default use upscaler's rules for the editor viewports.
+			Upscaler.SetupUpscalerViewExtensions(
+				bEnabled, 
+				/*/bEnableInEditorViewport*/ false);
+		});
+}
+
+void FDisplayClusterUpscaler::SetupUpscalerViewExtensionsForSettings(const FDisplayClusterUpscalerSettings& InUpscalerSettings)
+{
+	using namespace UE::DisplayCluster::Upscaler;
+
+	ForeachUpscaler(
+		[&](IUpscalerModularFeature& Upscaler)
+		{
+			if (Upscaler.GetName() == InUpscalerSettings.CustomUpscalerName)
+			{
+				// Enable all necessary ViewExtension of this upscaler
+				// DCRA preview in editor can use upscalers.
+				Upscaler.SetupUpscalerViewExtensionsForSettings(
+					InUpscalerSettings.CustomUpscalerSettings,
+					/**bEnableInEditorViewport*/ true);
+			}
+			else
+			{
+				// Disable all ViewExtensions for all other upscalers
+				Upscaler.SetupUpscalerViewExtensions(
+					/**bEnabled*/ false, /**bEnableInEditorViewport*/ false);
+			}
+		});
+}
\ No newline at end of file
diff --git a/Engine/Plugins/Runtime/nDisplay/Source/DisplayCluster/Private/Render/Upscaler/DisplayClusterUpscaler.h b/Engine/Plugins/Runtime/nDisplay/Source/DisplayCluster/Private/Render/Upscaler/DisplayClusterUpscaler.h
new file mode 100644
--- /dev/null
+++ b/Engine/Plugins/Runtime/nDisplay/Source/DisplayCluster/Private/Render/Upscaler/DisplayClusterUpscaler.h
@@ -0,0 +1,57 @@
+// Copyright Epic Games, Inc. All Rights Reserved.
+
+#pragma once
+
+#include "CoreMinimal.h"
+
+#include "Render/Upscaler/DisplayClusterUpscalerSettings.h"
+
+class FSceneViewFamilyContext;
+class FSceneView;
+
+/**
+ * nDisplay upscaling implementation.
+ */
+class FDisplayClusterUpscaler
+{
+public:
+	/*
+	* Setup all View Extensions of the upscaler to be enabled or disabled.
+	* Must be called before GatherActiveExtensions() during the rendering pass setup.
+	*
+	* @param bEnable - (in) Enable or disable all ViewExtensions of this upscaler.
+	*/
+	static void SetupUpscalerViewExtensions(bool bEnabled);
+
+	/**
+	* Enable only necessary view extensions used by upscaler for these input settings.
+	* Must be called before GatherActiveExtensions() during the rendering pass setup.
+	* 
+	* @param InUpscalerSettings - (in) the settings to apply for viewport
+	*/
+	static void SetupUpscalerViewExtensionsForSettings(const FDisplayClusterUpscalerSettings& InUpscalerSettings);
+
+	/** Configure upscaler and viewfamily for ertain settings.
+	* 
+	* @param InUpscalerSettings - (in) the settings to apply for viewport
+	* @param InScreenPercentage - (inj screen percentage of the viewport
+	* @param InDPIScale         - (in) The DPI scale of the window
+	* @param InOutViewFamily    - (in, out) The view family that will be configured.
+	* @param InOutViews         - (in, out) The views in the family that will be configured.
+	*/
+	static FName PostConfigureViewFamily(
+		const FDisplayClusterUpscalerSettings& InUpscalerSettings,
+		const float InScreenPercentage,
+		const float InDPIScale,
+		FSceneViewFamilyContext& InOutViewFamily,
+		const TArray<FSceneView*>& InOutViews);
+
+	/**
+	* Configure scene view for upscalers.
+	* @param InUpscalerSettings - (in) the settings to apply for view
+	* @param InOutViews         - (in, out) The views that will be configured. 
+	*/
+	static void SetupSceneView(
+		const FDisplayClusterUpscalerSettings& InUpscalerSettings,
+		FSceneView& InOutView);
+};
diff --git a/Engine/Plugins/Runtime/nDisplay/Source/DisplayCluster/Private/Render/Viewport/Configuration/DisplayClusterViewportConfigurationHelpers.cpp b/Engine/Plugins/Runtime/nDisplay/Source/DisplayCluster/Private/Render/Viewport/Configuration/DisplayClusterViewportConfigurationHelpers.cpp
--- a/Engine/Plugins/Runtime/nDisplay/Source/DisplayCluster/Private/Render/Viewport/Configuration/DisplayClusterViewportConfigurationHelpers.cpp
+++ b/Engine/Plugins/Runtime/nDisplay/Source/DisplayCluster/Private/Render/Viewport/Configuration/DisplayClusterViewportConfigurationHelpers.cpp
@@ -122,6 +122,9 @@ void FDisplayClusterViewportConfigurationHelpers::UpdateBaseViewportSetting(FDis
 	{
 		DstViewport.SetViewportBufferRatio(InRenderSettings.BufferRatio);
 
+		// Update viewport Upscaler method
+		InRenderSettings.UpscalerSettings.GetUpscalerSettings(InOutRenderSettings.UpscalerSettings);
+
 		DstViewport.UpdateConfiguration_Overscan(GetViewportOverscanSettings(InRenderSettings.Overscan));
 
 		DstViewport.UpdateConfiguration_PostRenderOverride(InRenderSettings.Replace);
diff --git a/Engine/Plugins/Runtime/nDisplay/Source/DisplayCluster/Private/Render/Viewport/Configuration/DisplayClusterViewportConfigurationHelpers_ICVFX.cpp b/Engine/Plugins/Runtime/nDisplay/Source/DisplayCluster/Private/Render/Viewport/Configuration/DisplayClusterViewportConfigurationHelpers_ICVFX.cpp
--- a/Engine/Plugins/Runtime/nDisplay/Source/DisplayCluster/Private/Render/Viewport/Configuration/DisplayClusterViewportConfigurationHelpers_ICVFX.cpp
+++ b/Engine/Plugins/Runtime/nDisplay/Source/DisplayCluster/Private/Render/Viewport/Configuration/DisplayClusterViewportConfigurationHelpers_ICVFX.cpp
@@ -527,6 +527,9 @@ void FDisplayClusterViewportConfigurationHelpers_ICVFX::UpdateCameraViewportSett
 	// Set viewport buffer ratio
 	DstViewport.SetViewportBufferRatio(InCameraSettings.GetCameraBufferRatio(*StageSettings));
 
+	// Update viewport Upscaler method
+	InCameraSettings.GetCameraUpscalerSettings(*StageSettings, InOutRenderSettings.UpscalerSettings);
+
 	// InCamera tile rendering.
 	FDisplayClusterViewportConfigurationHelpers_Tile::UpdateICVFXCameraViewportTileSettings(DstViewport, InCameraSettings.RenderSettings.Media);
 }
diff --git a/Engine/Plugins/Runtime/nDisplay/Source/DisplayCluster/Private/Render/Viewport/DisplayClusterViewport.cpp b/Engine/Plugins/Runtime/nDisplay/Source/DisplayCluster/Private/Render/Viewport/DisplayClusterViewport.cpp
--- a/Engine/Plugins/Runtime/nDisplay/Source/DisplayCluster/Private/Render/Viewport/DisplayClusterViewport.cpp
+++ b/Engine/Plugins/Runtime/nDisplay/Source/DisplayCluster/Private/Render/Viewport/DisplayClusterViewport.cpp
@@ -28,6 +28,7 @@
 #include "Render/Viewport/Preview/DisplayClusterViewportPreview.h"
 
 #include "Render/DisplayDevice/Components/DisplayClusterDisplayDeviceBaseComponent.h"
+#include "Render/Upscaler/DisplayClusterUpscaler.h"
 
 #include "EngineUtils.h"
 #include "OpenColorIORendering.h"
@@ -247,7 +248,7 @@ void FDisplayClusterViewport::ImplGatherActiveExtensions(const int32 ViewIndex,
 		{
 			if (ViewportManager->GetViewportManagerViewPointExtension().IsValid())
 			{
-				// To call SetupView() as early as possible: integration with other VEs
+				// To call SetupView() as early as possible: integration with other ViewExtensions
 				OutExtensions.Add(ViewportManager->GetViewportManagerViewPointExtension()->AsShared());
 			}
 			if (ViewportManager->GetViewportManagerViewExtension().IsValid())
@@ -400,6 +401,10 @@ void FDisplayClusterViewport::SetupSceneView(uint32 ContextNum, class UWorld* Wo
 		InDisplayDeviceComponent->SetupSceneView(*ViewportPreview, ContextNum, InOutViewFamily, InOutView);
 	}
 
+	// Handle upscalers
+	FDisplayClusterUpscaler::SetupSceneView(RenderSettings.UpscalerSettings, InOutView);
+
+
 	if (InOutView.State && World && World->Scene)
 	{
 		if (ShouldUseLumenPerView())
diff --git a/Engine/Plugins/Runtime/nDisplay/Source/DisplayCluster/Private/Render/Viewport/DisplayClusterViewportManager.cpp b/Engine/Plugins/Runtime/nDisplay/Source/DisplayCluster/Private/Render/Viewport/DisplayClusterViewportManager.cpp
--- a/Engine/Plugins/Runtime/nDisplay/Source/DisplayCluster/Private/Render/Viewport/DisplayClusterViewportManager.cpp
+++ b/Engine/Plugins/Runtime/nDisplay/Source/DisplayCluster/Private/Render/Viewport/DisplayClusterViewportManager.cpp
@@ -29,6 +29,7 @@
 #include "Render/Viewport/Configuration/DisplayClusterViewportConfiguration_Viewport.h"
 
 #include "Render/Viewport/DisplayClusterViewportStrings.h"
+#include "Render/Upscaler/DisplayClusterUpscaler.h"
 
 #include "IDisplayClusterWarpBlend.h"
 
@@ -825,106 +826,39 @@ void FDisplayClusterViewportManager::PostConfigureViewFamily(
 	const float InDPIScale
 )
 {
-#if WITH_EDITOR
-	if (GIsEditor)
-	{
-		// Force enable view family show flag for HighDPI derived's screen percentage.
-		ViewFamily.EngineShowFlags.ScreenPercentage = true;
-	}
-#endif
-
-	// Force screen percentage show flag to be turned off if not supported.
-	if (!ViewFamily.SupportsScreenPercentage())
-	{
-		ViewFamily.EngineShowFlags.ScreenPercentage = false;
-	}
-
-	const float DisplayClusterScreenPercentage = FMath::Max(0.01f, InFrameViewFamily.CustomBufferRatio);
-
-	// Get DPIDerivedResolutionFraction value:
-	float DPIDerivedResolutionFraction = 1.0f;
-	{
-#if WITH_EDITOR
-		// Automatically compute secondary resolution fraction from DPI.
-		// When in high res screenshot do not modify screen percentage based on dpi scale
-		if (GIsEditor && !GIsHighResScreenshot)
-		{
-			static auto CVarEditorViewportHighDPIPtr = IConsoleManager::Get().FindConsoleVariable(TEXT("r.Editor.Viewport.HighDPI"));
+	const float DisplayClusterScreenPercentage = FMath::Clamp(InFrameViewFamily.CustomBufferRatio, 0.01f, 1.0f);
 
-			if (CVarEditorViewportHighDPIPtr && CVarEditorViewportHighDPIPtr->GetInt() == 0)
-			{
-				DPIDerivedResolutionFraction = FMath::Min(1.0f / InDPIScale, 1.0f);
-			}
-		}
-#endif
-	}
+	// Configure upscaler
+	const FName CustomUpscalerName = FDisplayClusterUpscaler::PostConfigureViewFamily(InFrameViewFamily.UpscalerSettings, DisplayClusterScreenPercentage, InDPIScale, ViewFamily, Views);
 
-	// Setup the screen percentage and upscaling method for the view family.
-	if (ViewFamily.EngineShowFlags.ScreenPercentage)
-	{
-		// Set up secondary resolution fraction for the view family.
+	// Add ViewExtensions from the Custom Upscaler to this view family
+	if (!CustomUpscalerName.IsNone() && InFrameViewFamily.UpscalerSettings.CustomUpscalerName == CustomUpscalerName)
 		{
-			static IConsoleVariable* CVarCustomSecondaryScreenPercentage =
-				IConsoleManager::Get().FindConsoleVariable(TEXT("r.SecondaryScreenPercentage.GameViewport"), false);
-			const float CustomSecondaryScreenPercentage = CVarCustomSecondaryScreenPercentage ?
-				CVarCustomSecondaryScreenPercentage->GetFloat() : -1.f;
+		// Enable all necessary ViewExtension for the custom upscaler settings
+		FDisplayClusterUpscaler::SetupUpscalerViewExtensionsForSettings(InFrameViewFamily.UpscalerSettings);
 
-			if (CustomSecondaryScreenPercentage > 0.0)
+		// Iterate over all viewports, and collect all used ViewExtensions
+		for (const FDisplayClusterRenderFrameTargetView& FrameTargetViewIt : InFrameViewFamily.Views)
 			{
-				// Override secondary resolution fraction with CVar.
-				ViewFamily.SecondaryViewFraction = FMath::Min(CustomSecondaryScreenPercentage / 100.0f, 1.0f);
-			}
-			else
-			{
-				ViewFamily.SecondaryViewFraction = DPIDerivedResolutionFraction;
-			}
-
-			check(ViewFamily.SecondaryViewFraction > 0.0f);
-		}
-
-		// nDisplay cannot use the dynamic resolution feature because it uses multiple viewports (view families) in a single frame.
-		const bool bUsesDynamicResolution = false;
-
-		// If a screen percentage interface was not set by dynamic resolution, then create one matching legacy behavior.
-		if (ViewFamily.GetScreenPercentageInterface() == nullptr)
-		{
-			float GlobalResolutionFraction = 1.0f;
-			float SecondaryScreenPercentage = ViewFamily.SecondaryViewFraction;
-
-			if (ViewFamily.EngineShowFlags.ScreenPercentage)
-			{
-				// Get global view fraction set by r.ScreenPercentage.
-				GlobalResolutionFraction = DisplayClusterScreenPercentage;
-
-				// We need to split the screen percentage if below 0.5 because TAA upscaling only works well up to 2x.
-				if (GlobalResolutionFraction < 0.5f)
+			if (FrameTargetViewIt.Viewport.IsValid())
 				{
-					SecondaryScreenPercentage = 2.0f * GlobalResolutionFraction;
-					GlobalResolutionFraction = 0.5f;
-				}
-			}
-
-			ViewFamily.SetScreenPercentageInterface(new FLegacyScreenPercentageDriver(
-				ViewFamily, GlobalResolutionFraction));
-
-			ViewFamily.SecondaryViewFraction = SecondaryScreenPercentage;
-		}
+				TArray<FSceneViewExtensionRef> ViewExtensions = FrameTargetViewIt.Viewport->ToSharedRef()->GatherActiveExtensions(FrameTargetViewIt.ContextNum, InFrameTarget.Viewport);
 
-		check(ViewFamily.GetScreenPercentageInterface() != nullptr);
-
-		// Change screen percentage method to raw output when doing dynamic resolution with VR if not using TAA upsample.
-		if (EnumHasAnyFlags(InRenderingFlags, EDisplayClusterViewportRenderingFlags::StereoRendering)
-			&& bUsesDynamicResolution)
+				// Lookup for the new ViewExtensions, and register it for that ViewFamily
+				for (FSceneViewExtensionRef& ViewExtensionIt : ViewExtensions)
 		{
-			for (FSceneView* View : Views)
+					if (!ViewFamily.ViewExtensions.Contains(ViewExtensionIt))
 			{
-				if (View->PrimaryScreenPercentageMethod == EPrimaryScreenPercentageMethod::SpatialUpscale)
-				{
-					View->PrimaryScreenPercentageMethod = EPrimaryScreenPercentageMethod::RawOutput;
+						ViewFamily.ViewExtensions.Add(ViewExtensionIt);
+						ViewExtensionIt->SetupViewFamily(ViewFamily);
 				}
 			}
 		}
 	}
+
+		// Restore default behaviour of ViewExtensions of upscalers
+		FDisplayClusterUpscaler::SetupUpscalerViewExtensions(true);
+	}
 }
 
 void FDisplayClusterViewportManager::RenderFrame(FViewport* InViewport)
diff --git a/Engine/Plugins/Runtime/nDisplay/Source/DisplayCluster/Private/Render/Viewport/RenderFrame/DisplayClusterRenderFrameManager.cpp b/Engine/Plugins/Runtime/nDisplay/Source/DisplayCluster/Private/Render/Viewport/RenderFrame/DisplayClusterRenderFrameManager.cpp
--- a/Engine/Plugins/Runtime/nDisplay/Source/DisplayCluster/Private/Render/Viewport/RenderFrame/DisplayClusterRenderFrameManager.cpp
+++ b/Engine/Plugins/Runtime/nDisplay/Source/DisplayCluster/Private/Render/Viewport/RenderFrame/DisplayClusterRenderFrameManager.cpp
@@ -5,6 +5,7 @@
 #include "Render/Viewport/RenderFrame/DisplayClusterRenderFrameSettings.h"
 #include "Render/Viewport/Configuration/DisplayClusterViewportConfiguration.h"
 #include "Render/Viewport/DisplayClusterViewport.h"
+#include "Render/Upscaler/DisplayClusterUpscaler.h"
 
 #include "Misc/Parse.h"
 
@@ -79,8 +80,6 @@ bool FDisplayClusterRenderFrameManager::BuildRenderFrame(FViewport* InViewport,
 
 bool FDisplayClusterRenderFrameManager::BuildSimpleFrame(FViewport* InViewport, const TArray<TSharedPtr<FDisplayClusterViewport, ESPMode::ThreadSafe>>& InViewports, FDisplayClusterRenderFrame& OutRenderFrame)
 {
-	const FDisplayClusterRenderFrameSettings& InRenderFrameSettings = Configuration->GetRenderFrameSettings();
-
 	for (const TSharedPtr<FDisplayClusterViewport, ESPMode::ThreadSafe>& ViewportIt : InViewports)
 	{
 		if (ViewportIt.IsValid())
@@ -99,11 +98,19 @@ bool FDisplayClusterRenderFrameManager::BuildSimpleFrame(FViewport* InViewport,
 				{
 					FrameViewFamily.Views.Add(FrameView);
 					FrameViewFamily.CustomBufferRatio = ContextIt.CustomBufferRatio;
+					FrameViewFamily.UpscalerSettings = ViewportIt->GetRenderSettings().UpscalerSettings;
+
+					// Do not collect ViewExtension of upscalers for the initial list.
+					// We need to determine if we can use upscalers in the future when AntiAliasing methods are configured in SceneViews.
+					FDisplayClusterUpscaler::SetupUpscalerViewExtensions(false);
 					FrameViewFamily.ViewExtensions = ViewportIt->GatherActiveExtensions(ContextIt.ContextNum, InViewport);
+					FDisplayClusterUpscaler::SetupUpscalerViewExtensions(true);
 				}
 
 				FDisplayClusterRenderFrameTarget FrameRenderTarget;
 				{
+					FrameRenderTarget.Viewport = InViewport;
+
 					// Simple frame use unique RTT  for each viewport, so disable RTT when viewport rendering disabled
 					FrameRenderTarget.bShouldUseRenderTarget = ViewportIt->GetViewportResources(EDisplayClusterViewportResource::RenderTargets).Num() > 0;
 
diff --git a/Engine/Plugins/Runtime/nDisplay/Source/DisplayCluster/Public/Components/DisplayClusterICVFXCameraComponent.h b/Engine/Plugins/Runtime/nDisplay/Source/DisplayCluster/Public/Components/DisplayClusterICVFXCameraComponent.h
--- a/Engine/Plugins/Runtime/nDisplay/Source/DisplayCluster/Public/Components/DisplayClusterICVFXCameraComponent.h
+++ b/Engine/Plugins/Runtime/nDisplay/Source/DisplayCluster/Public/Components/DisplayClusterICVFXCameraComponent.h
@@ -124,6 +124,9 @@ private:
 	UPROPERTY(EditAnywhere, Transient, Category = "Inner Frustum", meta = (PropertyPath = "CameraSettings.HiddenICVFXViewports"))
 	FDisplayClusterEditorPropertyReference HiddenICVFXViewportsRef;
 
+	UPROPERTY(EditAnywhere, Transient, Category = "Inner Frustum", meta = (PropertyPath = "CameraSettings.UpscalerSettings"))
+	FDisplayClusterEditorPropertyReference UpscalerSettingsRef;
+
 	UPROPERTY(EditAnywhere, Transient, Category = "Inner Frustum", meta = (PropertyPath = "CameraSettings.BufferRatio"))
 	FDisplayClusterEditorPropertyReference BufferRatioRef;
 
diff --git a/Engine/Plugins/Runtime/nDisplay/Source/DisplayCluster/Public/DisplayClusterRootActor.h b/Engine/Plugins/Runtime/nDisplay/Source/DisplayCluster/Public/DisplayClusterRootActor.h
--- a/Engine/Plugins/Runtime/nDisplay/Source/DisplayCluster/Public/DisplayClusterRootActor.h
+++ b/Engine/Plugins/Runtime/nDisplay/Source/DisplayCluster/Public/DisplayClusterRootActor.h
@@ -298,6 +298,9 @@ private:
 	UPROPERTY(EditAnywhere, Transient, Category = Viewports, meta = (PropertyPath = "CurrentConfigData.RenderFrameSettings.ClusterICVFXOuterViewportBufferRatioMult"))
 	FDisplayClusterEditorPropertyReference ViewportScreenPercentageMultiplierRef;
 
+	UPROPERTY(EditInstanceOnly, Transient, Category = Viewports, meta = (DisplayName = "Viewport Upscaling Settings", PropertyPath = "CurrentConfigData.Cluster.Nodes.Viewports.RenderSettings.UpscalerSettings", ToolTip = "Adjust upscaler settings for an individual viewport."))
+	FDisplayClusterEditorPropertyReference ViewportUpscalerSettingsRef;
+
 	UPROPERTY(EditInstanceOnly, Transient, Category = Viewports, meta = (DisplayName = "Viewport Screen Percentage", PropertyPath = "CurrentConfigData.Cluster.Nodes.Viewports.RenderSettings.BufferRatio", ToolTip = "Adjust resolution scaling for an individual viewport.  Viewport Screen Percentage Multiplier is applied to this value."))
 	FDisplayClusterEditorPropertyReference ViewportScreenPercentageRef;
 
diff --git a/Engine/Plugins/Runtime/nDisplay/Source/DisplayCluster/Public/Render/Upscaler/DisplayClusterUpscalerSettings.h b/Engine/Plugins/Runtime/nDisplay/Source/DisplayCluster/Public/Render/Upscaler/DisplayClusterUpscalerSettings.h
new file mode 100644
--- /dev/null
+++ b/Engine/Plugins/Runtime/nDisplay/Source/DisplayCluster/Public/Render/Upscaler/DisplayClusterUpscalerSettings.h
@@ -0,0 +1,49 @@
+// Copyright Epic Games, Inc. All Rights Reserved.
+
+#pragma once
+
+#include "CoreMinimal.h"
+#include "Render/Viewport/Containers/DisplayClusterViewport_Enums.h"
+#include "StructUtils/PropertyBag.h"
+
+/** Determines AntiAliasing method. */
+enum class EDisplayClusterUpscalerAntiAliasingMethod : uint8
+{
+	// Use default AA from the project settings
+	Default = 0,
+
+	// Disable AA
+	None,
+
+	// Fast Approximate Anti-Aliasing (FXAA)
+	FXAA,
+
+	// Multisample Anti-Aliasing (MSAA)
+	// Only supported with forward shading.  MSAA sample count is controlled by r.MSAACount.
+	MSAA,
+
+	// Temporal Anti-Aliasing (TAA)
+	TAA,
+
+	// Temporal Super-Resolution (TSR)
+	TSR,
+};
+
+/**
+ * nDisplay Upscaler settings
+ */
+class FDisplayClusterUpscalerSettings
+{
+public:
+	FDisplayClusterUpscalerSettings() = default;
+
+public:
+	// Determines AntiAliasing method.
+	EDisplayClusterUpscalerAntiAliasingMethod AntiAliasingMethod = EDisplayClusterUpscalerAntiAliasingMethod::Default;
+
+	// Custom custom upscaler name
+	FName CustomUpscalerName;
+
+	// Custom upscaler settings
+	FInstancedPropertyBag CustomUpscalerSettings;
+};
diff --git a/Engine/Plugins/Runtime/nDisplay/Source/DisplayCluster/Public/Render/Viewport/Containers/DisplayClusterViewport_RenderSettings.h b/Engine/Plugins/Runtime/nDisplay/Source/DisplayCluster/Public/Render/Viewport/Containers/DisplayClusterViewport_RenderSettings.h
--- a/Engine/Plugins/Runtime/nDisplay/Source/DisplayCluster/Public/Render/Viewport/Containers/DisplayClusterViewport_RenderSettings.h
+++ b/Engine/Plugins/Runtime/nDisplay/Source/DisplayCluster/Public/Render/Viewport/Containers/DisplayClusterViewport_RenderSettings.h
@@ -9,6 +9,7 @@
 #include "DisplayClusterViewport_TileSettings.h"
 
 #include "Render/Viewport/Containers/DisplayClusterViewport_Enums.h"
+#include "Render/Upscaler/DisplayClusterUpscalerSettings.h"
 
 /**
  * nDisplay viewport render settings.
@@ -82,6 +83,9 @@ public:
 	// Performance: Render target adaptive resolution multiplier
 	float RenderTargetAdaptRatio = 1;
 
+	// Upscaler settings
+	FDisplayClusterUpscalerSettings UpscalerSettings;
+
 	// Viewport can overlap each other on backbuffer. This value uses to sorting order
 	int32 OverlapOrder = 0;
 
diff --git a/Engine/Plugins/Runtime/nDisplay/Source/DisplayCluster/Public/Render/Viewport/RenderFrame/DisplayClusterRenderFrame.h b/Engine/Plugins/Runtime/nDisplay/Source/DisplayCluster/Public/Render/Viewport/RenderFrame/DisplayClusterRenderFrame.h
--- a/Engine/Plugins/Runtime/nDisplay/Source/DisplayCluster/Public/Render/Viewport/RenderFrame/DisplayClusterRenderFrame.h
+++ b/Engine/Plugins/Runtime/nDisplay/Source/DisplayCluster/Public/Render/Viewport/RenderFrame/DisplayClusterRenderFrame.h
@@ -46,6 +46,9 @@ struct FDisplayClusterRenderFrameTargetViewFamily
 	// Customize ScreenPercentage feature for viewfamily
 	float CustomBufferRatio = 1;
 
+	// ViewFamily Upscaler settings
+	FDisplayClusterUpscalerSettings UpscalerSettings;
+
 	// Extensions that can modify view parameters
 	TArray<FSceneViewExtensionRef> ViewExtensions;
 
@@ -73,6 +76,9 @@ struct FDisplayClusterRenderFrameTarget
 	// Render target resource ref
 	TSharedPtr<FDisplayClusterViewportResource, ESPMode::ThreadSafe> RenderTargetResource;
 
+	// UE Viewport used to render
+	FViewport* Viewport = nullptr;
+
 	// Families, rendered on this target
 	TArray<FDisplayClusterRenderFrameTargetViewFamily> ViewFamilies;
 };
diff --git a/Engine/Plugins/Runtime/nDisplay/Source/DisplayClusterColorGrading/Private/DataModelGenerators/DisplayClusterColorGradingGenerator_RootActor.cpp b/Engine/Plugins/Runtime/nDisplay/Source/DisplayClusterColorGrading/Private/DataModelGenerators/DisplayClusterColorGradingGenerator_RootActor.cpp
--- a/Engine/Plugins/Runtime/nDisplay/Source/DisplayClusterColorGrading/Private/DataModelGenerators/DisplayClusterColorGradingGenerator_RootActor.cpp
+++ b/Engine/Plugins/Runtime/nDisplay/Source/DisplayClusterColorGrading/Private/DataModelGenerators/DisplayClusterColorGradingGenerator_RootActor.cpp
@@ -945,6 +945,7 @@ private:
 		};
 
 		IDetailCategoryBuilder& ICVFXCategoryBuilder = DetailBuilder.EditCategory(TEXT("CustomICVFXCategory"), LOCTEXT("CustomICVFXCategoryLabel", "In-Camera VFX"));
+		AddProperty(ICVFXCategoryBuilder, TEXT("UpscalerSettingsRef"));
 		AddProperty(ICVFXCategoryBuilder, TEXT("BufferRatioRef"));
 		AddProperty(ICVFXCategoryBuilder, TEXT("ExternalCameraActorRef"));
 		AddProperty(ICVFXCategoryBuilder, TEXT("HiddenICVFXViewportsRef"));
diff --git a/Engine/Plugins/Runtime/nDisplay/Source/DisplayClusterConfiguration/Private/DisplayClusterConfigurationTypes_ICVFX.cpp b/Engine/Plugins/Runtime/nDisplay/Source/DisplayClusterConfiguration/Private/DisplayClusterConfigurationTypes_ICVFX.cpp
--- a/Engine/Plugins/Runtime/nDisplay/Source/DisplayClusterConfiguration/Private/DisplayClusterConfigurationTypes_ICVFX.cpp
+++ b/Engine/Plugins/Runtime/nDisplay/Source/DisplayClusterConfiguration/Private/DisplayClusterConfigurationTypes_ICVFX.cpp
@@ -261,6 +261,12 @@ float FDisplayClusterConfigurationICVFX_CameraSettings::GetCameraBufferRatio(con
 	return BufferRatio;
 }
 
+void FDisplayClusterConfigurationICVFX_CameraSettings::GetCameraUpscalerSettings(const FDisplayClusterConfigurationICVFX_StageSettings& InStageSettings, FDisplayClusterUpscalerSettings& OutUpscalerSettings) const
+{
+	// The InStageSettings may also have default settings.
+	UpscalerSettings.GetUpscalerSettings(OutUpscalerSettings);
+}
+
 void FDisplayClusterConfigurationICVFX_CameraSettings::SetupViewInfo(const FDisplayClusterConfigurationICVFX_StageSettings& InStageSettings, FMinimalViewInfo& InOutViewInfo)
 {
 	RenderSettings.SetupViewInfo(InStageSettings, InOutViewInfo);
diff --git a/Engine/Plugins/Runtime/nDisplay/Source/DisplayClusterConfiguration/Private/DisplayClusterConfigurationTypes_Upscaler.cpp b/Engine/Plugins/Runtime/nDisplay/Source/DisplayClusterConfiguration/Private/DisplayClusterConfigurationTypes_Upscaler.cpp
new file mode 100644
--- /dev/null
+++ b/Engine/Plugins/Runtime/nDisplay/Source/DisplayClusterConfiguration/Private/DisplayClusterConfigurationTypes_Upscaler.cpp
@@ -0,0 +1,77 @@
+// Copyright Epic Games, Inc. All Rights Reserved.
+
+#include "DisplayClusterConfigurationTypes_Upscaler.h"
+#include "DisplayClusterConfigurationTypes.h"
+#include "DisplayClusterConfigurationTypes_Enums.h"
+
+#include "UObject/EnumProperty.h"
+#include "UObject/UnrealType.h"
+
+namespace UE::DisplayClusterConfiguration::UpscalerSettings
+{
+	/** Get configuration data for the current upscaling method. */
+	static inline void GetConfigurationData(const FDisplayClusterConfigurationUpscalerSettings& InSettings, FInstancedPropertyBag& OutPropertyBag)
+	{
+		OutPropertyBag = InSettings.EditingData;
+	}
+
+	/** Convert type EDisplayClusterConfigurationUpscalingMethod->EDisplayClusterUpscalerAntiAliasingMethod */
+	static inline EDisplayClusterUpscalerAntiAliasingMethod ToAntiAliasingMethod(const EDisplayClusterConfigurationUpscalingMethod InUpscalingMethod)
+	{
+		// Get AA method
+		switch (InUpscalingMethod)
+		{
+		case EDisplayClusterConfigurationUpscalingMethod::Disable: return EDisplayClusterUpscalerAntiAliasingMethod::None;
+
+		case EDisplayClusterConfigurationUpscalingMethod::FXAA: return EDisplayClusterUpscalerAntiAliasingMethod::FXAA;
+		case EDisplayClusterConfigurationUpscalingMethod::MSAA: return EDisplayClusterUpscalerAntiAliasingMethod::MSAA;
+		case EDisplayClusterConfigurationUpscalingMethod::TAA:  return EDisplayClusterUpscalerAntiAliasingMethod::TAA;
+		case EDisplayClusterConfigurationUpscalingMethod::TSR:  return EDisplayClusterUpscalerAntiAliasingMethod::TSR;
+
+		default:
+			break;
+		}
+
+		return EDisplayClusterUpscalerAntiAliasingMethod::Default;
+	};
+
+};
+
+bool FDisplayClusterConfigurationUpscalerSettings::GetUpscalerSettings(FDisplayClusterUpscalerSettings& OutUpscalerSettings) const
+{
+	using namespace UE::DisplayClusterConfiguration::UpscalerSettings;
+
+	// Find method name in the EDisplayClusterConfigurationUpscalingMethod
+	if (const UEnum* EnumClass = StaticEnum<EDisplayClusterConfigurationUpscalingMethod>())
+	{
+		for (int32 EnumElementIndex = 0; EnumElementIndex < EnumClass->NumEnums(); ++EnumElementIndex)
+		{
+			if (EnumClass->GetNameStringByIndex(EnumElementIndex) == MethodName
+#if WITH_METADATA
+				&& !EnumClass->HasMetaData(TEXT("Hidden"), EnumElementIndex)
+#endif /** WITH_METADATA */
+				)
+			{
+				OutUpscalerSettings.AntiAliasingMethod = ToAntiAliasingMethod(static_cast<EDisplayClusterConfigurationUpscalingMethod>(EnumElementIndex));
+
+				// Custom upscaler not used.
+				OutUpscalerSettings.CustomUpscalerName = NAME_None;
+				GetConfigurationData(*this, OutUpscalerSettings.CustomUpscalerSettings);
+
+				return true;
+			}
+		}
+	}
+
+	// Otherwise use custom upscaler
+	{
+		// Custom upscaler uses its own antialiasing rules.
+		OutUpscalerSettings.AntiAliasingMethod = EDisplayClusterUpscalerAntiAliasingMethod::Default;
+
+		// Set custom upscaler name and PropertyBag with settings.
+		OutUpscalerSettings.CustomUpscalerName = MethodName;
+		GetConfigurationData(*this, OutUpscalerSettings.CustomUpscalerSettings);
+	}
+
+	return MethodName != NAME_None;
+}
\ No newline at end of file
diff --git a/Engine/Plugins/Runtime/nDisplay/Source/DisplayClusterConfiguration/Public/DisplayClusterConfigurationTypes_Enums.h b/Engine/Plugins/Runtime/nDisplay/Source/DisplayClusterConfiguration/Public/DisplayClusterConfigurationTypes_Enums.h
--- a/Engine/Plugins/Runtime/nDisplay/Source/DisplayClusterConfiguration/Public/DisplayClusterConfigurationTypes_Enums.h
+++ b/Engine/Plugins/Runtime/nDisplay/Source/DisplayClusterConfiguration/Public/DisplayClusterConfigurationTypes_Enums.h
@@ -251,3 +251,29 @@ enum class EDisplayClusterConfigurationRootActorPreviewSettingsSource : uint8
 	// This option is used when the preview settings need to be customized in a different way, but not to change the RootActor properties.
 	Configuration
 };
+
+/** Determines Upscaling method.
+* Not used as an enum in the UI.
+* Used in the FDisplayClusterConfigurationUpscalerSettingsDetailCustomization.
+*/
+UENUM()
+enum class EDisplayClusterConfigurationUpscalingMethod : uint8
+{
+	/** Use configuration project settings. */
+	Default = 0 UMETA(DisplayName = "Default"),
+
+	/** No Antialiasing. No upscaling. */
+	Disable UMETA(DisplayName = "Disable"),
+
+	/** Fast Approximate Anti-Aliasing. No upscaling. */
+	FXAA UMETA(DisplayName = "FXAA"),
+
+	/** Multisample Anti-Aliasing. No upscaling. (Only supported with forward shading.) */
+	MSAA UMETA(DisplayName = "MSAA"),
+
+	/** Temporal Super-Resolution. Use TSR upscaler with built-in AntiAliasing. */
+	TSR UMETA(DisplayName = "TSR"),
+
+	/** Temporal Anti-Aliasing. No upscaling. */
+	TAA UMETA(DisplayName = "TAA"),
+};
diff --git a/Engine/Plugins/Runtime/nDisplay/Source/DisplayClusterConfiguration/Public/DisplayClusterConfigurationTypes_ICVFX.h b/Engine/Plugins/Runtime/nDisplay/Source/DisplayClusterConfiguration/Public/DisplayClusterConfigurationTypes_ICVFX.h
--- a/Engine/Plugins/Runtime/nDisplay/Source/DisplayClusterConfiguration/Public/DisplayClusterConfigurationTypes_ICVFX.h
+++ b/Engine/Plugins/Runtime/nDisplay/Source/DisplayClusterConfiguration/Public/DisplayClusterConfigurationTypes_ICVFX.h
@@ -12,6 +12,7 @@
 
 #include "CineCameraActor.h"
 
+#include "DisplayClusterConfigurationTypes_Upscaler.h"
 #include "DisplayClusterConfigurationTypes_Media.h"
 #include "DisplayClusterConfigurationTypes_PostRender.h"
 #include "DisplayClusterConfigurationTypes_Postprocess.h"
@@ -806,6 +807,9 @@ public:
 	/** Get camera buffer ratio. */
 	float GetCameraBufferRatio(const FDisplayClusterConfigurationICVFX_StageSettings& InStageSettings) const;
 
+	/** Get camera upscale settings. */
+	void GetCameraUpscalerSettings(const FDisplayClusterConfigurationICVFX_StageSettings& InStageSettings, FDisplayClusterUpscalerSettings& OutUpscalerSettings) const;
+
 	/** Get camera frame resolution. */
 	UE_DEPRECATED(5.5, "Use the UDisplayClusterICVFXCameraComponent::GetICVFXCameraFrameSize() instead")
 	FIntPoint GetCameraFrameSize(const FDisplayClusterConfigurationICVFX_StageSettings& InStageSettings, const UCineCameraComponent& InCineCameraComponent) const { return FIntPoint(0, 0); }
@@ -834,6 +838,10 @@ public:
 	UPROPERTY(EditAnywhere, BlueprintReadWrite, Category = "In Camera VFX", meta = (DisplayName = "Inner Frustum Screen Percentage", ClampMin = "0.05", UIMin = "0.05", ClampMax = "10.0", UIMax = "1.0"))
 	float BufferRatio = 1;
 
+	/** Upscaler settings for the inner frustum. */
+	UPROPERTY(EditAnywhere, BlueprintReadWrite, Category = "In Camera VFX", meta = (DisplayName = "Upscaler Settings"))
+	FDisplayClusterConfigurationUpscalerSettings UpscalerSettings;
+
 	/** Render a larger or smaller inner frame */
 	UPROPERTY(EditAnywhere, BlueprintReadWrite, Category = "In Camera VFX", meta = (DisplayName = "Inner Frustum Overscan"))
 	FDisplayClusterConfigurationICVFX_CameraCustomFrustum CustomFrustum;
diff --git a/Engine/Plugins/Runtime/nDisplay/Source/DisplayClusterConfiguration/Public/DisplayClusterConfigurationTypes_Upscaler.h b/Engine/Plugins/Runtime/nDisplay/Source/DisplayClusterConfiguration/Public/DisplayClusterConfigurationTypes_Upscaler.h
new file mode 100644
--- /dev/null
+++ b/Engine/Plugins/Runtime/nDisplay/Source/DisplayClusterConfiguration/Public/DisplayClusterConfigurationTypes_Upscaler.h
@@ -0,0 +1,45 @@
+// Copyright Epic Games, Inc. All Rights Reserved.
+
+#pragma once
+
+#include "CoreMinimal.h"
+
+#include "DisplayClusterConfigurationTypes_Enums.h"
+#include "Render/Upscaler/DisplayClusterUpscalerSettings.h"
+
+#include "StructUtils/PropertyBag.h"
+
+#include "DisplayClusterConfigurationTypes_Upscaler.generated.h"
+
+/**
+* Upscaler settings.
+* 
+* MethodName  - refers to one of the available methods
+* EditingData - Editable data for selected Method.
+*               FInstancedPropertyBag is an abstract container and is not handled on the nDisplay side.
+* 
+* MethodName is customized as a drop-down list where items are taken from:
+* - EDisplayClusterConfigurationUpscalingMethod enum elements
+* - Upscalers exposed in the IUpscalerModularFeature interface.
+* 
+* Data logic are implemented in the FDisplayClusterConfigurationUpscalerSettingsDetailCustomization.
+* The initial structure FInstancedPropertyBag is derived by the IUpscalerModularFeature::GetSettings().
+*/
+USTRUCT(BlueprintType)
+struct DISPLAYCLUSTERCONFIGURATION_API FDisplayClusterConfigurationUpscalerSettings
+{
+	GENERATED_BODY()
+
+public:
+	/** Get upscaler settings. */
+	bool GetUpscalerSettings(FDisplayClusterUpscalerSettings& OutUpscalerSettings) const;
+
+public:
+	/** Override the Upscaling method. */
+	UPROPERTY(EditAnywhere, BlueprintReadWrite, Category = "Upscaler", meta = (DisplayName = "Upscaling Method"))
+	FName MethodName;
+
+	/** Upscaling method settings. */
+	UPROPERTY(EditAnywhere, Category = "Upscaler", meta = (DisplayName = "Upscaling Settings"))
+	FInstancedPropertyBag EditingData;
+};
diff --git a/Engine/Plugins/Runtime/nDisplay/Source/DisplayClusterConfiguration/Public/DisplayClusterConfigurationTypes_Viewport.h b/Engine/Plugins/Runtime/nDisplay/Source/DisplayClusterConfiguration/Public/DisplayClusterConfigurationTypes_Viewport.h
--- a/Engine/Plugins/Runtime/nDisplay/Source/DisplayClusterConfiguration/Public/DisplayClusterConfigurationTypes_Viewport.h
+++ b/Engine/Plugins/Runtime/nDisplay/Source/DisplayClusterConfiguration/Public/DisplayClusterConfigurationTypes_Viewport.h
@@ -10,6 +10,7 @@
 #include "DisplayClusterConfigurationStrings.h"
 #include "DisplayClusterConfigurationTypes_Base.h"
 #include "DisplayClusterConfigurationTypes_ICVFX.h"
+#include "DisplayClusterConfigurationTypes_Upscaler.h"
 #include "DisplayClusterConfigurationTypes_Media.h"
 #include "DisplayClusterConfigurationTypes_PostRender.h"
 #include "DisplayClusterConfigurationTypes_Postprocess.h"
@@ -110,6 +111,10 @@ public:
 	UPROPERTY(EditAnywhere, BlueprintReadWrite, Category = "Stereo")
 	EDisplayClusterConfigurationViewport_StereoMode StereoMode = EDisplayClusterConfigurationViewport_StereoMode::Default;
 
+	/** Upscaler settings for the viewport. */
+	UPROPERTY(EditAnywhere, BlueprintReadWrite, Category = "Rendering", meta = (DisplayName = "Upscaler Settings"))
+	FDisplayClusterConfigurationUpscalerSettings UpscalerSettings;
+
 	/** Adjust resolution scaling for an individual viewport.  Viewport Screen Percentage Multiplier is applied to this value. */
 	UPROPERTY(EditAnywhere, BlueprintReadWrite, Category = "Rendering", meta = (DisplayName = "Screen Percentage", ClampMin = "0.05", UIMin = "0.05", ClampMax = "10.0", UIMax = "1.0"))
 	float BufferRatio = 1;
diff --git a/Engine/Plugins/Runtime/nDisplay/Source/DisplayClusterConfigurator/DisplayClusterConfigurator.Build.cs b/Engine/Plugins/Runtime/nDisplay/Source/DisplayClusterConfigurator/DisplayClusterConfigurator.Build.cs
--- a/Engine/Plugins/Runtime/nDisplay/Source/DisplayClusterConfigurator/DisplayClusterConfigurator.Build.cs
+++ b/Engine/Plugins/Runtime/nDisplay/Source/DisplayClusterConfigurator/DisplayClusterConfigurator.Build.cs
@@ -49,6 +49,7 @@ public class DisplayClusterConfigurator : ModuleRules
 				"Settings",
 				"Slate",
 				"SlateCore",
+				"StructUtilsEditor",
 				"ToolMenus",
 				"UnrealEd",
 				"SubobjectEditor",
diff --git a/Engine/Plugins/Runtime/nDisplay/Source/DisplayClusterConfigurator/Private/DisplayClusterConfiguratorModule.cpp b/Engine/Plugins/Runtime/nDisplay/Source/DisplayClusterConfigurator/Private/DisplayClusterConfiguratorModule.cpp
--- a/Engine/Plugins/Runtime/nDisplay/Source/DisplayClusterConfigurator/Private/DisplayClusterConfiguratorModule.cpp
+++ b/Engine/Plugins/Runtime/nDisplay/Source/DisplayClusterConfigurator/Private/DisplayClusterConfiguratorModule.cpp
@@ -33,6 +33,7 @@
 #include "Views/Details/Media/DisplayClusterConfiguratorMediaTileCustomization.h"
 #include "Views/Details/Components/DisplayClusterCameraComponentDetailsCustomization.h"
 #include "Views/Details/Policies/DisplayClusterConfiguratorPolicyDetailCustomization.h"
+#include "Views/Details/Upscaler/DisplayClusterConfigurationUpscalerSettingsDetailCustomization.h"
 
 #include "Blueprints/DisplayClusterBlueprint.h"
 #include "Components/DisplayClusterScreenComponent.h"
@@ -47,6 +48,7 @@
 #include "ISettingsModule.h"
 #include "AssetRegistry/AssetRegistryModule.h"
 #include "AssetTypeCategories.h"
+
 #include "DisplayClusterConfiguratorBlueprintEditor.h"
 #include "DisplayClusterConfiguratorToolbarExtensions.h"
 #include "HAL/IConsoleManager.h"
@@ -251,6 +253,7 @@ void FDisplayClusterConfiguratorModule::RegisterCustomLayouts()
 	REGISTER_PROPERTY_LAYOUT(FDisplayClusterConfigurationMediaOutputGroup, FDisplayClusterConfiguratorMediaFullFrameOutputCustomization);
 	REGISTER_PROPERTY_LAYOUT(FDisplayClusterConfigurationViewport_ColorGradingSettings, FDCConfiguratorColorGradingSettingsCustomization);
 	REGISTER_PROPERTY_LAYOUT(FDisplayClusterConfigurationViewport_ColorGradingWhiteBalanceSettings, FDCConfiguratorWhiteBalanceCustomization);
+	REGISTER_PROPERTY_LAYOUT(FDisplayClusterConfigurationUpscalerSettings, FDisplayClusterConfigurationUpscalerSettingsDetailCustomization);
 }
 
 void FDisplayClusterConfiguratorModule::UnregisterCustomLayouts()
diff --git a/Engine/Plugins/Runtime/nDisplay/Source/DisplayClusterConfigurator/Private/Views/Details/Upscaler/DisplayClusterConfigurationUpscalerSettingsDetailCustomization.cpp b/Engine/Plugins/Runtime/nDisplay/Source/DisplayClusterConfigurator/Private/Views/Details/Upscaler/DisplayClusterConfigurationUpscalerSettingsDetailCustomization.cpp
new file mode 100644
--- /dev/null
+++ b/Engine/Plugins/Runtime/nDisplay/Source/DisplayClusterConfigurator/Private/Views/Details/Upscaler/DisplayClusterConfigurationUpscalerSettingsDetailCustomization.cpp
@@ -0,0 +1,304 @@
+// Copyright Epic Games, Inc. All Rights Reserved.
+
+#include "Views/Details/Upscaler/DisplayClusterConfigurationUpscalerSettingsDetailCustomization.h"
+
+#include "DisplayClusterConfigurationTypes.h"
+#include "DisplayClusterConfigurationTypes_Enums.h"
+
+#include "IUpscalerModularFeature.h"
+
+#include "UObject/EnumProperty.h"
+#include "UObject/UnrealType.h"
+#include "UObject/Object.h"
+
+#include "Widgets/Input/SComboBox.h"
+
+#include "DetailWidgetRow.h"
+
+#include "Features/IModularFeatures.h"
+
+#include "PropertyEditorModule.h"
+#include "PropertyEditorDelegates.h"
+#include "PropertyBagDetails.h"
+#include "IDetailChildrenBuilder.h"
+
+#include "IPropertyUtilities.h"
+#include "PropertyHandle.h"
+
+using namespace UE::VirtualProduction;
+
+namespace UE::DisplayClusterConfigurator::UpscalerSettingsDetailCustomization
+{
+	/** Returns true if only the enabled upscaling methods can be seen by the user. */
+	static bool ShouldShowOnlyEnabledUpscalingMethods()
+	{
+		static IConsoleVariable* const ICVar = IConsoleManager::Get().FindConsoleVariable(TEXT("nDisplay.render.upscaling.HideDisabledMethods"));
+		if (ICVar && ICVar->GetInt() == 0)
+		{
+			return false;
+		}
+
+		return true;
+	}
+
+	/** Returns true if the method is enabled for the current project settings/hardware. */
+	static bool IsUpscalingMethodEnabled(const EDisplayClusterConfigurationUpscalingMethod InMethod)
+	{
+		return true;
+	}
+
+	/** Iterate over all default upscaler methods. */
+	static void ForeachDefaultUpscalerMethod(TFunction<void(const UEnum&, const int32, const EDisplayClusterConfigurationUpscalingMethod)>&& IteratorFunc)
+	{
+		// Collect embedded values from the EDisplayClusterConfigurationUpscalingMethod
+		if (const UEnum* EnumSource = StaticEnum<EDisplayClusterConfigurationUpscalingMethod>())
+		{
+			for (int32 EnumIndex = 0; EnumIndex < EnumSource->NumEnums() - 1; ++EnumIndex)
+			{
+				const bool bIsHidden = EnumSource->HasMetaData(TEXT("Hidden"), EnumIndex);
+				if (!bIsHidden)
+				{
+					const EDisplayClusterConfigurationUpscalingMethod UpscalingMethod = static_cast<EDisplayClusterConfigurationUpscalingMethod>(EnumIndex);
+					if (!ShouldShowOnlyEnabledUpscalingMethods() || IsUpscalingMethodEnabled(UpscalingMethod))
+					{
+						IteratorFunc(*EnumSource, EnumIndex, UpscalingMethod);
+					}
+				}
+			}
+		}
+	}
+
+	/** Iterate over all upscaler interfaces. */
+	static void ForeachUpscaler(TFunction<void(IUpscalerModularFeature&)>&& IteratorFunc)
+	{
+		static IModularFeatures& ModularFeatures = IModularFeatures::Get();
+
+		ModularFeatures.LockModularFeatureList();
+		{
+			const TArray<IUpscalerModularFeature*> Extenders = ModularFeatures.GetModularFeatureImplementations<IUpscalerModularFeature>(IUpscalerModularFeature::ModularFeatureName);
+			for (IUpscalerModularFeature* Extender : Extenders)
+			{
+				// Iterate over active upscalers
+				if (Extender
+					&& (!ShouldShowOnlyEnabledUpscalingMethods() || Extender->IsFeatureEnabled()))
+				{
+					IteratorFunc(*Extender);
+				}
+			}
+		}
+		ModularFeatures.UnlockModularFeatureList();
+	}
+
+	/** Get upscaling settings for method name. */
+	static void GetMethodSettings(const FName& InMethodName, FInstancedPropertyBag& InOutEditableData)
+	{
+		// Get settings from the upscaler
+		ForeachUpscaler([&](IUpscalerModularFeature& Upscaler)
+			{
+				if (Upscaler.GetName() == InMethodName)
+				{
+					Upscaler.GetSettings(InOutEditableData);
+				}
+			});
+	}
+};
+
+FDisplayClusterConfigurationUpscalerSettingsDetailCustomization::FDisplayClusterConfigurationUpscalerSettingsDetailCustomization()
+{ }
+
+void FDisplayClusterConfigurationUpscalerSettingsDetailCustomization::CustomizeHeader(
+	TSharedRef<IPropertyHandle> InPropertyHandle, FDetailWidgetRow& InHeaderRow, IPropertyTypeCustomizationUtils& InCustomizationUtils)
+{ }
+
+void FDisplayClusterConfigurationUpscalerSettingsDetailCustomization::CustomizeChildren(
+	TSharedRef<IPropertyHandle> InPropertyHandle,
+	IDetailChildrenBuilder& InChildBuilder,
+	IPropertyTypeCustomizationUtils& InCustomizationUtils)
+{
+	using namespace UE::DisplayClusterConfigurator::UpscalerSettingsDetailCustomization;
+
+	// CUSTOMIZE PROPERTY: UpscalingSettings (with UpscalingSettingsStorage)
+	MethodNameHandle = InPropertyHandle->GetChildHandle(GET_MEMBER_NAME_CHECKED(FDisplayClusterConfigurationUpscalerSettings, MethodName));
+	EditingDataHandle = InPropertyHandle->GetChildHandle(GET_MEMBER_NAME_CHECKED(FDisplayClusterConfigurationUpscalerSettings, EditingData));
+
+	if (!MethodNameHandle.IsValid() || !EditingDataHandle.IsValid())
+	{
+		MethodNameHandle.Reset();
+		EditingDataHandle.Reset();
+
+		return;
+	}
+
+	/** Find upscaling method from string. */
+	auto FindUpscalingMethod = [&](const FName& InMethodName)
+		{
+			for (const TSharedPtr<FUpscalerMethodEntry>& UpscalerMethod : UpscalerMethods)
+			{
+				if (UpscalerMethod.IsValid() && UpscalerMethod->Name == InMethodName)
+				{
+					return UpscalerMethod;
+				}
+			}
+
+			// Method is not available.
+			// When the DCRA configuration is from another project, but the current project/hardware does not have this method.
+			
+			// return first possible
+			return !UpscalerMethods.IsEmpty() ? UpscalerMethods[0] : nullptr;
+		};
+
+	auto RebuildUpscalerMethods = [&]()
+		{
+			// Collect embedded values from the EDisplayClusterConfigurationUpscalingMethod
+			ForeachDefaultUpscalerMethod([&](
+				const UEnum& StaticEnum,
+				const int32 ElementIndex,
+				const EDisplayClusterConfigurationUpscalingMethod InUpscalingMethod)
+				{
+					TSharedPtr<FUpscalerMethodEntry> NewElement = MakeShared<FUpscalerMethodEntry>(
+						FName(StaticEnum.GetNameStringByIndex(ElementIndex)),
+						StaticEnum.GetDisplayNameTextByIndex(ElementIndex),
+						StaticEnum.GetToolTipTextByIndex(ElementIndex));
+
+					UpscalerMethods.Add(NewElement);
+				});
+
+			// Collect custom  upscalers
+			ForeachUpscaler([&](IUpscalerModularFeature& Upscaler)
+				{
+					TSharedPtr<FUpscalerMethodEntry> NewElement = MakeShared<FUpscalerMethodEntry>(
+						Upscaler.GetName(),
+						Upscaler.GetDisplayName(),
+						Upscaler.GetTooltipText());
+
+					UpscalerMethods.Add(NewElement);
+				});
+		};
+
+	FName CurrentMethodName;
+
+	// CUSTOMIZE PROPERTY: UpscalingMethod
+	{
+		RebuildUpscalerMethods();
+
+		// Update current value
+		MethodNameHandle->GetValue(CurrentMethodName);
+		CurrentUpscalerMethod = FindUpscalingMethod(CurrentMethodName);
+
+		// Build dropdown list
+		InChildBuilder.AddCustomRow(FText::FromString("Selected Option"))
+			.NameContent()
+			[
+				MethodNameHandle->CreatePropertyNameWidget()
+			]
+			.ValueContent()
+			[
+				SNew(SComboBox<TSharedPtr<FUpscalerMethodEntry>>)
+					.OptionsSource(&UpscalerMethods)
+					.OnGenerateWidget_Lambda([](TSharedPtr<FUpscalerMethodEntry> InMethod)
+						{							
+							if (InMethod.IsValid())
+							{
+								return SNew(STextBlock)
+									.Text(InMethod->DisplayName.IsEmpty()
+										? FText::FromName(InMethod->Name)
+										: InMethod->DisplayName)
+									.ToolTipText(InMethod->Tooltip);
+							}
+
+							//@todo
+							return SNew(STextBlock).Text(FText::GetEmpty());
+						})
+					.OnSelectionChanged_Lambda([&](TSharedPtr<FUpscalerMethodEntry> NewUpscalerMethod, ESelectInfo::Type)
+						{
+							CurrentUpscalerMethod = NewUpscalerMethod;
+							FName MethodName = CurrentUpscalerMethod.IsValid() ? CurrentUpscalerMethod->Name : NAME_None;
+
+							// When MethodName changes, copy from PerMethodData[MethodName] to EditingData
+							const bool bResetEditingData = true;
+							UpdateEditingData(bResetEditingData);
+
+							MethodNameHandle->SetValue(MethodName);
+						})
+					[
+						SNew(STextBlock)
+							.Text_Lambda([&]() -> FText
+								{
+									return CurrentUpscalerMethod.IsValid() ? CurrentUpscalerMethod->DisplayName : FText::FromString(TEXT("Default"));
+								})
+							.ToolTipText_Lambda([&]() -> FText
+								{
+									return CurrentUpscalerMethod.IsValid() ? CurrentUpscalerMethod->Tooltip : FText::GetEmpty();
+								})
+					]
+			];
+	}
+
+	const bool bResetEditingData = !CurrentUpscalerMethod.IsValid() || CurrentMethodName != CurrentUpscalerMethod->Name;
+
+	// Initialize EditingData
+	UpdateEditingData(bResetEditingData);
+
+	// Update method name
+	if (CurrentUpscalerMethod.IsValid() && CurrentMethodName != CurrentUpscalerMethod->Name)
+	{
+		MethodNameHandle->SetValue(CurrentUpscalerMethod->Name);
+	}
+
+	uint32 NumFields = 0;
+	EditingDataHandle->GetNumChildren(NumFields);
+
+	// Show customization if we have at least 1 property in the bag
+	if (NumFields > 0)
+	{
+		TSharedRef<FPropertyBagInstanceDataDetails> EditingDataDetails =
+			MakeShared<FPropertyBagInstanceDataDetails>(
+				EditingDataHandle
+				, InCustomizationUtils.GetPropertyUtilities()
+				, /** bInFixedLayout */true);
+
+		InChildBuilder
+			.AddCustomBuilder(EditingDataDetails);
+	}
+}
+
+bool FDisplayClusterConfigurationUpscalerSettingsDetailCustomization::UpdateEditingData(bool bResetEditingData)
+{
+	if (!EditingDataHandle.IsValid() || !CurrentUpscalerMethod.IsValid())
+	{
+		return false;
+	}
+
+	/** Update EditingData property. */
+	using namespace UE::DisplayClusterConfigurator::UpscalerSettingsDetailCustomization;
+
+	// NewSettingsBag->EditingData
+	{
+		void* StructPointer = nullptr;
+		if (EditingDataHandle->GetValueData(StructPointer) == FPropertyAccess::Success && StructPointer)
+		{
+			FInstancedPropertyBag* EditingDataBag = static_cast<FInstancedPropertyBag*>(StructPointer);
+			if (bResetEditingData)
+			{
+				EditingDataBag->Reset();
+			}
+
+			if(!EditingDataBag->IsValid() || EditingDataBag->GetNumPropertiesInBag()==0)
+			{
+				// Get new settings
+				FInstancedPropertyBag NewSettingsBag;
+				GetMethodSettings(CurrentUpscalerMethod->Name, NewSettingsBag);
+
+				const bool bTheSame = NewSettingsBag.Identical(EditingDataBag, PPF_None);
+				if (!bTheSame)
+				{
+					EditingDataBag->MigrateToNewBagInstance(NewSettingsBag);
+
+					return true;
+				}
+			}
+		}
+	}
+
+	return false;
+}
\ No newline at end of file
diff --git a/Engine/Plugins/Runtime/nDisplay/Source/DisplayClusterConfigurator/Private/Views/Details/Upscaler/DisplayClusterConfigurationUpscalerSettingsDetailCustomization.h b/Engine/Plugins/Runtime/nDisplay/Source/DisplayClusterConfigurator/Private/Views/Details/Upscaler/DisplayClusterConfigurationUpscalerSettingsDetailCustomization.h
new file mode 100644
--- /dev/null
+++ b/Engine/Plugins/Runtime/nDisplay/Source/DisplayClusterConfigurator/Private/Views/Details/Upscaler/DisplayClusterConfigurationUpscalerSettingsDetailCustomization.h
@@ -0,0 +1,61 @@
+// Copyright Epic Games, Inc. All Rights Reserved.
+
+#pragma once
+
+#include "CoreMinimal.h"
+#include "IPropertyTypeCustomization.h"
+#include "DisplayClusterConfigurationTypes_Upscaler.h"
+
+/**
+* DisplayCluster UI customization for Upscaler Settings.
+*/
+class FDisplayClusterConfigurationUpscalerSettingsDetailCustomization : public IPropertyTypeCustomization
+{
+public:
+	FDisplayClusterConfigurationUpscalerSettingsDetailCustomization();
+
+	static TSharedRef<IPropertyTypeCustomization> MakeInstance()
+	{
+		return MakeShareable(new FDisplayClusterConfigurationUpscalerSettingsDetailCustomization);
+	}
+
+	//~Begin IDetailCustomization
+	virtual void CustomizeHeader(TSharedRef<IPropertyHandle> InPropertyHandle, FDetailWidgetRow& HeaderRow, IPropertyTypeCustomizationUtils& CustomizationUtils) override;
+	virtual void CustomizeChildren(TSharedRef<IPropertyHandle> InPropertyHandle, IDetailChildrenBuilder& StructBuilder, IPropertyTypeCustomizationUtils& CustomizationUtils) override;
+	//~~End IDetailCustomization
+
+protected:
+	// Initialize EditingDataHandle
+	bool UpdateEditingData(bool bResetBag);
+
+protected: // CUSTOMIZE PROPERTY: UpscalingMethod
+
+	/** Droplist element with the tooltip text. */
+	struct FUpscalerMethodEntry
+	{
+		FUpscalerMethodEntry() = default;
+		FUpscalerMethodEntry(
+			const FName& InName,
+			const FText& InDisplayName,
+			const FText& InTooltip)
+			: Name(InName), DisplayName(InDisplayName), Tooltip(InTooltip)
+		{ }
+
+	public:
+		// Key
+		FName Name;
+
+		// (opt) The display name shown in the UI. Use 'Name' is this value is empty.
+		FText DisplayName;
+
+		// (opt) Tooltip text
+		FText Tooltip;
+	};
+
+	TArray<TSharedPtr<FUpscalerMethodEntry>> UpscalerMethods;
+	TSharedPtr<FUpscalerMethodEntry> CurrentUpscalerMethod;
+
+	// property handles
+	TSharedPtr<IPropertyHandle> MethodNameHandle;
+	TSharedPtr<IPropertyHandle> EditingDataHandle;
+};
diff --git a/Engine/Plugins/Runtime/nDisplay/Source/DisplayClusterDetails/Private/DataModelGenerators/DisplayClusterDetailsGenerator_RootActor.cpp b/Engine/Plugins/Runtime/nDisplay/Source/DisplayClusterDetails/Private/DataModelGenerators/DisplayClusterDetailsGenerator_RootActor.cpp
--- a/Engine/Plugins/Runtime/nDisplay/Source/DisplayClusterDetails/Private/DataModelGenerators/DisplayClusterDetailsGenerator_RootActor.cpp
+++ b/Engine/Plugins/Runtime/nDisplay/Source/DisplayClusterDetails/Private/DataModelGenerators/DisplayClusterDetailsGenerator_RootActor.cpp
@@ -261,7 +261,10 @@ public:
 		};
 
 		IDetailCategoryBuilder& ICVFXCategoryBuilder = DetailBuilder.EditCategory(TEXT("CustomICVFXCategory"), LOCTEXT("CustomICVFXCategoryLabel", "In-Camera VFX"));
+
+		AddProperty(ICVFXCategoryBuilder, TEXT("UpscalerSettingsRef"));
 		AddProperty(ICVFXCategoryBuilder, TEXT("BufferRatioRef"));
+
 		AddProperty(ICVFXCategoryBuilder, TEXT("ExternalCameraActorRef"));
 		AddProperty(ICVFXCategoryBuilder, TEXT("HiddenICVFXViewportsRef"));
 
diff --git a/Engine/Source/Runtime/VirtualProduction/VirtualProduction/Private/VirtualProductionModule.cpp b/Engine/Source/Runtime/VirtualProduction/VirtualProduction/Private/VirtualProductionModule.cpp
new file mode 100644
--- /dev/null
+++ b/Engine/Source/Runtime/VirtualProduction/VirtualProduction/Private/VirtualProductionModule.cpp
@@ -0,0 +1,5 @@
+// Copyright Epic Games, Inc. All Rights Reserved.
+
+#include "Modules/ModuleManager.h"
+
+IMPLEMENT_MODULE(FDefaultModuleImpl, VirtualProduction);
diff --git a/Engine/Source/Runtime/VirtualProduction/VirtualProduction/Public/IUpscalerModularFeature.h b/Engine/Source/Runtime/VirtualProduction/VirtualProduction/Public/IUpscalerModularFeature.h
new file mode 100644
--- /dev/null
+++ b/Engine/Source/Runtime/VirtualProduction/VirtualProduction/Public/IUpscalerModularFeature.h
@@ -0,0 +1,141 @@
+// Copyright Epic Games, Inc. All Rights Reserved.
+
+#pragma once
+
+#include "CoreMinimal.h"
+#include "Features/IModularFeature.h"
+
+struct FInstancedPropertyBag;
+
+class FSceneViewFamilyContext;
+class FSceneView;
+
+namespace UE::VirtualProduction
+{
+	/**
+	* The upscaler modular feature parameters.
+	*/
+	struct FUpscalerModularFeatureParameters
+	{
+		// Upscaler screen percentage vaue.
+		float UpscalerScreenPercentage = 1.0f;
+
+		// Secondary screen percentage value.
+		float SecondaryScreenPercentage = 1.0f;
+	};
+
+	/**
+	 * Interface for a modular feature of a Upscaler.
+	 *
+	 * This interface provides a way to interact with the upscaler without requiring the presence
+	 * of their specific plugins. It allows modular features to be integrated and accessed within the
+	 * rendering pipeline in a flexible and decoupled manner.
+	 *
+	 * Classes implementing this interface can define custom behavior and settings for Scene View Extensions
+	 * while maintaining compatibility with other parts of the Unreal Engine system.
+	 */
+	class IUpscalerModularFeature
+		: public IModularFeature
+	{
+	public:
+		virtual ~IUpscalerModularFeature() = default;
+
+		/** The unique modular feature name. */
+		static constexpr const TCHAR* ModularFeatureName = TEXT("UpscalerModularFeature");
+
+	public:
+		/**
+		 * Returns the unique identifier name for this feature.
+		 */
+		virtual const FName& GetName() const = 0;
+
+		/**
+		 * Gets the display name shown in the UI.
+		 */
+		virtual const FText& GetDisplayName() const = 0;
+
+		/**
+		 * Returns a hint text for the feature.
+		 */
+		virtual const FText& GetTooltipText() const = 0;
+
+	public:
+		/**
+		 * Determines whether the feature is currently enabled and available for use.
+		 * This can depend on config, platform, or runtime logic.
+		 */
+		virtual bool IsFeatureEnabled() const = 0;
+
+		/**
+		* Setup all View Extensions of the upscaler to be enabled or disabled.
+		* Must be called before GatherActiveExtensions() during the rendering pass setup.
+		*
+		* @param bEnable                 - (in) Enable or disable all ViewExtensions of this upscaler.
+		* @param bEnableInEditorViewport - (in) Makes this upscaler's ViewExtensions work in the editor's viewports
+		*/
+		virtual void SetupUpscalerViewExtensions(const bool bEnable, const bool bEnableInEditorViewport) = 0;
+
+		/**
+		* Enable only necessary view extensions used by upscaler for these input settings.
+		* Must be called before GatherActiveExtensions() during the rendering pass setup.
+		*
+		* @param InUpscalerSettings      - (in) The upscaler settings which we want to use. Enables only necessary ViewExtensions.
+		* @param bEnableInEditorViewport - (in) Makes this upscaler's ViewExtensions work in the editor's viewports
+		*
+		* @return true if the settings can be used and upscaler is enabled
+		*/
+		virtual bool SetupUpscalerViewExtensionsForSettings(const FInstancedPropertyBag& InUpscalerSettings, const bool bEnableInEditorViewport) = 0;
+
+		/**
+		* Get default settings for this upscaler
+		*
+		* @param OuUpscalerSettings - (out) default settings
+		*
+		* @param return true if successed.
+		*/
+		virtual bool GetSettings(FInstancedPropertyBag& OuUpscalerSettings) const
+		{
+			return false;
+		}
+
+		/**
+		* Reset all stored per-view settings to default values
+		* The settings that was stored by the function PostConfigureViewFamily()
+		*
+		* @param return true if successed.
+		*/
+		virtual void ResetSettings()
+		{
+		}
+
+		/**
+		* Setup SceneView for this upscaler.
+		* Upscalers should configure the view properties (e.g. AntiAliasingMethod) for their rendering pipeline (temporal/spatial/etc.).
+		*
+		* @param InUpscalerSettings - (in) upscaler settings which we want to use. This settings should be associated with the InOutViews.
+		* @param InOutView          - (in) view to configure
+		*/
+		virtual void SetupSceneView(
+			const FInstancedPropertyBag& InUpscalerSettings,
+			FSceneView& InOutView) = 0;
+
+		/**
+		* Configure ViewFamily and it Views for specific settings and use these settings for the views.
+		* Note: Expected that these parameters should be configured:
+		*   InOutViewFamily.SetScreenPercentageInterface();
+		*   InOutViewFamily.SecondaryViewFraction;
+		* 241300
+		* Some features from the settings may require additional customization in the view family or view structures.
+		*
+		* @param InUpscalerSettings - (in) upscaler settings which we want to use. This settings should be associated with the InOutViews.
+		* @param InUpscalerParam    - (in) upscaler parameters that is used to configure view family.
+		* @param InOutViewFamily    - (in, out) The view family that will be configured.
+		* @param InOutViews         - (in, out) The views in the family that will be configured.
+		*/
+		virtual bool PostConfigureViewFamily(
+			const FInstancedPropertyBag& InUpscalerSettings,
+			const FUpscalerModularFeatureParameters& InUpscalerParam,
+			FSceneViewFamilyContext& InOutViewFamily,
+			const TArray<FSceneView*>& InOutViews) = 0;
+	};
+};
diff --git a/Engine/Source/Runtime/VirtualProduction/VirtualProduction/VirtualProduction.build.cs b/Engine/Source/Runtime/VirtualProduction/VirtualProduction/VirtualProduction.build.cs
new file mode 100644
--- /dev/null
+++ b/Engine/Source/Runtime/VirtualProduction/VirtualProduction/VirtualProduction.build.cs
@@ -0,0 +1,16 @@
+// Copyright Epic Games, Inc. All Rights Reserved.
+
+using UnrealBuildTool;
+
+public class VirtualProduction  : ModuleRules
+{
+	public VirtualProduction(ReadOnlyTargetRules ROTargetRules) : base(ROTargetRules)
+	{
+		PublicDependencyModuleNames.AddRange(
+			new string[] {
+				"Core",
+				"CoreUObject",
+				"Engine"
+			});
+	}
+}
