[nDisplay] [Virtual Production] Upscaler modular features

Added new argument to IUpscalerModularFeature::SetupSceneView() function.
Added logic to test the upscaler method on current hardware and settings.
Added new CVars:
- ""nDisplay.render.upscaling.Enable" - "nDisplay.render.upscaling.Enable
- ""nDisplay.render.upscaling.EnableCustomUpscalers" - "Allows the use of custom upscalers (default = 1)

diff --git a/Engine/Plugins/Runtime/nDisplay/Source/DisplayCluster/Private/Render/Upscaler/DisplayClusterUpscaler.cpp b/Engine/Plugins/Runtime/nDisplay/Source/DisplayCluster/Private/Render/Upscaler/DisplayClusterUpscaler.cpp
--- a/Engine/Plugins/Runtime/nDisplay/Source/DisplayCluster/Private/Render/Upscaler/DisplayClusterUpscaler.cpp
+++ b/Engine/Plugins/Runtime/nDisplay/Source/DisplayCluster/Private/Render/Upscaler/DisplayClusterUpscaler.cpp
@@ -22,9 +22,28 @@ namespace UE::DisplayCluster::Upscaler
 		ECVF_Default
 	);
 
-	static inline bool ShouldShowOnlyEnabledUpscalingMethods()
+	int32 GDisplayClusterRenderUpscaling_Enable = 1;
+	static FAutoConsoleVariableRef CVarDisplayClusterRenderUpscaling_Enable(
+		TEXT("nDisplay.render.upscaling.Enable"),
+		GDisplayClusterRenderUpscaling_Enable,
+		TEXT("Allows to use the upscaling settings (default = 1).\n")
+		TEXT("0 - Ignore upscaling settings.\n"),
+		ECVF_Default
+	);
+
+	int32 GDisplayClusterRenderUpscaling_EnableCustomUpscalers = 1;
+	static FAutoConsoleVariableRef CVarDisplayClusterRenderUpscaling_EnableCustomUpscalers(
+		TEXT("nDisplay.render.upscaling.EnableCustomUpscalers"),
+		GDisplayClusterRenderUpscaling_EnableCustomUpscalers,
+		TEXT("Allows the use of custom upscalers (default = 1).\n")
+		TEXT("0 - Don't use custom upscalers and render by default.\n"),
+		ECVF_Default
+	);
+
+	/** A custom upscaler may be used. */
+	static inline bool UseCustomUpscalers()
 	{
-		return GDisplayClusterRenderUpscaling_HideDisabledMethods!=0;
+		return GDisplayClusterRenderUpscaling_EnableCustomUpscalers && GDisplayClusterRenderUpscaling_Enable;
 	}
 
 	/** Iterate over all upscaler modular features. */
@@ -37,8 +56,8 @@ namespace UE::DisplayCluster::Upscaler
 			const TArray<IUpscalerModularFeature*> Extenders = ModularFeatures.GetModularFeatureImplementations<IUpscalerModularFeature>(IUpscalerModularFeature::ModularFeatureName);
 			for (IUpscalerModularFeature* Extender : Extenders)
 			{
-				// Iterate over active upscalers
-				if (Extender && (!ShouldShowOnlyEnabledUpscalingMethods() || Extender->IsFeatureEnabled()))
+				// Iterate over enabled upscalers
+				if (Extender && Extender->IsFeatureEnabled())
 				{
 					IteratorFunc(*Extender);
 				}
@@ -63,13 +82,66 @@ namespace UE::DisplayCluster::Upscaler
 		return 1.0f;
 	}
 
-	/** Check temporal upscaler CVar*/
-	static bool IsTemporalUpscalerActive()
+	/** Get secondary screen percentage fraction. */
+	static float GetSecondaryViewFraction(const float InDPIScale)
+	{
+		static IConsoleVariable* CVarCustomSecondaryScreenPercentage = IConsoleManager::Get().FindConsoleVariable(TEXT("r.SecondaryScreenPercentage.GameViewport"), false);
+		if (CVarCustomSecondaryScreenPercentage)
+		{
+			const float CustomSecondaryScreenPercentage = CVarCustomSecondaryScreenPercentage->GetFloat();
+		if (CustomSecondaryScreenPercentage > 0.0)
+		{
+			// Override secondary resolution fraction with CVar.
+			return FMath::Min(CustomSecondaryScreenPercentage / 100.0f, 1.0f);
+		}
+		}
+
+#if WITH_EDITOR
+			// Automatically compute secondary resolution fraction from DPI.
+			// When in high res screenshot do not modify screen percentage based on dpi scale
+			if (GIsEditor && !GIsHighResScreenshot)
+			{
+				static IConsoleVariable* CVarEditorViewportHighDPIPtr = IConsoleManager::Get().FindConsoleVariable(TEXT("r.Editor.Viewport.HighDPI"));
+				if (CVarEditorViewportHighDPIPtr && CVarEditorViewportHighDPIPtr->GetInt() == 0)
+				{
+
+					const float DPIDerivedResolutionFraction = FMath::Min(1.0f / InDPIScale, 1.0f);
+
+					return DPIDerivedResolutionFraction;
+				}
+			}
+#endif
+
+		return 1.0f;
+	}
+
+
+	/** Check if temporal upscaler can be used. */
+	static bool CanUseTemporalUpscaler()
 	{
 		static const IConsoleVariable* CVarTemporalAAUpscaler = IConsoleManager::Get().FindConsoleVariable(TEXT("r.TemporalAA.Upscaler"));
 		return CVarTemporalAAUpscaler && (CVarTemporalAAUpscaler->GetInt() != 0);
 	}
 
+	/** Check if MSAA can be used. */
+	static bool CanUseMSAA(const FStaticFeatureLevel InFeatureLevel)
+	{
+		static auto* MSAACountCVar = IConsoleManager::Get().FindTConsoleVariableDataInt(TEXT("r.MSAACount"));
+		if (MSAACountCVar->GetValueOnAnyThread() <= 0)
+		{
+			return false;
+		}
+
+		EShaderPlatform ShaderPlatform = GetFeatureLevelShaderPlatform(InFeatureLevel);
+		if ((InFeatureLevel >= ERHIFeatureLevel::SM5 && !IsForwardShadingEnabled(ShaderPlatform)))
+		{
+			// MSAA requires deferred renderer
+			return false;
+		}
+	
+		return true;
+	}
+
 	/** Temporal upscaler interface requires special AA method. */
 	static bool SceneViewSupportsTemporalUpscaler(const FSceneView& InSceneView)
 	{
@@ -87,67 +159,153 @@ namespace UE::DisplayCluster::Upscaler
 		return false;
 	}
 
-	/** Get secondary screen percentage fraction. */
-	static float GetSecondaryViewFraction(const float InDPIScale)
+	/** Return default AntiAliasing method. */
+	static EDisplayClusterUpscalerAntiAliasingMethod GetDefaultAntiAliasingMethod(const FStaticFeatureLevel InFeatureLevel)
 	{
-		static IConsoleVariable* CVarCustomSecondaryScreenPercentage = IConsoleManager::Get().FindConsoleVariable(TEXT("r.SecondaryScreenPercentage.GameViewport"), false);
-		if (CVarCustomSecondaryScreenPercentage)
+		static auto* DefaultAntiAliasingCvar = IConsoleManager::Get().FindTConsoleVariableDataInt(TEXT("r.AntiAliasingMethod"));
+		static EAntiAliasingMethod DefaultAntiAliasingMethod = EAntiAliasingMethod(FMath::Clamp<int32>(DefaultAntiAliasingCvar->GetValueOnAnyThread(), 0, AAM_MAX));
+
+		switch (DefaultAntiAliasingMethod)
 		{
-			const float CustomSecondaryScreenPercentage = CVarCustomSecondaryScreenPercentage->GetFloat();
-		if (CustomSecondaryScreenPercentage > 0.0)
+		case EAntiAliasingMethod::AAM_FXAA:
+			return EDisplayClusterUpscalerAntiAliasingMethod::FXAA;
+
+		case EAntiAliasingMethod::AAM_MSAA:
+			if (CanUseMSAA(InFeatureLevel))
+			{
+				return EDisplayClusterUpscalerAntiAliasingMethod::MSAA;
+			}
+			break;
+
+		case EAntiAliasingMethod::AAM_TemporalAA:
+			return EDisplayClusterUpscalerAntiAliasingMethod::TAA;
+
+		case EAntiAliasingMethod::AAM_TSR:
+			{
+				EShaderPlatform ShaderPlatform = GetFeatureLevelShaderPlatform(InFeatureLevel);
+				if (!SupportsTSR(ShaderPlatform))
+				{
+					// Fallback to UE4's TAA if TSR isn't supported on that platform
+					return EDisplayClusterUpscalerAntiAliasingMethod::TAA;
+				}
+			}
+			return EDisplayClusterUpscalerAntiAliasingMethod::TSR;
+
+		default:
+			break;
+		}
+
+		return EDisplayClusterUpscalerAntiAliasingMethod::None;
+	}
+
+	/** 
+	* Get the antialiasing method that can be used for the ViewFamily.
+	* Note: the logic of this function is based on the FSceneView::SetupAntiAliasingMethod() function.
+	* 
+	* @param InUpscalerSettings - (in) the upscaler settings that we want to use.
+	* @param InViewFamily       - (in) ViewFamily thyat own SceneView.
+	* @param InSceneView        - (in) SceneView to which the upscaler settings should be applied.
+	* 
+	* @return AntiAliasing method that can be used.
+	*/
+	static EDisplayClusterUpscalerAntiAliasingMethod GetSuitableAntiAliasingMethod(
+		const FDisplayClusterUpscalerSettings& InUpscalerSettings,
+		const FSceneViewFamily& InViewFamily,
+		const FSceneView& InSceneView)
+	{
+		const bool bWillApplyTemporalAA = InViewFamily.EngineShowFlags.PostProcessing || InSceneView.bIsPlanarReflection;
+
+		if (!bWillApplyTemporalAA || !InViewFamily.EngineShowFlags.AntiAliasing)
 		{
-			// Override secondary resolution fraction with CVar.
-			return FMath::Min(CustomSecondaryScreenPercentage / 100.0f, 1.0f);
+			return EDisplayClusterUpscalerAntiAliasingMethod::None;
 		}
+
+		// Get AntiAliasing method from the upscaler settings
+		EDisplayClusterUpscalerAntiAliasingMethod AntiAliasingMethod = InUpscalerSettings.AntiAliasingMethod;
+		
+		// Use the default method if the upscaler settings is disabled
+		if (!GDisplayClusterRenderUpscaling_Enable)
+		{
+			AntiAliasingMethod = EDisplayClusterUpscalerAntiAliasingMethod::Default;
 		}
 
-#if WITH_EDITOR
-			// Automatically compute secondary resolution fraction from DPI.
-			// When in high res screenshot do not modify screen percentage based on dpi scale
-			if (GIsEditor && !GIsHighResScreenshot)
+		// Back to default method if TemporalUpscaling not supported
+		if (AntiAliasingMethod == EDisplayClusterUpscalerAntiAliasingMethod::TAA
+			|| AntiAliasingMethod == EDisplayClusterUpscalerAntiAliasingMethod::TSR)
+		{
+			if (!CanUseTemporalUpscaler())
 			{
-				static IConsoleVariable* CVarEditorViewportHighDPIPtr = IConsoleManager::Get().FindConsoleVariable(TEXT("r.Editor.Viewport.HighDPI"));
-				if (CVarEditorViewportHighDPIPtr && CVarEditorViewportHighDPIPtr->GetInt() == 0)
-				{
+				AntiAliasingMethod = EDisplayClusterUpscalerAntiAliasingMethod::Default;
+			}
+		}
 
-					const float DPIDerivedResolutionFraction = FMath::Min(1.0f / InDPIScale, 1.0f);
+		// Back to default method is MSAA not supported
+		if (AntiAliasingMethod == EDisplayClusterUpscalerAntiAliasingMethod::MSAA && !CanUseMSAA(InSceneView.GetFeatureLevel()))
+		{
+			AntiAliasingMethod = EDisplayClusterUpscalerAntiAliasingMethod::Default;
+		}
 
-					return DPIDerivedResolutionFraction;
-				}
+		// Get default method from the project settings
+		if (AntiAliasingMethod == EDisplayClusterUpscalerAntiAliasingMethod::Default)
+		{
+			AntiAliasingMethod = GetDefaultAntiAliasingMethod(InSceneView.GetFeatureLevel());
+		}
+
+		if (AntiAliasingMethod == EDisplayClusterUpscalerAntiAliasingMethod::TAA)
+		{
+			// Downgrade TAA->FXAA
+			if (!InViewFamily.EngineShowFlags.TemporalAA || !InViewFamily.bRealtimeUpdate || !SupportsGen4TAA(InSceneView.GetShaderPlatform()))
+			{
+				return EDisplayClusterUpscalerAntiAliasingMethod::FXAA;
 			}
-#endif
+		}
+		else if (AntiAliasingMethod == EDisplayClusterUpscalerAntiAliasingMethod::TSR)
+		{
+			// Downgrade TSR->FXAA
+			if (!InViewFamily.EngineShowFlags.TemporalAA || !InViewFamily.bRealtimeUpdate || !SupportsTSR(InSceneView.GetShaderPlatform()))
+			{
+				return EDisplayClusterUpscalerAntiAliasingMethod::FXAA;
+			}
+		}
 
-		return 1.0f;
+		return AntiAliasingMethod;
 	}
 };
 
 void FDisplayClusterUpscaler::SetupSceneView(
 	const FDisplayClusterUpscalerSettings& InUpscalerSettings,
+	const FSceneViewFamily& InOutViewFamily,
 	FSceneView& InOutView)
 {
 	using namespace UE::DisplayCluster::Upscaler;
 
-	// Use custom AA methods
-	switch (InUpscalerSettings.AntiAliasingMethod)
+	// Get AntiAliasingMethod that can be used.
+	const EDisplayClusterUpscalerAntiAliasingMethod SuitableAntiAliasingMethod = GetSuitableAntiAliasingMethod(InUpscalerSettings, InOutViewFamily, InOutView);
+	switch (SuitableAntiAliasingMethod)
 	{
+	case EDisplayClusterUpscalerAntiAliasingMethod::None:
+		InOutView.AntiAliasingMethod = EAntiAliasingMethod::AAM_None;
+		InOutView.PrimaryScreenPercentageMethod = EPrimaryScreenPercentageMethod::SpatialUpscale;
+		break;
+
 	case EDisplayClusterUpscalerAntiAliasingMethod::FXAA:
 		InOutView.AntiAliasingMethod = EAntiAliasingMethod::AAM_FXAA;
+		InOutView.PrimaryScreenPercentageMethod = EPrimaryScreenPercentageMethod::SpatialUpscale;
 		break;
 
 	case EDisplayClusterUpscalerAntiAliasingMethod::MSAA:
 		InOutView.AntiAliasingMethod = EAntiAliasingMethod::AAM_MSAA;
-		break;
-
-	case EDisplayClusterUpscalerAntiAliasingMethod::None:
-		InOutView.AntiAliasingMethod = EAntiAliasingMethod::AAM_None;
+		InOutView.PrimaryScreenPercentageMethod = EPrimaryScreenPercentageMethod::SpatialUpscale;
 		break;
 
 	case EDisplayClusterUpscalerAntiAliasingMethod::TAA:
 		InOutView.AntiAliasingMethod = EAntiAliasingMethod::AAM_TemporalAA;
+		InOutView.PrimaryScreenPercentageMethod = EPrimaryScreenPercentageMethod::TemporalUpscale;
 		break;
 
 	case EDisplayClusterUpscalerAntiAliasingMethod::TSR:
 		InOutView.AntiAliasingMethod = EAntiAliasingMethod::AAM_TSR;
+		InOutView.PrimaryScreenPercentageMethod = EPrimaryScreenPercentageMethod::TemporalUpscale;
 		break;
 
 	default:
@@ -155,14 +313,17 @@ void FDisplayClusterUpscaler::SetupSceneView(
 	}
 
 	// Custom upscaler should configure view for the settings
-	ForeachUpscaler(
-		[&](IUpscalerModularFeature& Upscaler)
-		{
-			if (Upscaler.GetName() == InUpscalerSettings.CustomUpscalerName)
+	if (!InUpscalerSettings.CustomUpscalerName.IsNone() && UseCustomUpscalers())
+	{
+		ForeachUpscaler(
+			[&](IUpscalerModularFeature& Upscaler)
 			{
-				Upscaler.SetupSceneView(InUpscalerSettings.CustomUpscalerSettings, InOutView);
-			}
-		});
+				if (Upscaler.GetName() == InUpscalerSettings.CustomUpscalerName)
+				{
+					Upscaler.SetupSceneView(InUpscalerSettings.CustomUpscalerSettings, InOutViewFamily, InOutView);
+				}
+			});
+	}
 }
 
 FName FDisplayClusterUpscaler::PostConfigureViewFamily(
@@ -204,72 +365,71 @@ FName FDisplayClusterUpscaler::PostConfigureViewFamily(
 	}
 
 	FUpscalerModularFeatureParameters UpscalerParam;
-	UpscalerParam.UpscalerScreenPercentage  = InScreenPercentage;
+	UpscalerParam.UpscalerScreenPercentage = InScreenPercentage;
 	UpscalerParam.SecondaryScreenPercentage = GetSecondaryViewFraction(InDPIScale);
 
 	check(UpscalerParam.UpscalerScreenPercentage > 0.0f);
 	check(UpscalerParam.SecondaryScreenPercentage > 0.0f);
 
-	// Find  upscaler modular feature
-	FName UpscalerName = NAME_None;
-	ForeachUpscaler([&](IUpscalerModularFeature& Upscaler)
-		{
-			if (Upscaler.GetName() == InUpscalerSettings.CustomUpscalerName)
+	// Use custom upscaler
+	if (!InUpscalerSettings.CustomUpscalerName.IsNone() && UseCustomUpscalers())
+	{
+		// Find  upscaler modular feature
+		FName UpscalerName = NAME_None;
+		ForeachUpscaler([&](IUpscalerModularFeature& Upscaler)
 			{
-				if (Upscaler.PostConfigureViewFamily(
-					InUpscalerSettings.CustomUpscalerSettings, UpscalerParam,
-					InOutViewFamily, InOutViews))
+				if (Upscaler.GetName() == InUpscalerSettings.CustomUpscalerName)
 				{
-					// Return upscaler name on success
-					UpscalerName = Upscaler.GetName();
+					if (Upscaler.PostConfigureViewFamily(
+						InUpscalerSettings.CustomUpscalerSettings, UpscalerParam,
+						InOutViewFamily, InOutViews))
+					{
+						// Return upscaler name on success
+						UpscalerName = Upscaler.GetName();
+					}
 				}
-			}
-		});
+			});
 
-	if (UpscalerName != NAME_None)
-	{
-		return UpscalerName;
+		if (UpscalerName != NAME_None)
+		{
+			// Upscaler modular feature expected to set screen percentage interface
+			check(InOutViewFamily.GetScreenPercentageInterface())
+
+			return UpscalerName;
+		}
 	}
 
-	// Can't use Custom upscalers, back to the AntiAliasing method.
+	// Get the upscaler settings from the first SceneView.
+	const EPrimaryScreenPercentageMethod PrimaryScreenPercentageMethod = InOutViews[0]->PrimaryScreenPercentageMethod;
+	const EAntiAliasingMethod AntiAliasingMethod = InOutViews[0]->AntiAliasingMethod;
 
-	//@todo: logic below is in process/testing << need to discuss
-	switch (InUpscalerSettings.AntiAliasingMethod)
+	// Setup TSR & TAA
+	// We need to split the screen percentage if below 0.5 because TAA upscaling only works well up to 2x.
+	if(PrimaryScreenPercentageMethod == EPrimaryScreenPercentageMethod::TemporalUpscale
+	&&( AntiAliasingMethod == EAntiAliasingMethod::AAM_TemporalAA
+	 || AntiAliasingMethod == EAntiAliasingMethod::AAM_TSR))
 	{
-	case EDisplayClusterUpscalerAntiAliasingMethod::Default:
-	case EDisplayClusterUpscalerAntiAliasingMethod::TSR:
-		// We need to split the screen percentage if below 0.5 because TAA upscaling only works well up to 2x.
-		{
-			// Get global view fraction set by r.ScreenPercentage.
-			const float GlobalScreenPercentage = GetGlobalScreenPercentage();
+		// Get global view fraction set by r.ScreenPercentage.
+		const float GlobalScreenPercentage = GetGlobalScreenPercentage();
 
-			// Adjust for global screen% in math bellow
-			const float FinalResolutionFraction = UpscalerParam.UpscalerScreenPercentage * GlobalScreenPercentage;
-			if (FinalResolutionFraction < 0.5f)
-			{
-				UpscalerParam.SecondaryScreenPercentage *= 2.0f * FinalResolutionFraction;
+		// Adjust for global screen% in math bellow
+		const float FinalResolutionFraction = UpscalerParam.UpscalerScreenPercentage * GlobalScreenPercentage;
+		if (FinalResolutionFraction < 0.5f)
+		{
+			// Move the remaining ScreenPercentage to the spatial upscaler
+			const float RemainingScreenPercentage = 2.0f * UpscalerParam.UpscalerScreenPercentage;
+			UpscalerParam.SecondaryScreenPercentage *= RemainingScreenPercentage;
 
-				// Adjust value for the global screen%. Expectations that the actual TSR_Fraction=0.5 
-				UpscalerParam.UpscalerScreenPercentage = 0.5f / GlobalScreenPercentage;
-			}
+			// Adjust min value (0.5f) for the global screen%. Expectations that the actual TSR_Fraction=0.5 
+			UpscalerParam.UpscalerScreenPercentage = 0.5f / GlobalScreenPercentage;
 		}
+	}
 
-		InOutViewFamily.SetScreenPercentageInterface(new FLegacyScreenPercentageDriver(
-			InOutViewFamily, UpscalerParam.UpscalerScreenPercentage));
-
-		InOutViewFamily.SecondaryViewFraction = UpscalerParam.SecondaryScreenPercentage;
-		break;
-
-	// Don't use upscaling
-	default:
-		InOutViewFamily.SetScreenPercentageInterface(new FLegacyScreenPercentageDriver(
-			InOutViewFamily, 1.0f));
+	InOutViewFamily.SetScreenPercentageInterface(new FLegacyScreenPercentageDriver(
+		InOutViewFamily, UpscalerParam.UpscalerScreenPercentage));
 
-		InOutViewFamily.SecondaryViewFraction = UpscalerParam.SecondaryScreenPercentage;
-		break;
-	}
+	InOutViewFamily.SecondaryViewFraction = UpscalerParam.SecondaryScreenPercentage;
 
-	// In any other cases use default upscaler
 	return NAME_None;
 }
 
@@ -291,6 +451,14 @@ void FDisplayClusterUpscaler::SetupUpscalerViewExtensionsForSettings(const FDisp
 {
 	using namespace UE::DisplayCluster::Upscaler;
 
+	if (!UseCustomUpscalers())
+	{
+		// Disable all custom upscalers
+		SetupUpscalerViewExtensions(false);
+
+		return;
+	}
+
 	ForeachUpscaler(
 		[&](IUpscalerModularFeature& Upscaler)
 		{
diff --git a/Engine/Plugins/Runtime/nDisplay/Source/DisplayCluster/Private/Render/Upscaler/DisplayClusterUpscaler.h b/Engine/Plugins/Runtime/nDisplay/Source/DisplayCluster/Private/Render/Upscaler/DisplayClusterUpscaler.h
--- a/Engine/Plugins/Runtime/nDisplay/Source/DisplayCluster/Private/Render/Upscaler/DisplayClusterUpscaler.h
+++ b/Engine/Plugins/Runtime/nDisplay/Source/DisplayCluster/Private/Render/Upscaler/DisplayClusterUpscaler.h
@@ -7,6 +7,7 @@
 #include "Render/Upscaler/DisplayClusterUpscalerSettings.h"
 
 class FSceneViewFamilyContext;
+class FSceneViewFamily;
 class FSceneView;
 
 /**
@@ -49,9 +50,11 @@ public:
 	/**
 	* Configure scene view for upscalers.
 	* @param InUpscalerSettings - (in) the settings to apply for view
+	* @param InOutViewFamily    - (in) ViewFamily that own this view.
 	* @param InOutViews         - (in, out) The views that will be configured. 
 	*/
 	static void SetupSceneView(
 		const FDisplayClusterUpscalerSettings& InUpscalerSettings,
+		const FSceneViewFamily& InOutViewFamily,
 		FSceneView& InOutView);
 };
diff --git a/Engine/Plugins/Runtime/nDisplay/Source/DisplayCluster/Private/Render/Viewport/DisplayClusterViewport.cpp b/Engine/Plugins/Runtime/nDisplay/Source/DisplayCluster/Private/Render/Viewport/DisplayClusterViewport.cpp
--- a/Engine/Plugins/Runtime/nDisplay/Source/DisplayCluster/Private/Render/Viewport/DisplayClusterViewport.cpp
+++ b/Engine/Plugins/Runtime/nDisplay/Source/DisplayCluster/Private/Render/Viewport/DisplayClusterViewport.cpp
@@ -402,7 +402,7 @@ void FDisplayClusterViewport::SetupSceneView(uint32 ContextNum, class UWorld* Wo
 	}
 
 	// Handle upscalers
-	FDisplayClusterUpscaler::SetupSceneView(RenderSettings.UpscalerSettings, InOutView);
+	FDisplayClusterUpscaler::SetupSceneView(RenderSettings.UpscalerSettings, InOutViewFamily, InOutView);
 
 
 	if (InOutView.State && World && World->Scene)
diff --git a/Engine/Plugins/Runtime/nDisplay/Source/DisplayClusterConfiguration/Public/DisplayClusterConfigurationTypes_Enums.h b/Engine/Plugins/Runtime/nDisplay/Source/DisplayClusterConfiguration/Public/DisplayClusterConfigurationTypes_Enums.h
--- a/Engine/Plugins/Runtime/nDisplay/Source/DisplayClusterConfiguration/Public/DisplayClusterConfigurationTypes_Enums.h
+++ b/Engine/Plugins/Runtime/nDisplay/Source/DisplayClusterConfiguration/Public/DisplayClusterConfigurationTypes_Enums.h
@@ -262,18 +262,18 @@ enum class EDisplayClusterConfigurationUpscalingMethod : uint8
 	/** Use configuration project settings. */
 	Default = 0 UMETA(DisplayName = "Default"),
 
-	/** No Antialiasing. No upscaling. */
+	/** No Antialiasing. */
 	Disable UMETA(DisplayName = "Disable"),
 
-	/** Fast Approximate Anti-Aliasing. No upscaling. */
+	/** Fast Approximate Anti-Aliasing. */
 	FXAA UMETA(DisplayName = "FXAA"),
 
-	/** Multisample Anti-Aliasing. No upscaling. (Only supported with forward shading.) */
+	/** Multisample Anti-Aliasing. (Only supported with forward shading.) */
 	MSAA UMETA(DisplayName = "MSAA"),
 
+	/** Temporal Anti-Aliasing. */
+	TAA UMETA(DisplayName = "TAA"),
+
 	/** Temporal Super-Resolution. Use TSR upscaler with built-in AntiAliasing. */
 	TSR UMETA(DisplayName = "TSR"),
-
-	/** Temporal Anti-Aliasing. No upscaling. */
-	TAA UMETA(DisplayName = "TAA"),
 };
diff --git a/Engine/Plugins/Runtime/nDisplay/Source/DisplayClusterConfigurator/Private/Views/Details/Upscaler/DisplayClusterConfigurationUpscalerSettingsDetailCustomization.cpp b/Engine/Plugins/Runtime/nDisplay/Source/DisplayClusterConfigurator/Private/Views/Details/Upscaler/DisplayClusterConfigurationUpscalerSettingsDetailCustomization.cpp
--- a/Engine/Plugins/Runtime/nDisplay/Source/DisplayClusterConfigurator/Private/Views/Details/Upscaler/DisplayClusterConfigurationUpscalerSettingsDetailCustomization.cpp
+++ b/Engine/Plugins/Runtime/nDisplay/Source/DisplayClusterConfigurator/Private/Views/Details/Upscaler/DisplayClusterConfigurationUpscalerSettingsDetailCustomization.cpp
@@ -92,7 +92,12 @@ namespace UE::DisplayClusterConfigurator::UpscalerSettingsDetailCustomization
 	/** Get upscaling settings for method name. */
 	static void GetMethodSettings(const FName& InMethodName, FInstancedPropertyBag& InOutEditableData)
 	{
-		// Get settings from the upscaler
+		if (InMethodName.IsNone())
+		{
+			return;
+		}
+
+		// Get settings from the custom upscaler
 		ForeachUpscaler([&](IUpscalerModularFeature& Upscaler)
 			{
 				if (Upscaler.GetName() == InMethodName)
diff --git a/Engine/Source/Runtime/VirtualProduction/VirtualProduction/Public/IUpscalerModularFeature.h b/Engine/Source/Runtime/VirtualProduction/VirtualProduction/Public/IUpscalerModularFeature.h
--- a/Engine/Source/Runtime/VirtualProduction/VirtualProduction/Public/IUpscalerModularFeature.h
+++ b/Engine/Source/Runtime/VirtualProduction/VirtualProduction/Public/IUpscalerModularFeature.h
@@ -8,6 +8,7 @@
 struct FInstancedPropertyBag;
 
 class FSceneViewFamilyContext;
+class FSceneViewFamily;
 class FSceneView;
 
 namespace UE::VirtualProduction
@@ -113,10 +114,12 @@ namespace UE::VirtualProduction
 		* Upscalers should configure the view properties (e.g. AntiAliasingMethod) for their rendering pipeline (temporal/spatial/etc.).
 		*
 		* @param InUpscalerSettings - (in) upscaler settings which we want to use. This settings should be associated with the InOutViews.
-		* @param InOutView          - (in) view to configure
+		* @param InViewFamily       - (in) ViewFamily that own this view.
+		* @param InOutView          - (in) scene view to that will be configured.
 		*/
 		virtual void SetupSceneView(
 			const FInstancedPropertyBag& InUpscalerSettings,
+			const FSceneViewFamily& InViewFamily,
 			FSceneView& InOutView) = 0;
 
 		/**
