[nDisplay] [Virtual Production] Added global upscaler settings for Outer viewports and Inner Frustum.

diff --git a/Engine/Plugins/Runtime/nDisplay/Source/DisplayCluster/Private/Render/Viewport/Configuration/DisplayClusterViewportConfigurationHelpers.cpp b/Engine/Plugins/Runtime/nDisplay/Source/DisplayCluster/Private/Render/Viewport/Configuration/DisplayClusterViewportConfigurationHelpers.cpp
--- a/Engine/Plugins/Runtime/nDisplay/Source/DisplayCluster/Private/Render/Viewport/Configuration/DisplayClusterViewportConfigurationHelpers.cpp
+++ b/Engine/Plugins/Runtime/nDisplay/Source/DisplayCluster/Private/Render/Viewport/Configuration/DisplayClusterViewportConfigurationHelpers.cpp
@@ -122,8 +122,12 @@ void FDisplayClusterViewportConfigurationHelpers::UpdateBaseViewportSetting(FDis
 	{
 		DstViewport.SetViewportBufferRatio(InRenderSettings.BufferRatio);
 
+		// Get default upscaler settings
+		const FDisplayClusterConfigurationICVFX_StageSettings* StageSettings = DstViewport.Configuration->GetStageSettings();
+		const FDisplayClusterConfigurationUpscalerSettings* DefaultUpscalerSettings = StageSettings ? &StageSettings->OuterViewportUpscalerSettings : nullptr;
+
 		// Update viewport Upscaler method
-		InRenderSettings.UpscalerSettings.GetUpscalerSettings(InOutRenderSettings.UpscalerSettings);
+		InRenderSettings.UpscalerSettings.GetUpscalerSettings(DefaultUpscalerSettings, InOutRenderSettings.UpscalerSettings);
 
 		DstViewport.UpdateConfiguration_Overscan(GetViewportOverscanSettings(InRenderSettings.Overscan));
 
diff --git a/Engine/Plugins/Runtime/nDisplay/Source/DisplayCluster/Public/DisplayClusterRootActor.h b/Engine/Plugins/Runtime/nDisplay/Source/DisplayCluster/Public/DisplayClusterRootActor.h
--- a/Engine/Plugins/Runtime/nDisplay/Source/DisplayCluster/Public/DisplayClusterRootActor.h
+++ b/Engine/Plugins/Runtime/nDisplay/Source/DisplayCluster/Public/DisplayClusterRootActor.h
@@ -295,12 +295,15 @@ private:
 //////////////////////////////////////////////////////////////////////////////////////////////
 #if WITH_EDITORONLY_DATA
 private:
-	UPROPERTY(EditAnywhere, Transient, Category = Viewports, meta = (PropertyPath = "CurrentConfigData.RenderFrameSettings.ClusterICVFXOuterViewportBufferRatioMult"))
-	FDisplayClusterEditorPropertyReference ViewportScreenPercentageMultiplierRef;
+	UPROPERTY(EditAnywhere, Transient, Category = Viewports, meta = (PropertyPath = "CurrentConfigData.StageSettings.OuterViewportUpscalerSettings", DisplayName = "Default Upscaler Settings"))
+	FDisplayClusterEditorPropertyReference OuterViewportUpscalerSettingsRef;
 
 	UPROPERTY(EditInstanceOnly, Transient, Category = Viewports, meta = (DisplayName = "Viewport Upscaling Settings", PropertyPath = "CurrentConfigData.Cluster.Nodes.Viewports.RenderSettings.UpscalerSettings", ToolTip = "Adjust upscaler settings for an individual viewport."))
 	FDisplayClusterEditorPropertyReference ViewportUpscalerSettingsRef;
 
+	UPROPERTY(EditAnywhere, Transient, Category = Viewports, meta = (PropertyPath = "CurrentConfigData.RenderFrameSettings.ClusterICVFXOuterViewportBufferRatioMult"))
+	FDisplayClusterEditorPropertyReference ViewportScreenPercentageMultiplierRef;
+
 	UPROPERTY(EditInstanceOnly, Transient, Category = Viewports, meta = (DisplayName = "Viewport Screen Percentage", PropertyPath = "CurrentConfigData.Cluster.Nodes.Viewports.RenderSettings.BufferRatio", ToolTip = "Adjust resolution scaling for an individual viewport.  Viewport Screen Percentage Multiplier is applied to this value."))
 	FDisplayClusterEditorPropertyReference ViewportScreenPercentageRef;
 
@@ -319,6 +322,9 @@ private:
 	UPROPERTY(EditInstanceOnly, Transient, Category = Viewports, meta = (DisplayName = "Content Hidden Per-Viewport", PropertyPath = "CurrentConfigData.Cluster.Nodes.Viewports.RenderSettings.HiddenContent", ToolTip = "Hide actors from a specific viewport."))
 	FDisplayClusterEditorPropertyReference PerOuterHideListRef;
 	
+	UPROPERTY(EditAnywhere, Transient, Category = "In Camera VFX", meta = (PropertyPath = "CurrentConfigData.StageSettings.GlobalInnerFrustumUpscalerSettings", DisplayName = "Default Upscaler Settings"))
+	FDisplayClusterEditorPropertyReference GlobalInnerFrustumUpscalerSettingsRef;
+
 	UPROPERTY(EditAnywhere, Transient, Category = "In Camera VFX", meta = (PropertyPath = "CurrentConfigData.StageSettings.bEnableInnerFrustums"))
 	FDisplayClusterEditorPropertyReference EnableInnerFrustumsRef;
 
diff --git a/Engine/Plugins/Runtime/nDisplay/Source/DisplayClusterColorGrading/Private/DataModelGenerators/DisplayClusterColorGradingGenerator_RootActor.cpp b/Engine/Plugins/Runtime/nDisplay/Source/DisplayClusterColorGrading/Private/DataModelGenerators/DisplayClusterColorGradingGenerator_RootActor.cpp
--- a/Engine/Plugins/Runtime/nDisplay/Source/DisplayClusterColorGrading/Private/DataModelGenerators/DisplayClusterColorGradingGenerator_RootActor.cpp
+++ b/Engine/Plugins/Runtime/nDisplay/Source/DisplayClusterColorGrading/Private/DataModelGenerators/DisplayClusterColorGradingGenerator_RootActor.cpp
@@ -312,13 +312,16 @@ private:
 		};
 
 		IDetailCategoryBuilder& ViewportsCategoryBuilder = DetailBuilder.EditCategory(TEXT("CustomViewportsCategory"), LOCTEXT("CustomViewportsCategoryLabel", "Viewports"));
+		AddProperty(ViewportsCategoryBuilder, TEXT("OuterViewportUpscalerSettingsRef"));
 		AddProperty(ViewportsCategoryBuilder, TEXT("ViewportScreenPercentageMultiplierRef"));
 		AddProperty(ViewportsCategoryBuilder, TEXT("FreezeRenderOuterViewportsRef"));
 
 		IDetailCategoryBuilder& InnerFrustumCategoryBuilder = DetailBuilder.EditCategory(TEXT("CustomICVFXCategory"), LOCTEXT("CustomICVFXCategoryLabel", "In-Camera VFX"));
+		AddProperty(InnerFrustumCategoryBuilder, TEXT("GlobalInnerFrustumUpscalerSettingsRef"));
 		AddProperty(InnerFrustumCategoryBuilder, TEXT("ShowInnerFrustumOverlapsRef"));
 		AddProperty(InnerFrustumCategoryBuilder, GET_MEMBER_NAME_CHECKED(ADisplayClusterRootActor, InnerFrustumPriority));
 
+
 		IDetailCategoryBuilder& ViewportChromakeyCategoryBuilder = DetailBuilder.EditCategory(TEXT("CustomViewportChromakeyCategory"), LOCTEXT("CustomViewportChromakeyCategoryLabel", "Chromakey"));
 		AddProperty(ViewportChromakeyCategoryBuilder, TEXT("GlobalChromakeyColorRef"));
 
diff --git a/Engine/Plugins/Runtime/nDisplay/Source/DisplayClusterConfiguration/Private/DisplayClusterConfigurationTypes_ICVFX.cpp b/Engine/Plugins/Runtime/nDisplay/Source/DisplayClusterConfiguration/Private/DisplayClusterConfigurationTypes_ICVFX.cpp
--- a/Engine/Plugins/Runtime/nDisplay/Source/DisplayClusterConfiguration/Private/DisplayClusterConfigurationTypes_ICVFX.cpp
+++ b/Engine/Plugins/Runtime/nDisplay/Source/DisplayClusterConfiguration/Private/DisplayClusterConfigurationTypes_ICVFX.cpp
@@ -263,8 +263,9 @@ float FDisplayClusterConfigurationICVFX_CameraSettings::GetCameraBufferRatio(con
 
 void FDisplayClusterConfigurationICVFX_CameraSettings::GetCameraUpscalerSettings(const FDisplayClusterConfigurationICVFX_StageSettings& InStageSettings, FDisplayClusterUpscalerSettings& OutUpscalerSettings) const
 {
-	// The InStageSettings may also have default settings.
-	UpscalerSettings.GetUpscalerSettings(OutUpscalerSettings);
+	UpscalerSettings.GetUpscalerSettings(
+		&InStageSettings.GlobalInnerFrustumUpscalerSettings,
+		OutUpscalerSettings);
 }
 
 void FDisplayClusterConfigurationICVFX_CameraSettings::SetupViewInfo(const FDisplayClusterConfigurationICVFX_StageSettings& InStageSettings, FMinimalViewInfo& InOutViewInfo)
diff --git a/Engine/Plugins/Runtime/nDisplay/Source/DisplayClusterConfiguration/Private/DisplayClusterConfigurationTypes_Upscaler.cpp b/Engine/Plugins/Runtime/nDisplay/Source/DisplayClusterConfiguration/Private/DisplayClusterConfigurationTypes_Upscaler.cpp
--- a/Engine/Plugins/Runtime/nDisplay/Source/DisplayClusterConfiguration/Private/DisplayClusterConfigurationTypes_Upscaler.cpp
+++ b/Engine/Plugins/Runtime/nDisplay/Source/DisplayClusterConfiguration/Private/DisplayClusterConfigurationTypes_Upscaler.cpp
@@ -37,7 +37,9 @@ namespace UE::DisplayClusterConfiguration::UpscalerSettings
 
 };
 
-bool FDisplayClusterConfigurationUpscalerSettings::GetUpscalerSettings(FDisplayClusterUpscalerSettings& OutUpscalerSettings) const
+bool FDisplayClusterConfigurationUpscalerSettings::GetUpscalerSettings(
+	const FDisplayClusterConfigurationUpscalerSettings* InDefaultUpscalerSettings,
+	FDisplayClusterUpscalerSettings& OutUpscalerSettings) const
 {
 	using namespace UE::DisplayClusterConfiguration::UpscalerSettings;
 
@@ -54,6 +56,12 @@ bool FDisplayClusterConfigurationUpscalerSettings::GetUpscalerSettings(FDisplayC
 			{
 				OutUpscalerSettings.AntiAliasingMethod = ToAntiAliasingMethod(static_cast<EDisplayClusterConfigurationUpscalingMethod>(EnumElementIndex));
 
+				// Use default settings
+				if (OutUpscalerSettings.AntiAliasingMethod == EDisplayClusterUpscalerAntiAliasingMethod::Default && InDefaultUpscalerSettings)
+				{
+					return InDefaultUpscalerSettings->GetUpscalerSettings(nullptr, OutUpscalerSettings);
+				}
+
 				// Custom upscaler not used.
 				OutUpscalerSettings.CustomUpscalerName = NAME_None;
 				GetConfigurationData(*this, OutUpscalerSettings.CustomUpscalerSettings);
diff --git a/Engine/Plugins/Runtime/nDisplay/Source/DisplayClusterConfiguration/Public/DisplayClusterConfigurationTypes_ICVFX.h b/Engine/Plugins/Runtime/nDisplay/Source/DisplayClusterConfiguration/Public/DisplayClusterConfigurationTypes_ICVFX.h
--- a/Engine/Plugins/Runtime/nDisplay/Source/DisplayClusterConfiguration/Public/DisplayClusterConfigurationTypes_ICVFX.h
+++ b/Engine/Plugins/Runtime/nDisplay/Source/DisplayClusterConfiguration/Public/DisplayClusterConfigurationTypes_ICVFX.h
@@ -929,6 +929,10 @@ public:
 	UPROPERTY(EditAnywhere, BlueprintReadWrite, Category = "In Camera VFX", meta = (DisplayName = "Enable Inner Frustum Chromakey Overlap"))
 	bool bEnableInnerFrustumChromakeyOverlap = false;
 
+	/** Global upscaler settings that all ICVFX camera components can opt to use. */
+	UPROPERTY(EditAnywhere, BlueprintReadWrite, Category = "In Camera VFX", meta = (DisplayName = "Default Upscaler Settings"))
+	FDisplayClusterConfigurationUpscalerSettings GlobalInnerFrustumUpscalerSettings;
+
 	/** Default In-Cameras texture resolution. */
 	UPROPERTY(EditAnywhere, BlueprintReadWrite, Category = Configuration, meta = (DisplayName = "Default Frame Resolution"))
 	FDisplayClusterConfigurationICVFX_Size DefaultFrameSize;
@@ -940,6 +944,11 @@ public:
 	UPROPERTY(EditAnywhere, BlueprintReadWrite, Category = Viewports, meta = (DisplayName = "Freeze Viewports"))
 	bool bFreezeRenderOuterViewports = false;
 
+	/** Global upscaler settings for Outer viewports. */
+	UPROPERTY(EditAnywhere, BlueprintReadWrite, Category = Viewports, meta = (DisplayName = "Default Upscaler Settings"))
+	FDisplayClusterConfigurationUpscalerSettings OuterViewportUpscalerSettings;
+
+
 	// Hide list for all icvfx viewports (outer, inner, cameras, etc)
 	// (This allow to hide all actors from layers for icvfx render logic)
 	UPROPERTY(EditAnywhere, BlueprintReadWrite, Category = Viewports, meta = (DisplayName = "Content Hidden from Entire Cluster", ToolTip = "Content specified here will not appear anywhere in the nDisplay cluster", Substitutions = "LayersTooltip = Layers hidden from the entire nDisplay cluster, ActorsTooltip = Actors hidden from the entire nDisplay cluster"))
diff --git a/Engine/Plugins/Runtime/nDisplay/Source/DisplayClusterConfiguration/Public/DisplayClusterConfigurationTypes_Upscaler.h b/Engine/Plugins/Runtime/nDisplay/Source/DisplayClusterConfiguration/Public/DisplayClusterConfigurationTypes_Upscaler.h
--- a/Engine/Plugins/Runtime/nDisplay/Source/DisplayClusterConfiguration/Public/DisplayClusterConfigurationTypes_Upscaler.h
+++ b/Engine/Plugins/Runtime/nDisplay/Source/DisplayClusterConfiguration/Public/DisplayClusterConfigurationTypes_Upscaler.h
@@ -32,7 +32,9 @@ struct DISPLAYCLUSTERCONFIGURATION_API FDisplayClusterConfigurationUpscalerSetti
 
 public:
 	/** Get upscaler settings. */
-	bool GetUpscalerSettings(FDisplayClusterUpscalerSettings& OutUpscalerSettings) const;
+	bool GetUpscalerSettings(
+		const FDisplayClusterConfigurationUpscalerSettings* InDefaultUpscalerSettings,
+		FDisplayClusterUpscalerSettings& OutUpscalerSettings) const;
 
 public:
 	/** Override the Upscaling method. */
diff --git a/Engine/Plugins/Runtime/nDisplay/Source/DisplayClusterDetails/Private/DataModelGenerators/DisplayClusterDetailsGenerator_RootActor.cpp b/Engine/Plugins/Runtime/nDisplay/Source/DisplayClusterDetails/Private/DataModelGenerators/DisplayClusterDetailsGenerator_RootActor.cpp
--- a/Engine/Plugins/Runtime/nDisplay/Source/DisplayClusterDetails/Private/DataModelGenerators/DisplayClusterDetailsGenerator_RootActor.cpp
+++ b/Engine/Plugins/Runtime/nDisplay/Source/DisplayClusterDetails/Private/DataModelGenerators/DisplayClusterDetailsGenerator_RootActor.cpp
@@ -159,13 +159,16 @@ public:
 		}
 		
 		IDetailCategoryBuilder& ViewportsCategoryBuilder = DetailBuilder.EditCategory(TEXT("CustomViewportsCategory"), LOCTEXT("CustomViewportsCategoryLabel", "Viewports"));
+		AddProperty(ViewportsCategoryBuilder, TEXT("OuterViewportUpscalerSettingsRef"));
 		AddProperty(ViewportsCategoryBuilder, TEXT("ViewportScreenPercentageMultiplierRef"));
 		AddProperty(ViewportsCategoryBuilder, TEXT("FreezeRenderOuterViewportsRef"));
 
 		IDetailCategoryBuilder& InnerFrustumCategoryBuilder = DetailBuilder.EditCategory(TEXT("CustomICVFXCategory"), LOCTEXT("CustomICVFXCategoryLabel", "In-Camera VFX"));
+		AddProperty(InnerFrustumCategoryBuilder, TEXT("GlobalInnerFrustumUpscalerSettingsRef"));
 		AddProperty(InnerFrustumCategoryBuilder, TEXT("ShowInnerFrustumOverlapsRef"));
 		AddProperty(InnerFrustumCategoryBuilder, GET_MEMBER_NAME_CHECKED(ADisplayClusterRootActor, InnerFrustumPriority));
 
+
 		IDetailCategoryBuilder& ViewportChromakeyCategoryBuilder = DetailBuilder.EditCategory(TEXT("CustomViewportChromakeyCategory"), LOCTEXT("CustomViewportChromakeyCategoryLabel", "Chromakey"));
 		AddProperty(ViewportChromakeyCategoryBuilder, TEXT("GlobalChromakeyColorRef"));
 
