nDisplay: Fixed issue where legacy Unfreeze Viewports button was showing up in new editor toolbar. Also fixed issue where new Unfreeze Viewports toobar button wasn't showing up if level loaded didn't initially have any nDisplay stages.

diff --git a/Engine/Plugins/Runtime/nDisplay/Source/DisplayClusterConfigurator/Private/DisplayClusterConfiguratorModule.cpp b/Engine/Plugins/Runtime/nDisplay/Source/DisplayClusterConfigurator/Private/DisplayClusterConfiguratorModule.cpp
--- a/Engine/Plugins/Runtime/nDisplay/Source/DisplayClusterConfigurator/Private/DisplayClusterConfiguratorModule.cpp
+++ b/Engine/Plugins/Runtime/nDisplay/Source/DisplayClusterConfigurator/Private/DisplayClusterConfiguratorModule.cpp
@@ -57,6 +57,7 @@
 #include "ScopedTransaction.h"
 #include "SEditorViewportToolBarButton.h"
 #include "Subsystems/PanelExtensionSubsystem.h"
+#include "ViewportToolbar/UnrealEdViewportToolbar.h"
 #include "Views/Details/Cluster/DCConfiguratorColorGradingSettingsCustomization.h"
 #include "Views/Details/Cluster/DCConfiguratorWhiteBalanceCustomization.h"
 
@@ -519,6 +520,11 @@ EVisibility FDisplayClusterConfiguratorModule::GetViewportsFrozenWarningVisibili
 {
 	if (GEditor)
 	{
+		if (UE::UnrealEd::ShowNewViewportToolbars())
+		{
+			return EVisibility::Collapsed;
+		}
+
 		if (const UWorld* World = GEditor->GetEditorWorldContext().World())
 		{
 			for (TActorIterator<ADisplayClusterRootActor> It(World); It; ++It)
diff --git a/Engine/Plugins/Runtime/nDisplay/Source/DisplayClusterConfigurator/Private/DisplayClusterConfiguratorToolbarExtensions.cpp b/Engine/Plugins/Runtime/nDisplay/Source/DisplayClusterConfigurator/Private/DisplayClusterConfiguratorToolbarExtensions.cpp
--- a/Engine/Plugins/Runtime/nDisplay/Source/DisplayClusterConfigurator/Private/DisplayClusterConfiguratorToolbarExtensions.cpp
+++ b/Engine/Plugins/Runtime/nDisplay/Source/DisplayClusterConfigurator/Private/DisplayClusterConfiguratorToolbarExtensions.cpp
@@ -17,6 +17,22 @@ void FDisplayClusterConfiguratorToolbarExtensions::RegisterToolbarExtensions()
 	UToolMenu* Menu = UToolMenus::Get()->ExtendMenu("LevelEditor.ViewportToolBar.PerformanceAndScalability");
 	FToolMenuSection& Section = Menu->FindOrAddSection("FreezeNDisplayViewports", LOCTEXT("FreezeNDisplayViewportsSectionLabel", "Freeze nDisplay Viewports"));
 	Section.AddDynamicEntry(NAME_None, FNewToolMenuSectionDelegate::CreateSP(this, &FDisplayClusterConfiguratorToolbarExtensions::CreateFreezeViewportsMenu));
+	Section.Visibility = TAttribute<bool>::CreateLambda([this]()
+	{
+		// Check if there are any nDisplay stages in the level, as the menu section should only show up if there are nDisplay stages
+		if (GEditor)
+		{
+			if (const UWorld* World = GEditor->GetEditorWorldContext().World())
+			{
+				for (TActorIterator<ADisplayClusterRootActor> It(World); It; ++It)
+				{
+					return true;
+				}
+			}
+		}
+
+		return false;
+	});
 }
 
 void FDisplayClusterConfiguratorToolbarExtensions::UnregisterToolbarExtensions()
@@ -25,38 +41,32 @@ void FDisplayClusterConfiguratorToolbarExtensions::UnregisterToolbarExtensions()
 
 void FDisplayClusterConfiguratorToolbarExtensions::CreateFreezeViewportsMenu(FToolMenuSection& InSection)
 {
-	bool bAddedUnfreezeAll = false;
 	if (GEditor)
 	{
+		FToolMenuEntry UnfreezeAllEntry = FToolMenuEntry::InitMenuEntry(
+			FName("UnfreezeAllViewports"),
+			LOCTEXT("UnfreezeAllViewportsLabel", "Unfreeze All Viewports"),
+			LOCTEXT("ViewportsFrozenWarningToolTip", "nDisplay viewports are frozen. Click to unfreeze all frozen viewports."),
+			FSlateIcon(FDisplayClusterConfiguratorStyle::Get().GetStyleSetName(), "DisplayClusterConfigurator.LevelEditor.ViewportsFrozen"),
+			FUIAction(
+				FExecuteAction::CreateSP(this, &FDisplayClusterConfiguratorToolbarExtensions::UnfreezeAllViewports),
+				FCanExecuteAction::CreateSP(this, &FDisplayClusterConfiguratorToolbarExtensions::AreAnyViewportsFrozen)),
+			EUserInterfaceActionType::Button);
+
+		UnfreezeAllEntry.SetShowInToolbarTopLevel(TAttribute<bool>::CreateLambda(
+			[this]() -> bool
+			{
+				return AreAnyViewportsFrozen();
+			}
+		));
+
+		UnfreezeAllEntry.StyleNameOverride = "ViewportToolbarWarning";
+		InSection.AddEntry(UnfreezeAllEntry);
+
 		if (const UWorld* World = GEditor->GetEditorWorldContext().World())
 		{
 			for (TActorIterator<ADisplayClusterRootActor> It(World); It; ++It)
 			{
-				// If root actors have been found in the level, before the first one is added, add the "Unfreeze All Viewports" option
-				if (!bAddedUnfreezeAll)
-				{
-					FToolMenuEntry UnfreezeAllEntry = FToolMenuEntry::InitMenuEntry(
-						FName("UnfreezeAllViewports"),
-						LOCTEXT("UnfreezeAllViewportsLabel", "Unfreeze All Viewports"),
-						LOCTEXT("ViewportsFrozenWarningToolTip", "nDisplay viewports are frozen. Click to unfreeze all frozen viewports."),
-						FSlateIcon(FDisplayClusterConfiguratorStyle::Get().GetStyleSetName(), "DisplayClusterConfigurator.LevelEditor.ViewportsFrozen"),
-						FUIAction(
-							FExecuteAction::CreateSP(this, &FDisplayClusterConfiguratorToolbarExtensions::UnfreezeAllViewports),
-							FCanExecuteAction::CreateSP(this, &FDisplayClusterConfiguratorToolbarExtensions::AreAnyViewportsFrozen)),
-						EUserInterfaceActionType::Button);
-
-					UnfreezeAllEntry.SetShowInToolbarTopLevel(TAttribute<bool>::CreateLambda(
-						[this]() -> bool
-						{
-							return AreAnyViewportsFrozen();
-						}
-					));
-
-					UnfreezeAllEntry.StyleNameOverride = "ViewportToolbarWarning";
-					InSection.AddEntry(UnfreezeAllEntry);
-					bAddedUnfreezeAll = true;
-				}
-				
 				TWeakObjectPtr<ADisplayClusterRootActor> WeakPtr = *It;
 				FString RootActorLabel = It->GetActorNameOrLabel();
 				InSection.AddEntry(FToolMenuEntry::InitMenuEntry(
