[nDisplay] Added validation of some ICVFX media settings during BP compilation

#rnx

diff --git a/Engine/Plugins/Runtime/nDisplay/Source/DisplayClusterConfigurator/Private/DisplayClusterConfiguratorCompiler.cpp b/Engine/Plugins/Runtime/nDisplay/Source/DisplayClusterConfigurator/Private/DisplayClusterConfiguratorCompiler.cpp
--- a/Engine/Plugins/Runtime/nDisplay/Source/DisplayClusterConfigurator/Private/DisplayClusterConfiguratorCompiler.cpp
+++ b/Engine/Plugins/Runtime/nDisplay/Source/DisplayClusterConfigurator/Private/DisplayClusterConfiguratorCompiler.cpp
@@ -9,6 +9,7 @@
 #include "Blueprints/DisplayClusterBlueprint.h"
 #include "Blueprints/DisplayClusterBlueprintGeneratedClass.h"
 #include "Components/DisplayClusterCameraComponent.h"
+#include "Components/DisplayClusterICVFXCameraComponent.h"
 #include "Render/DisplayDevice/Components/DisplayClusterDisplayDeviceBaseComponent.h"
 #include "Camera/CameraComponent.h"
 
@@ -371,6 +372,46 @@ void FDisplayClusterConfiguratorKismetCompilerContext::ValidateConfiguration()
 			MessageLog.Warning(*LOCTEXT("HeadlessFullscreen", "Node @@ has both 'headless' and 'fullscreen' flags set. It's not forbidden but make sure it's intentional as the backbuffer will have the desktop's size.").ToString(), *ClusterNode.Key);
 		}
 	}
+
+	if (DCBlueprint && DCBlueprint->SimpleConstructionScript)
+	{
+		const TArray<USCS_Node*>& SCSNodes = Blueprint->SimpleConstructionScript->GetAllNodes();
+
+		// Iterate over all SCS nodes
+		for (const USCS_Node* const SCSNode : SCSNodes)
+		{
+			if (SCSNode && SCSNode->ComponentTemplate)
+			{
+				// Filter ICVFX camera components
+				const UActorComponent* const ComponentTemplate = SCSNode->ComponentTemplate;
+				if (const UDisplayClusterICVFXCameraComponent* const ICVFXCameraComponent = Cast<UDisplayClusterICVFXCameraComponent>(ComponentTemplate))
+				{
+					// Validate media settings
+					const FDisplayClusterConfigurationMediaICVFX& MediaSettings = ICVFXCameraComponent->CameraSettings.RenderSettings.Media;
+					if (MediaSettings.bEnable && MediaSettings.SplitType == EDisplayClusterConfigurationMediaSplitType::UniformTiles)
+					{
+						const int32 TilesCount = MediaSettings.TiledSplitLayout.X * MediaSettings.TiledSplitLayout.Y;
+						if (TilesCount < 2)
+						{
+							MessageLog.Error(*LOCTEXT("MediaInvalidTilesLayout", "Camera @@ has invalid split layout.").ToString(), ICVFXCameraComponent);
+						}
+						else
+						{
+							for (int32 GroupIdx = 0; GroupIdx < MediaSettings.TiledMediaInputGroups.Num(); ++GroupIdx)
+							{
+								const FDisplayClusterConfigurationMediaTiledInputGroup& TiledInputGroup = MediaSettings.TiledMediaInputGroups[GroupIdx];
+								if (TiledInputGroup.Tiles.Num() > TilesCount)
+								{
+									MessageLog.Error(*FText::Format(LOCTEXT("MediaInvalidTilesNum", "Camera @@ has invalid amount of tiles in the input group {0}."), FText::AsNumber(GroupIdx)).ToString(), ICVFXCameraComponent);
+								}
+							}
+						}
+					}
+				}
+
+			}
+		}
+	}
 }
 
 #undef LOCTEXT_NAMESPACE
