[nDIsplay] [Virtual Production] Improved logic related to the percentage range of the TSR/TAA upscaler.

diff --git a/Engine/Plugins/Runtime/nDisplay/Source/DisplayCluster/Private/Render/Upscaler/DisplayClusterUpscaler.cpp b/Engine/Plugins/Runtime/nDisplay/Source/DisplayCluster/Private/Render/Upscaler/DisplayClusterUpscaler.cpp
--- a/Engine/Plugins/Runtime/nDisplay/Source/DisplayCluster/Private/Render/Upscaler/DisplayClusterUpscaler.cpp
+++ b/Engine/Plugins/Runtime/nDisplay/Source/DisplayCluster/Private/Render/Upscaler/DisplayClusterUpscaler.cpp
@@ -274,13 +274,13 @@ namespace UE::DisplayCluster::Upscaler
 
 void FDisplayClusterUpscaler::SetupSceneView(
 	const FDisplayClusterUpscalerSettings& InUpscalerSettings,
-	const FSceneViewFamily& InOutViewFamily,
+	const FSceneViewFamily& InViewFamily,
 	FSceneView& InOutView)
 {
 	using namespace UE::DisplayCluster::Upscaler;
 
 	// Get AntiAliasingMethod that can be used.
-	const EDisplayClusterUpscalerAntiAliasingMethod SuitableAntiAliasingMethod = GetSuitableAntiAliasingMethod(InUpscalerSettings, InOutViewFamily, InOutView);
+	const EDisplayClusterUpscalerAntiAliasingMethod SuitableAntiAliasingMethod = GetSuitableAntiAliasingMethod(InUpscalerSettings, InViewFamily, InOutView);
 	switch (SuitableAntiAliasingMethod)
 	{
 	case EDisplayClusterUpscalerAntiAliasingMethod::None:
@@ -320,7 +320,7 @@ void FDisplayClusterUpscaler::SetupSceneView(
 			{
 				if (Upscaler.GetName() == InUpscalerSettings.CustomUpscalerName)
 				{
-					Upscaler.SetupSceneView(InUpscalerSettings.CustomUpscalerSettings, InOutViewFamily, InOutView);
+					Upscaler.SetupSceneView(InUpscalerSettings.CustomUpscalerSettings, InViewFamily, InOutView);
 				}
 			});
 	}
@@ -403,28 +403,67 @@ FName FDisplayClusterUpscaler::PostConfigureViewFamily(
 	const EPrimaryScreenPercentageMethod PrimaryScreenPercentageMethod = InOutViews[0]->PrimaryScreenPercentageMethod;
 	const EAntiAliasingMethod AntiAliasingMethod = InOutViews[0]->AntiAliasingMethod;
 
-	// Setup TSR & TAA
-	// We need to split the screen percentage if below 0.5 because TAA upscaling only works well up to 2x.
-	if(PrimaryScreenPercentageMethod == EPrimaryScreenPercentageMethod::TemporalUpscale
-	&&( AntiAliasingMethod == EAntiAliasingMethod::AAM_TemporalAA
-	 || AntiAliasingMethod == EAntiAliasingMethod::AAM_TSR))
+	// Get global view fraction set by r.ScreenPercentage.
+	const float GlobalScreenPercentage = GetGlobalScreenPercentage();
+
+	// Get TSR or TAA resolution fraction range
+	FVector2f TemporalUpscaleResolutionFractionRange = FVector2f::ZeroVector;
+	if (PrimaryScreenPercentageMethod == EPrimaryScreenPercentageMethod::TemporalUpscale)
 	{
-		// Get global view fraction set by r.ScreenPercentage.
-		const float GlobalScreenPercentage = GetGlobalScreenPercentage();
+		if (AntiAliasingMethod == EAntiAliasingMethod::AAM_TemporalAA)
+		{
+			// Sets the minimal and maximal screen percentage for TAAU. [0.5 .. 2.0]
+			TemporalUpscaleResolutionFractionRange = FVector2f(
+				ISceneViewFamilyScreenPercentage::kMinTAAUpsampleResolutionFraction,
+				ISceneViewFamilyScreenPercentage::kMaxTAAUpsampleResolutionFraction);
+		}
+		else if (AntiAliasingMethod == EAntiAliasingMethod::AAM_TSR)
+		{
+			// Sets the minimal and maximal screen percentage for TSR.  [0.25 .. 2.0]
+			TemporalUpscaleResolutionFractionRange = FVector2f(
+				ISceneViewFamilyScreenPercentage::kMinTSRResolutionFraction,
+				ISceneViewFamilyScreenPercentage::kMaxTSRResolutionFraction);
+		}
+	}
 
+	// We need to split the screen percentage value if it is outside the TSR/TAA range.
+	// All TemporalUpscaleResolutionFractionRange values are greater than zero.
+	if(TemporalUpscaleResolutionFractionRange.GetMin() > 0)
+	{
 		// Adjust for global screen% in math bellow
 		const float FinalResolutionFraction = UpscalerParam.UpscalerScreenPercentage * GlobalScreenPercentage;
-		if (FinalResolutionFraction < 0.5f)
+		if (FinalResolutionFraction < TemporalUpscaleResolutionFractionRange.X)
+		{
+			// Move the remaining ScreenPercentage to the spatial upscaler
+			const float RemainingScreenPercentage = UpscalerParam.UpscalerScreenPercentage / TemporalUpscaleResolutionFractionRange.X;
+			UpscalerParam.SecondaryScreenPercentage *= RemainingScreenPercentage;
+
+			// Adjust min value (0.5f) for the global screen%. Expectations that the actual TSR_Fraction=0.5 
+			UpscalerParam.UpscalerScreenPercentage = TemporalUpscaleResolutionFractionRange.X / GlobalScreenPercentage;
+		}
+		else if (FinalResolutionFraction > TemporalUpscaleResolutionFractionRange.Y)
 		{
 			// Move the remaining ScreenPercentage to the spatial upscaler
-			const float RemainingScreenPercentage = 2.0f * UpscalerParam.UpscalerScreenPercentage;
+			const float RemainingScreenPercentage = UpscalerParam.UpscalerScreenPercentage / TemporalUpscaleResolutionFractionRange.Y;
 			UpscalerParam.SecondaryScreenPercentage *= RemainingScreenPercentage;
 
 			// Adjust min value (0.5f) for the global screen%. Expectations that the actual TSR_Fraction=0.5 
-			UpscalerParam.UpscalerScreenPercentage = 0.5f / GlobalScreenPercentage;
+			UpscalerParam.UpscalerScreenPercentage = TemporalUpscaleResolutionFractionRange.Y / GlobalScreenPercentage;
 		}
 	}
 
+	const float FinalResolutionFraction = UpscalerParam.UpscalerScreenPercentage * GlobalScreenPercentage;
+
+	// Clamp final resolution fraction
+	// Sets the minimal and max screen percentage to [0.01 ... 4.0]
+	float ClampedFinalResolutionFraction = FMath::Clamp(
+		FinalResolutionFraction,
+		ISceneViewFamilyScreenPercentage::kMinResolutionFraction,
+		ISceneViewFamilyScreenPercentage::kMaxResolutionFraction);
+
+	UpscalerParam.UpscalerScreenPercentage = ClampedFinalResolutionFraction / GlobalScreenPercentage;
+
+	// Setup screen percentage
 	InOutViewFamily.SetScreenPercentageInterface(new FLegacyScreenPercentageDriver(
 		InOutViewFamily, UpscalerParam.UpscalerScreenPercentage));
 
diff --git a/Engine/Plugins/Runtime/nDisplay/Source/DisplayCluster/Private/Render/Upscaler/DisplayClusterUpscaler.h b/Engine/Plugins/Runtime/nDisplay/Source/DisplayCluster/Private/Render/Upscaler/DisplayClusterUpscaler.h
--- a/Engine/Plugins/Runtime/nDisplay/Source/DisplayCluster/Private/Render/Upscaler/DisplayClusterUpscaler.h
+++ b/Engine/Plugins/Runtime/nDisplay/Source/DisplayCluster/Private/Render/Upscaler/DisplayClusterUpscaler.h
@@ -50,11 +50,11 @@ public:
 	/**
 	* Configure scene view for upscalers.
 	* @param InUpscalerSettings - (in) the settings to apply for view
-	* @param InOutViewFamily    - (in) ViewFamily that own this view.
+	* @param InViewFamily       - (in) ViewFamily that own this view.
 	* @param InOutViews         - (in, out) The views that will be configured. 
 	*/
 	static void SetupSceneView(
 		const FDisplayClusterUpscalerSettings& InUpscalerSettings,
-		const FSceneViewFamily& InOutViewFamily,
+		const FSceneViewFamily& InViewFamily,
 		FSceneView& InOutView);
 };
diff --git a/Engine/Plugins/Runtime/nDisplay/Source/DisplayCluster/Private/Render/Viewport/DisplayClusterViewportManager.cpp b/Engine/Plugins/Runtime/nDisplay/Source/DisplayCluster/Private/Render/Viewport/DisplayClusterViewportManager.cpp
--- a/Engine/Plugins/Runtime/nDisplay/Source/DisplayCluster/Private/Render/Viewport/DisplayClusterViewportManager.cpp
+++ b/Engine/Plugins/Runtime/nDisplay/Source/DisplayCluster/Private/Render/Viewport/DisplayClusterViewportManager.cpp
@@ -826,10 +826,8 @@ void FDisplayClusterViewportManager::PostConfigureViewFamily(
 	const float InDPIScale
 )
 {
-	const float DisplayClusterScreenPercentage = FMath::Clamp(InFrameViewFamily.CustomBufferRatio, 0.01f, 1.0f);
-
 	// Configure upscaler
-	const FName CustomUpscalerName = FDisplayClusterUpscaler::PostConfigureViewFamily(InFrameViewFamily.UpscalerSettings, DisplayClusterScreenPercentage, InDPIScale, ViewFamily, Views);
+	const FName CustomUpscalerName = FDisplayClusterUpscaler::PostConfigureViewFamily(InFrameViewFamily.UpscalerSettings, InFrameViewFamily.CustomBufferRatio, InDPIScale, ViewFamily, Views);
 
 	// Add ViewExtensions from the Custom Upscaler to this view family
 	if (!CustomUpscalerName.IsNone() && InFrameViewFamily.UpscalerSettings.CustomUpscalerName == CustomUpscalerName)
