[nDisplay] [Virtual Production] DCRA freeze rendering issues

- Fixed freeze outers viewports issue
- Fixed DCRA preview rendering freeze issue

diff --git a/Engine/Plugins/Runtime/nDisplay/Source/DisplayCluster/Private/Game/EngineClasses/Scene/DisplayClusterRootActor.cpp b/Engine/Plugins/Runtime/nDisplay/Source/DisplayCluster/Private/Game/EngineClasses/Scene/DisplayClusterRootActor.cpp
--- a/Engine/Plugins/Runtime/nDisplay/Source/DisplayCluster/Private/Game/EngineClasses/Scene/DisplayClusterRootActor.cpp
+++ b/Engine/Plugins/Runtime/nDisplay/Source/DisplayCluster/Private/Game/EngineClasses/Scene/DisplayClusterRootActor.cpp
@@ -854,6 +854,12 @@ void ADisplayClusterRootActor::InitializeRootActor()
 	// where the instanced media settings may be incorrectly considered dirty and therefore
 	// not propagated from the template to the instance.
 	UE::DisplayCluster::RootActor::PropagateMediaSettingsFromDefault(this);
+
+	// At this point, the DCRA components have been re-created and these changes should be reflected in the ViewportManager.
+	if (IDisplayClusterViewportManager* ViewportManager = GetViewportManager())
+	{
+		ViewportManager->GetConfiguration().ReleaseConfiguration();
+	}
 }
 
 void ADisplayClusterRootActor::UpdateProceduralMeshComponentData(const UProceduralMeshComponent* InProceduralMeshComponent)
diff --git a/Engine/Plugins/Runtime/nDisplay/Source/DisplayCluster/Private/Render/Viewport/Configuration/DisplayClusterViewportConfiguration_ICVFX.cpp b/Engine/Plugins/Runtime/nDisplay/Source/DisplayCluster/Private/Render/Viewport/Configuration/DisplayClusterViewportConfiguration_ICVFX.cpp
--- a/Engine/Plugins/Runtime/nDisplay/Source/DisplayCluster/Private/Render/Viewport/Configuration/DisplayClusterViewportConfiguration_ICVFX.cpp
+++ b/Engine/Plugins/Runtime/nDisplay/Source/DisplayCluster/Private/Render/Viewport/Configuration/DisplayClusterViewportConfiguration_ICVFX.cpp
@@ -184,36 +184,6 @@ void FDisplayClusterViewportConfiguration_ICVFX::Update()
 		}
 	}
 
-	if (StageSettings->bFreezeRenderOuterViewports || RenderFrameSettings.IsPreviewFreezeRender())
-	{
-		// Freeze rendering for outer viewports
-		for (const TSharedPtr<FDisplayClusterViewport, ESPMode::ThreadSafe>& TargetIt : TargetViewports)
-		{
-			if (TargetIt.IsValid())
-			{
-				// Gain direct access to internal settings of the viewport:
-				FDisplayClusterViewport_RenderSettings& InOutRenderSettings = TargetIt->GetRenderSettingsImpl();
-
-				InOutRenderSettings.bFreezeRendering = true;
-			}
-		}
-
-		// Freeze render for lightcards when outer viewports freezed
-		if (StageSettings->Lightcard.bIgnoreOuterViewportsFreezingForLightcards == false)
-		{
-			for (const TSharedPtr<FDisplayClusterViewport, ESPMode::ThreadSafe>& ViewportIt : ViewportManager->ImplGetCurrentRenderFrameViewports())
-			{
-				if (ViewportIt.IsValid() && EnumHasAnyFlags(ViewportIt->GetRenderSettingsICVFX().RuntimeFlags, EDisplayClusterViewportRuntimeICVFXFlags::Lightcard | EDisplayClusterViewportRuntimeICVFXFlags::UVLightcard))
-				{
-					// Gain direct access to internal settings of the viewport:
-					FDisplayClusterViewport_RenderSettings& InOutRenderSettings = ViewportIt->GetRenderSettingsImpl();
-
-					InOutRenderSettings.bFreezeRendering = true;
-				}
-			}
-		}
-	}
-
 	ImplFinishReallocateViewports();
 }
 
diff --git a/Engine/Plugins/Runtime/nDisplay/Source/DisplayCluster/Private/Render/Viewport/DisplayClusterViewport.cpp b/Engine/Plugins/Runtime/nDisplay/Source/DisplayCluster/Private/Render/Viewport/DisplayClusterViewport.cpp
--- a/Engine/Plugins/Runtime/nDisplay/Source/DisplayCluster/Private/Render/Viewport/DisplayClusterViewport.cpp
+++ b/Engine/Plugins/Runtime/nDisplay/Source/DisplayCluster/Private/Render/Viewport/DisplayClusterViewport.cpp
@@ -516,29 +516,33 @@ bool FDisplayClusterViewport::UpdateFrameContexts(const uint32 InStereoViewIndex
 	const uint32 ViewportContextAmount = RenderSettings.bForceMono ? 1 : FrameTargetsAmount;
 
 	// Freeze the image in the viewport only after the frame has been rendered
-	if (RenderSettings.bFreezeRendering && RenderSettings.bEnable)
+	const bool bViewportRendered = Contexts.Num() > 0 && Contexts.Num() == Resources[EDisplayClusterViewportResource::InputShaderResources].Num();
+	if (bViewportRendered && RenderSettings.bEnable && ShouldFreezeRender())
 	{
-		// Freeze only when all resources valid
-		if (Contexts.Num() > 0 && Contexts.Num() == Resources[EDisplayClusterViewportResource::InputShaderResources].Num())
-		{
-			Resources.FreezeRendering(EDisplayClusterViewportResource::InputShaderResources);
-			Resources.FreezeRendering(EDisplayClusterViewportResource::AdditionalTargetableResources);
-			Resources.FreezeRendering(EDisplayClusterViewportResource::MipsShaderResources);
+		// Raise freeze flag for this viewport logic
+		RenderSettings.bFreezeRendering = true;
 
-			// Update context links for freezed viewport
-			for (FDisplayClusterViewport_Context& ContextIt : Contexts)
-			{
-				ContextIt.StereoscopicPass = FDisplayClusterViewportStereoscopicPass::EncodeStereoscopicPass(ContextIt.ContextNum, ViewportContextAmount, InFrameSettings);
-				ContextIt.StereoViewIndex = (int32)(InStereoViewIndex + ContextIt.ContextNum);
-				ContextIt.bDisableRender = true;
-				ContextIt.FrameTargetRect = FDisplayClusterViewportHelpers::GetValidViewportRect(DesiredFrameTargetRect, GetId(), TEXT("Context Frame"));
-			}
+		// Block image resources from being re-allocated
+		Resources.FreezeRendering(EDisplayClusterViewportResource::InputShaderResources);
+		Resources.FreezeRendering(EDisplayClusterViewportResource::AdditionalTargetableResources);
+		Resources.FreezeRendering(EDisplayClusterViewportResource::MipsShaderResources);
 
-			// Release only part of the resources, leaving resources that can be used by other viewports (viewport override feature)
-			Resources.ReleaseNotSharedResources();
+		// Update context links for freezed viewport
+		for (FDisplayClusterViewport_Context& ContextIt : Contexts)
+		{
+			ContextIt.StereoscopicPass = FDisplayClusterViewportStereoscopicPass::EncodeStereoscopicPass(ContextIt.ContextNum, ViewportContextAmount, InFrameSettings);
+			ContextIt.StereoViewIndex = (int32)(InStereoViewIndex + ContextIt.ContextNum);
+			ContextIt.bDisableRender = true;
+			ContextIt.FrameTargetRect = FDisplayClusterViewportHelpers::GetValidViewportRect(DesiredFrameTargetRect, GetId(), TEXT("Context Frame"));
 
-			return true;
+			// Reset context state, and invalidate caches from previous frame.
+			ContextIt.ContextState = EDisplayClusterViewportContextState::None;
 		}
+
+		// Release only part of the resources, leaving resources that can be used by other viewports (viewport override feature)
+		Resources.ReleaseNotSharedResources();
+
+		return true;
 	}
 
 	// Release old contexts
diff --git a/Engine/Plugins/Runtime/nDisplay/Source/DisplayCluster/Private/Render/Viewport/DisplayClusterViewport.h b/Engine/Plugins/Runtime/nDisplay/Source/DisplayCluster/Private/Render/Viewport/DisplayClusterViewport.h
--- a/Engine/Plugins/Runtime/nDisplay/Source/DisplayCluster/Private/Render/Viewport/DisplayClusterViewport.h
+++ b/Engine/Plugins/Runtime/nDisplay/Source/DisplayCluster/Private/Render/Viewport/DisplayClusterViewport.h
@@ -336,6 +336,9 @@ public:
 	/** Returns true if per-view Lumen scene is enabled. */
 	bool ShouldUseLumenPerView() const;
 
+	/** Return true if this viewport should freeze render. */
+	bool ShouldFreezeRender() const;
+
 	/** Returns true if textures in this viewport with an area larger than the MaxTextureArea limit
 	*  must be reduced to stay within that limit with the same aspect ratio.
 	*/
diff --git a/Engine/Plugins/Runtime/nDisplay/Source/DisplayCluster/Private/Render/Viewport/DisplayClusterViewport_Rules.cpp b/Engine/Plugins/Runtime/nDisplay/Source/DisplayCluster/Private/Render/Viewport/DisplayClusterViewport_Rules.cpp
--- a/Engine/Plugins/Runtime/nDisplay/Source/DisplayCluster/Private/Render/Viewport/DisplayClusterViewport_Rules.cpp
+++ b/Engine/Plugins/Runtime/nDisplay/Source/DisplayCluster/Private/Render/Viewport/DisplayClusterViewport_Rules.cpp
@@ -433,6 +433,37 @@ bool FDisplayClusterViewport::ShouldUseFullSizeFrameTargetableResource() const
 	return false;
 }
 
+bool FDisplayClusterViewport::ShouldFreezeRender() const
+{
+	// Freeze preview rendering for all viewport
+	if (Configuration->GetRenderFrameSettings().IsPreviewFreezeRender())
+	{
+		return true;
+	}
+
+	// ICVFX: Freeze only some viewports
+	if (const FDisplayClusterConfigurationICVFX_StageSettings* StageSettings = Configuration->GetStageSettings())
+	{	
+		if (StageSettings->bFreezeRenderOuterViewports)
+		{
+			// Freeze rendering for outer viewports
+			if (!IsInternalViewport())
+			{
+				return true;
+			}
+
+			// Enable\Disable freeze rendering for lightcards when outer viewports rendering also freezed. This will impact performance.
+			if (!StageSettings->Lightcard.bIgnoreOuterViewportsFreezingForLightcards
+				&& EnumHasAnyFlags(RenderSettingsICVFX.RuntimeFlags, EDisplayClusterViewportRuntimeICVFXFlags::Lightcard | EDisplayClusterViewportRuntimeICVFXFlags::UVLightcard))
+			{
+				return true;
+			}
+		}
+	}
+	
+	return false;
+}
+
 bool FDisplayClusterViewport::IsOpenColorIOEquals(const FDisplayClusterViewport& InViewport) const
 {
 	bool bEnabledOCIO_1 = OpenColorIO.IsValid();
diff --git a/Engine/Plugins/Runtime/nDisplay/Source/DisplayCluster/Private/Render/Viewport/Preview/DisplayClusterViewportManagerPreview.cpp b/Engine/Plugins/Runtime/nDisplay/Source/DisplayCluster/Private/Render/Viewport/Preview/DisplayClusterViewportManagerPreview.cpp
--- a/Engine/Plugins/Runtime/nDisplay/Source/DisplayCluster/Private/Render/Viewport/Preview/DisplayClusterViewportManagerPreview.cpp
+++ b/Engine/Plugins/Runtime/nDisplay/Source/DisplayCluster/Private/Render/Viewport/Preview/DisplayClusterViewportManagerPreview.cpp
@@ -135,12 +135,6 @@ void FDisplayClusterViewportManagerPreview::OnPreviewRenderTick()
 	}
 	TickPerFrameCounter = 0;
 
-
-	if (bEntireClusterRendered && RenderFrameSetting.IsPreviewFreezeRender())
-	{
-		return;
-	}
-
 	int32 ViewportsAmmount = RenderFrameSetting.PreviewSettings.ViewportsPerFrame;
 	int32 CycleDepth = 0;
 	FString FirstClusterNodeId;
@@ -187,13 +181,6 @@ void FDisplayClusterViewportManagerPreview::OnPreviewRenderTick()
 		{
 			// When rendering of all nodes in the cluster is complete, generate events for subscribers
 			OnEntireClusterPreviewGenerated.ExecuteIfBound();
-
-			// When preview rendering is freeze, we render a preview just once
-			if (RenderFrameSetting.IsPreviewFreezeRender())
-			{
-				bEntireClusterRendered = true;
-				break;
-			}
 		}
 	}
 }
@@ -222,6 +209,8 @@ void FDisplayClusterViewportManagerPreview::Release()
 			ViewportPreviewIt->Release();
 		}
 	}
+
+	ResetEntireClusterPreviewRendering();
 }
 
 const TArray<TSharedPtr<IDisplayClusterViewportPreview, ESPMode::ThreadSafe>> FDisplayClusterViewportManagerPreview::GetEntireClusterPreviewViewports() const
diff --git a/Engine/Plugins/Runtime/nDisplay/Source/DisplayCluster/Private/Render/Viewport/RenderTarget/DisplayClusterRenderTargetManager.cpp b/Engine/Plugins/Runtime/nDisplay/Source/DisplayCluster/Private/Render/Viewport/RenderTarget/DisplayClusterRenderTargetManager.cpp
--- a/Engine/Plugins/Runtime/nDisplay/Source/DisplayCluster/Private/Render/Viewport/RenderTarget/DisplayClusterRenderTargetManager.cpp
+++ b/Engine/Plugins/Runtime/nDisplay/Source/DisplayCluster/Private/Render/Viewport/RenderTarget/DisplayClusterRenderTargetManager.cpp
@@ -153,7 +153,7 @@ bool FDisplayClusterRenderTargetManager::AllocateRenderFrameResources(FViewport*
 		// Allocate viewport internal resources
 		for (const TSharedPtr<FDisplayClusterViewport, ESPMode::ThreadSafe>& ViewportIt : InViewports)
 		{
-			if (ViewportIt.IsValid() && ViewportIt->GetRenderSettings().bFreezeRendering == false)
+			if (ViewportIt.IsValid())
 			{
 				EPixelFormat CustomFormat = ImplGetCustomFormat(ViewportIt->GetRenderSettings().CaptureMode);
 
@@ -166,27 +166,43 @@ bool FDisplayClusterRenderTargetManager::AllocateRenderFrameResources(FViewport*
 					// Allocate per-viewport preview output textures:
 					if (ViewportIt->GetViewportResources(EDisplayClusterViewportResource::OutputPreviewTargetableResources).IsValidIndex(ContextNum))
 					{
-						ViewportIt->GetViewportResourcesImpl(EDisplayClusterViewportResource::OutputPreviewTargetableResources)[ContextNum] = ResourcesPool->AllocateResource(ViewportIt->GetId(), ContextSize, CustomFormat, EDisplayClusterViewportResourceSettingsFlags::PreviewTargetableTexture);
-						ImplViewportTextureResourceLogging(ViewportIt, ContextNum, TEXT("InputShader"), ViewportIt->GetViewportResources(EDisplayClusterViewportResource::OutputPreviewTargetableResources)[ContextNum]);
+						TSharedPtr<FDisplayClusterViewportResource, ESPMode::ThreadSafe>& ReallocResource = ViewportIt->GetViewportResourcesImpl(EDisplayClusterViewportResource::OutputPreviewTargetableResources)[ContextNum];
+						if (!ReallocResource.IsValid() || !EnumHasAnyFlags(ReallocResource->GetResourceState(), EDisplayClusterViewportResourceState::DisableReallocate))
+						{
+							ReallocResource = ResourcesPool->AllocateResource(ViewportIt->GetId(), ContextSize, CustomFormat, EDisplayClusterViewportResourceSettingsFlags::PreviewTargetableTexture);
+							ImplViewportTextureResourceLogging(ViewportIt, ContextNum, TEXT("InputShader"), ViewportIt->GetViewportResources(EDisplayClusterViewportResource::OutputPreviewTargetableResources)[ContextNum]);
+						}
 					}
 
 					if (ViewportIt->GetViewportResources(EDisplayClusterViewportResource::InputShaderResources).IsValidIndex(ContextNum))
 					{
-						ViewportIt->GetViewportResourcesImpl(EDisplayClusterViewportResource::InputShaderResources)[ContextNum] = ResourcesPool->AllocateResource(EmptyString, ContextSize, CustomFormat, EDisplayClusterViewportResourceSettingsFlags::ResolveTargetableTexture);
-						ImplViewportTextureResourceLogging(ViewportIt, ContextNum, TEXT("InputShader"), ViewportIt->GetViewportResources(EDisplayClusterViewportResource::InputShaderResources)[ContextNum]);
+						TSharedPtr<FDisplayClusterViewportResource, ESPMode::ThreadSafe>& ReallocResource = ViewportIt->GetViewportResourcesImpl(EDisplayClusterViewportResource::InputShaderResources)[ContextNum];
+						if (!ReallocResource.IsValid() || !EnumHasAnyFlags(ReallocResource->GetResourceState(), EDisplayClusterViewportResourceState::DisableReallocate))
+						{
+							ReallocResource = ResourcesPool->AllocateResource(EmptyString, ContextSize, CustomFormat, EDisplayClusterViewportResourceSettingsFlags::ResolveTargetableTexture);
+							ImplViewportTextureResourceLogging(ViewportIt, ContextNum, TEXT("InputShader"), ViewportIt->GetViewportResources(EDisplayClusterViewportResource::InputShaderResources)[ContextNum]);
+						}
 					}
 
 					// Allocate custom resources:
 					if (ViewportIt->GetViewportResources(EDisplayClusterViewportResource::AdditionalTargetableResources).IsValidIndex(ContextNum))
 					{
-						ViewportIt->GetViewportResourcesImpl(EDisplayClusterViewportResource::AdditionalTargetableResources)[ContextNum] = ResourcesPool->AllocateResource(EmptyString, ContextSize, CustomFormat, EDisplayClusterViewportResourceSettingsFlags::RenderTargetableTexture);
-						ImplViewportTextureResourceLogging(ViewportIt, ContextNum, TEXT("AdditionalRTT"), ViewportIt->GetViewportResources(EDisplayClusterViewportResource::AdditionalTargetableResources)[ContextNum]);
+						TSharedPtr<FDisplayClusterViewportResource, ESPMode::ThreadSafe>& ReallocResource = ViewportIt->GetViewportResourcesImpl(EDisplayClusterViewportResource::AdditionalTargetableResources)[ContextNum];
+						if (!ReallocResource.IsValid() || !EnumHasAnyFlags(ReallocResource->GetResourceState(), EDisplayClusterViewportResourceState::DisableReallocate))
+						{
+							ReallocResource = ResourcesPool->AllocateResource(EmptyString, ContextSize, CustomFormat, EDisplayClusterViewportResourceSettingsFlags::RenderTargetableTexture);
+							ImplViewportTextureResourceLogging(ViewportIt, ContextNum, TEXT("AdditionalRTT"), ViewportIt->GetViewportResources(EDisplayClusterViewportResource::AdditionalTargetableResources)[ContextNum]);
+						}
 					}
 
 					if (ViewportIt->GetViewportResources(EDisplayClusterViewportResource::MipsShaderResources).IsValidIndex(ContextNum))
 					{
-						ViewportIt->GetViewportResourcesImpl(EDisplayClusterViewportResource::MipsShaderResources)[ContextNum] = ResourcesPool->AllocateResource(EmptyString, ContextSize, CustomFormat, EDisplayClusterViewportResourceSettingsFlags::ResolveTargetableTexture, ContextIt.NumMips);
-						ImplViewportTextureResourceLogging(ViewportIt, ContextNum, TEXT("Mips"), ViewportIt->GetViewportResources(EDisplayClusterViewportResource::MipsShaderResources)[ContextNum]);
+						TSharedPtr<FDisplayClusterViewportResource, ESPMode::ThreadSafe>& ReallocResource = ViewportIt->GetViewportResourcesImpl(EDisplayClusterViewportResource::MipsShaderResources)[ContextNum];
+						if (!ReallocResource.IsValid() || !EnumHasAnyFlags(ReallocResource->GetResourceState(), EDisplayClusterViewportResourceState::DisableReallocate))
+						{
+							ReallocResource = ResourcesPool->AllocateResource(EmptyString, ContextSize, CustomFormat, EDisplayClusterViewportResourceSettingsFlags::ResolveTargetableTexture, ContextIt.NumMips);
+							ImplViewportTextureResourceLogging(ViewportIt, ContextNum, TEXT("Mips"), ViewportIt->GetViewportResources(EDisplayClusterViewportResource::MipsShaderResources)[ContextNum]);
+						}
 					}
 				}
 			}
