[nDIsplay] [Virtual Production] Fixed upscaler method selection logic for TSR/TAA upscaler.

diff --git a/Engine/Plugins/Runtime/nDisplay/Source/DisplayCluster/Private/Render/Upscaler/DisplayClusterUpscaler.cpp b/Engine/Plugins/Runtime/nDisplay/Source/DisplayCluster/Private/Render/Upscaler/DisplayClusterUpscaler.cpp
--- a/Engine/Plugins/Runtime/nDisplay/Source/DisplayCluster/Private/Render/Upscaler/DisplayClusterUpscaler.cpp
+++ b/Engine/Plugins/Runtime/nDisplay/Source/DisplayCluster/Private/Render/Upscaler/DisplayClusterUpscaler.cpp
@@ -268,6 +268,16 @@ namespace UE::DisplayCluster::Upscaler
 			}
 		}
 
+		if(AntiAliasingMethod == EDisplayClusterUpscalerAntiAliasingMethod::TAA
+		|| AntiAliasingMethod == EDisplayClusterUpscalerAntiAliasingMethod::TSR)
+		{
+			// TemporalAA/TSR requires the view to have a valid state.
+			if (InSceneView.State == nullptr)
+			{	
+				return EDisplayClusterUpscalerAntiAliasingMethod::FXAA;
+			}
+		}
+
 		return AntiAliasingMethod;
 	}
 };
