[nDisplay] [Virtual Production] Renamed UI

ViewPoint component:
- Renamed "Camera Post Process" to -> "Camera Settings"

diff --git a/Engine/Plugins/Runtime/nDisplay/Source/DisplayCluster/Public/Components/DisplayClusterCameraComponent.h b/Engine/Plugins/Runtime/nDisplay/Source/DisplayCluster/Public/Components/DisplayClusterCameraComponent.h
--- a/Engine/Plugins/Runtime/nDisplay/Source/DisplayCluster/Public/Components/DisplayClusterCameraComponent.h
+++ b/Engine/Plugins/Runtime/nDisplay/Source/DisplayCluster/Public/Components/DisplayClusterCameraComponent.h
@@ -299,44 +299,44 @@ protected:
 
 public:
 	/** Type of source camera used. */
-	UPROPERTY(EditAnywhere, BlueprintReadWrite, Category = "Camera Post Process", meta = (DisplayName = "Source Camera"))
+	UPROPERTY(EditAnywhere, BlueprintReadWrite, Category = "Camera Settings", meta = (DisplayName = "Source Camera"))
 	EDisplayClusterTargetCameraType TargetCameraType = EDisplayClusterTargetCameraType::None;
 
 	/** The name of the camera component that is used as the PP source. */
-	UPROPERTY(EditAnywhere, BlueprintReadWrite, Category = "Camera Post Process", meta = (DisplayName = "ICVFX Camera", EditCondition = "TargetCameraType==EDisplayClusterTargetCameraType::ICVFXCameraComponent", EditConditionHides))
+	UPROPERTY(EditAnywhere, BlueprintReadWrite, Category = "Camera Settings", meta = (DisplayName = "ICVFX Camera", EditCondition = "TargetCameraType==EDisplayClusterTargetCameraType::ICVFXCameraComponent", EditConditionHides))
 	FString ICVFXCameraComponentName;
 
-	UPROPERTY(EditAnywhere, BlueprintReadWrite, Category = "Camera Post Process", meta = (DisplayName = "Follow ICVFX Camera", EditCondition = "TargetCameraType==EDisplayClusterTargetCameraType::ICVFXCameraComponent", EditConditionHides))
+	UPROPERTY(EditAnywhere, BlueprintReadWrite, Category = "Camera Settings", meta = (DisplayName = "Follow ICVFX Camera", EditCondition = "TargetCameraType==EDisplayClusterTargetCameraType::ICVFXCameraComponent", EditConditionHides))
 	uint8 bUseICVFXCameraComponentTracking: 1 = 1;
 
 	/** Use a specific actor camera instead of a game camera. */
-	UPROPERTY(EditAnywhere, BlueprintReadWrite, Category = "Camera Post Process", meta = (DisplayName = "Camera Actor", EditCondition = "TargetCameraType==EDisplayClusterTargetCameraType::ExternalCineCameraActor", EditConditionHides))
+	UPROPERTY(EditAnywhere, BlueprintReadWrite, Category = "Camera Settings", meta = (DisplayName = "Camera Actor", EditCondition = "TargetCameraType==EDisplayClusterTargetCameraType::ExternalCineCameraActor", EditConditionHides))
 	TSoftObjectPtr<ACineCameraActor> ExternalCineCameraActor;
 
 public:
 
 	/** Use the PP settings from the specified camera. */
-	UPROPERTY(EditAnywhere, BlueprintReadWrite, Category = "Camera Post Process", meta = (DisplayName = "Include Post Process Settings", EditCondition = "TargetCameraType!=EDisplayClusterTargetCameraType::None", EditConditionHides))
+	UPROPERTY(EditAnywhere, BlueprintReadWrite, Category = "Camera Settings", meta = (DisplayName = "Include Post Process Settings", EditCondition = "TargetCameraType!=EDisplayClusterTargetCameraType::None", EditConditionHides))
 	uint8 bEnablePostProcess : 1 = 1;
 
 	/** Use the NearClippingPlane value from the specified cine camera. */
-	UPROPERTY(EditAnywhere, BlueprintReadWrite, Category = "Camera Post Process", meta = (DisplayName = "Include Custom Near Clipping Plane", EditCondition = "TargetCameraType!=EDisplayClusterTargetCameraType::None && TargetCameraType!=EDisplayClusterTargetCameraType::ICVFXCameraComponent", EditConditionHides))
+	UPROPERTY(EditAnywhere, BlueprintReadWrite, Category = "Camera Settings", meta = (DisplayName = "Include Custom Near Clipping Plane", EditCondition = "TargetCameraType!=EDisplayClusterTargetCameraType::None && TargetCameraType!=EDisplayClusterTargetCameraType::ICVFXCameraComponent", EditConditionHides))
 	uint8 bEnableNearClippingPlane : 1 = 1;
 
 	/** Enable the DoF PP settings from the specified camera. */
-	UPROPERTY(EditAnywhere, BlueprintReadWrite, Category = "Camera Post Process", meta = (DisplayName = "Include Depth Of Field", EditCondition = "TargetCameraType!=EDisplayClusterTargetCameraType::None && TargetCameraType!=EDisplayClusterTargetCameraType::ICVFXCameraComponent", EditConditionHides))
+	UPROPERTY(EditAnywhere, BlueprintReadWrite, Category = "Camera Settings", meta = (DisplayName = "Include Depth Of Field", EditCondition = "TargetCameraType!=EDisplayClusterTargetCameraType::None && TargetCameraType!=EDisplayClusterTargetCameraType::ICVFXCameraComponent", EditConditionHides))
 	uint8 bEnableDepthOfField : 1 = 1;
 
 	/** Use the DC Depth-Of-Field settings from the specified ICVFX camera. */
-	UPROPERTY(BlueprintReadWrite, Category = "Camera Post Process", meta = (DisplayName = "Include Inner Frustum Depth Of Field Compensation", EditCondition = "TargetCameraType==EDisplayClusterTargetCameraType::ICVFXCameraComponent", EditConditionHides))
+	UPROPERTY(BlueprintReadWrite, Category = "Camera Settings", meta = (DisplayName = "Include Inner Frustum Depth Of Field Compensation", EditCondition = "TargetCameraType==EDisplayClusterTargetCameraType::ICVFXCameraComponent", EditConditionHides))
 	uint8 bEnableICVFXDepthOfFieldCompensation : 1 = 1;
 
 	/** Use the DC ColorGrading from the specified ICVFX camera. */
-	UPROPERTY(EditAnywhere, BlueprintReadWrite, Category = "Camera Post Process", meta = (DisplayName = "Include Inner Frustum Color Grading", EditCondition = "TargetCameraType==EDisplayClusterTargetCameraType::ICVFXCameraComponent", EditConditionHides))
+	UPROPERTY(EditAnywhere, BlueprintReadWrite, Category = "Camera Settings", meta = (DisplayName = "Include Inner Frustum Color Grading", EditCondition = "TargetCameraType==EDisplayClusterTargetCameraType::ICVFXCameraComponent", EditConditionHides))
 	uint8 bEnableICVFXColorGrading : 1 = 1;
 
 	/** Use the DC Motion Blur settings from the specified ICVFX camera. */
-	UPROPERTY(BlueprintReadWrite, Category = "Camera Post Process", meta = (DisplayName = "Include Inner Frustum Motion Blur", EditCondition = "TargetCameraType==EDisplayClusterTargetCameraType::ICVFXCameraComponent", EditConditionHides))
+	UPROPERTY(BlueprintReadWrite, Category = "Camera Settings", meta = (DisplayName = "Include Inner Frustum Motion Blur", EditCondition = "TargetCameraType==EDisplayClusterTargetCameraType::ICVFXCameraComponent", EditConditionHides))
 	uint8 bEnableICVFXMotionBlur : 1 = 1;
 
 private:
diff --git a/Engine/Plugins/Runtime/nDisplay/Source/DisplayClusterConfiguration/Public/DisplayClusterConfigurationStrings.h b/Engine/Plugins/Runtime/nDisplay/Source/DisplayClusterConfiguration/Public/DisplayClusterConfigurationStrings.h
--- a/Engine/Plugins/Runtime/nDisplay/Source/DisplayClusterConfiguration/Public/DisplayClusterConfigurationStrings.h
+++ b/Engine/Plugins/Runtime/nDisplay/Source/DisplayClusterConfiguration/Public/DisplayClusterConfigurationStrings.h
@@ -50,7 +50,7 @@ namespace DisplayClusterConfigurationStrings
 		static constexpr const TCHAR* StereoCategory               = TEXT("Stereo");
 
 		static constexpr const TCHAR* ViewPointStereoCategory              = TEXT("Stereo");
-		static constexpr const TCHAR* ViewPointCameraPostProcessCategory   = TEXT("Camera Post Process");
+		static constexpr const TCHAR* ViewPointCameraPostProcessCategory   = TEXT("Camera Settings");
 		static constexpr const TCHAR* ViewPointInFrustumProjectionCategory = TEXT("Frustum Fit");
 	}
 
diff --git a/Engine/Plugins/Runtime/nDisplay/Source/DisplayClusterConfigurator/Private/DisplayClusterConfiguratorModule.cpp b/Engine/Plugins/Runtime/nDisplay/Source/DisplayClusterConfigurator/Private/DisplayClusterConfiguratorModule.cpp
--- a/Engine/Plugins/Runtime/nDisplay/Source/DisplayClusterConfigurator/Private/DisplayClusterConfiguratorModule.cpp
+++ b/Engine/Plugins/Runtime/nDisplay/Source/DisplayClusterConfigurator/Private/DisplayClusterConfiguratorModule.cpp
@@ -352,7 +352,7 @@ void FDisplayClusterConfiguratorModule::RegisterSectionMappings()
 	}
 	{
 		const TSharedRef<FPropertySection> Section = PropertyModule.FindOrCreateSection(UDisplayClusterCameraComponent::StaticClass()->GetFName(),
-			DisplayClusterConfigurationStrings::categories::ViewPointCameraPostProcessCategory, LOCTEXT("ViewPointCameraPostProcessSectionLabel", "Camera Post Process"));
+			DisplayClusterConfigurationStrings::categories::ViewPointCameraPostProcessCategory, LOCTEXT("ViewPointCameraPostProcessSectionLabel", "Camera Settings"));
 		Section->AddCategory(DisplayClusterConfigurationStrings::categories::ViewPointCameraPostProcessCategory);
 	}
 
