[nDisplay]  Fixed stereo rendering crash.

diff --git a/Engine/Plugins/Runtime/nDisplay/Source/DisplayCluster/Private/Render/Device/DisplayClusterDeviceBase.cpp b/Engine/Plugins/Runtime/nDisplay/Source/DisplayCluster/Private/Render/Device/DisplayClusterDeviceBase.cpp
--- a/Engine/Plugins/Runtime/nDisplay/Source/DisplayCluster/Private/Render/Device/DisplayClusterDeviceBase.cpp
+++ b/Engine/Plugins/Runtime/nDisplay/Source/DisplayCluster/Private/Render/Device/DisplayClusterDeviceBase.cpp
@@ -40,6 +40,9 @@
 
 #include "Render/Viewport/Configuration/DisplayClusterViewportConfigurationHelpers_Postprocess.h"
 
+#include "RenderGraphUtils.h"
+#include "RenderGraphBuilder.h"
+
 #include <utility>
 
 namespace UE::DisplayCluster::DeviceBaseHelpers
@@ -392,17 +395,10 @@ void FDisplayClusterDeviceBase::FinalizeNewFrame()
 
 DECLARE_GPU_STAT_NAMED(nDisplay_Device_RenderTexture, TEXT("nDisplay RenderDevice::RenderTexture"));
 
-#include "RenderGraphUtils.h"
-#include "RenderGraphBuilder.h"
-
-BEGIN_SHADER_PARAMETER_STRUCT(FResolveFrameTargetToBackBufferPass, )
-	RDG_TEXTURE_ACCESS(DstTexture, ERHIAccess::CopyDest)
-END_SHADER_PARAMETER_STRUCT()
-
 void FDisplayClusterDeviceBase::RenderTexture_RenderThread(class FRDGBuilder& GraphBuilder, FRDGTextureRef BackBuffer, FRDGTextureRef SrcTexture, FVector2f WindowSize) const
 {
-	const FIntPoint SrcSize = SrcTexture->GetRHI()->GetSizeXY();
-	const FIntPoint DstSize = BackBuffer->GetRHI()->GetSizeXY();
+	const FIntVector SrcSize = SrcTexture->Desc.GetSize();
+	const FIntVector DstSize = BackBuffer->Desc.GetSize();
 
 	FRHICopyTextureInfo CopyInfo;
 	CopyInfo.Size.X = FMath::Min(SrcSize.X, DstSize.X);
@@ -410,22 +406,14 @@ void FDisplayClusterDeviceBase::RenderTexture_RenderThread(class FRDGBuilder& Gr
 
 	AddCopyTexturePass(GraphBuilder, SrcTexture, BackBuffer, CopyInfo);
 
-	// TODO - Perform this with "add pass" type functions or otherwise make GraphBuilder of aware of transitions
-	FResolveFrameTargetToBackBufferPass* Params = GraphBuilder.AllocParameters<FResolveFrameTargetToBackBufferPass>();
-	Params->DstTexture = BackBuffer;
-
-	GraphBuilder.AddPass(RDG_EVENT_NAME("nDisplay_Device_RenderTexture"), Params, ERDGPassFlags::Copy,
-		[Params, WindowSize, this](FRHICommandListImmediate& RHICmdList)
+	if (RenderFrameMode == EDisplayClusterRenderFrameMode::Stereo)
+	{
+		if (FDisplayClusterViewportManagerProxy* ViewportManagerProxy = GetViewportManagerProxy_RenderThread())
 		{
-			if (RenderFrameMode == EDisplayClusterRenderFrameMode::Stereo)
-			{
-				if (FDisplayClusterViewportManagerProxy* ViewportManagerProxy = GetViewportManagerProxy_RenderThread())
-				{
-					// QuadBufStereo: Copy RIGHT_EYE to backbuffer
-					ViewportManagerProxy->ResolveFrameTargetToBackBuffer_RenderThread(RHICmdList, 1, 1, Params->DstTexture->GetRHI(), FVector2D(WindowSize));
-				}
-			}
-		});
+			// QuadBufStereo: Copy RIGHT_EYE to backbuffer
+			ViewportManagerProxy->ImplResolveFrameTargetToBackBuffer_RenderThread(GraphBuilder, 1, 1, BackBuffer, WindowSize);
+		}
+	}
 
 	const bool bClearTextureEnabled = CVarClearTextureEnabled.GetValueOnRenderThread() != 0;
 	if (bClearTextureEnabled)
diff --git a/Engine/Plugins/Runtime/nDisplay/Source/DisplayCluster/Private/Render/Viewport/DisplayClusterViewportHelpers.cpp b/Engine/Plugins/Runtime/nDisplay/Source/DisplayCluster/Private/Render/Viewport/DisplayClusterViewportHelpers.cpp
--- a/Engine/Plugins/Runtime/nDisplay/Source/DisplayCluster/Private/Render/Viewport/DisplayClusterViewportHelpers.cpp
+++ b/Engine/Plugins/Runtime/nDisplay/Source/DisplayCluster/Private/Render/Viewport/DisplayClusterViewportHelpers.cpp
@@ -226,22 +226,14 @@ EPixelFormat FDisplayClusterViewportHelpers::GetDefaultPixelFormat()
 }
 
 bool FDisplayClusterViewportHelpers::GetValidResourceRectsForResolve(
-	const FRHITexture* InSourceTexture,
-	const FRHITexture* InDestTexture,
+	const FIntPoint& InputResourceSize,
+	const FIntPoint& OutputResourceSize,
 	FIntRect& InOutSourceRect,
 	FIntRect& InOutDestRect)
 {
-	if (!InSourceTexture || !InDestTexture)
-	{
-		return false;
-	}
-
 	FIntRect SrcRect(InOutSourceRect);
 	FIntRect DestRect(InOutDestRect);
 
-	const FIntPoint& InputResourceSize = InSourceTexture->GetDesc().Extent;
-	const FIntPoint& OutputResourceSize = InDestTexture->GetDesc().Extent;
-
 	// If SrcRect or DestRect exceeds the texture size, RHI will crash. Let's adjust it to the texture size.
 	{
 		SrcRect.Min.X = FMath::Clamp(SrcRect.Min.X, 0, InputResourceSize.X);
@@ -265,15 +257,24 @@ bool FDisplayClusterViewportHelpers::GetValidResourceRectsForResolve(
 		FMath::Max(0, -InOutDestRect.Min.X),
 		FMath::Max(0, -InOutDestRect.Min.Y));
 
-	// SrcRect.Min and DestRect.Min always > 0
-
 	// Check the SrcRect and DestRect
+	// SrcRect.Min and DestRect.Min should be always > 0
 	if (SrcRect.Size().GetMin() <= 0 || DestRect.Size().GetMin() <= 0)
 	{
 		// The SrcRect or DestRect is invalid.
 		return false;
 	}
 
+	// if the SrcRect.Max>DestRect.Max, truncate
+	if (SrcRect.Width() > DestRect.Width())
+	{
+		SrcRect.Max.X = SrcRect.Min.X + DestRect.Width();
+	}
+	if (SrcRect.Height() > DestRect.Height())
+	{
+		SrcRect.Max.Y = SrcRect.Min.Y + DestRect.Height();
+	}
+
 	// Can be resolved
 	InOutSourceRect = SrcRect;
 	InOutDestRect = DestRect;
diff --git a/Engine/Plugins/Runtime/nDisplay/Source/DisplayCluster/Private/Render/Viewport/DisplayClusterViewportHelpers.h b/Engine/Plugins/Runtime/nDisplay/Source/DisplayCluster/Private/Render/Viewport/DisplayClusterViewportHelpers.h
--- a/Engine/Plugins/Runtime/nDisplay/Source/DisplayCluster/Private/Render/Viewport/DisplayClusterViewportHelpers.h
+++ b/Engine/Plugins/Runtime/nDisplay/Source/DisplayCluster/Private/Render/Viewport/DisplayClusterViewportHelpers.h
@@ -4,10 +4,6 @@
 #include "CoreMinimal.h"
 #include "PixelFormat.h"
 
-class FDisplayClusterViewportResource;
-class FRHITexture;
-struct FDisplayClusterRenderFrameSettings;
-
 /**
  * Helper class for DC viewports
  */
@@ -67,7 +63,7 @@ public:
 	* 
 	* @return - the mips value to be used.
 	*/
-	static int32 GetMaxTextureNumMips(const FDisplayClusterRenderFrameSettings& InRenderFrameSettings, const int32 InNumMips);
+	static int32 GetMaxTextureNumMips(const struct FDisplayClusterRenderFrameSettings& InRenderFrameSettings, const int32 InNumMips);
 
 	/**
 	 * Getting the default pixel format for preview rendering
@@ -84,16 +80,16 @@ public:
 	* If any rect exceeds the texture size, RHI will crash.
 	* This function adjusts the rects to the size of the textures.
 	* 
-	* @param InSourceTexture - (in) src RHI texture
-	* @param InDestTexture   - (in) dest RHI texture
+	* @param InSourceTextureExtent - (in) Source RHI texture size
+	* @param InDestTextureExtent   - (in) Destination RHI texture size
 	* @param InOutSourceRect - (in, out) src rect
 	* @param InOutDestRect   - (in, out) dest rect
 	* 
 	* @return false, if it isn't possible
 	*/
 	static bool GetValidResourceRectsForResolve(
-		const FRHITexture* InSourceTexture,
-		const FRHITexture* InDestTexture,
+		const FIntPoint& InSourceTextureExtent,
+		const FIntPoint& InDestTextureExtent,
 		FIntRect& InOutSourceRect,
 		FIntRect& InOutDestRect);
 };
diff --git a/Engine/Plugins/Runtime/nDisplay/Source/DisplayCluster/Private/Render/Viewport/DisplayClusterViewportManagerProxy.cpp b/Engine/Plugins/Runtime/nDisplay/Source/DisplayCluster/Private/Render/Viewport/DisplayClusterViewportManagerProxy.cpp
--- a/Engine/Plugins/Runtime/nDisplay/Source/DisplayCluster/Private/Render/Viewport/DisplayClusterViewportManagerProxy.cpp
+++ b/Engine/Plugins/Runtime/nDisplay/Source/DisplayCluster/Private/Render/Viewport/DisplayClusterViewportManagerProxy.cpp
@@ -568,10 +568,86 @@ bool FDisplayClusterViewportManagerProxy::GetFrameTargets_RenderThread(TArray<FR
 	return false;
 }
 
+bool FDisplayClusterViewportManagerProxy::ImplResolveFrameTargetToBackBuffer_RenderThread(FRDGBuilder& GraphBuilder, const uint32 InContextNum, const int32 DestArrayIndex, FRDGTextureRef DestTexture, FVector2f WindowSize) const
+{
+	check(IsInRenderingThread());
+
+	if (!DestTexture)
+	{
+		return false;
+	}
+
+	// Output texture must supports quadbuffer stereo.
+	if (DestArrayIndex > 0)
+	{
+		// Check if the backbuffer texture has a second slice for stereo.
+		const bool bStereoTexture = DestTexture->Desc.Dimension == ETextureDimension::Texture2DArray && DestTexture->Desc.ArraySize > 1;
+		if (!bStereoTexture)
+		{
+			return false;
+		}
+	}
+
+	TArray<FRHITexture*>   FrameResources;
+	TArray<FIntPoint>      TargetOffsets;
+	if (GetFrameTargets_RenderThread(FrameResources, TargetOffsets) && FrameResources.IsValidIndex(InContextNum))
+	{
+		// Use internal frame textures as source
+		if (FRHITexture* FrameTexture = FrameResources[InContextNum])
+		{
+			FIntRect SrcRect(FIntPoint::ZeroValue, FrameTexture->GetDesc().Extent);
+
+			const FIntPoint& DestOffset = TargetOffsets[InContextNum];
+			FIntRect DestRect(DestOffset, DestOffset + FrameTexture->GetDesc().Extent);
+
+			// Check if resources with the specified regions can be resolved.
+			if (!FDisplayClusterViewportHelpers::GetValidResourceRectsForResolve(FrameTexture->GetDesc().Extent, DestTexture->Desc.Extent, SrcRect, DestRect))
+			{
+				// The SrcRect or DestRect is invalid.
+				return false;
+			}
+
+			FRHICopyTextureInfo CopyInfo;
+
+			CopyInfo.SourceSliceIndex = 0;
+			CopyInfo.DestSliceIndex = DestArrayIndex;
+
+			CopyInfo.SourcePosition = FIntVector(SrcRect.Min.X, SrcRect.Min.Y, 0);
+			CopyInfo.DestPosition = FIntVector(DestRect.Min.X, DestRect.Min.Y, 0);
+
+			CopyInfo.Size = FIntVector(DestRect.Width(), DestRect.Height(), 0);
+
+			FRDGTextureRef SrcTexture = RegisterExternalTexture(GraphBuilder, FrameTexture->GetTexture2D(), TEXT("DisplayClusterStereoTexture"));
+
+			AddCopyTexturePass(GraphBuilder, SrcTexture, DestTexture, CopyInfo);
+
+			return true;
+		}
+	}
+
+	return false;
+}
+
 bool FDisplayClusterViewportManagerProxy::ResolveFrameTargetToBackBuffer_RenderThread(FRHICommandListImmediate& RHICmdList, const uint32 InContextNum, const int32 DestArrayIndex, FRHITexture* DestTexture, FVector2D WindowSize) const
 {
 	check(IsInRenderingThread());
 
+	if (!DestTexture)
+	{
+		return false;
+	}
+
+	// Output texture must supports quadbuffer stereo.
+	if (DestArrayIndex > 0)
+	{
+		// Check if the backbuffer texture has a second slice for stereo.
+		const bool bStereoTexture = DestTexture->GetDesc().Dimension == ETextureDimension::Texture2DArray && DestTexture->GetDesc().ArraySize > 1;
+		if (!bStereoTexture)
+		{
+			return false;
+		}
+	}
+
 	TArray<FRHITexture*>   FrameResources;
 	TArray<FIntPoint>        TargetOffsets;
 	if (GetFrameTargets_RenderThread(FrameResources, TargetOffsets) && FrameResources.IsValidIndex(InContextNum))
@@ -585,7 +661,7 @@ bool FDisplayClusterViewportManagerProxy::ResolveFrameTargetToBackBuffer_RenderT
 			FIntRect DestRect(DestOffset, DestOffset + FrameTexture->GetDesc().Extent);
 
 			// Check if resources with the specified regions can be resolved.
-			if (!FDisplayClusterViewportHelpers::GetValidResourceRectsForResolve(FrameTexture, DestTexture, SrcRect, DestRect))
+			if (!FDisplayClusterViewportHelpers::GetValidResourceRectsForResolve(FrameTexture->GetDesc().Extent, DestTexture->GetDesc().Extent, SrcRect, DestRect))
 			{
 				// The SrcRect or DestRect is invalid.
 				return false;
diff --git a/Engine/Plugins/Runtime/nDisplay/Source/DisplayCluster/Private/Render/Viewport/DisplayClusterViewportManagerProxy.h b/Engine/Plugins/Runtime/nDisplay/Source/DisplayCluster/Private/Render/Viewport/DisplayClusterViewportManagerProxy.h
--- a/Engine/Plugins/Runtime/nDisplay/Source/DisplayCluster/Private/Render/Viewport/DisplayClusterViewportManagerProxy.h
+++ b/Engine/Plugins/Runtime/nDisplay/Source/DisplayCluster/Private/Render/Viewport/DisplayClusterViewportManagerProxy.h
@@ -8,6 +8,10 @@
 #include "Render/Viewport/DisplayClusterViewportProxy.h"
 #include "Render/Viewport/Containers/DisplayClusterViewport_Enums.h"
 #include "Render/Viewport/RenderFrame/DisplayClusterRenderFrameSettings.h"
+
+#include "RenderGraphUtils.h"
+#include "RenderGraphBuilder.h"
+
 #include "Templates/SharedPointer.h"
 
 class FDisplayClusterRenderTargetManager;
@@ -68,6 +72,19 @@ public:
 	virtual bool ResolveFrameTargetToBackBuffer_RenderThread(FRHICommandListImmediate& RHICmdList, const uint32 InContextNum, const int32 DestArrayIndex, FRHITexture* DestTexture, FVector2D WindowSize) const override;
 	//~~IDisplayClusterViewportManagerProxy
 
+	/**
+	* Resolve to backbuffer
+	* [RDG version]
+	* 
+	* @param GraphBuilder   - Render Graph interface
+	* @param InContextNum   - renderframe source context num
+	* @param DestArrayIndex - dest array index on backbuffer
+	* @param DestTexture    - dest RDG texture
+	* @param WindowSize     - dest backbuffer window size
+	*
+	* @return - true if success
+	*/
+	bool ImplResolveFrameTargetToBackBuffer_RenderThread(FRDGBuilder& GraphBuilder, const uint32 InContextNum, const int32 DestArrayIndex, FRDGTextureRef DestTexture, FVector2f WindowSize) const;
 
 	/** Release all textures on render thread. */
 	void ReleaseTextures_RenderThread();
