nDisplay: Remove dependency on Composure.

The entire plugin was being pulled only to calculate a projection matrix.

diff --git a/Engine/Plugins/Runtime/nDisplay/Source/DisplayClusterProjection/DisplayClusterProjection.Build.cs b/Engine/Plugins/Runtime/nDisplay/Source/DisplayClusterProjection/DisplayClusterProjection.Build.cs
--- a/Engine/Plugins/Runtime/nDisplay/Source/DisplayClusterProjection/DisplayClusterProjection.Build.cs
+++ b/Engine/Plugins/Runtime/nDisplay/Source/DisplayClusterProjection/DisplayClusterProjection.Build.cs
@@ -28,7 +28,6 @@ public class DisplayClusterProjection : ModuleRules
 
 		PrivateDependencyModuleNames.AddRange(
 			new string[] {
-				"Composure",
 				"DisplayCluster",
 				"ProceduralMeshComponent",
 				"Projects",
diff --git a/Engine/Plugins/Runtime/nDisplay/Source/DisplayClusterProjection/Private/Policy/Camera/DisplayClusterProjectionCameraPolicy.cpp b/Engine/Plugins/Runtime/nDisplay/Source/DisplayClusterProjection/Private/Policy/Camera/DisplayClusterProjectionCameraPolicy.cpp
--- a/Engine/Plugins/Runtime/nDisplay/Source/DisplayClusterProjection/Private/Policy/Camera/DisplayClusterProjectionCameraPolicy.cpp
+++ b/Engine/Plugins/Runtime/nDisplay/Source/DisplayClusterProjection/Private/Policy/Camera/DisplayClusterProjectionCameraPolicy.cpp
@@ -10,7 +10,7 @@
 #include "Engine/World.h"
 #include "Camera/CameraComponent.h"
 #include "CineCameraComponent.h"
-#include "ComposurePostMoves.h"
+#include "Math/PerspectiveMatrix.h"
 
 #include "Render/Viewport/IDisplayClusterViewport.h"
 #include "Render/Viewport/IDisplayClusterViewportManager.h"
@@ -122,8 +122,16 @@ bool FDisplayClusterProjectionCameraPolicy::GetProjectionMatrix(IDisplayClusterV
 	}
 	else
 	{
-		FComposurePostMoveSettings ComposureSettings;
-		OutPrjMatrix = ComposureSettings.GetProjectionMatrix(ClampedCameraFOVDegrees, CameraAspectRatio);
+		const float HalfFOV = 0.5 * FMath::DegreesToRadians(ClampedCameraFOVDegrees);
+
+		OutPrjMatrix = FReversedZPerspectiveMatrix(
+			HalfFOV,
+			HalfFOV,
+			/* XAxisMultiplier = */ 1.0,
+			/* YAxisMultiplier = */ CameraAspectRatio,
+			GNearClippingPlane,
+			GNearClippingPlane
+		);
 	}
 
 	if (!CameraSettings.OffCenterProjectionOffset.IsZero())
diff --git a/Engine/Plugins/Runtime/nDisplay/nDisplay.uplugin b/Engine/Plugins/Runtime/nDisplay/nDisplay.uplugin
--- a/Engine/Plugins/Runtime/nDisplay/nDisplay.uplugin
+++ b/Engine/Plugins/Runtime/nDisplay/nDisplay.uplugin
@@ -260,10 +260,6 @@
 			"Name": "ColorGrading",
 			"Enabled": true
 		},
-		{
-			"Name": "Composure",
-			"Enabled": true
-		},
 		{
 			"Name": "ConcertSyncClient",
 			"Enabled": true
