[nDisplay] Fixed nDisplay EasyBlend projection policy DX12 dynamic viewpoint stuttering by replacing Easyblend's render queues with commands
Based on https://github.com/EpicGames/UnrealEngine/pull/13231

diff --git a/Engine/Plugins/Runtime/nDisplay/Source/DisplayClusterProjection/Private/Policy/EasyBlend/Windows/DX12/DisplayClusterProjectionEasyBlendLibraryDX12.cpp b/Engine/Plugins/Runtime/nDisplay/Source/DisplayClusterProjection/Private/Policy/EasyBlend/Windows/DX12/DisplayClusterProjectionEasyBlendLibraryDX12.cpp
--- a/Engine/Plugins/Runtime/nDisplay/Source/DisplayClusterProjection/Private/Policy/EasyBlend/Windows/DX12/DisplayClusterProjectionEasyBlendLibraryDX12.cpp
+++ b/Engine/Plugins/Runtime/nDisplay/Source/DisplayClusterProjection/Private/Policy/EasyBlend/Windows/DX12/DisplayClusterProjectionEasyBlendLibraryDX12.cpp
@@ -36,8 +36,8 @@ bool FDisplayClusterProjectionEasyBlendLibraryDX12::InitializeDLL()
 		ImportWarpBlendFunc(EasyBlendSDK_SetEyepoint);
 		ImportWarpBlendFunc(EasyBlendSDK_GetHeadingPitchRoll);
 
-		ImportWarpBlendFunc(EasyBlendSDK_InitializeDX12_CommandQueue);
-		ImportWarpBlendFunc(EasyBlendSDK_TransformInputToOutputDX12_CommandQueue);
+		ImportWarpBlendFunc(EasyBlendSDK_InitializeDX12_Device);
+		ImportWarpBlendFunc(EasyBlendSDK_TransformInputToOutputDX12_CommandList);
 	}
 
 	// All EasyBlend  DLL functions are initialized successfully
@@ -56,8 +56,8 @@ void FDisplayClusterProjectionEasyBlendLibraryDX12::ReleaseDLL()
 	ReleaseWarpBlendFunc(EasyBlendSDK_SetEyepoint);
 	ReleaseWarpBlendFunc(EasyBlendSDK_GetHeadingPitchRoll);
 
-	ReleaseWarpBlendFunc(EasyBlendSDK_InitializeDX12_CommandQueue);
-	ReleaseWarpBlendFunc(EasyBlendSDK_TransformInputToOutputDX12_CommandQueue);
+	ReleaseWarpBlendFunc(EasyBlendSDK_InitializeDX12_Device);
+	ReleaseWarpBlendFunc(EasyBlendSDK_TransformInputToOutputDX12_CommandList);
 
 	if (EasyBlend_DLL_Handler)
 	{
diff --git a/Engine/Plugins/Runtime/nDisplay/Source/DisplayClusterProjection/Private/Policy/EasyBlend/Windows/DX12/DisplayClusterProjectionEasyBlendLibraryDX12.h b/Engine/Plugins/Runtime/nDisplay/Source/DisplayClusterProjection/Private/Policy/EasyBlend/Windows/DX12/DisplayClusterProjectionEasyBlendLibraryDX12.h
--- a/Engine/Plugins/Runtime/nDisplay/Source/DisplayClusterProjection/Private/Policy/EasyBlend/Windows/DX12/DisplayClusterProjectionEasyBlendLibraryDX12.h
+++ b/Engine/Plugins/Runtime/nDisplay/Source/DisplayClusterProjection/Private/Policy/EasyBlend/Windows/DX12/DisplayClusterProjectionEasyBlendLibraryDX12.h
@@ -35,8 +35,8 @@ public:
 	DeclareEasyBlendFunc(EasyBlendSDK_SetEyepoint);
 	DeclareEasyBlendFunc(EasyBlendSDK_GetHeadingPitchRoll);
 
-	DeclareEasyBlendFunc(EasyBlendSDK_InitializeDX12_CommandQueue);
-	DeclareEasyBlendFunc(EasyBlendSDK_TransformInputToOutputDX12_CommandQueue);
+	DeclareEasyBlendFunc(EasyBlendSDK_InitializeDX12_Device);
+	DeclareEasyBlendFunc(EasyBlendSDK_TransformInputToOutputDX12_CommandList);
 
 private:
 	/** Imprort all functionf from EasyBlend DLL. Return true on success. */
diff --git a/Engine/Plugins/Runtime/nDisplay/Source/DisplayClusterProjection/Private/Policy/EasyBlend/Windows/DX12/DisplayClusterProjectionEasyBlendLibraryDX12Types.h b/Engine/Plugins/Runtime/nDisplay/Source/DisplayClusterProjection/Private/Policy/EasyBlend/Windows/DX12/DisplayClusterProjectionEasyBlendLibraryDX12Types.h
--- a/Engine/Plugins/Runtime/nDisplay/Source/DisplayClusterProjection/Private/Policy/EasyBlend/Windows/DX12/DisplayClusterProjectionEasyBlendLibraryDX12Types.h
+++ b/Engine/Plugins/Runtime/nDisplay/Source/DisplayClusterProjection/Private/Policy/EasyBlend/Windows/DX12/DisplayClusterProjectionEasyBlendLibraryDX12Types.h
@@ -42,9 +42,9 @@ typedef EasyBlendSDKError(__stdcall* EasyBlendSDK_UninitializeProc)(EasyBlendSDK
 * Initialize the DX12 renderer
 * 
 * @param msm           - EasyBlend SDK mesh initialized using EasyBlendSDK_SDKMODE_DX12
-* @param pCommandQueue - cmd queue
+* @param pDevice       - D3D12 device object pointer
 */
-typedef EasyBlendSDKError(__stdcall* EasyBlendSDK_InitializeDX12_CommandQueueProc)(EasyBlendSDK_Mesh* msm, ID3D12CommandQueue* pCommandQueue);
+typedef EasyBlendSDKError(__stdcall* EasyBlendSDK_InitializeDX12_DeviceProc)(EasyBlendSDK_Mesh* msm, ID3D12Device* pDevice);
 
 /**
 * Fill a given command list with commands that will warp the image on the back buffer and 
@@ -57,7 +57,7 @@ typedef EasyBlendSDKError(__stdcall* EasyBlendSDK_InitializeDX12_CommandQueuePro
 * @param msm       - EasyBlend SDK mesh initialized using EasyBlendSDK_SDKMODE_DX12
 * @param pCmdQueue - D3D12 command list that is ready to accept commands
 */
-typedef EasyBlendSDKError(__stdcall* EasyBlendSDK_TransformInputToOutputDX12_CommandQueueProc)(EasyBlendSDK_Mesh* msm, ID3D12Resource* InputTexture, ID3D12Resource* OutputTexture, D3D12_CPU_DESCRIPTOR_HANDLE OutputRTV);
+typedef EasyBlendSDKError(__stdcall* EasyBlendSDK_TransformInputToOutputDX12_CommandListProc)(EasyBlendSDK_Mesh* msm, ID3D12Resource* InputTexture, ID3D12Resource* OutputTexture, D3D12_CPU_DESCRIPTOR_HANDLE OutputRTV, ID3D12GraphicsCommandList* pCommandList);
 
 /**
 * This command returns the view angles for the instance of the SDK.
diff --git a/Engine/Plugins/Runtime/nDisplay/Source/DisplayClusterProjection/Private/Policy/EasyBlend/Windows/DX12/DisplayClusterProjectionEasyBlendPolicyViewDataDX12.cpp b/Engine/Plugins/Runtime/nDisplay/Source/DisplayClusterProjection/Private/Policy/EasyBlend/Windows/DX12/DisplayClusterProjectionEasyBlendPolicyViewDataDX12.cpp
--- a/Engine/Plugins/Runtime/nDisplay/Source/DisplayClusterProjection/Private/Policy/EasyBlend/Windows/DX12/DisplayClusterProjectionEasyBlendPolicyViewDataDX12.cpp
+++ b/Engine/Plugins/Runtime/nDisplay/Source/DisplayClusterProjection/Private/Policy/EasyBlend/Windows/DX12/DisplayClusterProjectionEasyBlendPolicyViewDataDX12.cpp
@@ -159,11 +159,16 @@ bool FDisplayClusterProjectionEasyBlendPolicyViewDataDX12::ApplyWarpBlend_Render
 		FScopeLock DataScopeLock(&ViewData->EasyBlendMeshDataAccessCS);
 		TSharedRef<FDisplayClusterProjectionEasyBlendLibraryDX12, ESPMode::ThreadSafe>& EasyBlendAPI = FDisplayClusterProjectionEasyBlendLibraryDX12::Get();
 
-		// Always use device 0 for EasyBlend
-		const uint32 DevIndex = 0;
 		ID3D12DynamicRHI* D3D12DynamicRHI = GetID3D12DynamicRHI();
-		ID3D12CommandQueue* D3D12CommandQueue = D3D12DynamicRHI ? D3D12DynamicRHI->RHIGetCommandQueue() : nullptr;
-		if (!D3D12CommandQueue)
+		if (!D3D12DynamicRHI)
+		{
+			return;
+		}
+
+		// Ensure EasyBlend uses the same GPU device as the source texture to support multi-GPU configurations
+		const uint32 DevIndex = D3D12DynamicRHI->RHIGetResourceDeviceIndex(InputTexture);
+		ID3D12Device* D3D12Device = D3D12DynamicRHI->RHIGetDevice(DevIndex);
+		if (!D3D12Device)
 		{
 			return;
 		}
@@ -174,7 +179,7 @@ bool FDisplayClusterProjectionEasyBlendPolicyViewDataDX12::ApplyWarpBlend_Render
 		// Initialize EasyBlend internals
 		if (!ViewData->bIsRenderResourcesInitialized)
 		{
-			const EasyBlendSDKError Result = EasyBlendAPI->EasyBlendSDK_InitializeDX12_CommandQueue(ViewData->EasyBlendMeshData.Get(), D3D12CommandQueue);
+			const EasyBlendSDKError Result = EasyBlendAPI->EasyBlendSDK_InitializeDX12_Device(ViewData->EasyBlendMeshData.Get(), D3D12Device);
 			if (!EasyBlendSDK_SUCCEEDED(Result))
 			{
 				ViewData->bIsRenderResourcesInitialized = false;
@@ -188,11 +193,12 @@ bool FDisplayClusterProjectionEasyBlendPolicyViewDataDX12::ApplyWarpBlend_Render
 		ID3D12Resource* SrcTexture = D3D12DynamicRHI->RHIGetResource(InputTexture);
 		ID3D12Resource* DestTexture = D3D12DynamicRHI->RHIGetResource(OutputTexture);
 		const D3D12_CPU_DESCRIPTOR_HANDLE RTVHandle = D3D12DynamicRHI->RHIGetRenderTargetView(OutputTexture);
+		ID3D12GraphicsCommandList* D3D12GraphicsCommandList = D3D12DynamicRHI->RHIGetGraphicsCommandList(ExecutingCmdList, DevIndex);
 
-		if (SrcTexture && DestTexture && D3D12CommandQueue)
+		if (SrcTexture && DestTexture && D3D12GraphicsCommandList)
 		{
 			// Render EasyBlend on proxy
-			const EasyBlendSDKError Result2 = EasyBlendAPI->EasyBlendSDK_TransformInputToOutputDX12_CommandQueue(ViewData->EasyBlendMeshData.Get(), SrcTexture, DestTexture, RTVHandle);
+			const EasyBlendSDKError Result2 = EasyBlendAPI->EasyBlendSDK_TransformInputToOutputDX12_CommandList(ViewData->EasyBlendMeshData.Get(), SrcTexture, DestTexture, RTVHandle, D3D12GraphicsCommandList);
 		}
 	});
 	RHICmdList.Transition(FRHITransitionInfo(OutputTexture, ERHIAccess::RTV, ERHIAccess::Present));
