nDisplay: Initial fix to a crash when stage app connects.

Initial fix to a crash due to the GraphBuilder's FRDGBlackboard getting two instances of the same struct types.
Only the preview scene with the nDisplay actor will be rendered without the lightcard proxies overlaid on top, which means that they won't be visible outside the stage meshes.

#rnx

diff --git a/Engine/Plugins/Runtime/nDisplay/Source/DisplayClusterLightCardEditorShaders/Private/DisplayClusterMeshProjectionRenderer.cpp b/Engine/Plugins/Runtime/nDisplay/Source/DisplayClusterLightCardEditorShaders/Private/DisplayClusterMeshProjectionRenderer.cpp
--- a/Engine/Plugins/Runtime/nDisplay/Source/DisplayClusterLightCardEditorShaders/Private/DisplayClusterMeshProjectionRenderer.cpp
+++ b/Engine/Plugins/Runtime/nDisplay/Source/DisplayClusterLightCardEditorShaders/Private/DisplayClusterMeshProjectionRenderer.cpp
@@ -1811,10 +1811,8 @@ void FDisplayClusterMeshProjectionRenderer::RenderScenes(FCanvas* Canvas, const
 {
 	Canvas->Flush_GameThread();
 
-	/**
-	* This code is based on the existing FDisplayClusterMeshProjectionRenderer::Render() function.
-	* The idea is to create code similar to the existing code.
-	*/
+	// This code is based on the existing FDisplayClusterMeshProjectionRenderer::Render() function.
+
 	FRenderTarget* RenderTarget = Canvas->GetRenderTarget();
 	const bool bIsHitTesting = Canvas->IsHitTesting();
 	FHitProxyConsumer* HitProxyConsumer = Canvas->GetHitProxyConsumer();
@@ -1860,6 +1858,11 @@ void FDisplayClusterMeshProjectionRenderer::RenderScenes(FCanvas* Canvas, const
 
 				GetRendererModule().CreateAndInitSingleView(RHICmdList, &ViewFamily, &NewInitOptions);
 				SceneViews.Add(ViewFamily.Views[0]);
+
+				// Initial fix to a crash due to the GraphBuilder's FRDGBlackboard getting two instances of the same struct types.
+				// Only the first scene (PreviewWorld->Scene) will be rendered, which implies that the light cards / stage actors 
+				// won't be visible outside the stage meshes. This code might need to be updated to that composites.
+				break;
 			}
 
 			FRDGTextureRef OutputTexture = GraphBuilder.RegisterExternalTexture(CreateRenderTarget(RenderTarget->GetRenderTargetTexture(), TEXT("ViewRenderTarget")));
