[nDisplay]  Disable TSR for preview in scene
- added a new DCRootActor property `bPreviewEnableTSR`

Reworked rules for ViewState.
-Addeed a new function FDisplayClusterViewport::ShouldUseViewStates()

diff --git a/Engine/Plugins/Runtime/nDisplay/Source/DisplayCluster/Private/Game/EngineClasses/Scene/DisplayClusterRootActor.cpp b/Engine/Plugins/Runtime/nDisplay/Source/DisplayCluster/Private/Game/EngineClasses/Scene/DisplayClusterRootActor.cpp
--- a/Engine/Plugins/Runtime/nDisplay/Source/DisplayCluster/Private/Game/EngineClasses/Scene/DisplayClusterRootActor.cpp
+++ b/Engine/Plugins/Runtime/nDisplay/Source/DisplayCluster/Private/Game/EngineClasses/Scene/DisplayClusterRootActor.cpp
@@ -1357,6 +1357,7 @@ FDisplayClusterViewport_PreviewSettings ADisplayClusterRootActor::GetPreviewSett
 	OutPreviewSettings.bPreviewInGameRenderFrustum = bPreviewInGameRenderFrustum;
 
 	OutPreviewSettings.bPreviewEnableHoldoutComposite = bPreviewEnableHoldoutComposite;
+	OutPreviewSettings.bPreviewEnableTSR = bPreviewEnableTSR;
 
 	return OutPreviewSettings;
 }
@@ -1634,6 +1635,7 @@ void ADisplayClusterRootActor::OverrideRootActorProperties(const FDisplayCluster
 	OVERRIDE_PROPERTY(bPreviewEnablePostProcess);
 	OVERRIDE_PROPERTY(bPreviewEnableOCIO);
 	OVERRIDE_PROPERTY(bPreviewEnableHoldoutComposite);
+	OVERRIDE_PROPERTY(bPreviewEnableTSR);
 	OVERRIDE_PROPERTY(bPreviewEnableOverlayMaterial);
 	OVERRIDE_PROPERTY(bEnablePreviewTechvis);
 	OVERRIDE_PROPERTY(bEnablePreviewMesh);
diff --git a/Engine/Plugins/Runtime/nDisplay/Source/DisplayCluster/Private/Render/Upscaler/DisplayClusterUpscaler.cpp b/Engine/Plugins/Runtime/nDisplay/Source/DisplayCluster/Private/Render/Upscaler/DisplayClusterUpscaler.cpp
--- a/Engine/Plugins/Runtime/nDisplay/Source/DisplayCluster/Private/Render/Upscaler/DisplayClusterUpscaler.cpp
+++ b/Engine/Plugins/Runtime/nDisplay/Source/DisplayCluster/Private/Render/Upscaler/DisplayClusterUpscaler.cpp
@@ -1,9 +1,9 @@
 // Copyright Epic Games, Inc. All Rights Reserved.
 
 #include "Render/Upscaler/DisplayClusterUpscaler.h"
+#include "Render/Viewport/DisplayClusterViewport.h"
 
 #include "IUpscalerModularFeature.h"
-
 #include "Features/IModularFeatures.h"
 
 #include "LegacyScreenPercentageDriver.h"
@@ -202,6 +202,7 @@ namespace UE::DisplayCluster::Upscaler
 	* Get the antialiasing method that can be used for the ViewFamily.
 	* Note: the logic of this function is based on the FSceneView::SetupAntiAliasingMethod() function.
 	* 
+	* @param InViewport         - (in) the DC viewport
 	* @param InUpscalerSettings - (in) the upscaler settings that we want to use.
 	* @param InViewFamily       - (in) ViewFamily thyat own SceneView.
 	* @param InSceneView        - (in) SceneView to which the upscaler settings should be applied.
@@ -209,6 +210,7 @@ namespace UE::DisplayCluster::Upscaler
 	* @return AntiAliasing method that can be used.
 	*/
 	static EDisplayClusterUpscalerAntiAliasingMethod GetSuitableAntiAliasingMethod(
+		const FDisplayClusterViewport& InViewport,
 		const FDisplayClusterUpscalerSettings& InUpscalerSettings,
 		const FSceneViewFamily& InViewFamily,
 		const FSceneView& InSceneView)
@@ -276,6 +278,15 @@ namespace UE::DisplayCluster::Upscaler
 			{	
 				return EDisplayClusterUpscalerAntiAliasingMethod::FXAA;
 			}
+
+			// Disable TSR/TAA for preview
+			if (!InViewport.Configuration->GetPreviewSettings().bPreviewEnableTSR)
+			{
+				if (InViewport.Configuration->IsPreviewRendering())
+				{
+					return EDisplayClusterUpscalerAntiAliasingMethod::FXAA;
+				}
+			}
 		}
 
 		return AntiAliasingMethod;
@@ -283,6 +294,7 @@ namespace UE::DisplayCluster::Upscaler
 };
 
 void FDisplayClusterUpscaler::SetupSceneView(
+	const FDisplayClusterViewport& InViewport,
 	const FDisplayClusterUpscalerSettings& InUpscalerSettings,
 	const FSceneViewFamily& InViewFamily,
 	FSceneView& InOutView)
@@ -290,7 +302,7 @@ void FDisplayClusterUpscaler::SetupSceneView(
 	using namespace UE::DisplayCluster::Upscaler;
 
 	// Get AntiAliasingMethod that can be used.
-	const EDisplayClusterUpscalerAntiAliasingMethod SuitableAntiAliasingMethod = GetSuitableAntiAliasingMethod(InUpscalerSettings, InViewFamily, InOutView);
+	const EDisplayClusterUpscalerAntiAliasingMethod SuitableAntiAliasingMethod = GetSuitableAntiAliasingMethod(InViewport, InUpscalerSettings, InViewFamily, InOutView);
 	switch (SuitableAntiAliasingMethod)
 	{
 	case EDisplayClusterUpscalerAntiAliasingMethod::None:
diff --git a/Engine/Plugins/Runtime/nDisplay/Source/DisplayCluster/Private/Render/Upscaler/DisplayClusterUpscaler.h b/Engine/Plugins/Runtime/nDisplay/Source/DisplayCluster/Private/Render/Upscaler/DisplayClusterUpscaler.h
--- a/Engine/Plugins/Runtime/nDisplay/Source/DisplayCluster/Private/Render/Upscaler/DisplayClusterUpscaler.h
+++ b/Engine/Plugins/Runtime/nDisplay/Source/DisplayCluster/Private/Render/Upscaler/DisplayClusterUpscaler.h
@@ -6,6 +6,7 @@
 
 #include "Render/Upscaler/DisplayClusterUpscalerSettings.h"
 
+class FDisplayClusterViewport;
 class FSceneViewFamilyContext;
 class FSceneViewFamily;
 class FSceneView;
@@ -49,11 +50,13 @@ public:
 
 	/**
 	* Configure scene view for upscalers.
+	* @param InViewport         - (in) the viewport who ask for this setup
 	* @param InUpscalerSettings - (in) the settings to apply for view
 	* @param InViewFamily       - (in) ViewFamily that own this view.
 	* @param InOutViews         - (in, out) The views that will be configured. 
 	*/
 	static void SetupSceneView(
+		const FDisplayClusterViewport& InViewport,
 		const FDisplayClusterUpscalerSettings& InUpscalerSettings,
 		const FSceneViewFamily& InViewFamily,
 		FSceneView& InOutView);
diff --git a/Engine/Plugins/Runtime/nDisplay/Source/DisplayCluster/Private/Render/Viewport/DisplayClusterViewport.cpp b/Engine/Plugins/Runtime/nDisplay/Source/DisplayCluster/Private/Render/Viewport/DisplayClusterViewport.cpp
--- a/Engine/Plugins/Runtime/nDisplay/Source/DisplayCluster/Private/Render/Viewport/DisplayClusterViewport.cpp
+++ b/Engine/Plugins/Runtime/nDisplay/Source/DisplayCluster/Private/Render/Viewport/DisplayClusterViewport.cpp
@@ -402,7 +402,7 @@ void FDisplayClusterViewport::SetupSceneView(uint32 ContextNum, class UWorld* Wo
 	}
 
 	// Handle upscalers
-	FDisplayClusterUpscaler::SetupSceneView(RenderSettings.UpscalerSettings, InOutViewFamily, InOutView);
+	FDisplayClusterUpscaler::SetupSceneView(*this, RenderSettings.UpscalerSettings, InOutViewFamily, InOutView);
 
 
 	if (InOutView.State && World && World->Scene)
@@ -527,6 +527,9 @@ bool FDisplayClusterViewport::UpdateFrameContexts(const uint32 InStereoViewIndex
 		Resources.FreezeRendering(EDisplayClusterViewportResource::AdditionalTargetableResources);
 		Resources.FreezeRendering(EDisplayClusterViewportResource::MipsShaderResources);
 
+		// Release ViewState resources if they are not used for rendering.
+		ViewStates.Empty();
+
 		// Update context links for freezed viewport
 		for (FDisplayClusterViewport_Context& ContextIt : Contexts)
 		{
diff --git a/Engine/Plugins/Runtime/nDisplay/Source/DisplayCluster/Private/Render/Viewport/DisplayClusterViewport.h b/Engine/Plugins/Runtime/nDisplay/Source/DisplayCluster/Private/Render/Viewport/DisplayClusterViewport.h
--- a/Engine/Plugins/Runtime/nDisplay/Source/DisplayCluster/Private/Render/Viewport/DisplayClusterViewport.h
+++ b/Engine/Plugins/Runtime/nDisplay/Source/DisplayCluster/Private/Render/Viewport/DisplayClusterViewport.h
@@ -339,6 +339,9 @@ public:
 	/** Return true if this viewport should freeze render. */
 	bool ShouldFreezeRender() const;
 
+	/** Return true if ViewState should be used by this viewport.*/
+	bool ShouldUseViewStates() const;
+
 	/** Returns true if textures in this viewport with an area larger than the MaxTextureArea limit
 	*  must be reduced to stay within that limit with the same aspect ratio.
 	*/
diff --git a/Engine/Plugins/Runtime/nDisplay/Source/DisplayCluster/Private/Render/Viewport/DisplayClusterViewport_EditorPreview.cpp b/Engine/Plugins/Runtime/nDisplay/Source/DisplayCluster/Private/Render/Viewport/DisplayClusterViewport_EditorPreview.cpp
--- a/Engine/Plugins/Runtime/nDisplay/Source/DisplayCluster/Private/Render/Viewport/DisplayClusterViewport_EditorPreview.cpp
+++ b/Engine/Plugins/Runtime/nDisplay/Source/DisplayCluster/Private/Render/Viewport/DisplayClusterViewport_EditorPreview.cpp
@@ -12,23 +12,6 @@
 
 #include "SceneManagement.h"
 
-///////////////////////////////////////////////////////////////////////////////////////
-int32 GDisplayClusterPreviewEnableViewState = 1;
-static FAutoConsoleVariableRef CVarDisplayClusterPreviewEnableViewState(
-	TEXT("nDisplay.preview.EnableViewState"),
-	GDisplayClusterPreviewEnableViewState,
-	TEXT("Enable view state for preview (0 - disable).\n"),
-	ECVF_RenderThreadSafe
-);
-
-int32 GDisplayClusterPreviewEnableConfiguratorViewState = 0;
-static FAutoConsoleVariableRef CVarDisplayClusterPreviewEnableConfiguratorViewState(
-	TEXT("nDisplay.preview.EnableConfiguratorViewState"),
-	GDisplayClusterPreviewEnableConfiguratorViewState,
-	TEXT("Enable view state for preview in Configurator window (0 - disable).\n"),
-	ECVF_RenderThreadSafe
-);
-
 ///////////////////////////////////////////////////////////////////////////////////////
 //          FDisplayClusterViewport
 ///////////////////////////////////////////////////////////////////////////////////////
@@ -79,9 +62,9 @@ void FDisplayClusterViewport::CleanupViewState()
 
 FSceneViewStateInterface* FDisplayClusterViewport::GetViewState(uint32 ViewIndex)
 {
-	if (GDisplayClusterPreviewEnableViewState == 0 || (GDisplayClusterPreviewEnableConfiguratorViewState == 0 && Configuration->IsCurrentWorldHasAnyType(EWorldType::EditorPreview)))
+	if(!ShouldUseViewStates())
 	{
-		// Disable ViewState
+		// Release ViewState resources if they are not used for rendering.
 		ViewStates.Empty();
 
 		return nullptr;
@@ -107,8 +90,4 @@ FSceneViewStateInterface* FDisplayClusterViewport::GetViewState(uint32 ViewIndex
 	}
 
 	return ViewStates[ViewIndex]->GetReference();
-}
-
-
-
-
+}
\ No newline at end of file
diff --git a/Engine/Plugins/Runtime/nDisplay/Source/DisplayCluster/Private/Render/Viewport/DisplayClusterViewport_Rules.cpp b/Engine/Plugins/Runtime/nDisplay/Source/DisplayCluster/Private/Render/Viewport/DisplayClusterViewport_Rules.cpp
--- a/Engine/Plugins/Runtime/nDisplay/Source/DisplayCluster/Private/Render/Viewport/DisplayClusterViewport_Rules.cpp
+++ b/Engine/Plugins/Runtime/nDisplay/Source/DisplayCluster/Private/Render/Viewport/DisplayClusterViewport_Rules.cpp
@@ -79,6 +79,23 @@ static FAutoConsoleVariableRef CVarDisplayClusterLumenPerView(
 	ECVF_RenderThreadSafe
 );
 
+int32 GDisplayClusterPreviewEnableViewState = 1;
+static FAutoConsoleVariableRef CVarDisplayClusterPreviewEnableViewState(
+	TEXT("nDisplay.preview.EnableViewState"),
+	GDisplayClusterPreviewEnableViewState,
+	TEXT("Enable view state for preview (0 - disable).\n"),
+	ECVF_RenderThreadSafe
+);
+
+int32 GDisplayClusterPreviewEnableConfiguratorViewState = 0;
+static FAutoConsoleVariableRef CVarDisplayClusterPreviewEnableConfiguratorViewState(
+	TEXT("nDisplay.preview.EnableConfiguratorViewState"),
+	GDisplayClusterPreviewEnableConfiguratorViewState,
+	TEXT("Enable view state for preview in Configurator window (0 - disable).\n"),
+	ECVF_RenderThreadSafe
+);
+
+
 ///////////////////////////////////////////////////////////////////////////////////////
 // FDisplayClusterViewport
 ///////////////////////////////////////////////////////////////////////////////////////
@@ -464,6 +481,65 @@ bool FDisplayClusterViewport::ShouldFreezeRender() const
 	return false;
 }
 
+bool FDisplayClusterViewport::ShouldUseViewStates() const
+{
+	// RenderMode rules
+	switch (Configuration->GetRenderFrameSettings().RenderMode)
+	{
+		// Cluster node rendering
+	case EDisplayClusterRenderFrameMode::Mono:
+	case EDisplayClusterRenderFrameMode::Stereo:
+	case EDisplayClusterRenderFrameMode::SideBySide:
+	case EDisplayClusterRenderFrameMode::TopBottom:
+		// These modes use ViewState from ULocalPLayer, outside of nDisplay.
+		return false;
+
+		// MRQ rendering
+	case EDisplayClusterRenderFrameMode::MRQ_Mono:
+		// These modes use ViewState from MRQ, outside of nDisplay.
+		return false;
+
+	// Preview-In-Scene
+	case EDisplayClusterRenderFrameMode::PreviewProxyHitInScene:
+	case EDisplayClusterRenderFrameMode::PreviewInScene:
+		if (GDisplayClusterPreviewEnableViewState == 0)
+		{
+			return false;
+		}
+
+		// Disable ViewState for the preview inside the Configurator.
+		if (GDisplayClusterPreviewEnableConfiguratorViewState == 0 && Configuration->IsCurrentWorldHasAnyType(EWorldType::EditorPreview))
+		{
+			return false;
+		}
+
+		break;
+
+		// PIE
+	case EDisplayClusterRenderFrameMode::PIE_Mono:
+	case EDisplayClusterRenderFrameMode::PIE_SideBySide:
+	case EDisplayClusterRenderFrameMode::PIE_TopBottom:
+		// We don't have any rules for PIE rendering
+		break;
+
+
+	default:
+		return false;
+	}
+
+	// Viewport type rules
+	switch (RenderSettings.CaptureMode)
+	{
+	case EDisplayClusterViewportCaptureMode::Chromakey:
+		return false;
+
+	default:
+		break;
+	}
+
+	return true;
+}
+
 bool FDisplayClusterViewport::IsOpenColorIOEquals(const FDisplayClusterViewport& InViewport) const
 {
 	bool bEnabledOCIO_1 = OpenColorIO.IsValid();
diff --git a/Engine/Plugins/Runtime/nDisplay/Source/DisplayCluster/Public/DisplayClusterRootActor.h b/Engine/Plugins/Runtime/nDisplay/Source/DisplayCluster/Public/DisplayClusterRootActor.h
--- a/Engine/Plugins/Runtime/nDisplay/Source/DisplayCluster/Public/DisplayClusterRootActor.h
+++ b/Engine/Plugins/Runtime/nDisplay/Source/DisplayCluster/Public/DisplayClusterRootActor.h
@@ -491,6 +491,12 @@ public:
 	UPROPERTY(EditAnywhere, BlueprintReadWrite, Category = "Preview", meta = (DisplayName = "Enable Preview Holdout", EditCondition = "bPreviewEnablePostProcess"), BlueprintSetter = SetPreviewEnableHoldoutComposite)
 	bool bPreviewEnableHoldoutComposite = true;
 
+	/** Enable TemporalAA/TSR in preview.
+	* Increases memory consumption. This will impact editor performance. 
+	*/
+	UPROPERTY(EditAnywhere, BlueprintReadWrite, Category = "Preview", meta = (DisplayName = "Allow TemporalAA/TSR in Preview"))
+	bool bPreviewEnableTSR = false;
+
 	/** Enable OCIO on preview. */
 	UPROPERTY(EditAnywhere, BlueprintReadWrite, Category = "Preview", meta = (DisplayName = "Enable OCIO on Preview"))
 	bool bPreviewEnableOCIO = false;
diff --git a/Engine/Plugins/Runtime/nDisplay/Source/DisplayCluster/Public/DisplayClusterRootActorContainers.h b/Engine/Plugins/Runtime/nDisplay/Source/DisplayCluster/Public/DisplayClusterRootActorContainers.h
--- a/Engine/Plugins/Runtime/nDisplay/Source/DisplayCluster/Public/DisplayClusterRootActorContainers.h
+++ b/Engine/Plugins/Runtime/nDisplay/Source/DisplayCluster/Public/DisplayClusterRootActorContainers.h
@@ -27,9 +27,15 @@ struct FDisplayClusterRootActorPropertyOverrides
 
 	/** Enable OCIO on preview. */
 	TOptional<bool> bPreviewEnableOCIO;
+
 	/** Enables the use of 'HoldoutComposite' plugin for the preview rendering. */
 	TOptional<bool> bPreviewEnableHoldoutComposite;
 
+	/** Enable TemporalAA/TSR in preview.
+	* Increases memory consumption. This will impact editor performance.
+	*/
+	TOptional<bool> bPreviewEnableTSR;
+
 	/** Show overlay material on the preview mesh when preview rendering is enabled (UMeshComponent::OverlayMaterial). */
 	TOptional<bool> bPreviewEnableOverlayMaterial;
 
diff --git a/Engine/Plugins/Runtime/nDisplay/Source/DisplayCluster/Public/Render/Viewport/Containers/DisplayClusterViewport_PreviewSettings.h b/Engine/Plugins/Runtime/nDisplay/Source/DisplayCluster/Public/Render/Viewport/Containers/DisplayClusterViewport_PreviewSettings.h
--- a/Engine/Plugins/Runtime/nDisplay/Source/DisplayCluster/Public/Render/Viewport/Containers/DisplayClusterViewport_PreviewSettings.h
+++ b/Engine/Plugins/Runtime/nDisplay/Source/DisplayCluster/Public/Render/Viewport/Containers/DisplayClusterViewport_PreviewSettings.h
@@ -25,6 +25,11 @@ struct FDisplayClusterViewport_PreviewSettings
 	/** Enables the use of 'HoldoutComposite' plugin for the preview rendering. */
 	bool bPreviewEnableHoldoutComposite = true;
 
+	/** Enable TemporalAA/TSR in preview.
+	* Increases memory consumption. This will impact editor performance.
+	*/
+	bool bPreviewEnableTSR = false;
+
 	// Preview uses techvis mode for rendering
 	bool bEnablePreviewTechvis = false;
 
