[nDisplay] Fixed tiled media issue when overscan + late OCIO are used

#rnx

diff --git a/Engine/Plugins/Runtime/nDisplay/Source/DisplayClusterMedia/Private/Capture/DisplayClusterMediaCaptureViewportBase.cpp b/Engine/Plugins/Runtime/nDisplay/Source/DisplayClusterMedia/Private/Capture/DisplayClusterMediaCaptureViewportBase.cpp
--- a/Engine/Plugins/Runtime/nDisplay/Source/DisplayClusterMedia/Private/Capture/DisplayClusterMediaCaptureViewportBase.cpp
+++ b/Engine/Plugins/Runtime/nDisplay/Source/DisplayClusterMedia/Private/Capture/DisplayClusterMediaCaptureViewportBase.cpp
@@ -236,7 +236,7 @@ void FDisplayClusterMediaCaptureViewportBase::OnPostRenderViewFamily_RenderThrea
 		TArray<FIntRect>     Regions;
 
 		// Get RHI texture and pass it to the media capture pipeline
-		if (ViewportProxy->GetResourcesWithRects_RenderThread(EDisplayClusterViewportResourceType::InternalRenderTargetResource, Textures, Regions))
+		if (ViewportProxy->GetResourcesWithRects_RenderThread(EDisplayClusterViewportResourceType::InternalRenderTargetEntireRectResource, Textures, Regions))
 		{
 			if (Textures.Num() > 0 && Regions.Num() > 0 && Textures[0])
 			{
diff --git a/Engine/Plugins/Runtime/nDisplay/Source/DisplayClusterMedia/Private/Input/DisplayClusterMediaInputBase.cpp b/Engine/Plugins/Runtime/nDisplay/Source/DisplayClusterMedia/Private/Input/DisplayClusterMediaInputBase.cpp
--- a/Engine/Plugins/Runtime/nDisplay/Source/DisplayClusterMedia/Private/Input/DisplayClusterMediaInputBase.cpp
+++ b/Engine/Plugins/Runtime/nDisplay/Source/DisplayClusterMedia/Private/Input/DisplayClusterMediaInputBase.cpp
@@ -240,7 +240,7 @@ void FDisplayClusterMediaInputBase::ImportMediaDataOCIO_RenderThread(FRHICommand
 		FRDGTextureRef TextureLinear = GraphBuilder.CreateTexture(InterimTextureDesc, TEXT("DC.MediaTextureTempLinear"));
 
 		// Add PQ-decode pass
-		FDisplayClusterShaderParameters_MediaPQ Parameters{ SrcTexture, InSrcRect, TextureLinear, InDstRect };
+		FDisplayClusterShaderParameters_MediaPQ Parameters{ SrcTexture, InSrcRect, TextureLinear, FIntRect{ FIntPoint::ZeroValue, InDstRect.Size() }};
 		IDisplayClusterShaders::Get().AddPQToLinearPass(GraphBuilder, Parameters);
 
 		// Use this texture in the OCIO pass
diff --git a/Engine/Plugins/Runtime/nDisplay/Source/DisplayClusterMedia/Private/Input/DisplayClusterMediaInputViewportBase.cpp b/Engine/Plugins/Runtime/nDisplay/Source/DisplayClusterMedia/Private/Input/DisplayClusterMediaInputViewportBase.cpp
--- a/Engine/Plugins/Runtime/nDisplay/Source/DisplayClusterMedia/Private/Input/DisplayClusterMediaInputViewportBase.cpp
+++ b/Engine/Plugins/Runtime/nDisplay/Source/DisplayClusterMedia/Private/Input/DisplayClusterMediaInputViewportBase.cpp
@@ -135,7 +135,8 @@ void FDisplayClusterMediaInputViewportBase::OnPostCrossGpuTransfer_RenderThread(
 			TArray<FRHITexture*> Textures;
 			TArray<FIntRect>     Regions;
 
-			if (PlaybackViewport->GetResourcesWithRects_RenderThread(EDisplayClusterViewportResourceType::InternalRenderTargetResource, Textures, Regions))
+			// Proceed with a proper texture resource
+			if (PlaybackViewport->GetResourcesWithRects_RenderThread(EDisplayClusterViewportResourceType::InternalRenderTargetEntireRectResource, Textures, Regions))
 			{
 				if (Textures.Num() > 0 && Regions.Num() > 0 && Textures[0])
 				{
