[nDisplay] Reworked IUpscalerModularFeature interface.

IUpscalerModularFeature
- Added new functions AddSceneViewExtensionIsActiveFunctors(),
RemoveSceneViewExtensionIsActiveFunctor()

FDisplayClusterSceneViewViewExtensionContext
- Added direct reference to the DC viewport instead of TSharedPtr<IDisplayClusterViewportConfiguration> + FString ViewportID
The functor logic on the nDisplay side will be based on the FDisplayClusterSceneViewViewExtensionContext and will require access to the DC viewport to retrieve data.
To avoid searching the viewport by name, the context will use a direct reference to the viewport object.

diff --git a/Engine/Plugins/Runtime/nDisplay/Source/DisplayCluster/Private/Render/Upscaler/DisplayClusterUpscaler.cpp b/Engine/Plugins/Runtime/nDisplay/Source/DisplayCluster/Private/Render/Upscaler/DisplayClusterUpscaler.cpp
--- a/Engine/Plugins/Runtime/nDisplay/Source/DisplayCluster/Private/Render/Upscaler/DisplayClusterUpscaler.cpp
+++ b/Engine/Plugins/Runtime/nDisplay/Source/DisplayCluster/Private/Render/Upscaler/DisplayClusterUpscaler.cpp
@@ -2,9 +2,11 @@
 
 #include "Render/Upscaler/DisplayClusterUpscaler.h"
 #include "Render/Viewport/DisplayClusterViewport.h"
+#include "DisplayClusterSceneViewExtensions.h"
 
 #include "IUpscalerModularFeature.h"
 #include "Features/IModularFeatures.h"
+#include "SceneViewExtensionContext.h"
 
 #include "LegacyScreenPercentageDriver.h"
 #include "SceneView.h"
@@ -47,23 +49,88 @@ namespace UE::DisplayCluster::Upscaler
 	}
 
 	/** Iterate over all upscaler modular features. */
-	static void ForeachUpscaler(TFunction<void(IUpscalerModularFeature&)>&& IteratorFunc)
+	static void ForeachUpscaler(TFunction<void(IUpscalerModularFeature&)> IteratorFunc)
 	{
 		static IModularFeatures& ModularFeatures = IModularFeatures::Get();
+		static TMap<FName, FGuid> SceneViewExtensionIsActiveFunctorsMap;
+
+		// If the modular function will be disabled at runtime, we should forget about it
+		// And then restore again when it is enabled.
+		TMap<FName, FGuid> NewSceneViewExtensionIsActiveFunctorsMap;
 
 		ModularFeatures.LockModularFeatureList();
+		ON_SCOPE_EXIT
+		{
+			ModularFeatures.UnlockModularFeatureList();
+		};
+
+		// Iterate over enabled upscalers
+		const TArray<IUpscalerModularFeature*> UpscalerModularFeatures = ModularFeatures.GetModularFeatureImplementations<IUpscalerModularFeature>(IUpscalerModularFeature::ModularFeatureName);
+		for (IUpscalerModularFeature* UpscalerModularFeature : UpscalerModularFeatures)
 		{
-			const TArray<IUpscalerModularFeature*> Extenders = ModularFeatures.GetModularFeatureImplementations<IUpscalerModularFeature>(IUpscalerModularFeature::ModularFeatureName);
-			for (IUpscalerModularFeature* Extender : Extenders)
+			if (UpscalerModularFeature == nullptr)
+			{
+				continue;
+			}
+
+			// Register IsActive Functors
+			const FName& UpscalerName = UpscalerModularFeature->GetName();
+			if (SceneViewExtensionIsActiveFunctorsMap.Contains(UpscalerName))
 			{
-				// Iterate over enabled upscalers
-				if (Extender && Extender->IsFeatureEnabled())
+				// Use exist Guid
+				NewSceneViewExtensionIsActiveFunctorsMap.Emplace(UpscalerName, SceneViewExtensionIsActiveFunctorsMap[UpscalerName]);
+			}
+			else
+			{
+				FSceneViewExtensionIsActiveFunctor IsActiveFunctor;
+				IsActiveFunctor.IsActiveFunction = [UpscalerName](const ISceneViewExtension* SceneViewExtension, const FSceneViewExtensionContext& Context)
 				{
-					IteratorFunc(*Extender);
+					// Note: SceneViewExtension this pointer is always passed to functors from FSceneViewExtensionBase::IsActiveThisFrame, we confirm this remains so.
+					check(SceneViewExtension);
+
+					// Is nDisplay context
+					static const FDisplayClusterSceneViewExtensionContext DCViewExtensionContext;
+					if (Context.IsA(MoveTempIfPossible(DCViewExtensionContext)))
+					{
+						// Is viewport exist
+						const FDisplayClusterSceneViewExtensionContext& DisplayContext = static_cast<const FDisplayClusterSceneViewExtensionContext&>(Context);
+						if (DisplayContext.DisplayClusterViewport.IsValid())
+						{
+							const FDisplayClusterUpscalerSettings& UpscalerSettings = DisplayContext.DisplayClusterViewport->GetRenderSettings().UpscalerSettings;
+							if (!UpscalerSettings.bIsActive)
+							{
+								// Custom upscalers are disabled for this viewport
+								return TOptional<bool>(false);
+							}
+
+							if (UpscalerSettings.CustomUpscalerName != UpscalerName)
+							{
+								// This viewport uses a different upscaler.
+								return TOptional<bool>(false);
+							}
+
+							return TOptional<bool>(true);
+						}
+					}
+
+					return TOptional<bool>();
+				};
+
+				if (UpscalerModularFeature->AddSceneViewExtensionIsActiveFunctor(IsActiveFunctor))
+				{
+					NewSceneViewExtensionIsActiveFunctorsMap.Emplace(UpscalerName, IsActiveFunctor.GetGuid());
 				}
 			}
+
+
+			if (UpscalerModularFeature->IsFeatureEnabled())
+			{
+				IteratorFunc(*UpscalerModularFeature);
+			}
 		}
-		ModularFeatures.UnlockModularFeatureList();
+
+		// Update
+		SceneViewExtensionIsActiveFunctorsMap = NewSceneViewExtensionIsActiveFunctorsMap;
 	}
 
 	/** Get global fraction. */
@@ -342,7 +409,7 @@ void FDisplayClusterUpscaler::SetupSceneView(
 			{
 				if (Upscaler.GetName() == InUpscalerSettings.CustomUpscalerName)
 				{
-					Upscaler.SetupSceneView(InUpscalerSettings.CustomUpscalerSettings, InViewFamily, InOutView);
+					Upscaler.SetupSceneView(InUpscalerSettings.CustomUpscalerSettings, InOutView);
 				}
 			});
 	}
@@ -404,7 +471,7 @@ FName FDisplayClusterUpscaler::PostConfigureViewFamily(
 				{
 					if (Upscaler.PostConfigureViewFamily(
 						InUpscalerSettings.CustomUpscalerSettings, UpscalerParam,
-						InOutViewFamily, InOutViews))
+						InOutViewFamily))
 					{
 						// Return upscaler name on success
 						UpscalerName = Upscaler.GetName();
@@ -492,50 +559,4 @@ FName FDisplayClusterUpscaler::PostConfigureViewFamily(
 	InOutViewFamily.SecondaryViewFraction = UpscalerParam.SecondaryScreenPercentage;
 
 	return NAME_None;
-}
-
-void FDisplayClusterUpscaler::SetupUpscalerViewExtensions(bool bEnabled)
-{
-	using namespace UE::DisplayCluster::Upscaler;
-
-	ForeachUpscaler(
-		[&](IUpscalerModularFeature& Upscaler)
-	{
-			// By default use upscaler's rules for the editor viewports.
-			Upscaler.SetupUpscalerViewExtensions(
-				bEnabled, 
-				/*/bEnableInEditorViewport*/ false);
-		});
-}
-
-void FDisplayClusterUpscaler::SetupUpscalerViewExtensionsForSettings(const FDisplayClusterUpscalerSettings& InUpscalerSettings)
-{
-	using namespace UE::DisplayCluster::Upscaler;
-
-	if (!UseCustomUpscalers())
-	{
-		// Disable all custom upscalers
-		SetupUpscalerViewExtensions(false);
-
-		return;
-	}
-
-	ForeachUpscaler(
-		[&](IUpscalerModularFeature& Upscaler)
-		{
-			if (Upscaler.GetName() == InUpscalerSettings.CustomUpscalerName)
-			{
-				// Enable all necessary ViewExtension of this upscaler
-				// DCRA preview in editor can use upscalers.
-				Upscaler.SetupUpscalerViewExtensionsForSettings(
-					InUpscalerSettings.CustomUpscalerSettings,
-					/**bEnableInEditorViewport*/ true);
-			}
-			else
-			{
-				// Disable all ViewExtensions for all other upscalers
-				Upscaler.SetupUpscalerViewExtensions(
-					/**bEnabled*/ false, /**bEnableInEditorViewport*/ false);
-			}
-		});
 }
\ No newline at end of file
diff --git a/Engine/Plugins/Runtime/nDisplay/Source/DisplayCluster/Private/Render/Upscaler/DisplayClusterUpscaler.h b/Engine/Plugins/Runtime/nDisplay/Source/DisplayCluster/Private/Render/Upscaler/DisplayClusterUpscaler.h
--- a/Engine/Plugins/Runtime/nDisplay/Source/DisplayCluster/Private/Render/Upscaler/DisplayClusterUpscaler.h
+++ b/Engine/Plugins/Runtime/nDisplay/Source/DisplayCluster/Private/Render/Upscaler/DisplayClusterUpscaler.h
@@ -17,22 +17,6 @@ class FSceneView;
 class FDisplayClusterUpscaler
 {
 public:
-	/*
-	* Setup all View Extensions of the upscaler to be enabled or disabled.
-	* Must be called before GatherActiveExtensions() during the rendering pass setup.
-	*
-	* @param bEnable - (in) Enable or disable all ViewExtensions of this upscaler.
-	*/
-	static void SetupUpscalerViewExtensions(bool bEnabled);
-
-	/**
-	* Enable only necessary view extensions used by upscaler for these input settings.
-	* Must be called before GatherActiveExtensions() during the rendering pass setup.
-	* 
-	* @param InUpscalerSettings - (in) the settings to apply for viewport
-	*/
-	static void SetupUpscalerViewExtensionsForSettings(const FDisplayClusterUpscalerSettings& InUpscalerSettings);
-
 	/** Configure upscaler and viewfamily for ertain settings.
 	* 
 	* @param InUpscalerSettings - (in) the settings to apply for viewport
diff --git a/Engine/Plugins/Runtime/nDisplay/Source/DisplayCluster/Private/Render/Viewport/DisplayClusterViewport.cpp b/Engine/Plugins/Runtime/nDisplay/Source/DisplayCluster/Private/Render/Viewport/DisplayClusterViewport.cpp
--- a/Engine/Plugins/Runtime/nDisplay/Source/DisplayCluster/Private/Render/Viewport/DisplayClusterViewport.cpp
+++ b/Engine/Plugins/Runtime/nDisplay/Source/DisplayCluster/Private/Render/Viewport/DisplayClusterViewport.cpp
@@ -223,7 +223,7 @@ void FDisplayClusterViewport::ImplGatherActiveExtensions(const int32 ViewIndex,
 	case EDisplayClusterViewportCaptureMode::MoviePipeline:
 		if (InViewport)
 		{
-			FDisplayClusterSceneViewExtensionContext ViewExtensionContext(InViewport, Configuration, GetId());
+			FDisplayClusterSceneViewExtensionContext ViewExtensionContext(InViewport, AsShared());
 			OutExtensions = GEngine->ViewExtensions->GatherActiveExtensions(ViewExtensionContext);
 			return;
 		}
@@ -231,7 +231,7 @@ void FDisplayClusterViewport::ImplGatherActiveExtensions(const int32 ViewIndex,
 		{
 			if (UWorld* CurrentWorld = Configuration->GetCurrentWorld())
 			{
-				FDisplayClusterSceneViewExtensionContext ViewExtensionContext(CurrentWorld->Scene, Configuration, GetId());
+				FDisplayClusterSceneViewExtensionContext ViewExtensionContext(CurrentWorld->Scene, AsShared());
 				OutExtensions = GEngine->ViewExtensions->GatherActiveExtensions(ViewExtensionContext);
 				return;
 			}
diff --git a/Engine/Plugins/Runtime/nDisplay/Source/DisplayCluster/Private/Render/Viewport/DisplayClusterViewport.h b/Engine/Plugins/Runtime/nDisplay/Source/DisplayCluster/Private/Render/Viewport/DisplayClusterViewport.h
--- a/Engine/Plugins/Runtime/nDisplay/Source/DisplayCluster/Private/Render/Viewport/DisplayClusterViewport.h
+++ b/Engine/Plugins/Runtime/nDisplay/Source/DisplayCluster/Private/Render/Viewport/DisplayClusterViewport.h
@@ -77,6 +77,15 @@ public:
 		return Configuration.Get();
 	}
 
+	virtual const TSharedRef<IDisplayClusterViewportConfiguration, ESPMode::ThreadSafe> GetConfigurationRef() override
+	{
+		return Configuration;
+	}
+	virtual const TSharedRef<const IDisplayClusterViewportConfiguration, ESPMode::ThreadSafe> GetConfigurationRef() const override
+	{
+		return Configuration;
+	}
+
 	/** Get viewport preview API */
 	virtual IDisplayClusterViewportPreview& GetViewportPreview() const override;
 
diff --git a/Engine/Plugins/Runtime/nDisplay/Source/DisplayCluster/Private/Render/Viewport/DisplayClusterViewportFrameStatsViewExtension.cpp b/Engine/Plugins/Runtime/nDisplay/Source/DisplayCluster/Private/Render/Viewport/DisplayClusterViewportFrameStatsViewExtension.cpp
--- a/Engine/Plugins/Runtime/nDisplay/Source/DisplayCluster/Private/Render/Viewport/DisplayClusterViewportFrameStatsViewExtension.cpp
+++ b/Engine/Plugins/Runtime/nDisplay/Source/DisplayCluster/Private/Render/Viewport/DisplayClusterViewportFrameStatsViewExtension.cpp
@@ -177,7 +177,7 @@ bool FDisplayClusterViewportFrameStatsViewExtension::IsActiveThisFrame_Internal(
 		if (Context.IsA(MoveTempIfPossible(DCViewExtensionContext)))
 		{
 			const FDisplayClusterSceneViewExtensionContext& DisplayContext = static_cast<const FDisplayClusterSceneViewExtensionContext&>(Context);
-			if (DisplayContext.Configuration == Configuration)
+			if (DisplayContext.IsSameDisplayClusterViewportConfiguration(Configuration))
 			{
 				// Apply only for DC viewports
 				return true;
diff --git a/Engine/Plugins/Runtime/nDisplay/Source/DisplayCluster/Private/Render/Viewport/DisplayClusterViewportManager.cpp b/Engine/Plugins/Runtime/nDisplay/Source/DisplayCluster/Private/Render/Viewport/DisplayClusterViewportManager.cpp
--- a/Engine/Plugins/Runtime/nDisplay/Source/DisplayCluster/Private/Render/Viewport/DisplayClusterViewportManager.cpp
+++ b/Engine/Plugins/Runtime/nDisplay/Source/DisplayCluster/Private/Render/Viewport/DisplayClusterViewportManager.cpp
@@ -831,9 +831,6 @@ void FDisplayClusterViewportManager::PostConfigureViewFamily(
 	// Add ViewExtensions from the Custom Upscaler to this view family
 	if (!CustomUpscalerName.IsNone() && InFrameViewFamily.UpscalerSettings.CustomUpscalerName == CustomUpscalerName)
 	{
-		// Enable all necessary ViewExtension for the custom upscaler settings
-		FDisplayClusterUpscaler::SetupUpscalerViewExtensionsForSettings(InFrameViewFamily.UpscalerSettings);
-
 		// Iterate over all viewports, and collect all used ViewExtensions
 		for (const FDisplayClusterRenderFrameTargetView& FrameTargetViewIt : InFrameViewFamily.Views)
 		{
@@ -852,9 +849,6 @@ void FDisplayClusterViewportManager::PostConfigureViewFamily(
 				}
 			}
 		}
-
-		// Restore default behaviour of ViewExtensions of upscalers
-		FDisplayClusterUpscaler::SetupUpscalerViewExtensions(true);
 	}
 }
 
diff --git a/Engine/Plugins/Runtime/nDisplay/Source/DisplayCluster/Private/Render/Viewport/DisplayClusterViewportManager.h b/Engine/Plugins/Runtime/nDisplay/Source/DisplayCluster/Private/Render/Viewport/DisplayClusterViewportManager.h
--- a/Engine/Plugins/Runtime/nDisplay/Source/DisplayCluster/Private/Render/Viewport/DisplayClusterViewportManager.h
+++ b/Engine/Plugins/Runtime/nDisplay/Source/DisplayCluster/Private/Render/Viewport/DisplayClusterViewportManager.h
@@ -85,6 +85,15 @@ public:
 		return Configuration.Get();
 	}
 
+	virtual const TSharedRef<IDisplayClusterViewportConfiguration, ESPMode::ThreadSafe> GetConfigurationRef() override
+	{
+		return Configuration;
+	}
+	virtual const TSharedRef<const IDisplayClusterViewportConfiguration, ESPMode::ThreadSafe> GetConfigurationRef() const override
+	{
+		return Configuration;
+	}
+
 	virtual bool BeginNewFrame(FViewport* InViewport, FDisplayClusterRenderFrame& OutRenderFrame) override;
 	virtual void InitializeNewFrame() override;
 	virtual void FinalizeNewFrame() override;
diff --git a/Engine/Plugins/Runtime/nDisplay/Source/DisplayCluster/Private/Render/Viewport/DisplayClusterViewportManagerViewExtension.cpp b/Engine/Plugins/Runtime/nDisplay/Source/DisplayCluster/Private/Render/Viewport/DisplayClusterViewportManagerViewExtension.cpp
--- a/Engine/Plugins/Runtime/nDisplay/Source/DisplayCluster/Private/Render/Viewport/DisplayClusterViewportManagerViewExtension.cpp
+++ b/Engine/Plugins/Runtime/nDisplay/Source/DisplayCluster/Private/Render/Viewport/DisplayClusterViewportManagerViewExtension.cpp
@@ -189,7 +189,7 @@ bool FDisplayClusterViewportManagerViewExtension::IsActiveThisFrame_Internal(con
 	if (Context.IsA(MoveTempIfPossible(DCViewExtensionContext)))
 	{
 		const FDisplayClusterSceneViewExtensionContext& DisplayContext = static_cast<const FDisplayClusterSceneViewExtensionContext&>(Context);
-		if (DisplayContext.Configuration == Configuration)
+		if (DisplayContext.IsSameDisplayClusterViewportConfiguration(Configuration))
 		{
 			// Apply only for DC viewports
 			return true;
diff --git a/Engine/Plugins/Runtime/nDisplay/Source/DisplayCluster/Private/Render/Viewport/RenderFrame/DisplayClusterRenderFrameManager.cpp b/Engine/Plugins/Runtime/nDisplay/Source/DisplayCluster/Private/Render/Viewport/RenderFrame/DisplayClusterRenderFrameManager.cpp
--- a/Engine/Plugins/Runtime/nDisplay/Source/DisplayCluster/Private/Render/Viewport/RenderFrame/DisplayClusterRenderFrameManager.cpp
+++ b/Engine/Plugins/Runtime/nDisplay/Source/DisplayCluster/Private/Render/Viewport/RenderFrame/DisplayClusterRenderFrameManager.cpp
@@ -84,6 +84,9 @@ bool FDisplayClusterRenderFrameManager::BuildSimpleFrame(FViewport* InViewport,
 	{
 		if (ViewportIt.IsValid())
 		{
+			// Exclude all ViewExtensions. They are added later, in FDisplayClusterViewportManager::PostConfigureViewFamily()
+			ViewportIt->GetRenderSettingsImpl().UpscalerSettings.bIsActive = false;
+
 			for (const FDisplayClusterViewport_Context& ContextIt : ViewportIt->GetContexts())
 			{
 				FDisplayClusterRenderFrameTargetView FrameView;
@@ -100,11 +103,8 @@ bool FDisplayClusterRenderFrameManager::BuildSimpleFrame(FViewport* InViewport,
 					FrameViewFamily.CustomBufferRatio = ContextIt.CustomBufferRatio;
 					FrameViewFamily.UpscalerSettings = ViewportIt->GetRenderSettings().UpscalerSettings;
 
-					// Do not collect ViewExtension of upscalers for the initial list.
 					// We need to determine if we can use upscalers in the future when AntiAliasing methods are configured in SceneViews.
-					FDisplayClusterUpscaler::SetupUpscalerViewExtensions(false);
 					FrameViewFamily.ViewExtensions = ViewportIt->GatherActiveExtensions(ContextIt.ContextNum, InViewport);
-					FDisplayClusterUpscaler::SetupUpscalerViewExtensions(true);
 				}
 
 				FDisplayClusterRenderFrameTarget FrameRenderTarget;
@@ -122,6 +122,9 @@ bool FDisplayClusterRenderFrameManager::BuildSimpleFrame(FViewport* InViewport,
 
 				OutRenderFrame.RenderTargets.Add(FrameRenderTarget);
 			}
+
+			// After the GatherActiveExtensions() enable upscalers ViewExtensions back.
+			ViewportIt->GetRenderSettingsImpl().UpscalerSettings.bIsActive = true;
 		}
 	}
 
diff --git a/Engine/Plugins/Runtime/nDisplay/Source/DisplayCluster/Public/DisplayClusterSceneViewExtensions.h b/Engine/Plugins/Runtime/nDisplay/Source/DisplayCluster/Public/DisplayClusterSceneViewExtensions.h
--- a/Engine/Plugins/Runtime/nDisplay/Source/DisplayCluster/Public/DisplayClusterSceneViewExtensions.h
+++ b/Engine/Plugins/Runtime/nDisplay/Source/DisplayCluster/Public/DisplayClusterSceneViewExtensions.h
@@ -4,7 +4,7 @@
 
 #include "CoreMinimal.h"
 #include "SceneViewExtension.h"
-#include "Render/Viewport/IDisplayClusterViewportConfiguration.h"
+#include "Render/Viewport/IDisplayClusterViewport.h"
 
 /**
  * Contains information about the context in which this scene view extension will be used.
@@ -26,22 +26,33 @@ public:
 		: FSceneViewExtensionContext()
 	{ }
 
-	FDisplayClusterSceneViewExtensionContext(FViewport* InViewport, const TSharedRef<IDisplayClusterViewportConfiguration, ESPMode::ThreadSafe>& InConfiguration, const FString& InViewportId)
+	FDisplayClusterSceneViewExtensionContext(
+		FViewport* InViewport,
+		const TSharedRef<const IDisplayClusterViewport, ESPMode::ThreadSafe>& InDisplayClusterViewport)
 		: FSceneViewExtensionContext(InViewport)
-		, ViewportId(InViewportId)
-		, Configuration(InConfiguration)
+		, DisplayClusterViewport(InDisplayClusterViewport)
 	{ }
 
-	FDisplayClusterSceneViewExtensionContext(FSceneInterface* InScene, const TSharedRef<IDisplayClusterViewportConfiguration, ESPMode::ThreadSafe>& InConfiguration, const FString& InViewportId)
+	FDisplayClusterSceneViewExtensionContext(
+		FSceneInterface* InScene,
+		const TSharedRef<const IDisplayClusterViewport, ESPMode::ThreadSafe>& InDisplayClusterViewport)
 		: FSceneViewExtensionContext(InScene)
-		, ViewportId(InViewportId)
-		, Configuration(InConfiguration)
+		, DisplayClusterViewport(InDisplayClusterViewport)
 	{ }
 
-public:
-	// The id of the nDisplay viewport being rendered.
-	const FString ViewportId;
+	/** Returns true if this viewport context refers to the same configuration. */
+	bool IsSameDisplayClusterViewportConfiguration(const TSharedRef<const IDisplayClusterViewportConfiguration, ESPMode::ThreadSafe> InConfigurationRef) const
+	{
+		if (DisplayClusterViewport.IsValid()
+			&& DisplayClusterViewport->GetConfigurationRef() == InConfigurationRef)
+		{
+			return true;
+		}
+
+		return false;
+	}
 
-	// Reference to viewport manager
-	const TSharedPtr<IDisplayClusterViewportConfiguration, ESPMode::ThreadSafe> Configuration;
+public:
+	// Reference to the DC viewport
+	const TSharedPtr<const IDisplayClusterViewport, ESPMode::ThreadSafe> DisplayClusterViewport;
 };
diff --git a/Engine/Plugins/Runtime/nDisplay/Source/DisplayCluster/Public/Render/Upscaler/DisplayClusterUpscalerSettings.h b/Engine/Plugins/Runtime/nDisplay/Source/DisplayCluster/Public/Render/Upscaler/DisplayClusterUpscalerSettings.h
--- a/Engine/Plugins/Runtime/nDisplay/Source/DisplayCluster/Public/Render/Upscaler/DisplayClusterUpscalerSettings.h
+++ b/Engine/Plugins/Runtime/nDisplay/Source/DisplayCluster/Public/Render/Upscaler/DisplayClusterUpscalerSettings.h
@@ -46,4 +46,8 @@ public:
 
 	// Custom upscaler settings
 	FInstancedPropertyBag CustomUpscalerSettings;
+
+	// parameters for the Functor logic(IUpscalerModularFeature::AddSceneViewExtensionIsActiveFunctor()).
+	// Enable to use upscaler ViewExtension
+	bool bIsActive = true;
 };
diff --git a/Engine/Plugins/Runtime/nDisplay/Source/DisplayCluster/Public/Render/Viewport/IDisplayClusterViewport.h b/Engine/Plugins/Runtime/nDisplay/Source/DisplayCluster/Public/Render/Viewport/IDisplayClusterViewport.h
--- a/Engine/Plugins/Runtime/nDisplay/Source/DisplayCluster/Public/Render/Viewport/IDisplayClusterViewport.h
+++ b/Engine/Plugins/Runtime/nDisplay/Source/DisplayCluster/Public/Render/Viewport/IDisplayClusterViewport.h
@@ -48,6 +48,12 @@ public:
 	/** [const] Get viewport manager configuration interface. */
 	virtual const IDisplayClusterViewportConfiguration& GetConfiguration() const = 0;
 
+	/** Get viewport manager configuration interface reference. */
+	virtual const TSharedRef<IDisplayClusterViewportConfiguration, ESPMode::ThreadSafe> GetConfigurationRef() = 0;
+
+	/** [const] Get viewport manager configuration interface reference. */
+	virtual const TSharedRef<const IDisplayClusterViewportConfiguration, ESPMode::ThreadSafe>  GetConfigurationRef() const = 0;
+
 	/** Get viewport preview API */
 	virtual IDisplayClusterViewportPreview& GetViewportPreview() const = 0;
 
diff --git a/Engine/Plugins/Runtime/nDisplay/Source/DisplayCluster/Public/Render/Viewport/IDisplayClusterViewportManager.h b/Engine/Plugins/Runtime/nDisplay/Source/DisplayCluster/Public/Render/Viewport/IDisplayClusterViewportManager.h
--- a/Engine/Plugins/Runtime/nDisplay/Source/DisplayCluster/Public/Render/Viewport/IDisplayClusterViewportManager.h
+++ b/Engine/Plugins/Runtime/nDisplay/Source/DisplayCluster/Public/Render/Viewport/IDisplayClusterViewportManager.h
@@ -51,6 +51,11 @@ public:
 	virtual IDisplayClusterViewportConfiguration& GetConfiguration() = 0;
 	virtual const IDisplayClusterViewportConfiguration& GetConfiguration() const = 0;
 
+	/** Get viewport manager configuration interface reference. */
+	virtual const TSharedRef<IDisplayClusterViewportConfiguration, ESPMode::ThreadSafe> GetConfigurationRef() = 0;
+	virtual const TSharedRef<const IDisplayClusterViewportConfiguration, ESPMode::ThreadSafe>  GetConfigurationRef() const = 0;
+
+
 	/** Get viewport manager preview API */
 	virtual IDisplayClusterViewportManagerPreview& GetViewportManagerPreview() = 0;
 	virtual const IDisplayClusterViewportManagerPreview& GetViewportManagerPreview() const = 0;
diff --git a/Engine/Plugins/Runtime/nDisplay/Source/DisplayClusterConfigurator/Private/Views/Details/Upscaler/DisplayClusterConfigurationUpscalerSettingsDetailCustomization.cpp b/Engine/Plugins/Runtime/nDisplay/Source/DisplayClusterConfigurator/Private/Views/Details/Upscaler/DisplayClusterConfigurationUpscalerSettingsDetailCustomization.cpp
--- a/Engine/Plugins/Runtime/nDisplay/Source/DisplayClusterConfigurator/Private/Views/Details/Upscaler/DisplayClusterConfigurationUpscalerSettingsDetailCustomization.cpp
+++ b/Engine/Plugins/Runtime/nDisplay/Source/DisplayClusterConfigurator/Private/Views/Details/Upscaler/DisplayClusterConfigurationUpscalerSettingsDetailCustomization.cpp
@@ -48,7 +48,7 @@ namespace UE::DisplayClusterConfigurator::UpscalerSettingsDetailCustomization
 	}
 
 	/** Iterate over all default upscaler methods. */
-	static void ForeachDefaultUpscalerMethod(TFunction<void(const UEnum&, const int32, const EDisplayClusterConfigurationUpscalingMethod)>&& IteratorFunc)
+	static void ForeachDefaultUpscalerMethod(TFunction<void(const UEnum&, const int32, const EDisplayClusterConfigurationUpscalingMethod)> IteratorFunc)
 	{
 		// Collect embedded values from the EDisplayClusterConfigurationUpscalingMethod
 		if (const UEnum* EnumSource = StaticEnum<EDisplayClusterConfigurationUpscalingMethod>())
@@ -69,24 +69,26 @@ namespace UE::DisplayClusterConfigurator::UpscalerSettingsDetailCustomization
 	}
 
 	/** Iterate over all upscaler interfaces. */
-	static void ForeachUpscaler(TFunction<void(IUpscalerModularFeature&)>&& IteratorFunc)
+	static void ForeachUpscaler(TFunction<void(IUpscalerModularFeature&)> IteratorFunc)
 	{
 		static IModularFeatures& ModularFeatures = IModularFeatures::Get();
 
 		ModularFeatures.LockModularFeatureList();
+		ON_SCOPE_EXIT
 		{
-			const TArray<IUpscalerModularFeature*> Extenders = ModularFeatures.GetModularFeatureImplementations<IUpscalerModularFeature>(IUpscalerModularFeature::ModularFeatureName);
-			for (IUpscalerModularFeature* Extender : Extenders)
+			ModularFeatures.UnlockModularFeatureList();
+		};
+
+		const TArray<IUpscalerModularFeature*> Extenders = ModularFeatures.GetModularFeatureImplementations<IUpscalerModularFeature>(IUpscalerModularFeature::ModularFeatureName);
+		for (IUpscalerModularFeature* Extender : Extenders)
+		{
+			// Iterate over active upscalers
+			if (Extender
+				&& (!ShouldShowOnlyEnabledUpscalingMethods() || Extender->IsFeatureEnabled()))
 			{
-				// Iterate over active upscalers
-				if (Extender
-					&& (!ShouldShowOnlyEnabledUpscalingMethods() || Extender->IsFeatureEnabled()))
-				{
-					IteratorFunc(*Extender);
-				}
+				IteratorFunc(*Extender);
 			}
 		}
-		ModularFeatures.UnlockModularFeatureList();
 	}
 
 	/** Get upscaling settings for method name. */
diff --git a/Engine/Source/Runtime/VirtualProduction/VirtualProduction/Public/IUpscalerModularFeature.h b/Engine/Source/Runtime/VirtualProduction/VirtualProduction/Public/IUpscalerModularFeature.h
--- a/Engine/Source/Runtime/VirtualProduction/VirtualProduction/Public/IUpscalerModularFeature.h
+++ b/Engine/Source/Runtime/VirtualProduction/VirtualProduction/Public/IUpscalerModularFeature.h
@@ -6,6 +6,7 @@
 #include "Features/IModularFeature.h"
 
 struct FInstancedPropertyBag;
+struct FSceneViewExtensionIsActiveFunctor;
 
 class FSceneViewFamilyContext;
 class FSceneViewFamily;
@@ -68,58 +69,44 @@ namespace UE::VirtualProduction
 		virtual bool IsFeatureEnabled() const = 0;
 
 		/**
-		* Setup all View Extensions of the upscaler to be enabled or disabled.
-		* Must be called before GatherActiveExtensions() during the rendering pass setup.
-		*
-		* @param bEnable                 - (in) Enable or disable all ViewExtensions of this upscaler.
-		* @param bEnableInEditorViewport - (in) Makes this upscaler's ViewExtensions work in the editor's viewports
+		* Adds a functor and returns a GUID to identify it.
+		* 
+		* @param IsActiveFunction - the functor with custom logic for the upscaler ViewExtension
+		* 
+		* @return true on success
 		*/
-		virtual void SetupUpscalerViewExtensions(const bool bEnable, const bool bEnableInEditorViewport) = 0;
+		virtual bool AddSceneViewExtensionIsActiveFunctor(const FSceneViewExtensionIsActiveFunctor& IsActiveFunction) = 0;
 
-		/**
-		* Enable only necessary view extensions used by upscaler for these input settings.
-		* Must be called before GatherActiveExtensions() during the rendering pass setup.
-		*
-		* @param InUpscalerSettings      - (in) The upscaler settings which we want to use. Enables only necessary ViewExtensions.
-		* @param bEnableInEditorViewport - (in) Makes this upscaler's ViewExtensions work in the editor's viewports
+		/** 
+		* Removes a functor by GUID.
+		* 
+		* @param FunctorGuid - use value from the FSceneViewExtensionIsActiveFunctor::GetGuid()
 		*
-		* @return true if the settings can be used and upscaler is enabled
+		* @return true if the functor exists and has been deleted
 		*/
-		virtual bool SetupUpscalerViewExtensionsForSettings(const FInstancedPropertyBag& InUpscalerSettings, const bool bEnableInEditorViewport) = 0;
+		virtual bool RemoveSceneViewExtensionIsActiveFunctor(const FGuid& FunctorGuid) = 0;
 
 		/**
 		* Get default settings for this upscaler
 		*
 		* @param OuUpscalerSettings - (out) default settings
-		*
-		* @param return true if successed.
+		* 
+		* @return true if successed.
 		*/
 		virtual bool GetSettings(FInstancedPropertyBag& OuUpscalerSettings) const
 		{
 			return false;
 		}
 
-		/**
-		* Reset all stored per-view settings to default values
-		* The settings that was stored by the function PostConfigureViewFamily()
-		*
-		* @param return true if successed.
-		*/
-		virtual void ResetSettings()
-		{
-		}
-
 		/**
 		* Setup SceneView for this upscaler.
 		* Upscalers should configure the view properties (e.g. AntiAliasingMethod) for their rendering pipeline (temporal/spatial/etc.).
 		*
 		* @param InUpscalerSettings - (in) upscaler settings which we want to use. This settings should be associated with the InOutViews.
-		* @param InViewFamily       - (in) ViewFamily that own this view.
 		* @param InOutView          - (in) scene view to that will be configured.
 		*/
 		virtual void SetupSceneView(
 			const FInstancedPropertyBag& InUpscalerSettings,
-			const FSceneViewFamily& InViewFamily,
 			FSceneView& InOutView) = 0;
 
 		/**
@@ -133,12 +120,10 @@ namespace UE::VirtualProduction
 		* @param InUpscalerSettings - (in) upscaler settings which we want to use. This settings should be associated with the InOutViews.
 		* @param InUpscalerParam    - (in) upscaler parameters that is used to configure view family.
 		* @param InOutViewFamily    - (in, out) The view family that will be configured.
-		* @param InOutViews         - (in, out) The views in the family that will be configured.
 		*/
 		virtual bool PostConfigureViewFamily(
 			const FInstancedPropertyBag& InUpscalerSettings,
 			const FUpscalerModularFeatureParameters& InUpscalerParam,
-			FSceneViewFamilyContext& InOutViewFamily,
-			const TArray<FSceneView*>& InOutViews) = 0;
+			FSceneViewFamilyContext& InOutViewFamily) = 0;
 	};
 };
