[nDisplay] Fixed LightCard rendering issues.

Render UVLC and LC with the same settings in the linear color space.
- Viewport Textures in linear colorspace

Fixed gamma correction parameter for OCIO.

Fixed typo: virtual FDisplayClusterViewportRenderTargetResource::GetDisplayGamma() - missed override keyword

diff --git a/Engine/Plugins/Runtime/nDisplay/Shaders/Private/MPCDIOverlayShaders.usf b/Engine/Plugins/Runtime/nDisplay/Shaders/Private/MPCDIOverlayShaders.usf
--- a/Engine/Plugins/Runtime/nDisplay/Shaders/Private/MPCDIOverlayShaders.usf
+++ b/Engine/Plugins/Runtime/nDisplay/Shaders/Private/MPCDIOverlayShaders.usf
@@ -40,16 +40,9 @@ float3 ComposeOverlayColor(float4 InColor, float4 InOverlayColor)
 /** Apply final gamma for LightCard color. */
 float4 GetFinalLightcardColor(float4 InColor)
 {
-    // Lightcard color is premultiplied. Unpremultiply before gamma correction.
-	float4 OutColor = InvertedUnpremultiply(InColor);
-
-	// Apply gamma conversion for unpremultiplied color
-	OutColor.rgb = DisplayGammaCorrection(OutColor.rgb, LightCardGamma);
-
-	// And finally premultiply it back
-	OutColor = InvertedPremultiply(OutColor);
-	
-	return OutColor;
+	// ICVFX expects all lightcards in final color space (with OCIO applied)
+	// Gamma should not be converted after OCIO.
+	return InColor;
 }
 
 // Additive render pass with 2 overlays
diff --git a/Engine/Plugins/Runtime/nDisplay/Shaders/Private/UVLightCardGammaCorrection.usf b/Engine/Plugins/Runtime/nDisplay/Shaders/Private/UVLightCardGammaCorrection.usf
--- a/Engine/Plugins/Runtime/nDisplay/Shaders/Private/UVLightCardGammaCorrection.usf
+++ b/Engine/Plugins/Runtime/nDisplay/Shaders/Private/UVLightCardGammaCorrection.usf
@@ -5,19 +5,14 @@
 
 Texture2D InputTexture;
 SamplerState InputSampler;
-float InverseGamma = 1.0;
 
 
 float4 GammaCorrectionPS(noperspective float4 UVAndScreenPos : TEXCOORD0) : SV_Target0
 {
 	float4 OutColor = Texture2DSample(InputTexture, InputSampler, UVAndScreenPos.xy);
-	
-	// LightCard textures are premultiplied, so we need to unpremultiply before any gamma correction
-	OutColor = Unpremultiply(OutColor);
-
-	OutColor.rgb = DisplayGammaCorrection(OutColor.rgb, InverseGamma);
-
-	OutColor = Premultiply(OutColor);
 
+	// Always render UVLC in linear color space.
+	
+	// Invert alpha.
 	return float4(OutColor.rgb, 1.0 - OutColor.a);
 }
\ No newline at end of file
diff --git a/Engine/Plugins/Runtime/nDisplay/Source/DisplayCluster/Private/Render/Viewport/DisplayClusterViewport.cpp b/Engine/Plugins/Runtime/nDisplay/Source/DisplayCluster/Private/Render/Viewport/DisplayClusterViewport.cpp
--- a/Engine/Plugins/Runtime/nDisplay/Source/DisplayCluster/Private/Render/Viewport/DisplayClusterViewport.cpp
+++ b/Engine/Plugins/Runtime/nDisplay/Source/DisplayCluster/Private/Render/Viewport/DisplayClusterViewport.cpp
@@ -339,6 +339,9 @@ void FDisplayClusterViewport::SetupSceneView(uint32 ContextNum, class UWorld* Wo
 	check(IsInGameThread());
 	check(Contexts.IsValidIndex(ContextNum));
 
+	// Configure the capture mode to be used.
+	IDisplayClusterViewportManager::SetupSceneView(RenderSettings.CaptureMode, InOutViewFamily, InOutView);
+
 	// MRQ only uses viewport visibility settings
 	if(RenderSettings.CaptureMode == EDisplayClusterViewportCaptureMode::MoviePipeline)
 	{
@@ -374,18 +377,6 @@ void FDisplayClusterViewport::SetupSceneView(uint32 ContextNum, class UWorld* Wo
 		InOutView.bAllowCrossGPUTransfer = false;
 	}
 
-	// Disable raytracing for lightcard and chromakey
-	switch (RenderSettings.CaptureMode)
-	{
-	case EDisplayClusterViewportCaptureMode::Chromakey:
-	case EDisplayClusterViewportCaptureMode::Lightcard:
-		InOutView.bAllowRayTracing = false;
-		break;
-
-	default:
-		break;
-	}
-
 	// Apply visibility settigns to view
 	VisibilitySettings.SetupSceneView(InOutView);
 
diff --git a/Engine/Plugins/Runtime/nDisplay/Source/DisplayCluster/Private/Render/Viewport/DisplayClusterViewportManager.cpp b/Engine/Plugins/Runtime/nDisplay/Source/DisplayCluster/Private/Render/Viewport/DisplayClusterViewportManager.cpp
--- a/Engine/Plugins/Runtime/nDisplay/Source/DisplayCluster/Private/Render/Viewport/DisplayClusterViewportManager.cpp
+++ b/Engine/Plugins/Runtime/nDisplay/Source/DisplayCluster/Private/Render/Viewport/DisplayClusterViewportManager.cpp
@@ -148,6 +148,30 @@ TSharedRef<IDisplayClusterViewportManager, ESPMode::ThreadSafe> IDisplayClusterV
 	return ViewportManager;
 }
 
+void IDisplayClusterViewportManager::SetupSceneView(const EDisplayClusterViewportCaptureMode InMode, FSceneViewFamily& InOutViewFamily, FSceneView& InOutView)
+{
+	if (InMode == EDisplayClusterViewportCaptureMode::Lightcard)
+	{
+		// The UVLC texture is shared and can be used for any kind of viewport (with or without OCIO).
+		// Because of UVLC texture sharing, always render any kind of LightCard viewports as for OCIO.
+		// This will result in the same rendering result, and as a result the LC and UVLC colors will match.
+		FOpenColorIORendering::PrepareView(InOutViewFamily, InOutView);
+	}
+
+	// Disable raytracing for lightcard and chromakey
+	switch (InMode)
+	{
+	case EDisplayClusterViewportCaptureMode::Chromakey:
+	case EDisplayClusterViewportCaptureMode::Lightcard:
+
+		InOutView.bAllowRayTracing = false;
+		break;
+
+	default:
+		break;
+	}
+}
+
 void IDisplayClusterViewportManager::SetupEngineShowFlags(const EDisplayClusterViewportCaptureMode InMode, FEngineShowFlags& InOutEngineShowFlags)
 {
 	// Control NaniteMeshes for ChromaKey and Lightcards:
@@ -754,6 +778,10 @@ FSceneViewFamily::ConstructionValues FDisplayClusterViewportManager::CreateViewF
 	{
 	case EDisplayClusterViewportCaptureMode::Chromakey:
 	case EDisplayClusterViewportCaptureMode::Lightcard:
+
+		// Disable postprocess for LightCard & CK
+		InEngineShowFlags.PostProcessing = 0;
+
 		switch (RenderFrameSettings.AlphaChannelCaptureMode)
 		{
 		case EDisplayClusterRenderFrameAlphaChannelCaptureMode::Copy:
@@ -801,7 +829,7 @@ void FDisplayClusterViewportManager::ConfigureViewFamily(
 {
 	// Note: EngineShowFlags should have already been configured in CreateViewFamilyConstructionValues.
 	ViewFamily.SceneCaptureCompositeMode = ESceneCaptureCompositeMode::SCCM_Overwrite;
-	if (!ShouldRenderFinalColor())
+	if (!ShouldRenderFinalColor() || InFrameTarget.CaptureMode == EDisplayClusterViewportCaptureMode::Lightcard)
 	{
 		ViewFamily.SceneCaptureSource = ESceneCaptureSource::SCS_SceneColorHDR;
 	}
diff --git a/Engine/Plugins/Runtime/nDisplay/Source/DisplayCluster/Private/Render/Viewport/DisplayClusterViewportProxy.cpp b/Engine/Plugins/Runtime/nDisplay/Source/DisplayCluster/Private/Render/Viewport/DisplayClusterViewportProxy.cpp
--- a/Engine/Plugins/Runtime/nDisplay/Source/DisplayCluster/Private/Render/Viewport/DisplayClusterViewportProxy.cpp
+++ b/Engine/Plugins/Runtime/nDisplay/Source/DisplayCluster/Private/Render/Viewport/DisplayClusterViewportProxy.cpp
@@ -126,12 +126,26 @@ const FDisplayClusterViewportProxy& FDisplayClusterViewportProxy::GetRenderingVi
 	return *this;
 }
 
+bool FDisplayClusterViewportProxy::IsPostProcessDisabled() const
+{
+	if (ConfigurationProxy->GetRenderFrameSettings().IsPostProcessDisabled())
+	{
+		return true;
+	}
+
+	if (EnumHasAnyFlags(RenderSettingsICVFX.RuntimeFlags, EDisplayClusterViewportRuntimeICVFXFlags::UVLightcard | EDisplayClusterViewportRuntimeICVFXFlags::Lightcard | EDisplayClusterViewportRuntimeICVFXFlags::Chromakey))
+	{
+		return true;
+	}
+
+	return false;
+}
+
 EDisplayClusterViewportOpenColorIOMode FDisplayClusterViewportProxy::GetOpenColorIOMode() const
 {
 	if (OpenColorIO.IsValid() && OpenColorIO->IsValid_RenderThread())
 	{
-		if (EnumHasAnyFlags(RenderSettingsICVFX.RuntimeFlags, EDisplayClusterViewportRuntimeICVFXFlags::UVLightcard | EDisplayClusterViewportRuntimeICVFXFlags::Lightcard | EDisplayClusterViewportRuntimeICVFXFlags::Chromakey)
-			|| ConfigurationProxy->GetRenderFrameSettings().IsPostProcessDisabled())
+		if (IsPostProcessDisabled())
 		{
 			// Rendering without post-processing, OCIO is applied last, to the RTT texture of the viewport
 			return EDisplayClusterViewportOpenColorIOMode::Resolved;
@@ -467,24 +481,17 @@ void FDisplayClusterViewportProxy::ReleaseTextures_RenderThread()
 
 FDisplayClusterColorEncoding FDisplayClusterViewportProxy::GetResourceColorEncoding_RenderThread(const EDisplayClusterViewportResourceType InResourceType) const
 {
-	// Color modifiers:
-	EDisplayClusterColorPremultiply Premultiply = EDisplayClusterColorPremultiply::None;
-	if (EnumHasAnyFlags(RenderSettingsICVFX.RuntimeFlags, EDisplayClusterViewportRuntimeICVFXFlags::UVLightcard | EDisplayClusterViewportRuntimeICVFXFlags::Lightcard))
-	{
-		// All LightCard textures are premultiplied with inverted alpha.
-		Premultiply = EDisplayClusterColorPremultiply::InvertPremultiply;
-	}
-		
 	const EDisplayClusterViewportResourceType ResourceType = GetResourceType_RenderThread(InResourceType);
 	switch (ResourceType)
 	{
 	case EDisplayClusterViewportResourceType::InternalRenderTargetEntireRectResource:
 	case EDisplayClusterViewportResourceType::InternalRenderTargetResource:
-		if (EnumHasAllFlags(RenderSettingsICVFX.RuntimeFlags, EDisplayClusterViewportRuntimeICVFXFlags::UVLightcard))
+
+		if (EnumHasAnyFlags(RenderSettingsICVFX.RuntimeFlags, EDisplayClusterViewportRuntimeICVFXFlags::UVLightcard | EDisplayClusterViewportRuntimeICVFXFlags::Lightcard))
 		{
-			// The UVLightCard viewport is not actually rendered, so dont use context data
-			// This external texture will always be in linear color space.
-			return { EDisplayClusterColorEncoding::Linear, Premultiply };
+			// The LightCard viewport should always be rendered in linear color space with inverted alpha.
+			// ICVFX expects all lightcards in linear color space (blending purpose)
+			return { EDisplayClusterColorEncoding::Linear, EDisplayClusterColorPremultiply::InvertPremultiply };
 		}
 
 		if (!Contexts.IsEmpty())
@@ -493,46 +500,45 @@ FDisplayClusterColorEncoding FDisplayClusterViewportProxy::GetResourceColorEncod
 			const FEngineShowFlags& EngineShowFlags = Contexts[0].RenderThreadData.EngineShowFlags;
 			if (EngineShowFlags.Tonemapper == 0 || EngineShowFlags.PostProcessing == 0)
 			{
-				return { EDisplayClusterColorEncoding::Linear, Premultiply };
+				return { EDisplayClusterColorEncoding::Linear };
 			}
 
 			const float DefaultDisplayGamma = UTextureRenderTarget::GetDefaultDisplayGamma();
 			const float DisplayGamma = Contexts[0].RenderThreadData.EngineDisplayGamma;
 			if (DisplayGamma == DefaultDisplayGamma)
 			{
-				return { EDisplayClusterColorEncoding::Gamma,Premultiply };
+				return { EDisplayClusterColorEncoding::Gamma };
 			}
 
 			// Custom gamma value is different from default
-			return { DisplayGamma, Premultiply };
+			return { DisplayGamma };
 		}
 
 		break;
 
-	case EDisplayClusterViewportResourceType::InputShaderResource:
-	case EDisplayClusterViewportResourceType::MipsShaderResource:
-		if (EnumHasAnyFlags(RenderSettingsICVFX.RuntimeFlags, EDisplayClusterViewportRuntimeICVFXFlags::UVLightcard | EDisplayClusterViewportRuntimeICVFXFlags::Lightcard))
-		{
-			// ICVFX expects all lightcards in linear color space
-			return { EDisplayClusterColorEncoding::Linear, Premultiply };
-		}
-		break;
-
-		// Preview Output
+	// Preview Output
 	case EDisplayClusterViewportResourceType::OutputPreviewTargetableResource:
 		if (ConfigurationProxy->GetRenderFrameSettings().ShouldUseHoldout())
 		{
 			// The HoldoutComposite plugin expects the input in linear gamma.
-			return { EDisplayClusterColorEncoding::Linear, Premultiply };
+			return { EDisplayClusterColorEncoding::Linear };
 		}
+		break;
 
+	default:
 		break;
 	}
 
-	if (ConfigurationProxy->GetRenderFrameSettings().IsPostProcessDisabled())
+	if (EnumHasAnyFlags(RenderSettingsICVFX.RuntimeFlags, EDisplayClusterViewportRuntimeICVFXFlags::UVLightcard | EDisplayClusterViewportRuntimeICVFXFlags::Lightcard))
+	{
+		// after the OCIO color space isn't linear anymore for LightCards.
+		return { EDisplayClusterColorEncoding::Gamma, EDisplayClusterColorPremultiply::InvertPremultiply };
+	}
+
+	if (IsPostProcessDisabled())
 	{
-		return { EDisplayClusterColorEncoding::Linear, Premultiply };
+		return { EDisplayClusterColorEncoding::Linear };
 	}
 
-	return { EDisplayClusterColorEncoding::Gamma, Premultiply };
+	return { EDisplayClusterColorEncoding::Gamma };
 }
\ No newline at end of file
diff --git a/Engine/Plugins/Runtime/nDisplay/Source/DisplayCluster/Private/Render/Viewport/DisplayClusterViewportProxy.h b/Engine/Plugins/Runtime/nDisplay/Source/DisplayCluster/Private/Render/Viewport/DisplayClusterViewportProxy.h
--- a/Engine/Plugins/Runtime/nDisplay/Source/DisplayCluster/Private/Render/Viewport/DisplayClusterViewportProxy.h
+++ b/Engine/Plugins/Runtime/nDisplay/Source/DisplayCluster/Private/Render/Viewport/DisplayClusterViewportProxy.h
@@ -324,6 +324,9 @@ private:
 
 	/** Check if there is an RTT source (internal or external) in this viewport proxy. */
 	bool IsInputRenderTargetResourceExists() const;
+	
+	/** Return true if postprocess is disbled for this viewport. */
+	bool IsPostProcessDisabled() const;
 
 public:
 	// Configuration for proxy
diff --git a/Engine/Plugins/Runtime/nDisplay/Source/DisplayCluster/Private/Render/Viewport/DisplayClusterViewport_OpenColorIO.cpp b/Engine/Plugins/Runtime/nDisplay/Source/DisplayCluster/Private/Render/Viewport/DisplayClusterViewport_OpenColorIO.cpp
--- a/Engine/Plugins/Runtime/nDisplay/Source/DisplayCluster/Private/Render/Viewport/DisplayClusterViewport_OpenColorIO.cpp
+++ b/Engine/Plugins/Runtime/nDisplay/Source/DisplayCluster/Private/Render/Viewport/DisplayClusterViewport_OpenColorIO.cpp
@@ -55,24 +55,19 @@ bool FDisplayClusterViewport_OpenColorIO::IsConversionSettingsEqual(const FOpenC
 
 float FDisplayClusterViewport_OpenColorIO::GetGammaCorrection(const struct FDisplayClusterColorEncoding& InColorEncoding)
 {
-	// Convert input to the sRGB: default DisplayGamma is 2.2 which corresponds to SRGB
-	const EDisplayClusterColorEncoding EqualEncoding = InColorEncoding.GetEqualEncoding();
-	switch (EqualEncoding)
+	// OCIO shader applies gamma correction BEFORE unpremult
+	// do not do gamma correction for pre-multiplied textures
+	if (InColorEncoding.Premultiply == EDisplayClusterColorPremultiply::None)
 	{
-	case EDisplayClusterColorEncoding::Linear: // Linear -> DefaultGamma/sRGB
-		return UTextureRenderTarget::GetDefaultDisplayGamma();
-
-	case EDisplayClusterColorEncoding::Gamma: // CustomGamma -> DefaultGamma/sRGB
-		return UTextureRenderTarget::GetDefaultDisplayGamma() / InColorEncoding.GammaValue;
-
-	case EDisplayClusterColorEncoding::sRGB: // DefaultGamma/sRGB
-		return 1.f;
-
-	default:
-		break;
-}
+		const EDisplayClusterColorEncoding EqualEncoding = InColorEncoding.GetEqualEncoding();
+		if (EqualEncoding == EDisplayClusterColorEncoding::Gamma)
+		{
+			// Get gamma correction value
+			return UTextureRenderTarget::GetDefaultDisplayGamma() / InColorEncoding.GammaValue;
+		}
+	}
 
-	// Unsupported input encoding, but we'll return 1 anyway, meaning we'll use the input as is.	
+	// No gamma correction
 	return 1.f;
 }
 
diff --git a/Engine/Plugins/Runtime/nDisplay/Source/DisplayCluster/Private/Render/Viewport/LightCard/DisplayClusterViewportLightCardManager.cpp b/Engine/Plugins/Runtime/nDisplay/Source/DisplayCluster/Private/Render/Viewport/LightCard/DisplayClusterViewportLightCardManager.cpp
--- a/Engine/Plugins/Runtime/nDisplay/Source/DisplayCluster/Private/Render/Viewport/LightCard/DisplayClusterViewportLightCardManager.cpp
+++ b/Engine/Plugins/Runtime/nDisplay/Source/DisplayCluster/Private/Render/Viewport/LightCard/DisplayClusterViewportLightCardManager.cpp
@@ -300,19 +300,6 @@ void FDisplayClusterViewportLightCardManager::RenderUVLightCard(const EDisplayCl
 		{
 			FDisplayClusterShaderParameters_UVLightCards UVLightCardParameters;
 			UVLightCardParameters.ProjectionPlaneSize = ADisplayClusterLightCardActor::UVPlaneDefaultSize;
-			if (!ViewportManager->ShouldRenderFinalColor())
-			{
-				// SceneCaptureSource = ESceneCaptureSource::SCS_SceneColorHDR
-				// Output gamma is linear
-				UVLightCardParameters.bRenderFinalColor = false;
-				UVLightCardParameters.LightCardGamma = 1;
-			}
-			else
-			{
-				// Output gamma is sRGB, have to convert to linear
-				UVLightCardParameters.bRenderFinalColor = true;
-				UVLightCardParameters.LightCardGamma = 2.2;
-			}
 
 			// Store any components that were invisible but forced to be visible so they can be set back to invisible after the render
 			TArray<UPrimitiveComponent*> ComponentsToUnload;
diff --git a/Engine/Plugins/Runtime/nDisplay/Source/DisplayCluster/Private/Render/Viewport/RenderTarget/DisplayClusterRenderTargetManager.cpp b/Engine/Plugins/Runtime/nDisplay/Source/DisplayCluster/Private/Render/Viewport/RenderTarget/DisplayClusterRenderTargetManager.cpp
--- a/Engine/Plugins/Runtime/nDisplay/Source/DisplayCluster/Private/Render/Viewport/RenderTarget/DisplayClusterRenderTargetManager.cpp
+++ b/Engine/Plugins/Runtime/nDisplay/Source/DisplayCluster/Private/Render/Viewport/RenderTarget/DisplayClusterRenderTargetManager.cpp
@@ -29,6 +29,25 @@ static FAutoConsoleVariableRef CVarDisplayClusterSceneColorFormat(
 ////////////////////////////////////////////////////////////////////////////////////////////////////////////
 namespace UE::DisplayCluster::RenderTargetManager
 {
+	/** Return custom flags for capture mode. */
+	static EDisplayClusterViewportResourceSettingsFlags ImplGetCustomFlags(EDisplayClusterViewportCaptureMode CaptureMode)
+	{
+		switch (CaptureMode)
+		{
+		case EDisplayClusterViewportCaptureMode::Chromakey:
+		case EDisplayClusterViewportCaptureMode::Lightcard:
+
+			// LightCards renders in linear color space
+			return EDisplayClusterViewportResourceSettingsFlags::LinearColorSpace;
+
+		default:
+			break;
+		}
+
+		return EDisplayClusterViewportResourceSettingsFlags::None;
+	}
+
+	/** Return texture format for capture mode. */
 	static EPixelFormat ImplGetCustomFormat(EDisplayClusterViewportCaptureMode CaptureMode)
 	{
 		EPixelFormat OutPixelFormat = EPixelFormat::PF_Unknown;
@@ -119,8 +138,11 @@ bool FDisplayClusterRenderTargetManager::AllocateRenderFrameResources(FViewport*
 		{
 			if (FrameRenderTargetIt.bShouldUseRenderTarget)
 			{
+				const EDisplayClusterViewportResourceSettingsFlags CustomFlags = ImplGetCustomFlags(FrameRenderTargetIt.CaptureMode);
+				const EPixelFormat CustomFormat = ImplGetCustomFormat(FrameRenderTargetIt.CaptureMode);
+
 				// reallocate
-				TSharedPtr<FDisplayClusterViewportResource, ESPMode::ThreadSafe> NewResource = ResourcesPool->AllocateResource(EmptyString, FrameRenderTargetIt.RenderTargetSize, ImplGetCustomFormat(FrameRenderTargetIt.CaptureMode), EDisplayClusterViewportResourceSettingsFlags::RenderTarget);
+				TSharedPtr<FDisplayClusterViewportResource, ESPMode::ThreadSafe> NewResource = ResourcesPool->AllocateResource(EmptyString, FrameRenderTargetIt.RenderTargetSize, CustomFormat, EDisplayClusterViewportResourceSettingsFlags::RenderTarget | CustomFlags);
 				if (NewResource.IsValid())
 				{
 					// Set RenderFrame resource
@@ -155,7 +177,10 @@ bool FDisplayClusterRenderTargetManager::AllocateRenderFrameResources(FViewport*
 		{
 			if (ViewportIt.IsValid())
 			{
-				EPixelFormat CustomFormat = ImplGetCustomFormat(ViewportIt->GetRenderSettings().CaptureMode);
+				const EDisplayClusterViewportCaptureMode CaptureMode = ViewportIt->GetRenderSettings().CaptureMode;
+
+				const EDisplayClusterViewportResourceSettingsFlags CustomFlags = ImplGetCustomFlags(CaptureMode);
+				const EPixelFormat CustomFormat = ImplGetCustomFormat(CaptureMode);
 
 				// Allocate all context resources:
 				for (const FDisplayClusterViewport_Context& ContextIt : ViewportIt->GetContexts())
@@ -169,7 +194,7 @@ bool FDisplayClusterRenderTargetManager::AllocateRenderFrameResources(FViewport*
 						TSharedPtr<FDisplayClusterViewportResource, ESPMode::ThreadSafe>& ReallocResource = ViewportIt->GetViewportResourcesImpl(EDisplayClusterViewportResource::OutputPreviewTargetableResources)[ContextNum];
 						if (!ReallocResource.IsValid() || !EnumHasAnyFlags(ReallocResource->GetResourceState(), EDisplayClusterViewportResourceState::DisableReallocate))
 						{
-							ReallocResource = ResourcesPool->AllocateResource(ViewportIt->GetId(), ContextSize, CustomFormat, EDisplayClusterViewportResourceSettingsFlags::PreviewTargetableTexture);
+							ReallocResource = ResourcesPool->AllocateResource(ViewportIt->GetId(), ContextSize, CustomFormat, EDisplayClusterViewportResourceSettingsFlags::PreviewTargetableTexture | CustomFlags);
 							ImplViewportTextureResourceLogging(ViewportIt, ContextNum, TEXT("InputShader"), ViewportIt->GetViewportResources(EDisplayClusterViewportResource::OutputPreviewTargetableResources)[ContextNum]);
 						}
 					}
@@ -179,7 +204,7 @@ bool FDisplayClusterRenderTargetManager::AllocateRenderFrameResources(FViewport*
 						TSharedPtr<FDisplayClusterViewportResource, ESPMode::ThreadSafe>& ReallocResource = ViewportIt->GetViewportResourcesImpl(EDisplayClusterViewportResource::InputShaderResources)[ContextNum];
 						if (!ReallocResource.IsValid() || !EnumHasAnyFlags(ReallocResource->GetResourceState(), EDisplayClusterViewportResourceState::DisableReallocate))
 						{
-							ReallocResource = ResourcesPool->AllocateResource(EmptyString, ContextSize, CustomFormat, EDisplayClusterViewportResourceSettingsFlags::ResolveTargetableTexture);
+							ReallocResource = ResourcesPool->AllocateResource(EmptyString, ContextSize, CustomFormat, EDisplayClusterViewportResourceSettingsFlags::ResolveTargetableTexture | CustomFlags);
 							ImplViewportTextureResourceLogging(ViewportIt, ContextNum, TEXT("InputShader"), ViewportIt->GetViewportResources(EDisplayClusterViewportResource::InputShaderResources)[ContextNum]);
 						}
 					}
@@ -190,7 +215,7 @@ bool FDisplayClusterRenderTargetManager::AllocateRenderFrameResources(FViewport*
 						TSharedPtr<FDisplayClusterViewportResource, ESPMode::ThreadSafe>& ReallocResource = ViewportIt->GetViewportResourcesImpl(EDisplayClusterViewportResource::AdditionalTargetableResources)[ContextNum];
 						if (!ReallocResource.IsValid() || !EnumHasAnyFlags(ReallocResource->GetResourceState(), EDisplayClusterViewportResourceState::DisableReallocate))
 						{
-							ReallocResource = ResourcesPool->AllocateResource(EmptyString, ContextSize, CustomFormat, EDisplayClusterViewportResourceSettingsFlags::RenderTargetableTexture);
+							ReallocResource = ResourcesPool->AllocateResource(EmptyString, ContextSize, CustomFormat, EDisplayClusterViewportResourceSettingsFlags::RenderTargetableTexture | CustomFlags);
 							ImplViewportTextureResourceLogging(ViewportIt, ContextNum, TEXT("AdditionalRTT"), ViewportIt->GetViewportResources(EDisplayClusterViewportResource::AdditionalTargetableResources)[ContextNum]);
 						}
 					}
@@ -200,7 +225,7 @@ bool FDisplayClusterRenderTargetManager::AllocateRenderFrameResources(FViewport*
 						TSharedPtr<FDisplayClusterViewportResource, ESPMode::ThreadSafe>& ReallocResource = ViewportIt->GetViewportResourcesImpl(EDisplayClusterViewportResource::MipsShaderResources)[ContextNum];
 						if (!ReallocResource.IsValid() || !EnumHasAnyFlags(ReallocResource->GetResourceState(), EDisplayClusterViewportResourceState::DisableReallocate))
 						{
-							ReallocResource = ResourcesPool->AllocateResource(EmptyString, ContextSize, CustomFormat, EDisplayClusterViewportResourceSettingsFlags::ResolveTargetableTexture, ContextIt.NumMips);
+							ReallocResource = ResourcesPool->AllocateResource(EmptyString, ContextSize, CustomFormat, EDisplayClusterViewportResourceSettingsFlags::ResolveTargetableTexture | CustomFlags, ContextIt.NumMips);
 							ImplViewportTextureResourceLogging(ViewportIt, ContextNum, TEXT("Mips"), ViewportIt->GetViewportResources(EDisplayClusterViewportResource::MipsShaderResources)[ContextNum]);
 						}
 					}
diff --git a/Engine/Plugins/Runtime/nDisplay/Source/DisplayCluster/Private/Render/Viewport/RenderTarget/DisplayClusterRenderTargetResource.h b/Engine/Plugins/Runtime/nDisplay/Source/DisplayCluster/Private/Render/Viewport/RenderTarget/DisplayClusterRenderTargetResource.h
--- a/Engine/Plugins/Runtime/nDisplay/Source/DisplayCluster/Private/Render/Viewport/RenderTarget/DisplayClusterRenderTargetResource.h
+++ b/Engine/Plugins/Runtime/nDisplay/Source/DisplayCluster/Private/Render/Viewport/RenderTarget/DisplayClusterRenderTargetResource.h
@@ -108,7 +108,7 @@ public:
 		return GetResourceSettings().GetSizeXY();
 	}
 
-	virtual float GetDisplayGamma() const 
+	virtual float GetDisplayGamma() const override
 	{
 		return GetResourceSettings().GetDisplayGamma();
 	}
diff --git a/Engine/Plugins/Runtime/nDisplay/Source/DisplayCluster/Private/Render/Viewport/Resource/DisplayClusterViewportResourceSettings.cpp b/Engine/Plugins/Runtime/nDisplay/Source/DisplayCluster/Private/Render/Viewport/Resource/DisplayClusterViewportResourceSettings.cpp
--- a/Engine/Plugins/Runtime/nDisplay/Source/DisplayCluster/Private/Render/Viewport/Resource/DisplayClusterViewportResourceSettings.cpp
+++ b/Engine/Plugins/Runtime/nDisplay/Source/DisplayCluster/Private/Render/Viewport/Resource/DisplayClusterViewportResourceSettings.cpp
@@ -58,4 +58,10 @@ FDisplayClusterViewportResourceSettings::FDisplayClusterViewportResourceSettings
 	, ResourceFlags(InBaseSettings.ResourceFlags)
 {
 	EnumAddFlags(ResourceFlags, InResourceFlags);
+
+	// For a texture in linear color space, set Gamma to 1.0f
+	if (EnumHasAnyFlags(InResourceFlags, EDisplayClusterViewportResourceSettingsFlags::LinearColorSpace))
+	{
+		DisplayGamma = 1.0f;
+	}
 }
diff --git a/Engine/Plugins/Runtime/nDisplay/Source/DisplayCluster/Private/Render/Viewport/Resource/DisplayClusterViewportResourceSettings.h b/Engine/Plugins/Runtime/nDisplay/Source/DisplayCluster/Private/Render/Viewport/Resource/DisplayClusterViewportResourceSettings.h
--- a/Engine/Plugins/Runtime/nDisplay/Source/DisplayCluster/Private/Render/Viewport/Resource/DisplayClusterViewportResourceSettings.h
+++ b/Engine/Plugins/Runtime/nDisplay/Source/DisplayCluster/Private/Render/Viewport/Resource/DisplayClusterViewportResourceSettings.h
@@ -33,6 +33,10 @@ enum class EDisplayClusterViewportResourceSettingsFlags : uint8
 	// This is the texture that is used to render the DCRA preview
 	// (contains the UTexture that is used for UMaterialInstanceDynamic).
 	PreviewTargetableTexture = 1 << 5,
+
+	// This texture is in linear color space
+	// DisplayGamma should be set to 1.0f
+	LinearColorSpace = 1 << 6,
 };
 ENUM_CLASS_FLAGS(EDisplayClusterViewportResourceSettingsFlags);
 
diff --git a/Engine/Plugins/Runtime/nDisplay/Source/DisplayCluster/Public/Render/Viewport/IDisplayClusterViewportManager.h b/Engine/Plugins/Runtime/nDisplay/Source/DisplayCluster/Public/Render/Viewport/IDisplayClusterViewportManager.h
--- a/Engine/Plugins/Runtime/nDisplay/Source/DisplayCluster/Public/Render/Viewport/IDisplayClusterViewportManager.h
+++ b/Engine/Plugins/Runtime/nDisplay/Source/DisplayCluster/Public/Render/Viewport/IDisplayClusterViewportManager.h
@@ -107,6 +107,9 @@ public:
 	/** This function sets the engine flags corresponding to the rendering mode. */
 	static void SetupEngineShowFlags(const EDisplayClusterViewportCaptureMode InMode, FEngineShowFlags& InOutEngineShowFlags);
 
+	/** Sets the ViewFamily and SceneView properties correspond to the rendering mode. */
+	static void SetupSceneView(const EDisplayClusterViewportCaptureMode InMode, FSceneViewFamily& InOutViewFamily, FSceneView& InOutView);
+
 	/**
 	* Initialize view family, using rules
 	* [Game thread func]
diff --git a/Engine/Plugins/Runtime/nDisplay/Source/DisplayClusterShaders/Private/Shaders/DisplayClusterShadersPreprocess_UVLightCards.cpp b/Engine/Plugins/Runtime/nDisplay/Source/DisplayClusterShaders/Private/Shaders/DisplayClusterShadersPreprocess_UVLightCards.cpp
--- a/Engine/Plugins/Runtime/nDisplay/Source/DisplayClusterShaders/Private/Shaders/DisplayClusterShadersPreprocess_UVLightCards.cpp
+++ b/Engine/Plugins/Runtime/nDisplay/Source/DisplayClusterShaders/Private/Shaders/DisplayClusterShadersPreprocess_UVLightCards.cpp
@@ -181,7 +181,6 @@ class FGammaCorrectionPS : public FGlobalShader
 	BEGIN_SHADER_PARAMETER_STRUCT(FParameters, )
 		SHADER_PARAMETER_RDG_TEXTURE(Texture2D, InputTexture)
 		SHADER_PARAMETER_SAMPLER(SamplerState, InputSampler)
-		SHADER_PARAMETER(float, InverseGamma)
 		RENDER_TARGET_BINDING_SLOTS()
 		END_SHADER_PARAMETER_STRUCT()
 };
@@ -211,6 +210,7 @@ bool FDisplayClusterShadersPreprocess_UVLightCards::RenderPreprocess_UVLightCard
 		InRenderTarget,
 		InScene,
 		EngineShowFlags)
+		.SetResolveScene(true)
 		.SetTime(FGameTime::GetTimeSinceAppStart()));
 
 	FScenePrimitiveRenderingContextScopeHelper ScenePrimitiveRenderingContextScopeHelper(GetRendererModule().BeginScenePrimitiveRendering(GraphBuilder, &ViewFamily));
@@ -258,15 +258,10 @@ bool FDisplayClusterShadersPreprocess_UVLightCards::RenderPreprocess_UVLightCard
 	GetRendererModule().CreateAndInitSingleView(RHICmdList, &ViewFamily, &ViewInitOptions);
 	FSceneView* View = (FSceneView*)ViewFamily.Views[0];
 
-	if (!InParameters.bRenderFinalColor)
-	{
-		ViewFamily.SceneCaptureSource = ESceneCaptureSource::SCS_SceneColorHDR;
-	}
+	ViewFamily.SceneCaptureCompositeMode = ESceneCaptureCompositeMode::SCCM_Overwrite;
 
-	ViewFamily.EngineShowFlags.SetToneCurve(false);
-	// This flags sets tonemapper to output to ETonemapperOutputDevice::LinearNoToneCurve
-	View->FinalPostProcessSettings.bOverride_ToneCurveAmount = 1;
-	View->FinalPostProcessSettings.ToneCurveAmount = 0.0;
+	// Use the same render settings as for the 3D LC.
+	IDisplayClusterViewportManager::SetupSceneView(EDisplayClusterViewportCaptureMode::Lightcard, ViewFamily, *View);
 
 	FRDGTextureRef OutputTexture = GraphBuilder.RegisterExternalTexture(CreateRenderTarget(InRenderTarget->GetRenderTargetTexture(), TEXT("UVLightCardRenderTarget")));
 	FRenderTargetBinding OutputRenderTargetBinding(OutputTexture, ERenderTargetLoadAction::EClear);
@@ -310,7 +305,6 @@ bool FDisplayClusterShadersPreprocess_UVLightCards::RenderPreprocess_UVLightCard
 	FGammaCorrectionPS::FParameters* GammaCorrectionParameters = GraphBuilder.AllocParameters<FGammaCorrectionPS::FParameters>();
 	GammaCorrectionParameters->InputTexture = ColorTexture;
 	GammaCorrectionParameters->InputSampler = TStaticSamplerState<>::GetRHI();
-	GammaCorrectionParameters->InverseGamma = 1.0 / InParameters.LightCardGamma;
 	GammaCorrectionParameters->RenderTargets[0] = OutputRenderTargetBinding;
 
 	FGlobalShaderMap* GlobalShaderMap = GetGlobalShaderMap(View->FeatureLevel);
diff --git a/Engine/Plugins/Runtime/nDisplay/Source/DisplayClusterShaders/Public/ShaderParameters/DisplayClusterShaderParameters_UVLightCards.h b/Engine/Plugins/Runtime/nDisplay/Source/DisplayClusterShaders/Public/ShaderParameters/DisplayClusterShaderParameters_UVLightCards.h
--- a/Engine/Plugins/Runtime/nDisplay/Source/DisplayClusterShaders/Public/ShaderParameters/DisplayClusterShaderParameters_UVLightCards.h
+++ b/Engine/Plugins/Runtime/nDisplay/Source/DisplayClusterShaders/Public/ShaderParameters/DisplayClusterShaderParameters_UVLightCards.h
@@ -13,11 +13,5 @@ struct FDisplayClusterShaderParameters_UVLightCards
 	TArray<FPrimitiveSceneProxy*> PrimitivesToRender;
 
 	/** The size of the plane in world units to project the UV space onto */
-	float ProjectionPlaneSize;
-
-	/** The gamma to apply to the light cards when gamma correcting them */
-	float LightCardGamma = 1.0;
-
-	/** Indicates whether to render the final color of the light cards or not */
-	bool bRenderFinalColor;
+	float ProjectionPlaneSize = 200.0f;
 };
