nDisplay: Added ability to override global upscaler settings on a per-viewport basis, and made various enhancements and fixes to the upscaler settings details customization.

diff --git a/Engine/Plugins/Runtime/nDisplay/Source/DisplayCluster/Public/Components/DisplayClusterICVFXCameraComponent.h b/Engine/Plugins/Runtime/nDisplay/Source/DisplayCluster/Public/Components/DisplayClusterICVFXCameraComponent.h
--- a/Engine/Plugins/Runtime/nDisplay/Source/DisplayCluster/Public/Components/DisplayClusterICVFXCameraComponent.h
+++ b/Engine/Plugins/Runtime/nDisplay/Source/DisplayCluster/Public/Components/DisplayClusterICVFXCameraComponent.h
@@ -151,14 +151,14 @@ private:
 	UPROPERTY(EditAnywhere, Transient, Category = "Inner Frustum", meta = (PropertyPath = "CameraSettings.RenderSettings.GenerateMips"))
 	FDisplayClusterEditorPropertyReference GenerateMipsRef;
 
-	UPROPERTY(EditAnywhere, Transient, Category = "ICVFX Camera", meta = (PropertyPath = "CameraSettings.ExternalCameraActor"))
+	UPROPERTY(EditAnywhere, Transient, Category = "Camera ICVFX", meta = (PropertyPath = "CameraSettings.ExternalCameraActor"))
 	FDisplayClusterEditorPropertyReference ExternalCameraActorRef;
 
 	/** Exposed reference to the camera's inner depth of field settings */
-	UPROPERTY(EditAnywhere, Transient, Category = "ICVFX Camera", meta = (PropertyPath = "CameraSettings.CameraDepthOfField", DisplayName = "ICVFX Depth of Field"))
+	UPROPERTY(EditAnywhere, Transient, Category = "Camera ICVFX", meta = (PropertyPath = "CameraSettings.CameraDepthOfField", DisplayName = "ICVFX Depth of Field"))
 	FDisplayClusterEditorPropertyReference CameraDepthOfFieldRef;
 
-	UPROPERTY(EditAnywhere, Transient, Category = "ICVFX Camera", meta = (PropertyPath = "CameraSettings.CameraMotionBlur", DisplayName = "ICVFX Camera Motion Blur"))
+	UPROPERTY(EditAnywhere, Transient, Category = "Camera ICVFX", meta = (PropertyPath = "CameraSettings.CameraMotionBlur", DisplayName = "ICVFX Camera Motion Blur"))
 	FDisplayClusterEditorPropertyReference CameraMotionBlurRef;
 
 	UPROPERTY(EditAnywhere, Transient, Category = "Inner Frustum", meta = (PropertyPath = "CameraSettings.CameraHideList"))
diff --git a/Engine/Plugins/Runtime/nDisplay/Source/DisplayCluster/Public/DisplayClusterRootActor.h b/Engine/Plugins/Runtime/nDisplay/Source/DisplayCluster/Public/DisplayClusterRootActor.h
--- a/Engine/Plugins/Runtime/nDisplay/Source/DisplayCluster/Public/DisplayClusterRootActor.h
+++ b/Engine/Plugins/Runtime/nDisplay/Source/DisplayCluster/Public/DisplayClusterRootActor.h
@@ -298,7 +298,7 @@ private:
 	UPROPERTY(EditAnywhere, Transient, Category = Viewports, meta = (PropertyPath = "CurrentConfigData.StageSettings.OuterViewportUpscalerSettings", DisplayName = "Default Upscaler Settings"))
 	FDisplayClusterEditorPropertyReference OuterViewportUpscalerSettingsRef;
 
-	UPROPERTY(EditInstanceOnly, Transient, Category = Viewports, meta = (DisplayName = "Viewport Upscaling Settings", PropertyPath = "CurrentConfigData.Cluster.Nodes.Viewports.RenderSettings.UpscalerSettings", ToolTip = "Adjust upscaler settings for an individual viewport."))
+	UPROPERTY(EditInstanceOnly, Transient, Category = Viewports, meta = (DisplayName = "Per-Viewport Upscaling Setting Overrides", PropertyPath = "CurrentConfigData.Cluster.Nodes.Viewports.RenderSettings.UpscalerSettings", ToolTip = "Adjust upscaler settings for an individual viewport."))
 	FDisplayClusterEditorPropertyReference ViewportUpscalerSettingsRef;
 
 	UPROPERTY(EditAnywhere, Transient, Category = Viewports, meta = (PropertyPath = "CurrentConfigData.RenderFrameSettings.ClusterICVFXOuterViewportBufferRatioMult"))
diff --git a/Engine/Plugins/Runtime/nDisplay/Source/DisplayClusterConfiguration/Private/DisplayClusterConfigurationTypes.cpp b/Engine/Plugins/Runtime/nDisplay/Source/DisplayClusterConfiguration/Private/DisplayClusterConfigurationTypes.cpp
--- a/Engine/Plugins/Runtime/nDisplay/Source/DisplayClusterConfiguration/Private/DisplayClusterConfigurationTypes.cpp
+++ b/Engine/Plugins/Runtime/nDisplay/Source/DisplayClusterConfiguration/Private/DisplayClusterConfigurationTypes.cpp
@@ -598,3 +598,67 @@ UDisplayClusterConfigurationData* UDisplayClusterConfigurationData::CreateNewCon
 
 	return NewConfigData;
 }
+
+#if WITH_EDITOR
+void UDisplayClusterConfigurationData::PostEditChangeChainProperty(struct FPropertyChangedChainEvent& PropertyChangedEvent)
+{
+	auto ChainContainsProperty = [&PropertyChangedEvent](const FName& InPropertyName)
+	{
+		TDoubleLinkedList<FProperty*>::TDoubleLinkedListNode* CurrentNode = PropertyChangedEvent.PropertyChain.GetActiveNode();
+		while (CurrentNode)
+		{
+			if (FProperty* Property = CurrentNode->GetValue())
+			{
+				if (Property->GetFName() == InPropertyName)
+				{
+					return  true;
+				}
+			}
+
+			CurrentNode = CurrentNode->GetPrevNode();
+		}
+
+		return false;
+	};
+	
+	const FName PropertyName = PropertyChangedEvent.GetPropertyName();
+	if (PropertyName == GET_MEMBER_NAME_CHECKED(FDisplayClusterConfigurationUpscalerSettings, MethodName)
+		&& ChainContainsProperty(GET_MEMBER_NAME_CHECKED(FDisplayClusterConfigurationICVFX_StageSettings, OuterViewportUpscalerSettings)))
+	{
+		Modify();
+		
+		// All viewport upscaling methods should match the global outer viewport upscaling method
+		for (TTuple<FString, TObjectPtr<UDisplayClusterConfigurationClusterNode>>& Node : Cluster->Nodes)
+		{
+			for (TTuple<FString, TObjectPtr<UDisplayClusterConfigurationViewport>>& Viewport : Node.Value->Viewports)
+			{
+				Viewport.Value->Modify();
+				
+				Viewport.Value->RenderSettings.UpscalerSettings.MethodName = StageSettings.OuterViewportUpscalerSettings.MethodName;
+				Viewport.Value->RenderSettings.UpscalerSettings.EditingData = StageSettings.OuterViewportUpscalerSettings.EditingData;
+				Viewport.Value->RenderSettings.UpscalerSettings.ParameterOverrideGuids.Empty();
+			}
+		}
+	}
+	else if (ChainContainsProperty(GET_MEMBER_NAME_CHECKED(FDisplayClusterConfigurationUpscalerSettings, EditingData))
+		&& ChainContainsProperty(GET_MEMBER_NAME_CHECKED(FDisplayClusterConfigurationICVFX_StageSettings, OuterViewportUpscalerSettings)))
+	{
+		Modify();
+		
+		// A property in the upscaler settings property bag has changed. Propagate the global property bag values to the viewports, considering viewport overrides
+		for (TTuple<FString, TObjectPtr<UDisplayClusterConfigurationClusterNode>>& Node : Cluster->Nodes)
+		{
+			for (TTuple<FString, TObjectPtr<UDisplayClusterConfigurationViewport>>& Viewport : Node.Value->Viewports)
+			{
+				Viewport.Value->Modify();
+				
+				Viewport.Value->RenderSettings.UpscalerSettings.EditingData.MigrateToNewBagInstanceWithOverrides(
+					StageSettings.OuterViewportUpscalerSettings.EditingData,
+					Viewport.Value->RenderSettings.UpscalerSettings.ParameterOverrideGuids);
+			}
+		}
+	}
+	
+	Super::PostEditChangeProperty(PropertyChangedEvent);
+}
+#endif
diff --git a/Engine/Plugins/Runtime/nDisplay/Source/DisplayClusterConfiguration/Private/DisplayClusterConfigurationTypes_Upscaler.cpp b/Engine/Plugins/Runtime/nDisplay/Source/DisplayClusterConfiguration/Private/DisplayClusterConfigurationTypes_Upscaler.cpp
--- a/Engine/Plugins/Runtime/nDisplay/Source/DisplayClusterConfiguration/Private/DisplayClusterConfigurationTypes_Upscaler.cpp
+++ b/Engine/Plugins/Runtime/nDisplay/Source/DisplayClusterConfiguration/Private/DisplayClusterConfigurationTypes_Upscaler.cpp
@@ -65,6 +65,12 @@ bool FDisplayClusterConfigurationUpscalerSettings::GetUpscalerSettings(
 				// Custom upscaler not used.
 				OutUpscalerSettings.CustomUpscalerName = NAME_None;
 				GetConfigurationData(*this, OutUpscalerSettings.CustomUpscalerSettings);
+				
+				// Migrate any data from the default upscaler settings that isn't explicitly overridden in the parameter overrides
+				if (InDefaultUpscalerSettings)
+				{
+					OutUpscalerSettings.CustomUpscalerSettings.MigrateToNewBagInstanceWithOverrides(InDefaultUpscalerSettings->EditingData, ParameterOverrideGuids);
+				}
 
 				return true;
 			}
@@ -79,6 +85,12 @@ bool FDisplayClusterConfigurationUpscalerSettings::GetUpscalerSettings(
 		// Set custom upscaler name and PropertyBag with settings.
 		OutUpscalerSettings.CustomUpscalerName = MethodName;
 		GetConfigurationData(*this, OutUpscalerSettings.CustomUpscalerSettings);
+
+		// Migrate any data from the default upscaler settings that isn't explicitly overridden in the parameter overrides
+		if (InDefaultUpscalerSettings)
+		{
+			OutUpscalerSettings.CustomUpscalerSettings.MigrateToNewBagInstanceWithOverrides(InDefaultUpscalerSettings->EditingData, ParameterOverrideGuids);
+		}
 	}
 
 	return MethodName != NAME_None;
diff --git a/Engine/Plugins/Runtime/nDisplay/Source/DisplayClusterConfiguration/Public/DisplayClusterConfigurationTypes.h b/Engine/Plugins/Runtime/nDisplay/Source/DisplayClusterConfiguration/Public/DisplayClusterConfigurationTypes.h
--- a/Engine/Plugins/Runtime/nDisplay/Source/DisplayClusterConfiguration/Public/DisplayClusterConfigurationTypes.h
+++ b/Engine/Plugins/Runtime/nDisplay/Source/DisplayClusterConfiguration/Public/DisplayClusterConfigurationTypes.h
@@ -667,4 +667,11 @@ private:
 	FDisplayClusterEditorPropertyReference DefaultFrameSizeRef;
 
 #endif // WITH_EDITORONLY_DATA
+
+#if WITH_EDITOR
+public:
+	// UObject interface
+	virtual void PostEditChangeChainProperty(struct FPropertyChangedChainEvent& PropertyChangedEvent) override;
+	// End of UObject interface
+#endif
 };
diff --git a/Engine/Plugins/Runtime/nDisplay/Source/DisplayClusterConfiguration/Public/DisplayClusterConfigurationTypes_Enums.h b/Engine/Plugins/Runtime/nDisplay/Source/DisplayClusterConfiguration/Public/DisplayClusterConfigurationTypes_Enums.h
--- a/Engine/Plugins/Runtime/nDisplay/Source/DisplayClusterConfiguration/Public/DisplayClusterConfigurationTypes_Enums.h
+++ b/Engine/Plugins/Runtime/nDisplay/Source/DisplayClusterConfiguration/Public/DisplayClusterConfigurationTypes_Enums.h
@@ -266,14 +266,14 @@ enum class EDisplayClusterConfigurationUpscalingMethod : uint8
 	Disable UMETA(DisplayName = "Disable"),
 
 	/** Fast Approximate Anti-Aliasing. */
-	FXAA UMETA(DisplayName = "FXAA"),
+	FXAA UMETA(DisplayName = "Fast Approximate Anti-Aliasing (FXAA)"),
 
 	/** Multisample Anti-Aliasing. (Only supported with forward shading.) */
-	MSAA UMETA(DisplayName = "MSAA"),
+	MSAA UMETA(DisplayName = "Multisample Anti-Aliasing (MSAA)"),
 
 	/** Temporal Anti-Aliasing. */
-	TAA UMETA(DisplayName = "TAA"),
+	TAA UMETA(DisplayName = "Temporal Anti-Aliasing (TAA)"),
 
 	/** Temporal Super-Resolution. Use TSR upscaler with built-in AntiAliasing. */
-	TSR UMETA(DisplayName = "TSR"),
+	TSR UMETA(DisplayName = "Temporal Super-Resolution (TSR)"),
 };
diff --git a/Engine/Plugins/Runtime/nDisplay/Source/DisplayClusterConfiguration/Public/DisplayClusterConfigurationTypes_Upscaler.h b/Engine/Plugins/Runtime/nDisplay/Source/DisplayClusterConfiguration/Public/DisplayClusterConfigurationTypes_Upscaler.h
--- a/Engine/Plugins/Runtime/nDisplay/Source/DisplayClusterConfiguration/Public/DisplayClusterConfigurationTypes_Upscaler.h
+++ b/Engine/Plugins/Runtime/nDisplay/Source/DisplayClusterConfiguration/Public/DisplayClusterConfigurationTypes_Upscaler.h
@@ -44,4 +44,12 @@ public:
 	/** Upscaling method settings. */
 	UPROPERTY(EditAnywhere, Category = "Upscaler", meta = (DisplayName = "Upscaling Settings"))
 	FInstancedPropertyBag EditingData;
+
+	/**
+	 * In cases where these upscaler settings are overriding global upscaler settings, this list stores the IDs of the properties
+	 * that will override the global settings. These are exposed and utilized when the owner of the UpscalerSettings struct uses
+	 * the WithOverrides metadata specifier
+	 */
+	UPROPERTY()
+	TArray<FGuid> ParameterOverrideGuids;
 };
diff --git a/Engine/Plugins/Runtime/nDisplay/Source/DisplayClusterConfiguration/Public/DisplayClusterConfigurationTypes_Viewport.h b/Engine/Plugins/Runtime/nDisplay/Source/DisplayClusterConfiguration/Public/DisplayClusterConfigurationTypes_Viewport.h
--- a/Engine/Plugins/Runtime/nDisplay/Source/DisplayClusterConfiguration/Public/DisplayClusterConfigurationTypes_Viewport.h
+++ b/Engine/Plugins/Runtime/nDisplay/Source/DisplayClusterConfiguration/Public/DisplayClusterConfigurationTypes_Viewport.h
@@ -112,7 +112,7 @@ public:
 	EDisplayClusterConfigurationViewport_StereoMode StereoMode = EDisplayClusterConfigurationViewport_StereoMode::Default;
 
 	/** Upscaler settings for the viewport. */
-	UPROPERTY(EditAnywhere, BlueprintReadWrite, Category = "Rendering", meta = (DisplayName = "Upscaler Settings"))
+	UPROPERTY(EditAnywhere, BlueprintReadWrite, Category = "Rendering", meta = (DisplayName = "Upscaler Settings", NoResetToDefault, WithOverrides))
 	FDisplayClusterConfigurationUpscalerSettings UpscalerSettings;
 
 	/** Adjust resolution scaling for an individual viewport.  Viewport Screen Percentage Multiplier is applied to this value. */
diff --git a/Engine/Plugins/Runtime/nDisplay/Source/DisplayClusterConfigurator/Private/Views/Details/Components/DisplayClusterICVFXCameraComponentDetailsCustomization.cpp b/Engine/Plugins/Runtime/nDisplay/Source/DisplayClusterConfigurator/Private/Views/Details/Components/DisplayClusterICVFXCameraComponentDetailsCustomization.cpp
--- a/Engine/Plugins/Runtime/nDisplay/Source/DisplayClusterConfigurator/Private/Views/Details/Components/DisplayClusterICVFXCameraComponentDetailsCustomization.cpp
+++ b/Engine/Plugins/Runtime/nDisplay/Source/DisplayClusterConfigurator/Private/Views/Details/Components/DisplayClusterICVFXCameraComponentDetailsCustomization.cpp
@@ -127,6 +127,29 @@ void FDisplayClusterICVFXCameraComponentDetailsCustomization::CustomizeDetails(I
 		CameraCategory.AddProperty(GET_MEMBER_NAME_CHECKED(UDisplayClusterICVFXCameraComponent, CurrentAperture), UCineCameraComponent::StaticClass());
 	}
 
+	if (EditedObject.IsValid())
+	{
+		const FName UpscaleMethodName = EditedObject->CameraSettings.UpscalerSettings.MethodName;
+		bool bIsCustomUpscaleMethod = UpscaleMethodName != NAME_None;
+	
+		if (const UEnum* EnumClass = StaticEnum<EDisplayClusterConfigurationUpscalingMethod>())
+		{
+			for (int32 EnumElementIndex = 0; EnumElementIndex < EnumClass->NumEnums(); ++EnumElementIndex)
+			{
+				if (EnumClass->GetNameStringByIndex(EnumElementIndex) == UpscaleMethodName && !EnumClass->HasMetaData(TEXT("Hidden"), EnumElementIndex))
+				{
+					bIsCustomUpscaleMethod = false;
+				}
+			}
+		}
+	
+		if (bIsCustomUpscaleMethod)
+		{
+			// If the upscale method is a custom upscale method, hide the screen percentage
+			InLayoutBuilder.GetProperty(GET_MEMBER_NAME_CHECKED(UDisplayClusterICVFXCameraComponent, BufferRatioRef))->MarkHiddenByCustomization();
+		}
+	}
+
 	InLayoutBuilder.SortCategories(DisplayClusterICVFXCameraComponentDetailsCustomizationUtils::SortCategories);
 
 	// Most of the properties in the camera settings are exposed through property references, so hide the camera settings property.
diff --git a/Engine/Plugins/Runtime/nDisplay/Source/DisplayClusterConfigurator/Private/Views/Details/DisplayClusterRootActorDetailsCustomization.cpp b/Engine/Plugins/Runtime/nDisplay/Source/DisplayClusterConfigurator/Private/Views/Details/DisplayClusterRootActorDetailsCustomization.cpp
--- a/Engine/Plugins/Runtime/nDisplay/Source/DisplayClusterConfigurator/Private/Views/Details/DisplayClusterRootActorDetailsCustomization.cpp
+++ b/Engine/Plugins/Runtime/nDisplay/Source/DisplayClusterConfigurator/Private/Views/Details/DisplayClusterRootActorDetailsCustomization.cpp
@@ -202,6 +202,38 @@ void FDisplayClusterRootActorDetailsCustomization::CustomizeDetails(IDetailLayou
 				];
 		}
 	}
+
+	if (UDisplayClusterConfigurationData* ConfigData = GetConfigData())
+	{
+		const FName UpscaleMethodName = ConfigData->StageSettings.OuterViewportUpscalerSettings.MethodName;
+		bool bIsCustomUpscaleMethod = UpscaleMethodName != NAME_None;
+
+		if (const UEnum* EnumClass = StaticEnum<EDisplayClusterConfigurationUpscalingMethod>())
+		{
+			for (int32 EnumElementIndex = 0; EnumElementIndex < EnumClass->NumEnums(); ++EnumElementIndex)
+			{
+				if (EnumClass->GetNameStringByIndex(EnumElementIndex) == UpscaleMethodName && !EnumClass->HasMetaData(TEXT("Hidden"), EnumElementIndex))
+				{
+					bIsCustomUpscaleMethod = false;
+				}
+			}
+		}
+
+		if (!bIsCustomUpscaleMethod)
+		{
+			// If the upscale method is not a custom upscale method, hide the per-viewport upscaler settings list
+			InLayoutBuilder.GetProperty(GET_MEMBER_NAME_CHECKED(ADisplayClusterRootActor, ViewportUpscalerSettingsRef))->MarkHiddenByCustomization();
+		}
+		else
+		{
+			// Otherwise, hide the screen percentage properties
+			InLayoutBuilder.GetProperty(GET_MEMBER_NAME_CHECKED(ADisplayClusterRootActor, ViewportScreenPercentageMultiplierRef))->MarkHiddenByCustomization();
+			InLayoutBuilder.GetProperty(GET_MEMBER_NAME_CHECKED(ADisplayClusterRootActor, ViewportScreenPercentageRef))->MarkHiddenByCustomization();
+
+			InLayoutBuilder.GetProperty(GET_MEMBER_NAME_CHECKED(ADisplayClusterRootActor, ViewportUpscalerSettingsRef))->SetPropertyDisplayName(
+				FText::Format(LOCTEXT("PerViewportUpscalerSettingOverrides", "Per-Viewport {0} Setting Overrides"), FText::FromName(UpscaleMethodName)));
+		}
+	}
 }
 
 TSharedRef<SWidget> FDisplayClusterRootActorDetailsCustomization::CreateCustomNodeIdWidget()
diff --git a/Engine/Plugins/Runtime/nDisplay/Source/DisplayClusterConfigurator/Private/Views/Details/Upscaler/DisplayClusterConfigurationUpscalerSettingsDetailCustomization.cpp b/Engine/Plugins/Runtime/nDisplay/Source/DisplayClusterConfigurator/Private/Views/Details/Upscaler/DisplayClusterConfigurationUpscalerSettingsDetailCustomization.cpp
--- a/Engine/Plugins/Runtime/nDisplay/Source/DisplayClusterConfigurator/Private/Views/Details/Upscaler/DisplayClusterConfigurationUpscalerSettingsDetailCustomization.cpp
+++ b/Engine/Plugins/Runtime/nDisplay/Source/DisplayClusterConfigurator/Private/Views/Details/Upscaler/DisplayClusterConfigurationUpscalerSettingsDetailCustomization.cpp
@@ -2,6 +2,7 @@
 
 #include "Views/Details/Upscaler/DisplayClusterConfigurationUpscalerSettingsDetailCustomization.h"
 
+#include "DetailLayoutBuilder.h"
 #include "DisplayClusterConfigurationTypes.h"
 #include "DisplayClusterConfigurationTypes_Enums.h"
 
@@ -14,6 +15,7 @@
 #include "Widgets/Input/SComboBox.h"
 
 #include "DetailWidgetRow.h"
+#include "DisplayClusterRootActor.h"
 
 #include "Features/IModularFeatures.h"
 
@@ -24,6 +26,9 @@
 
 #include "IPropertyUtilities.h"
 #include "PropertyHandle.h"
+#include "ScopedTransaction.h"
+
+#define LOCTEXT_NAMESPACE "DisplayClusterConfigurationUpscalerSettingsDetailCustomization"
 
 using namespace UE::VirtualProduction;
 
@@ -108,6 +113,222 @@ namespace UE::DisplayClusterConfigurator::UpscalerSettingsDetailCustomization
 				}
 			});
 	}
+	
+	class FPropertyBagCustomization : public FPropertyBagInstanceDataDetails
+	{
+	public:
+		FPropertyBagCustomization(
+			const TArray<TWeakObjectPtr<UObject>>& InOwningObjects,
+			TSharedPtr<IPropertyHandle> InParentStructProperty,
+			TSharedPtr<IPropertyHandle> InPropertyBagHandle,
+			const TSharedPtr<IPropertyUtilities>& InPropUtils,
+			bool bInIsOverridable)
+			: FPropertyBagInstanceDataDetails(InPropertyBagHandle, InPropUtils, /* bInFixedLayout */ true)
+			, OwningObjects(InOwningObjects)
+			, ParentStructProperty(InParentStructProperty)
+			, bIsOverridable(bInIsOverridable)
+		{ }
+
+		virtual void OnChildRowAdded(IDetailPropertyRow& ChildRow) override
+		{
+			FPropertyBagInstanceDataDetails::OnChildRowAdded(ChildRow);
+
+			if (!bIsOverridable)
+			{
+				ChildRow.OverrideResetToDefault(FResetToDefaultOverride::Create(
+					FIsResetToDefaultVisible::CreateLambda([this](TSharedPtr<IPropertyHandle> PropertyHandle)
+					{
+						const FName PropertyName = PropertyHandle->GetProperty()->GetFName();
+						return !IsPropertyDefaultValue(PropertyName);
+					}),
+					FResetToDefaultHandler::CreateLambda([this](TSharedPtr<IPropertyHandle> PropertyHandle)
+					{
+						const FName PropertyName = PropertyHandle->GetProperty()->GetFName();
+
+						FScopedTransaction Transaction(FText::Format(LOCTEXT("ResetToDefault", "Reset {0} to default value"), FText::FromName(PropertyName)));
+						PropertyHandle->NotifyPreChange();
+
+						SetPropertyToDefaultValue(PropertyName);
+
+						PropertyHandle->NotifyPostChange(EPropertyChangeType::ValueSet);
+						PropertyHandle->NotifyFinishedChangingProperties();
+					})));
+			}
+		}
+		
+		virtual bool HasPropertyOverrides() const override { return bIsOverridable; }
+		
+		virtual void PreChangeOverrides() override
+		{
+			ParentStructProperty->NotifyPreChange();
+		}
+
+		virtual void PostChangeOverrides() override
+		{
+			ParentStructProperty->NotifyPostChange(EPropertyChangeType::ValueSet);
+			ParentStructProperty->NotifyFinishedChangingProperties();
+		}
+		
+		virtual void EnumeratePropertyBags(TSharedPtr<IPropertyHandle> PropertyBagHandle, const EnumeratePropertyBagFuncRef& Func) const override
+		{
+			struct FUpscalerSettingsOverrideProvider : public IPropertyBagOverrideProvider
+			{
+			public:
+				FUpscalerSettingsOverrideProvider(FDisplayClusterConfigurationUpscalerSettings& InUpscalerSettings)
+					: UpscalerSettings(InUpscalerSettings)
+				{
+				}
+		
+				virtual bool IsPropertyOverridden(const FGuid PropertyID) const override
+				{
+					return UpscalerSettings.ParameterOverrideGuids.Contains(PropertyID);
+				}
+		
+				virtual void SetPropertyOverride(const FGuid PropertyID, const bool bIsOverridden) const override
+				{
+					if (bIsOverridden)
+					{
+						UpscalerSettings.ParameterOverrideGuids.AddUnique(PropertyID);
+					}
+					else
+					{
+						UpscalerSettings.ParameterOverrideGuids.Remove(PropertyID);
+					}
+				}
+
+			private:
+				FDisplayClusterConfigurationUpscalerSettings& UpscalerSettings;
+			};
+			
+			ParentStructProperty->EnumerateRawData([this, Func](void* RawData, const int32 DataIndex, const int32 NumDatas)
+			{
+				if (!OwningObjects.IsValidIndex(DataIndex) || !OwningObjects[DataIndex].IsValid())
+				{
+					return true;
+				}
+
+				UObject* OwningObject = OwningObjects[DataIndex].Get();
+				if (FDisplayClusterConfigurationUpscalerSettings* UpscalerSettings = static_cast<FDisplayClusterConfigurationUpscalerSettings*>(RawData))
+				{
+					if (ParentStructProperty->GetProperty()->GetNameCPP() == GET_MEMBER_NAME_CHECKED(FDisplayClusterConfigurationViewport_RenderSettings, UpscalerSettings))
+					{
+						if (ADisplayClusterRootActor* OwningRootActor = Cast<ADisplayClusterRootActor>(OwningObject))
+						{
+							const FDisplayClusterConfigurationUpscalerSettings& GlobalUpscalerSettings = OwningRootActor->GetConfigData()->StageSettings.OuterViewportUpscalerSettings;
+							const FInstancedPropertyBag& DefaultParameters = GlobalUpscalerSettings.EditingData;
+							FInstancedPropertyBag& Parameters = UpscalerSettings->EditingData;
+							FUpscalerSettingsOverrideProvider OverrideProvider(*UpscalerSettings);
+							if (!Func(DefaultParameters, Parameters, OverrideProvider))
+							{
+								return false;
+							}
+						}
+					}
+				}
+				return true;
+			});
+		}
+
+	private:
+		bool IsPropertyDefaultValue(const FName& PropertyName)
+		{
+			TSharedPtr<IPropertyHandle> MethodNameHandle = ParentStructProperty->GetChildHandle(GET_MEMBER_NAME_CHECKED(FDisplayClusterConfigurationUpscalerSettings, MethodName));
+
+			FName MethodName;
+			MethodNameHandle->GetValue(MethodName);
+			
+			FInstancedPropertyBag DefaultPropertyBag;
+			GetMethodSettings(MethodName, DefaultPropertyBag);
+			
+			const FPropertyBagPropertyDesc* DefaultPropertyDesc = DefaultPropertyBag.FindPropertyDescByName(PropertyName);
+			
+			bool bIsDifferent = false;
+			ParentStructProperty->EnumerateRawData([this, &PropertyName, &bIsDifferent, &DefaultPropertyBag, &DefaultPropertyDesc](void* RawData, const int32 DataIndex, const int32 NumDatas)
+			{
+				if (FDisplayClusterConfigurationUpscalerSettings* UpscalerSettings = static_cast<FDisplayClusterConfigurationUpscalerSettings*>(RawData))
+				{
+					const FPropertyBagPropertyDesc* PropertyDesc = UpscalerSettings->EditingData.FindPropertyDescByName(PropertyName);
+					if (!ArePropertiesIdentical(DefaultPropertyDesc, DefaultPropertyBag, PropertyDesc, UpscalerSettings->EditingData))
+					{
+						bIsDifferent = true;
+					}
+				}
+
+				return true;
+			});
+
+			return !bIsDifferent;
+		}
+
+		void SetPropertyToDefaultValue(const FName& PropertyName)
+		{
+			TSharedPtr<IPropertyHandle> MethodNameHandle = ParentStructProperty->GetChildHandle(GET_MEMBER_NAME_CHECKED(FDisplayClusterConfigurationUpscalerSettings, MethodName));
+
+			FName MethodName;
+			MethodNameHandle->GetValue(MethodName);
+			
+			FInstancedPropertyBag DefaultPropertyBag;
+			GetMethodSettings(MethodName, DefaultPropertyBag);
+			
+			const FPropertyBagPropertyDesc* DefaultPropertyDesc = DefaultPropertyBag.FindPropertyDescByName(PropertyName);
+
+			
+			bool bIsDifferent = false;
+			ParentStructProperty->EnumerateRawData([this, &PropertyName, &bIsDifferent, &DefaultPropertyBag, &DefaultPropertyDesc](void* RawData, const int32 DataIndex, const int32 NumDatas)
+			{
+				if (FDisplayClusterConfigurationUpscalerSettings* UpscalerSettings = static_cast<FDisplayClusterConfigurationUpscalerSettings*>(RawData))
+				{
+					const FPropertyBagPropertyDesc* PropertyDesc = UpscalerSettings->EditingData.FindPropertyDescByName(PropertyName);
+					CopyPropertyValue(DefaultPropertyDesc, DefaultPropertyBag, PropertyDesc, UpscalerSettings->EditingData);
+				}
+
+				return true;
+			});
+		}
+		
+		bool ArePropertiesIdentical(const FPropertyBagPropertyDesc* InSourcePropertyDesc, const FInstancedPropertyBag& InSourceInstance, const FPropertyBagPropertyDesc* InTargetPropertyDesc, const FInstancedPropertyBag& InTargetInstance)
+		{
+			if (!InSourceInstance.IsValid() || !InTargetInstance.IsValid() || !InSourcePropertyDesc || !InSourcePropertyDesc->CachedProperty || !InTargetPropertyDesc || !InTargetPropertyDesc->CachedProperty)
+			{
+				return false;
+			}
+
+			if (!InSourcePropertyDesc->CompatibleType(*InTargetPropertyDesc))
+			{
+				return false;
+			}
+
+			const uint8* SourceValueAddress = InSourceInstance.GetValue().GetMemory() + InSourcePropertyDesc->CachedProperty->GetOffset_ForInternal();
+			const uint8* TargetValueAddress = InTargetInstance.GetValue().GetMemory() + InTargetPropertyDesc->CachedProperty->GetOffset_ForInternal();
+
+			return InSourcePropertyDesc->CachedProperty->Identical(SourceValueAddress, TargetValueAddress);
+		}
+
+		void CopyPropertyValue(const FPropertyBagPropertyDesc* InSourcePropertyDesc, const FInstancedPropertyBag& InSourceInstance, const FPropertyBagPropertyDesc* InTargetPropertyDesc, FInstancedPropertyBag& InTargetInstance)
+		{
+			if (!InSourceInstance.IsValid() || !InTargetInstance.IsValid() || !InSourcePropertyDesc || !InSourcePropertyDesc->CachedProperty || !InTargetPropertyDesc || !InTargetPropertyDesc->CachedProperty)
+			{
+				return;
+			}
+
+			// Can't copy if they are not compatible.
+			if (!InSourcePropertyDesc->CompatibleType(*InTargetPropertyDesc))
+			{
+				return;
+			}
+
+			const uint8* SourceValueAddress = InSourceInstance.GetValue().GetMemory() + InSourcePropertyDesc->CachedProperty->GetOffset_ForInternal();
+			uint8* TargetValueAddress = InTargetInstance.GetMutableValue().GetMemory() + InTargetPropertyDesc->CachedProperty->GetOffset_ForInternal();
+
+			InSourcePropertyDesc->CachedProperty->CopyCompleteValue(TargetValueAddress, SourceValueAddress);
+		}
+		
+	private:
+		const TArray<TWeakObjectPtr<UObject>>& OwningObjects;
+		TSharedPtr<IPropertyHandle> ParentStructProperty;
+		bool bIsOverridable;
+	};
+	
 };
 
 FDisplayClusterConfigurationUpscalerSettingsDetailCustomization::FDisplayClusterConfigurationUpscalerSettingsDetailCustomization()
@@ -125,6 +346,9 @@ void FDisplayClusterConfigurationUpscalerSettingsDetailCustomization::CustomizeC
 	using namespace UE::DisplayClusterConfigurator::UpscalerSettingsDetailCustomization;
 
 	// CUSTOMIZE PROPERTY: UpscalingSettings (with UpscalingSettingsStorage)
+
+	const bool bWithOverrides = InPropertyHandle->HasMetaData(TEXT("WithOverrides"));
+
 	MethodNameHandle = InPropertyHandle->GetChildHandle(GET_MEMBER_NAME_CHECKED(FDisplayClusterConfigurationUpscalerSettings, MethodName));
 	EditingDataHandle = InPropertyHandle->GetChildHandle(GET_MEMBER_NAME_CHECKED(FDisplayClusterConfigurationUpscalerSettings, EditingData));
 
@@ -135,7 +359,7 @@ void FDisplayClusterConfigurationUpscalerSettingsDetailCustomization::CustomizeC
 
 		return;
 	}
-
+	
 	/** Find upscaling method from string. */
 	auto FindUpscalingMethod = [&](const FName& InMethodName)
 		{
@@ -192,77 +416,113 @@ void FDisplayClusterConfigurationUpscalerSettingsDetailCustomization::CustomizeC
 		MethodNameHandle->GetValue(CurrentMethodName);
 		CurrentUpscalerMethod = FindUpscalingMethod(CurrentMethodName);
 
-		// Build dropdown list
-		InChildBuilder.AddCustomRow(FText::FromString("Selected Option"))
-			.NameContent()
-			[
-				MethodNameHandle->CreatePropertyNameWidget()
-			]
-			.ValueContent()
-			[
-				SNew(SComboBox<TSharedPtr<FUpscalerMethodEntry>>)
-					.OptionsSource(&UpscalerMethods)
-					.OnGenerateWidget_Lambda([](TSharedPtr<FUpscalerMethodEntry> InMethod)
-						{							
-							if (InMethod.IsValid())
+		if (!bWithOverrides)
+		{
+			// Build dropdown list
+			InChildBuilder.AddCustomRow(FText::FromString("Selected Option"))
+				.NameContent()
+				[
+					MethodNameHandle->CreatePropertyNameWidget()
+				]
+				.ValueContent()
+				[
+					SNew(SComboBox<TSharedPtr<FUpscalerMethodEntry>>)
+						.OptionsSource(&UpscalerMethods)
+						.OnGenerateWidget_Lambda([](TSharedPtr<FUpscalerMethodEntry> InMethod)
+							{							
+								if (InMethod.IsValid())
+								{
+									return SNew(STextBlock)
+										.Text(InMethod->DisplayName.IsEmpty()
+											? FText::FromName(InMethod->Name)
+											: InMethod->DisplayName)
+										.ToolTipText(InMethod->Tooltip);
+								}
+
+								//@todo
+								return SNew(STextBlock).Text(FText::GetEmpty());
+							})
+						.OnSelectionChanged_Lambda([&](TSharedPtr<FUpscalerMethodEntry> NewUpscalerMethod, ESelectInfo::Type)
 							{
-								return SNew(STextBlock)
-									.Text(InMethod->DisplayName.IsEmpty()
-										? FText::FromName(InMethod->Name)
-										: InMethod->DisplayName)
-									.ToolTipText(InMethod->Tooltip);
-							}
+								// Create a transaction around the whole process of setting the method name so that changing the editor data property bag
+								// also gets serialized to the transaction buffer
+								FScopedTransaction Transaction(LOCTEXT("EditMethodNameTransaction", "Edit Method Name"));
+							
+								CurrentUpscalerMethod = NewUpscalerMethod;
+								FName MethodName = CurrentUpscalerMethod.IsValid() ? CurrentUpscalerMethod->Name : NAME_None;
+							
+								// When MethodName changes, copy from PerMethodData[MethodName] to EditingData
+								const bool bResetEditingData = true;
+								UpdateEditingData(bResetEditingData);
+							
+								MethodNameHandle->SetValue(MethodName);
+							})
+						[
+							SNew(STextBlock)
+								.Font(IDetailLayoutBuilder::GetDetailFont())
+								.Text_Lambda([&]() -> FText
+									{
+										return CurrentUpscalerMethod.IsValid() ? CurrentUpscalerMethod->DisplayName : FText::FromString(TEXT("Default"));
+									})
+								.ToolTipText_Lambda([&]() -> FText
+									{
+										return CurrentUpscalerMethod.IsValid() ? CurrentUpscalerMethod->Tooltip : FText::GetEmpty();
+									})
+						]
+				]
+				.OverrideResetToDefault(FResetToDefaultOverride::Create(
+ 					TAttribute<bool>::CreateLambda([this] 
+					{
+ 						FName MethodName;
+ 						MethodNameHandle->GetValue(MethodName);
 
-							//@todo
-							return SNew(STextBlock).Text(FText::GetEmpty());
-						})
-					.OnSelectionChanged_Lambda([&](TSharedPtr<FUpscalerMethodEntry> NewUpscalerMethod, ESelectInfo::Type)
-						{
-							CurrentUpscalerMethod = NewUpscalerMethod;
-							FName MethodName = CurrentUpscalerMethod.IsValid() ? CurrentUpscalerMethod->Name : NAME_None;
-
-							// When MethodName changes, copy from PerMethodData[MethodName] to EditingData
-							const bool bResetEditingData = true;
-							UpdateEditingData(bResetEditingData);
-
-							MethodNameHandle->SetValue(MethodName);
-						})
-					[
-						SNew(STextBlock)
-							.Text_Lambda([&]() -> FText
-								{
-									return CurrentUpscalerMethod.IsValid() ? CurrentUpscalerMethod->DisplayName : FText::FromString(TEXT("Default"));
-								})
-							.ToolTipText_Lambda([&]() -> FText
-								{
-									return CurrentUpscalerMethod.IsValid() ? CurrentUpscalerMethod->Tooltip : FText::GetEmpty();
-								})
-					]
-			];
+ 						FName DefaultName = !UpscalerMethods.IsEmpty() ? UpscalerMethods[0]->Name : NAME_None;
+ 						return MethodName != NAME_None && MethodName != DefaultName;
+					}),
+ 					FSimpleDelegate::CreateLambda([this] 
+					{
+ 						CurrentUpscalerMethod = !UpscalerMethods.IsEmpty() ? UpscalerMethods[0] : nullptr;
+ 						const FName MethodName = CurrentUpscalerMethod ? CurrentUpscalerMethod->Name : NAME_None;
+ 						
+ 						const bool bResetEditingData = true;
+ 						UpdateEditingData(bResetEditingData);
+ 						
+ 						MethodNameHandle->SetValue(MethodName, EPropertyValueSetFlags::ResetToDefault);
+					})
+				));
+		}
+		else
+		{
+			MethodNameHandle->MarkHiddenByCustomization();
+		}
 	}
-
+	
 	const bool bResetEditingData = !CurrentUpscalerMethod.IsValid() || CurrentMethodName != CurrentUpscalerMethod->Name;
-
-	// Initialize EditingData
-	UpdateEditingData(bResetEditingData);
-
+	
 	// Update method name
 	if (CurrentUpscalerMethod.IsValid() && CurrentMethodName != CurrentUpscalerMethod->Name)
 	{
-		MethodNameHandle->SetValue(CurrentUpscalerMethod->Name);
+		// Update method name interactively and without a transaction to prevent script reconstruction from running immediately when the customizer is initialized
+		MethodNameHandle->SetValue(CurrentUpscalerMethod->Name, EPropertyValueSetFlags::InteractiveChange | EPropertyValueSetFlags::NotTransactable);
 	}
-
+	
+	// Initialize EditingData
+	UpdateEditingData(bResetEditingData);
+	
 	uint32 NumFields = 0;
 	EditingDataHandle->GetNumChildren(NumFields);
+	EditingDataHandle->MarkResetToDefaultCustomized(true);
 
 	// Show customization if we have at least 1 property in the bag
 	if (NumFields > 0)
 	{
-		TSharedRef<FPropertyBagInstanceDataDetails> EditingDataDetails =
-			MakeShared<FPropertyBagInstanceDataDetails>(
-				EditingDataHandle
+		TSharedRef<FPropertyBagCustomization> EditingDataDetails =
+			MakeShared<FPropertyBagCustomization>(
+				InCustomizationUtils.GetPropertyUtilities()->GetSelectedObjects(),
+				InPropertyHandle
+				, EditingDataHandle
 				, InCustomizationUtils.GetPropertyUtilities()
-				, /** bInFixedLayout */true);
+				, bWithOverrides);
 
 		InChildBuilder
 			.AddCustomBuilder(EditingDataDetails);
@@ -284,10 +544,15 @@ bool FDisplayClusterConfigurationUpscalerSettingsDetailCustomization::UpdateEdit
 		void* StructPointer = nullptr;
 		if (EditingDataHandle->GetValueData(StructPointer) == FPropertyAccess::Success && StructPointer)
 		{
+			bool bEditingDataChanged = false;
+			bool bNewUpscalerMethodApplied = false;
+			
 			FInstancedPropertyBag* EditingDataBag = static_cast<FInstancedPropertyBag*>(StructPointer);
 			if (bResetEditingData)
 			{
+				EditingDataHandle->NotifyPreChange();
 				EditingDataBag->Reset();
+				bEditingDataChanged = true;
 			}
 
 			if(!EditingDataBag->IsValid() || EditingDataBag->GetNumPropertiesInBag()==0)
@@ -299,13 +564,30 @@ bool FDisplayClusterConfigurationUpscalerSettingsDetailCustomization::UpdateEdit
 				const bool bTheSame = NewSettingsBag.Identical(EditingDataBag, PPF_None);
 				if (!bTheSame)
 				{
+					// If we haven't already notified for the change, do it now
+					if (!bEditingDataChanged)
+					{
+						EditingDataHandle->NotifyPreChange();
+					}
+					
 					EditingDataBag->MigrateToNewBagInstance(NewSettingsBag);
-
-					return true;
+					bEditingDataChanged = true;
+					bNewUpscalerMethodApplied = true;
 				}
 			}
+
+			if (bEditingDataChanged)
+			{
+				// Set the change interactively to prevent script reconstruction from being run before method name can also be set
+				EditingDataHandle->NotifyPostChange(EPropertyChangeType::Interactive);
+				EditingDataHandle->NotifyFinishedChangingProperties();
+			}
+
+			return bNewUpscalerMethodApplied;
 		}
 	}
 
 	return false;
-}
\ No newline at end of file
+}
+
+#undef LOCTEXT_NAMESPACE
\ No newline at end of file
diff --git a/Engine/Plugins/Runtime/nDisplay/Source/DisplayClusterDetails/Private/Drawer/SDisplayClusterDetailsDrawer.cpp b/Engine/Plugins/Runtime/nDisplay/Source/DisplayClusterDetails/Private/Drawer/SDisplayClusterDetailsDrawer.cpp
--- a/Engine/Plugins/Runtime/nDisplay/Source/DisplayClusterDetails/Private/Drawer/SDisplayClusterDetailsDrawer.cpp
+++ b/Engine/Plugins/Runtime/nDisplay/Source/DisplayClusterDetails/Private/Drawer/SDisplayClusterDetailsDrawer.cpp
@@ -478,8 +478,8 @@ void SDisplayClusterDetailsDrawer::OnObjectsReplaced(const TMap<UObject*, UObjec
 
 	if (bNeedsFullRefresh)
 	{
-		const bool bPreserveDrawerState = true;
-		Refresh(bPreserveDrawerState);
+		// Wait until the next tick so that we aren't undercutting any details customizations that may want to do logic after invoking an object reconstruction
+		bRefreshOnNextTick = true;
 	}
 	else if (bNeedsListRefresh && ObjectListView)
 	{
