[nDisplay] Fixed RHI validation failures.

diff --git a/Engine/Plugins/Runtime/nDisplay/Source/DisplayClusterShaders/Private/Shaders/DisplayClusterShadersCopyTexture.cpp b/Engine/Plugins/Runtime/nDisplay/Source/DisplayClusterShaders/Private/Shaders/DisplayClusterShadersCopyTexture.cpp
--- a/Engine/Plugins/Runtime/nDisplay/Source/DisplayClusterShaders/Private/Shaders/DisplayClusterShadersCopyTexture.cpp
+++ b/Engine/Plugins/Runtime/nDisplay/Source/DisplayClusterShaders/Private/Shaders/DisplayClusterShadersCopyTexture.cpp
@@ -104,6 +104,12 @@ namespace UE::DisplayClusterShaders::Private
 		RDG_TEXTURE_ACCESS(Output, ERHIAccess::CopyDest)
 	END_SHADER_PARAMETER_STRUCT()
 
+	/** RDG screen pixel shader parameters. */
+	BEGIN_SHADER_PARAMETER_STRUCT(FDisplayClusterScreenPixelShaderTextureParameters, )
+		RDG_TEXTURE_ACCESS(Input, ERHIAccess::SRVGraphicsPixel)
+		RDG_TEXTURE_ACCESS(Output, ERHIAccess::RTV)
+	END_SHADER_PARAMETER_STRUCT()
+
 	/**
 	* Pixel shaders parameters for 'FColorEncodingCopyRectPS'
 	*/
@@ -320,9 +326,6 @@ namespace UE::DisplayClusterShaders::Private
 			const FIntRect& SrcRect = InputContext.Rect;
 			const FIntRect& DestRect = OutputContext.Rect;
 
-			RHICmdList.Transition(FRHITransitionInfo(SrcTexture, ERHIAccess::Unknown, ERHIAccess::SRVMask));
-			RHICmdList.Transition(FRHITransitionInfo(DestTexture, ERHIAccess::Unknown, ERHIAccess::RTV));
-
 			FRHIRenderPassInfo RPInfo(DestTexture, ERenderTargetActions::Load_Store);
 			RHICmdList.BeginRenderPass(RPInfo, TEXT("nDisplay.Shaders.ColorEncodingCopyRect"));
 			
@@ -356,7 +359,6 @@ namespace UE::DisplayClusterShaders::Private
 			}
 
 			RHICmdList.EndRenderPass();
-			RHICmdList.Transition(FRHITransitionInfo(DestTexture, ERHIAccess::Unknown, ERHIAccess::SRVMask));
 
 			return true;
 		}
@@ -392,7 +394,14 @@ bool FDisplayClusterShadersCopyTexture::ColorEncodingCopyRect_RenderThread(
 {
 	using namespace UE::DisplayClusterShaders::Private;
 
-	return FColorEncodingCopyRectPS::RenderPass(RHICmdList, Input.TextureRHI, Output.TextureRHI, Input, Output, Settings);
+	RHICmdList.Transition(FRHITransitionInfo(Input.TextureRHI, ERHIAccess::Unknown, ERHIAccess::SRVMask));
+	RHICmdList.Transition(FRHITransitionInfo(Output.TextureRHI, ERHIAccess::Unknown, ERHIAccess::RTV));
+
+	bool bResult = FColorEncodingCopyRectPS::RenderPass(RHICmdList, Input.TextureRHI, Output.TextureRHI, Input, Output, Settings);
+
+	RHICmdList.Transition(FRHITransitionInfo(Output.TextureRHI, ERHIAccess::Unknown, ERHIAccess::SRVMask));
+
+	return bResult;
 }
 
 /**
@@ -412,7 +421,7 @@ bool FDisplayClusterShadersCopyTexture::AddPassColorEncodingCopyRect_RenderThrea
 	}
 
 	// Initialize render pass parameters.
-	FDisplayClusterCopyTextureParameters* Parameters = GraphBuilder.AllocParameters<FDisplayClusterCopyTextureParameters>();
+	FDisplayClusterScreenPixelShaderTextureParameters* Parameters = GraphBuilder.AllocParameters<FDisplayClusterScreenPixelShaderTextureParameters>();
 	Parameters->Input  = Input.TextureRDG;
 	Parameters->Output = Output.TextureRDG;
 
@@ -428,7 +437,7 @@ bool FDisplayClusterShadersCopyTexture::AddPassColorEncodingCopyRect_RenderThrea
 	GraphBuilder.AddPass(
 		RDG_EVENT_NAME("nDisplayShaders.ResampleTexture"),
 		Parameters,
-		ERDGPassFlags::Copy | ERDGPassFlags::NeverCull,
+		ERDGPassFlags::Raster | ERDGPassFlags::SkipRenderPass | ERDGPassFlags::NeverCull,
 		[Parameters, Input, Output, Settings](FRHICommandListImmediate& RHICmdList)
 	{
 		FColorEncodingCopyRectPS::RenderPass(
