nDisplay: Fixed issue where ICVFX camera upscaler settings were getting overridden by global settings

diff --git a/Engine/Plugins/Runtime/nDisplay/Source/DisplayCluster/Private/Render/Viewport/Configuration/DisplayClusterViewportConfigurationHelpers.cpp b/Engine/Plugins/Runtime/nDisplay/Source/DisplayCluster/Private/Render/Viewport/Configuration/DisplayClusterViewportConfigurationHelpers.cpp
--- a/Engine/Plugins/Runtime/nDisplay/Source/DisplayCluster/Private/Render/Viewport/Configuration/DisplayClusterViewportConfigurationHelpers.cpp
+++ b/Engine/Plugins/Runtime/nDisplay/Source/DisplayCluster/Private/Render/Viewport/Configuration/DisplayClusterViewportConfigurationHelpers.cpp
@@ -129,6 +129,13 @@ void FDisplayClusterViewportConfigurationHelpers::UpdateBaseViewportSetting(FDis
 		// Update viewport Upscaler method
 		InRenderSettings.UpscalerSettings.GetUpscalerSettings(DefaultUpscalerSettings, InOutRenderSettings.UpscalerSettings);
 
+		// UE-290901: Temporary fix to avoid API changes by moving this out of GetUpscalerSettings so that it only applies to outer viewport render settings, not ICVFX settings
+		// TODO: Fold back into GetUpscalerSettings
+		if (DefaultUpscalerSettings)
+		{
+			InOutRenderSettings.UpscalerSettings.CustomUpscalerSettings.MigrateToNewBagInstanceWithOverrides(DefaultUpscalerSettings->EditingData, InRenderSettings.UpscalerSettings.ParameterOverrideGuids);
+		}
+		
 		DstViewport.UpdateConfiguration_Overscan(GetViewportOverscanSettings(InRenderSettings.Overscan));
 
 		DstViewport.UpdateConfiguration_PostRenderOverride(InRenderSettings.Replace);
diff --git a/Engine/Plugins/Runtime/nDisplay/Source/DisplayClusterConfiguration/Private/DisplayClusterConfigurationTypes_Upscaler.cpp b/Engine/Plugins/Runtime/nDisplay/Source/DisplayClusterConfiguration/Private/DisplayClusterConfigurationTypes_Upscaler.cpp
--- a/Engine/Plugins/Runtime/nDisplay/Source/DisplayClusterConfiguration/Private/DisplayClusterConfigurationTypes_Upscaler.cpp
+++ b/Engine/Plugins/Runtime/nDisplay/Source/DisplayClusterConfiguration/Private/DisplayClusterConfigurationTypes_Upscaler.cpp
@@ -65,12 +65,6 @@ bool FDisplayClusterConfigurationUpscalerSettings::GetUpscalerSettings(
 				// Custom upscaler not used.
 				OutUpscalerSettings.CustomUpscalerName = NAME_None;
 				GetConfigurationData(*this, OutUpscalerSettings.CustomUpscalerSettings);
-				
-				// Migrate any data from the default upscaler settings that isn't explicitly overridden in the parameter overrides
-				if (InDefaultUpscalerSettings)
-				{
-					OutUpscalerSettings.CustomUpscalerSettings.MigrateToNewBagInstanceWithOverrides(InDefaultUpscalerSettings->EditingData, ParameterOverrideGuids);
-				}
 
 				return true;
 			}
@@ -85,12 +79,6 @@ bool FDisplayClusterConfigurationUpscalerSettings::GetUpscalerSettings(
 		// Set custom upscaler name and PropertyBag with settings.
 		OutUpscalerSettings.CustomUpscalerName = MethodName;
 		GetConfigurationData(*this, OutUpscalerSettings.CustomUpscalerSettings);
-
-		// Migrate any data from the default upscaler settings that isn't explicitly overridden in the parameter overrides
-		if (InDefaultUpscalerSettings)
-		{
-			OutUpscalerSettings.CustomUpscalerSettings.MigrateToNewBagInstanceWithOverrides(InDefaultUpscalerSettings->EditingData, ParameterOverrideGuids);
-		}
 	}
 
 	return MethodName != NAME_None;
