The old Level Viewport Toolbar Panel Extensions are now deprecated.

# - Allow panel extension names to be marked deprecated & log warnings when extensions are registered.
# - Remove existing panel extensions that have integrated with the new toolbar.

[RN] Major Editor

diff --git a/Engine/Plugins/Experimental/DaySequence/Source/DaySequenceEditor/Private/DaySequenceEditorModule.cpp b/Engine/Plugins/Experimental/DaySequence/Source/DaySequenceEditor/Private/DaySequenceEditorModule.cpp
--- a/Engine/Plugins/Experimental/DaySequence/Source/DaySequenceEditor/Private/DaySequenceEditorModule.cpp
+++ b/Engine/Plugins/Experimental/DaySequence/Source/DaySequenceEditor/Private/DaySequenceEditorModule.cpp
@@ -71,8 +71,6 @@ public:
 	virtual void StartupModule() override;
 	virtual void ShutdownModule() override;
 
-	void OnPostEngineInit();
-	
 	// FGCObject interface
 	virtual void AddReferencedObjects( FReferenceCollector& Collector ) override;
 	virtual FString GetReferencerName() const override;
@@ -116,7 +114,6 @@ private:
 	void DeregisterSettings();
 
 	void CreateDaySequenceViewportMenu(UToolMenu* Menu);
-	TSharedRef<SWidget> CreateDaySequenceViewportToolbarExtension(FWeakObjectPtr ExtensionContext);
 
 	/** Shared widget functions */
 	static bool IsMenuVisible();
@@ -166,8 +163,6 @@ void FDaySequenceEditorModule::StartupModule()
 {
 	// This code will execute after your module is loaded into memory; the exact timing is specified in the .uplugin file per-module
 
-	FCoreDelegates::OnPostEngineInit.AddRaw(this, &FDaySequenceEditorModule::OnPostEngineInit);
-	
 	FDaySequenceEditorStyle::Initialize();
 	FDaySequenceEditorCommands::Register();
 
@@ -344,23 +339,6 @@ void FDaySequenceEditorModule::ShutdownModule()
 	FDaySequenceEditorCommands::Unregister();
 }
 
-void FDaySequenceEditorModule::OnPostEngineInit()
-{
-	if (GEditor)
-	{
-		if (UPanelExtensionSubsystem* PanelExtensionSubsystem = GEditor->GetEditorSubsystem<UPanelExtensionSubsystem>())
-		{
-			if (!PanelExtensionSubsystem->IsPanelFactoryRegistered(DaySequenceViewportToolBarExtensionName))
-			{
-				FPanelExtensionFactory DaySequenceViewportMenuWidget;
-				DaySequenceViewportMenuWidget.CreateExtensionWidget = FPanelExtensionFactory::FCreateExtensionWidget::CreateRaw(this, &FDaySequenceEditorModule::CreateDaySequenceViewportToolbarExtension);
-				DaySequenceViewportMenuWidget.Identifier = DaySequenceViewportToolBarExtensionName;
-				PanelExtensionSubsystem->RegisterPanelFactory("LevelViewportToolBar.LeftExtension", DaySequenceViewportMenuWidget);
-			}
-		}
-	}
-}
-
 void FDaySequenceEditorModule::AddReferencedObjects( FReferenceCollector& Collector )
 {
 	if (Settings)
@@ -616,29 +594,6 @@ TSharedRef<SWidget> FDaySequenceEditorModule::CreateTimeOfDayWidget() const
 	return UToolMenus::Get()->GenerateWidget(MenuName, MenuContext);
 }
 
-TSharedRef<SWidget> FDaySequenceEditorModule::CreateDaySequenceViewportToolbarExtension(FWeakObjectPtr ExtensionContext)
-{
-	if (const UViewportToolBarContext* ExtensionContextObject = Cast<UViewportToolBarContext>(ExtensionContext.Get()); ensure(ExtensionContextObject))
-	{
-		if (TSharedPtr<SViewportToolBar> Toolbar = ExtensionContextObject->ViewportToolBar.Pin())
-		{
-			return SNew(SEditorViewportToolbarMenu)
-				.ParentToolBar(Toolbar)
-				.Label(LOCTEXT("DaySequenceMenuLabel", "Time of Day"))
-				.LabelIcon(FDaySequenceEditorStyle::Get()->GetBrush("DaySequenceEditor.ViewportToolBar"))
-				.OnGetMenuContent_Lambda([this]()
-				{
-					return CreateTimeOfDayWidget();
-				})
-				.Visibility_Lambda([this]()
-				{
-					return IsMenuVisible() && UE::UnrealEd::ShowOldViewportToolbars() ? EVisibility::Visible : EVisibility::Collapsed;
-				});
-		}
-	}
-	return SNullWidget::NullWidget;
-}
-
 TSharedRef<SWidget> FDaySequenceEditorModule::CreateDaySequencePreviewWidget()
 {
 	return
diff --git a/Engine/Plugins/MovieScene/LevelSequenceEditor/Source/LevelSequenceEditor/Private/CinematicViewport/SCinematicLevelViewport.cpp b/Engine/Plugins/MovieScene/LevelSequenceEditor/Source/LevelSequenceEditor/Private/CinematicViewport/SCinematicLevelViewport.cpp
--- a/Engine/Plugins/MovieScene/LevelSequenceEditor/Source/LevelSequenceEditor/Private/CinematicViewport/SCinematicLevelViewport.cpp
+++ b/Engine/Plugins/MovieScene/LevelSequenceEditor/Source/LevelSequenceEditor/Private/CinematicViewport/SCinematicLevelViewport.cpp
@@ -218,39 +218,11 @@ private:
 
 class SCinematicPreviewViewport : public SLevelViewport
 {
-	static FName GetOldToolbarOverlayExtensionId()
-	{
-		return "LevelViewportToolBar.RightmostExtension";
-	}
-
 public:
-	virtual ~SCinematicPreviewViewport() override
-	{
-		// Unregistering extension from old toolbar
-		{
-			UPanelExtensionSubsystem* PanelExtensionSubsystem = GEditor->GetEditorSubsystem<UPanelExtensionSubsystem>();
-			if (!PanelExtensionSubsystem)
-			{
-				return;
-			}
-
-			if (!PanelExtensionSubsystem->IsPanelFactoryRegistered(OverlayPanelExtensionFactory.Identifier))
-			{
-				return;
-			}
-
-			PanelExtensionSubsystem->UnregisterPanelFactory(
-				OverlayPanelExtensionFactory.Identifier, GetOldToolbarOverlayExtensionId()
-			);
-		}
-	}
-
 	virtual const FSlateBrush* OnGetViewportBorderBrush() const override { return nullptr; }
 	virtual bool IsActorEditorContextVisible() const { return false; }
 	virtual EVisibility GetSelectedActorsCurrentLevelTextVisibility() const override { return EVisibility::Collapsed; }
-	virtual EVisibility GetViewportControlsVisibility() const override { return EVisibility::Collapsed; }
 
-	virtual TSharedPtr<SWidget> MakeViewportToolbar() { return nullptr; }
 	virtual TSharedPtr<SWidget> BuildViewportToolbar() override
 	{
 		return nullptr;
@@ -272,58 +244,6 @@ public:
 		CinematicLevelViewport = InCinematicLevelViewport;
 	}
 
-	void ExtendOldLevelToolbar()
-	{
-		if (!GEditor)
-		{
-			return;
-		}
-
-		auto CinematicLevelViewportPinned = CinematicLevelViewport.Pin();
-		if (!CinematicLevelViewportPinned)
-		{
-			return;
-		}
-
-		if (UPanelExtensionSubsystem* PanelExtensionSubsystem = GEditor->GetEditorSubsystem<UPanelExtensionSubsystem>())
-		{
-			OverlayPanelExtensionFactory.Identifier = "FilmOverlayOptions";
-			if (!PanelExtensionSubsystem->IsPanelFactoryRegistered(OverlayPanelExtensionFactory.Identifier))
-			{
-				OverlayPanelExtensionFactory.CreateExtensionWidget =
-					FPanelExtensionFactory::FCreateExtensionWidget::CreateLambda(
-						[OverlayWidgetWeak = CinematicLevelViewportPinned->GetOverlayWidget().ToWeakPtr(),
-						 CommandListWeak = CommandList.ToWeakPtr()](FWeakObjectPtr) -> TSharedRef<SWidget>
-						{
-							if (const TSharedPtr<SFilmOverlay> OverlayWidgetPinned = OverlayWidgetWeak.Pin())
-							{
-								TSharedRef<SFilmOverlayOptions> FilmOverlayOptions =
-									SNew(SFilmOverlayOptions, OverlayWidgetPinned);
-									
-								FilmOverlayOptions->SetVisibility(TAttribute<EVisibility>::CreateLambda([]
-								{
-									return UE::UnrealEd::ShowOldViewportToolbars() ? EVisibility::Visible : EVisibility::Collapsed;
-								}));
-
-								if (const TSharedPtr<FUICommandList> CommandListPinned = CommandListWeak.Pin())
-								{
-									FilmOverlayOptions->BindCommands(CommandListPinned.ToSharedRef());
-								}
-
-								return FilmOverlayOptions;
-							}
-
-							return SNullWidget::NullWidget;
-						}
-					);
-
-				PanelExtensionSubsystem->RegisterPanelFactory(
-					GetOldToolbarOverlayExtensionId(), OverlayPanelExtensionFactory
-				);
-			}
-		}
-	}
-
 	static void ExtendLevelToolbar()
 	{
 		if (UToolMenu* const Submenu = UToolMenus::Get()->ExtendMenu("LevelEditor.ViewportToolbar.Camera"))
@@ -344,10 +264,6 @@ public:
 	{
 		// Adding Film Overlay Options widget to both new and old toolbars
 		ExtendLevelToolbar();
-		ExtendOldLevelToolbar();
-
-		TSharedRef<SWidget> LegacyToolbar = SLevelViewport::MakeViewportToolbar().ToSharedRef();
-		TSharedRef<SWidget> Toolbar = SLevelViewport::BuildViewportToolbar().ToSharedRef();
 
 		// clang-format off
 		TSharedRef<SVerticalBox> ToolbarsWidget =
@@ -355,16 +271,7 @@ public:
 			+ SVerticalBox::Slot()
 			.AutoHeight()
 			[
-				SNew(SBox)
-				.Visibility_Lambda(
-					[this]() -> EVisibility
-					{
-						return  UE::UnrealEd::ShowNewViewportToolbars() ? EVisibility::Visible: EVisibility::Collapsed;
-					}
-				)
-				[
-					Toolbar
-				]
+				SLevelViewport::BuildViewportToolbar().ToSharedRef()
 			]
 			+ SVerticalBox::Slot()
 			.MaxHeight(1.0f)
@@ -372,20 +279,6 @@ public:
 				CreateViewportIndicatorWidget(
 					TAttribute<EVisibility>::CreateSP(this, &SCinematicPreviewViewport::OnGetCinematicViewportIndicatorVisibility)
 					).ToSharedRef()
-			]
-			+ SVerticalBox::Slot()
-			.AutoHeight()
-			[
-				SNew(SBox)
-				.Visibility_Lambda(
-					[this]() -> EVisibility
-					{
-						return  UE::UnrealEd::ShowOldViewportToolbars() ? EVisibility::Visible: EVisibility::Collapsed;
-					}
-				)
-				[
-					LegacyToolbar
-				]
 			];
 		// clang-format on
 
@@ -419,8 +312,6 @@ private:
 	bool bShowToolbar;
 
 	TWeakPtr<SCinematicLevelViewport> CinematicLevelViewport;
-
-	FPanelExtensionFactory OverlayPanelExtensionFactory;
 };
 
 BEGIN_SLATE_FUNCTION_BUILD_OPTIMIZATION
diff --git a/Engine/Plugins/Runtime/nDisplay/Source/DisplayClusterConfigurator/Private/DisplayClusterConfiguratorModule.cpp b/Engine/Plugins/Runtime/nDisplay/Source/DisplayClusterConfigurator/Private/DisplayClusterConfiguratorModule.cpp
--- a/Engine/Plugins/Runtime/nDisplay/Source/DisplayClusterConfigurator/Private/DisplayClusterConfiguratorModule.cpp
+++ b/Engine/Plugins/Runtime/nDisplay/Source/DisplayClusterConfigurator/Private/DisplayClusterConfiguratorModule.cpp
@@ -77,11 +77,6 @@
 
 #define LOCTEXT_NAMESPACE "DisplayClusterConfigurator"
 
-namespace UE::DisplayClusterConfiguratorModuleConstants
-{
-	const static FName LevelEditorViewportExtensionIdentifier = TEXT("DisplayClusterLevelViewportExtension");
-}
-
 void FDisplayClusterConfiguratorModule::StartupModule()
 {
 	FCoreDelegates::OnPostEngineInit.AddRaw(this, &FDisplayClusterConfiguratorModule::OnPostEngineInit);
@@ -161,8 +156,6 @@ void FDisplayClusterConfiguratorModule::ShutdownModule()
 	MenuExtensibilityManager.Reset();
 	ToolBarExtensibilityManager.Reset();
 
-	UnregisterPanelExtensions();
-
 	IKismetCompilerInterface& KismetCompilerModule = FModuleManager::GetModuleChecked<IKismetCompilerInterface>("KismetCompiler");
 	KismetCompilerModule.GetCompilers().Remove(&BlueprintCompiler);
 
@@ -175,7 +168,6 @@ void FDisplayClusterConfiguratorModule::ShutdownModule()
 
 void FDisplayClusterConfiguratorModule::OnPostEngineInit()
 {
-	RegisterPanelExtensions();
 	ToolbarExtensions->RegisterToolbarExtensions();
 }
 
@@ -446,101 +438,12 @@ void FDisplayClusterConfiguratorModule::UnregisterSectionMappings()
 	}
 }
 
-void FDisplayClusterConfiguratorModule::RegisterPanelExtensions()
-{
-	if (GEditor)
-	{
-		FDisplayClusterConfiguratorBlueprintEditor::RegisterPanelExtensionFactory();
-	
-		if (UPanelExtensionSubsystem* PanelExtensionSubsystem = GEditor->GetEditorSubsystem<UPanelExtensionSubsystem>())
-		{
-			if (!PanelExtensionSubsystem->IsPanelFactoryRegistered(UE::DisplayClusterConfiguratorModuleConstants::LevelEditorViewportExtensionIdentifier))
-			{
-				FPanelExtensionFactory LevelViewportToolbarExtensionWidget;
-				LevelViewportToolbarExtensionWidget.CreateExtensionWidget = FPanelExtensionFactory::FCreateExtensionWidget::CreateStatic(&FDisplayClusterConfiguratorModule::OnExtendLevelEditorViewportToolbar);
-				LevelViewportToolbarExtensionWidget.Identifier = UE::DisplayClusterConfiguratorModuleConstants::LevelEditorViewportExtensionIdentifier;
-
-				PanelExtensionSubsystem->RegisterPanelFactory(TEXT("LevelViewportToolBar.LeftExtension"), LevelViewportToolbarExtensionWidget);
-			}
-		}
-	}
-}
-
-void FDisplayClusterConfiguratorModule::UnregisterPanelExtensions()
-{
-	if (GEditor)
-	{
-		FDisplayClusterConfiguratorBlueprintEditor::UnregisterPanelExtensionFactory();
-
-		if (UPanelExtensionSubsystem* PanelExtensionSubsystem = GEditor->GetEditorSubsystem<UPanelExtensionSubsystem>())
-		{
-			PanelExtensionSubsystem->UnregisterPanelFactory(UE::DisplayClusterConfiguratorModuleConstants::LevelEditorViewportExtensionIdentifier);
-		}
-	}
-}
-
 TSharedPtr<FKismetCompilerContext> FDisplayClusterConfiguratorModule::GetCompilerForDisplayClusterBP(UBlueprint* BP,
                                                                                                      FCompilerResultsLog& InMessageLog, const FKismetCompilerOptions& InCompileOptions)
 {
 	return TSharedPtr<FKismetCompilerContext>(new FDisplayClusterConfiguratorKismetCompilerContext(CastChecked<UDisplayClusterBlueprint>(BP), InMessageLog, InCompileOptions));
 }
 
-TSharedRef<SWidget> FDisplayClusterConfiguratorModule::OnExtendLevelEditorViewportToolbar(FWeakObjectPtr ExtensionContext)
-{
-	return SNew(SEditorViewportToolBarButton)	
-	.ButtonType(EUserInterfaceActionType::Button)
-	.ButtonStyle(&FAppStyle::Get().GetWidgetStyle<FButtonStyle>("EditorViewportToolBar.WarningButton"))
-	.OnClicked_Static(OnViewportsFrozenWarningClicked)
-	.Visibility_Static(GetViewportsFrozenWarningVisibility)
-	.ToolTipText(LOCTEXT("DisplayClusterViewportsFrozenOff_ToolTip", "nDisplay viewports are frozen. Click to unfreeze the viewports."))
-	.Content()
-	[
-		SNew(STextBlock)
-		.TextStyle(&FAppStyle::Get().GetWidgetStyle<FTextBlockStyle>("SmallText"))
-		.Text(LOCTEXT("DisplayClusterViewportsFrozen", "nDisplay Viewports Frozen"))
-	];
-}
-
-FReply FDisplayClusterConfiguratorModule::OnViewportsFrozenWarningClicked()
-{
-	if (GEditor)
-	{
-		FScopedTransaction Transaction(LOCTEXT("UnfreezeViewports", "Unfreeze viewports"));
-		
-		for (TActorIterator<ADisplayClusterRootActor> It(GEditor->GetEditorWorldContext().World()); It; ++It)
-		{
-			It->SetFreezeOuterViewports(false);
-		}
-	}
-
-	return FReply::Handled();
-}
-
-EVisibility FDisplayClusterConfiguratorModule::GetViewportsFrozenWarningVisibility()
-{
-	if (GEditor)
-	{
-		if (UE::UnrealEd::ShowNewViewportToolbars())
-		{
-			return EVisibility::Collapsed;
-		}
-
-		if (const UWorld* World = GEditor->GetEditorWorldContext().World())
-		{
-			for (TActorIterator<ADisplayClusterRootActor> It(World); It; ++It)
-			{
-				const UDisplayClusterConfigurationData* ConfigData = It->GetConfigData();
-				if (ConfigData && ConfigData->StageSettings.bFreezeRenderOuterViewports)
-				{
-					return EVisibility::Visible;
-				}
-			}
-		}
-	}
-
-	return EVisibility::Collapsed;
-}
-
 IMPLEMENT_MODULE(FDisplayClusterConfiguratorModule, DisplayClusterConfigurator);
 
 #undef LOCTEXT_NAMESPACE
diff --git a/Engine/Plugins/Runtime/nDisplay/Source/DisplayClusterConfigurator/Private/DisplayClusterConfiguratorModule.h b/Engine/Plugins/Runtime/nDisplay/Source/DisplayClusterConfigurator/Private/DisplayClusterConfiguratorModule.h
--- a/Engine/Plugins/Runtime/nDisplay/Source/DisplayClusterConfigurator/Private/DisplayClusterConfiguratorModule.h
+++ b/Engine/Plugins/Runtime/nDisplay/Source/DisplayClusterConfigurator/Private/DisplayClusterConfiguratorModule.h
@@ -43,23 +43,8 @@ private:
 	void RegisterSectionMappings();
 	void UnregisterSectionMappings();
 
-	/** Register any panel extensions used by this module */
-	void RegisterPanelExtensions();
-
-	/** Unregister any active panel extensions used by this module */
-	void UnregisterPanelExtensions();
-
 	static TSharedPtr<FKismetCompilerContext> GetCompilerForDisplayClusterBP(UBlueprint* BP, FCompilerResultsLog& InMessageLog, const FKismetCompilerOptions& InCompileOptions);
 
-	/** Callback when the level editor viewport creates its toolbar */
-	static TSharedRef<SWidget> OnExtendLevelEditorViewportToolbar(FWeakObjectPtr ExtensionContext);
-
-	/** Callback when our viewports frozen button is clicked */
-	static FReply OnViewportsFrozenWarningClicked();
-
-	/** If the viewports frozen warning button should be displayed */
-	static EVisibility GetViewportsFrozenWarningVisibility();
-
 private:
 	TArray<TSharedPtr<IAssetTypeActions>> CreatedAssetTypeActions;
 	TSharedPtr<FExtensibilityManager> MenuExtensibilityManager;
diff --git a/Engine/Source/Editor/LevelEditor/Private/LevelEditor.cpp b/Engine/Source/Editor/LevelEditor/Private/LevelEditor.cpp
--- a/Engine/Source/Editor/LevelEditor/Private/LevelEditor.cpp
+++ b/Engine/Source/Editor/LevelEditor/Private/LevelEditor.cpp
@@ -85,6 +85,25 @@ const FName LevelEditorTabIds::LevelEditorHierarchicalLODOutliner(TEXT("LevelEdi
 const FName LevelEditorTabIds::OutputLog(TEXT("OutputLog"));
 const FName LevelEditorTabIds::LevelEditorEnvironmentLightingViewer(TEXT("LevelEditorEnvironmentLightingViewer"));
 
+namespace UE::Editor::LevelEditor::Private
+{
+	void OnPostEngineInit()
+	{
+		if (GEditor)
+		{
+			if (UPanelExtensionSubsystem* PanelExtensionSubsystem = GEditor->GetEditorSubsystem<UPanelExtensionSubsystem>())
+			{
+				FText DeprecationMessage = NSLOCTEXT("UnrealEditor", "LevelViewportToolbarPanelExtensionDeprecationMessage", "Please extend the \"LevelEditor.ViewportToolbar\" Tool Menus menu instead.");
+			
+				PanelExtensionSubsystem->MarkDeprecatedExtension("LevelViewportToolBar.LeftExtension", DeprecationMessage);
+				PanelExtensionSubsystem->MarkDeprecatedExtension("LevelViewportToolBar.MiddleExtension", DeprecationMessage);
+				PanelExtensionSubsystem->MarkDeprecatedExtension("LevelViewportToolBar.RightExtension", DeprecationMessage);
+				PanelExtensionSubsystem->MarkDeprecatedExtension("LevelViewportToolBar.RightmostExtension", DeprecationMessage);
+			}
+		}
+	}
+}
+
 FLevelEditorModule::FLevelEditorModule()
 	: ToggleImmersiveConsoleCommand(
 	TEXT( "LevelEditor.ToggleImmersive" ),
@@ -311,6 +330,8 @@ void FLevelEditorModule::StartupModule()
 	OutlinerSettings = MakeShared<FLevelEditorOutlinerSettings>();
 	OutlinerSettings->Initialize();
 	OutlinerSettings->SetupBuiltInCategories();
+	
+	FCoreDelegates::OnPostEngineInit.AddStatic(&UE::Editor::LevelEditor::Private::OnPostEngineInit);
 }
 
 /**
diff --git a/Engine/Source/Editor/UnrealEd/Private/Subsystems/PanelExtensionSubsystem.cpp b/Engine/Source/Editor/UnrealEd/Private/Subsystems/PanelExtensionSubsystem.cpp
--- a/Engine/Source/Editor/UnrealEd/Private/Subsystems/PanelExtensionSubsystem.cpp
+++ b/Engine/Source/Editor/UnrealEd/Private/Subsystems/PanelExtensionSubsystem.cpp
@@ -8,6 +8,8 @@
 #include "FileHelpers.h"
 #include "EditorReimportHandler.h"
 
+DECLARE_LOG_CATEGORY_CLASS(LogPanelExtensions, Log, All);
+
 SExtensionPanel::~SExtensionPanel()
 {
 	if (GEditor && ExtensionPanelID != NAME_None)
@@ -75,6 +77,11 @@ void UPanelExtensionSubsystem::Deinitialize()
 
 void UPanelExtensionSubsystem::RegisterPanelFactory(FName ExtensionPanelID, const FPanelExtensionFactory& InPanelExtensionFactory)
 {
+	if (const FText* DeprecationMessage = DeprecatedExtensions.Find(ExtensionPanelID))
+	{
+		UE_LOG(LogPanelExtensions, Warning, TEXT("Extension panel id \"%s\" has been deprecated. %s"), *ExtensionPanelID.ToString(), *DeprecationMessage->ToString());
+	}
+
 	TArray<FPanelExtensionFactory>& Factories = ExtensionPointMap.FindOrAdd(ExtensionPanelID);
 	Factories.Add(InPanelExtensionFactory);
 	Factories.StableSort([](const FPanelExtensionFactory& One, const FPanelExtensionFactory& Two)
@@ -117,6 +124,19 @@ bool UPanelExtensionSubsystem::IsPanelFactoryRegistered(FName Identifier, FName
 	return false;
 }
 
+void UPanelExtensionSubsystem::MarkDeprecatedExtension(FName ExtensionID, const FText& Message)
+{
+	DeprecatedExtensions.Add(ExtensionID, Message);
+	
+	if (const TArray<FPanelExtensionFactory>* Factories = ExtensionPointMap.Find(ExtensionID))
+	{
+		if (Factories->Num())
+		{
+			UE_LOG(LogPanelExtensions, Warning, TEXT("Extension panel id \"%s\" has been deprecated. %s"), *ExtensionID.ToString(), *Message.ToString());
+		}
+	}
+}
+
 TSharedRef<SWidget> UPanelExtensionSubsystem::CreateWidget(FName ExtensionPanelID, FWeakObjectPtr ExtensionContext)
 {
 	const TArray<FPanelExtensionFactory>* ExtensionArray = ExtensionPointMap.Find(ExtensionPanelID);
diff --git a/Engine/Source/Editor/UnrealEd/Public/Subsystems/PanelExtensionSubsystem.h b/Engine/Source/Editor/UnrealEd/Public/Subsystems/PanelExtensionSubsystem.h
--- a/Engine/Source/Editor/UnrealEd/Public/Subsystems/PanelExtensionSubsystem.h
+++ b/Engine/Source/Editor/UnrealEd/Public/Subsystems/PanelExtensionSubsystem.h
@@ -105,6 +105,8 @@ public:
 	UNREALED_API void RegisterPanelFactory(FName ExtensionPanelID, const FPanelExtensionFactory& InPanelExtensionFactory);
 	UNREALED_API void UnregisterPanelFactory(FName Identifier, FName ExtensionPanelID = NAME_None);
 	UNREALED_API bool IsPanelFactoryRegistered(FName Identifier, FName ExtensionPanelID = NAME_None) const;
+	
+	UNREALED_API void MarkDeprecatedExtension(FName ExtensionID, const FText& Message);
 
 protected:
 	friend class SExtensionPanel;
@@ -116,4 +118,5 @@ protected:
 private:
 	TMap<FName, TArray<FPanelExtensionFactory>> ExtensionPointMap;
 	TMap<FName, FPanelFactoryRegistryChanged> PanelFactoryRegistryChangedCallbackMap;
+	TMap<FName, FText> DeprecatedExtensions;
 };
