fix build errors when WITH_NVAPI=0

#rnx

diff --git a/Engine/Plugins/Runtime/nDisplay/Source/DisplayCluster/Private/Render/Synchronization/DisplayClusterRenderSyncPolicyFactoryInternal.cpp b/Engine/Plugins/Runtime/nDisplay/Source/DisplayCluster/Private/Render/Synchronization/DisplayClusterRenderSyncPolicyFactoryInternal.cpp
--- a/Engine/Plugins/Runtime/nDisplay/Source/DisplayCluster/Private/Render/Synchronization/DisplayClusterRenderSyncPolicyFactoryInternal.cpp
+++ b/Engine/Plugins/Runtime/nDisplay/Source/DisplayCluster/Private/Render/Synchronization/DisplayClusterRenderSyncPolicyFactoryInternal.cpp
@@ -62,12 +62,12 @@ TSharedPtr<IDisplayClusterRenderSyncPolicy> FDisplayClusterRenderSyncPolicyFacto
 	else if (InPolicyType.Equals(DisplayClusterConfigurationStrings::config::cluster::render_sync::Nvidia, ESearchCase::IgnoreCase))
 	{
 #if PLATFORM_WINDOWS
-		if (RHIType == ERHIInterfaceType::D3D12 && CVarUseNvidiaPresentBarrierPolicy.GetValueOnGameThread())
+		if (RHIType == ERHIInterfaceType::D3D12 && CVarUseNvidiaPresentBarrierPolicy.GetValueOnGameThread() && WITH_NVAPI)
 		{
 			UE_LOG(LogDisplayClusterRenderSync, Log, TEXT("Experimental NVIDIA synchronization approach will be used!"));
 			return MakeShared<FDisplayClusterRenderSyncPolicyNvidiaPresentBarrier>(Parameters);
 		}
-		else if (RHIType == ERHIInterfaceType::D3D11 || RHIType == ERHIInterfaceType::D3D12)
+		else if (RHIType == ERHIInterfaceType::D3D11 || RHIType == ERHIInterfaceType::D3D12 && WITH_NVAPI)
 		{
 			return MakeShared<FDisplayClusterRenderSyncPolicyNvidiaSwapBarrier>(Parameters);
 		}
diff --git a/Engine/Plugins/Runtime/nDisplay/Source/DisplayCluster/Private/Render/Synchronization/Windows/DisplayClusterRenderSyncPolicyNvidiaPresentBarrierWindows.cpp b/Engine/Plugins/Runtime/nDisplay/Source/DisplayCluster/Private/Render/Synchronization/Windows/DisplayClusterRenderSyncPolicyNvidiaPresentBarrierWindows.cpp
--- a/Engine/Plugins/Runtime/nDisplay/Source/DisplayCluster/Private/Render/Synchronization/Windows/DisplayClusterRenderSyncPolicyNvidiaPresentBarrierWindows.cpp
+++ b/Engine/Plugins/Runtime/nDisplay/Source/DisplayCluster/Private/Render/Synchronization/Windows/DisplayClusterRenderSyncPolicyNvidiaPresentBarrierWindows.cpp
@@ -15,7 +15,9 @@
 THIRD_PARTY_INCLUDES_START
 	#include "dxgi1_3.h"
 	#include "d3d12.h"
+#if WITH_NVAPI
 	#include "nvapi.h"
+#endif
 	#include <combaseapi.h>
 THIRD_PARTY_INCLUDES_END
 #include "Windows/HideWindowsPlatformTypes.h"
@@ -29,8 +31,10 @@ namespace DisplayClusterRenderSyncPolicyNvidiaPresentBarrier_Data_Windows
 	IDXGISwapChain2* DXGISwapChain = nullptr;
 	// Synchronization fence
 	ID3D12Fence*     PresentBarrierFence = nullptr;
+#if WITH_NVAPI
 	// Barrier client handle
 	NvPresentBarrierClientHandle NvBarrierClientHandle = nullptr;
+#endif
 	// Backbuffer references to keep the resources alive
 	TArray<TRefCountPtr<ID3D12Resource>> BackBuffers;
 }
@@ -307,6 +311,7 @@ void FDisplayClusterRenderSyncPolicyNvidiaPresentBarrier::OnFramePresented(bool
 }
 
 // A helper function to convert ID->TEXT for friendly logging
+#if WITH_NVAPI
 static constexpr const TCHAR* GetSyncModeName(NV_PRESENT_BARRIER_SYNC_MODE SyncMode)
 {
 	switch (SyncMode)
@@ -335,6 +340,7 @@ static constexpr const TCHAR* GetSyncModeName(NV_PRESENT_BARRIER_SYNC_MODE SyncM
 		return TEXT("UNKNOWN_STATE");
 	}
 }
+#endif //WITH_NVAPI
 
 void FDisplayClusterRenderSyncPolicyNvidiaPresentBarrier::LogPresentBarrierStats()
 {
diff --git a/Engine/Plugins/Runtime/nDisplay/Source/DisplayCluster/Private/Render/Synchronization/Windows/DisplayClusterRenderSyncPolicyNvidiaSwapBarrierWindows.cpp b/Engine/Plugins/Runtime/nDisplay/Source/DisplayCluster/Private/Render/Synchronization/Windows/DisplayClusterRenderSyncPolicyNvidiaSwapBarrierWindows.cpp
--- a/Engine/Plugins/Runtime/nDisplay/Source/DisplayCluster/Private/Render/Synchronization/Windows/DisplayClusterRenderSyncPolicyNvidiaSwapBarrierWindows.cpp
+++ b/Engine/Plugins/Runtime/nDisplay/Source/DisplayCluster/Private/Render/Synchronization/Windows/DisplayClusterRenderSyncPolicyNvidiaSwapBarrierWindows.cpp
@@ -17,7 +17,9 @@
 THIRD_PARTY_INCLUDES_START
 	#include "dxgi1_3.h"
 	#include "d3d12.h"
+#if WITH_NVAPI
 	#include "nvapi.h"
+#endif
 THIRD_PARTY_INCLUDES_END
 #include "Windows/HideWindowsPlatformTypes.h"
 
@@ -164,6 +166,7 @@ bool FDisplayClusterRenderSyncPolicyNvidiaSwapBarrier::InitializeNvidiaSwapLock(
 	using namespace DisplayClusterRenderSyncPolicyNvidiaSwapBarrier_Data_Windows;
 
 	bNvSyncInitializationCalled = true;
+#if WITH_NVAPI
 
 	// Get D3D1XDevice
 	D3DDevice = static_cast<IUnknown*>(GDynamicRHI->RHIGetNativeDevice());
@@ -185,7 +188,6 @@ bool FDisplayClusterRenderSyncPolicyNvidiaSwapBarrier::InitializeNvidiaSwapLock(
 	NvU32 MaxGroups = 0;
 	NvU32 MaxBarriers = 0;
 
-#if WITH_NVAPI
 	// Get amount of available groups and barriers
 	NvAPI_Status NvApiResult = NvAPI_D3D1x_QueryMaxSwapGroup(D3DDevice, &MaxGroups, &MaxBarriers);
 	if (NvApiResult != NVAPI_OK)
diff --git a/Engine/Plugins/Runtime/nDisplay/Source/DisplayClusterStageMonitoring/Private/DisplayClusterStageMonitoringModule.cpp b/Engine/Plugins/Runtime/nDisplay/Source/DisplayClusterStageMonitoring/Private/DisplayClusterStageMonitoringModule.cpp
--- a/Engine/Plugins/Runtime/nDisplay/Source/DisplayClusterStageMonitoring/Private/DisplayClusterStageMonitoringModule.cpp
+++ b/Engine/Plugins/Runtime/nDisplay/Source/DisplayClusterStageMonitoring/Private/DisplayClusterStageMonitoringModule.cpp
@@ -50,7 +50,7 @@ private:
 				const FName CurrentSyncPolicyName = CurrentSyncPolicy->GetName();
 				if (CurrentSyncPolicyName == NvidiaPolicy)
 				{
-					if (GetDefault<UDisplayClusterStageMonitoringSettings>()->ShouldEnableNvidiaWatchdog())
+					if (GetDefault<UDisplayClusterStageMonitoringSettings>()->ShouldEnableNvidiaWatchdog() && WITH_NVAPI)
 					{
 						NvidiaSyncWatchdog = MakeUnique<FNvidiaSyncWatchdog>();
 					}
diff --git a/Engine/Plugins/Runtime/nDisplay/Source/DisplayClusterStageMonitoring/Private/NvidiaSyncWatchdog.cpp b/Engine/Plugins/Runtime/nDisplay/Source/DisplayClusterStageMonitoring/Private/NvidiaSyncWatchdog.cpp
--- a/Engine/Plugins/Runtime/nDisplay/Source/DisplayClusterStageMonitoring/Private/NvidiaSyncWatchdog.cpp
+++ b/Engine/Plugins/Runtime/nDisplay/Source/DisplayClusterStageMonitoring/Private/NvidiaSyncWatchdog.cpp
@@ -16,7 +16,9 @@
 #pragma warning (disable : 4005) 	// Disable macro redefinition warning for compatibility with Windows SDK 8+
 #include "Windows/AllowWindowsPlatformTypes.h"
 #include "d3d11.h"
+#if WITH_NVAPI
 #include "nvapi.h"
+#endif
 #include "Windows/HideWindowsPlatformTypes.h"
 #pragma warning(pop)
 
@@ -63,6 +65,7 @@ void FNvidiaSyncWatchdog::OnPresentationPreSynchronization_RHIThread()
 
 void FNvidiaSyncWatchdog::OnPresentationPostSynchronization_RHIThread()
 {
+#if WITH_NVAPI
 	const uint64 PostSyncCycles = FPlatformTime::Cycles64();
 
 	//First time, this won't be valid
@@ -127,6 +130,7 @@ void FNvidiaSyncWatchdog::OnPresentationPostSynchronization_RHIThread()
 
 		return Result;
 	});
+#endif
 }
 
 FString FNvidiaSyncEvent::ToString() const
