ndisplay: Added safety check. In some cases when Custom Size is set on ICVFX camera the DisplayClusterFrameQueueCacheTexture and DisplayClusterViewportRenderTargetResource do not match in size and cause a crash.

#rnx

diff --git a/Engine/Plugins/Runtime/nDisplay/Source/DisplayClusterMedia/Private/Synchronization/LatencyQueue/DisplayClusterFrameQueue.cpp b/Engine/Plugins/Runtime/nDisplay/Source/DisplayClusterMedia/Private/Synchronization/LatencyQueue/DisplayClusterFrameQueue.cpp
--- a/Engine/Plugins/Runtime/nDisplay/Source/DisplayClusterMedia/Private/Synchronization/LatencyQueue/DisplayClusterFrameQueue.cpp
+++ b/Engine/Plugins/Runtime/nDisplay/Source/DisplayClusterMedia/Private/Synchronization/LatencyQueue/DisplayClusterFrameQueue.cpp
@@ -116,7 +116,7 @@ void FDisplayClusterFrameQueue::HandleProcessLatency_RenderThread(FRHICommandLis
 						{
 							// Get viewport texture
 							TArray<FRHITexture*> Textures;
-							ViewportProxyIt->GetResources_RenderThread(EDisplayClusterViewportResourceType::InternalRenderTargetResource, Textures);
+							ViewportProxyIt->GetResources_RenderThread(EDisplayClusterViewportResourceType::InternalRenderTargetEntireRectResource, Textures);
 
 							// Save current, and replace with an older
 							if (Textures.Num() > 0)
diff --git a/Engine/Plugins/Runtime/nDisplay/Source/DisplayClusterMedia/Private/Synchronization/LatencyQueue/DisplayClusterFrameQueueItem.cpp b/Engine/Plugins/Runtime/nDisplay/Source/DisplayClusterMedia/Private/Synchronization/LatencyQueue/DisplayClusterFrameQueueItem.cpp
--- a/Engine/Plugins/Runtime/nDisplay/Source/DisplayClusterMedia/Private/Synchronization/LatencyQueue/DisplayClusterFrameQueueItem.cpp
+++ b/Engine/Plugins/Runtime/nDisplay/Source/DisplayClusterMedia/Private/Synchronization/LatencyQueue/DisplayClusterFrameQueueItem.cpp
@@ -3,6 +3,7 @@
 #include "Synchronization/LatencyQueue/DisplayClusterFrameQueueItem.h"
 
 #include "DisplayClusterConfigurationTypes.h"
+#include "DisplayClusterMediaHelpers.h"
 #include "DisplayClusterMediaLog.h"
 #include "DisplayClusterRootActor.h"
 
@@ -63,7 +64,22 @@ void FDisplayClusterFrameQueueItem::LoadView(FRHICommandListImmediate& RHICmdLis
 	const FDisplayClusterFrameQueueItemView* const View = Views.Find(ViewportId);
 	if (View && View->Texture.IsValid())
 	{
-		RHICmdList.CopyTexture(View->Texture.GetReference(), Texture, FRHICopyTextureInfo());
+		if (View->Texture->GetDesc().Extent != Texture->GetDesc().Extent 
+			||	View->Texture->GetFormat() != Texture->GetFormat())
+		{
+			DisplayClusterMediaHelpers::ResampleTexture_RenderThread
+			(
+				RHICmdList,
+				View->Texture.GetReference(),
+				Texture,
+				{ FIntPoint::ZeroValue, View->Texture->GetDesc().Extent },
+				{ FIntPoint::ZeroValue, Texture->GetDesc().Extent }
+			);
+		}
+		else
+		{
+			RHICmdList.CopyTexture(View->Texture.GetReference(), Texture, FRHICopyTextureInfo());
+		}
 	}
 }
 
