[nDisplay] Fixed "invalid layout [1x1]" warnings when media is enabled, but only full-frame output is used

#rnx

diff --git a/Engine/Plugins/Runtime/nDisplay/Source/DisplayClusterMedia/Private/DisplayClusterMediaModule.cpp b/Engine/Plugins/Runtime/nDisplay/Source/DisplayClusterMedia/Private/DisplayClusterMediaModule.cpp
--- a/Engine/Plugins/Runtime/nDisplay/Source/DisplayClusterMedia/Private/DisplayClusterMediaModule.cpp
+++ b/Engine/Plugins/Runtime/nDisplay/Source/DisplayClusterMedia/Private/DisplayClusterMediaModule.cpp
@@ -235,8 +235,11 @@ void FDisplayClusterMediaModule::InitializeBackbufferUniformTilesOutput(const UD
 			// Validate tile layout
 			if (!bLayoutIsValid)
 			{
-				UE_LOG(LogDisplayClusterMedia, Warning, TEXT("Invalid layout [%dx%d] was requested for backbuffer capture. Max layout is [%dx%d]."),
-					MediaSettings.TiledSplitLayout.X, MediaSettings.TiledSplitLayout.Y, MaxLayout.X, MaxLayout.Y);
+				if (MediaSettings.TiledMediaOutputs.Num() > 0)
+				{
+					UE_LOG(LogDisplayClusterMedia, Warning, TEXT("Invalid layout [%dx%d] was requested for backbuffer capture. Max layout is [%dx%d]."),
+						MediaSettings.TiledSplitLayout.X, MediaSettings.TiledSplitLayout.Y, MaxLayout.X, MaxLayout.Y);
+				}
 
 				return;
 			}
