nDisplay: Fix Add Component button no longer showing in blueprint editor.

Fixes regression caused by 43119433.

diff --git a/Engine/Plugins/Runtime/nDisplay/Source/DisplayClusterConfigurator/Private/DisplayClusterConfiguratorModule.cpp b/Engine/Plugins/Runtime/nDisplay/Source/DisplayClusterConfigurator/Private/DisplayClusterConfiguratorModule.cpp
--- a/Engine/Plugins/Runtime/nDisplay/Source/DisplayClusterConfigurator/Private/DisplayClusterConfiguratorModule.cpp
+++ b/Engine/Plugins/Runtime/nDisplay/Source/DisplayClusterConfigurator/Private/DisplayClusterConfiguratorModule.cpp
@@ -156,6 +156,11 @@ void FDisplayClusterConfiguratorModule::ShutdownModule()
 	MenuExtensibilityManager.Reset();
 	ToolBarExtensibilityManager.Reset();
 
+	if (GEditor)
+	{
+		FDisplayClusterConfiguratorBlueprintEditor::UnregisterPanelExtensionFactory();
+	}
+
 	IKismetCompilerInterface& KismetCompilerModule = FModuleManager::GetModuleChecked<IKismetCompilerInterface>("KismetCompiler");
 	KismetCompilerModule.GetCompilers().Remove(&BlueprintCompiler);
 
@@ -168,6 +173,11 @@ void FDisplayClusterConfiguratorModule::ShutdownModule()
 
 void FDisplayClusterConfiguratorModule::OnPostEngineInit()
 {
+	if (GEditor)
+	{
+		FDisplayClusterConfiguratorBlueprintEditor::RegisterPanelExtensionFactory();
+	}
+
 	ToolbarExtensions->RegisterToolbarExtensions();
 }
 
