[nDisplay] Fixed media capture for the viewports that has overridden content
- Added another callback that is triggered for all overridden viewports
- Process that new callback on the media capture side

#rnx

diff --git a/Engine/Plugins/Runtime/nDisplay/Source/DisplayCluster/Private/DisplayClusterCallbacks.h b/Engine/Plugins/Runtime/nDisplay/Source/DisplayCluster/Private/DisplayClusterCallbacks.h
--- a/Engine/Plugins/Runtime/nDisplay/Source/DisplayCluster/Private/DisplayClusterCallbacks.h
+++ b/Engine/Plugins/Runtime/nDisplay/Source/DisplayCluster/Private/DisplayClusterCallbacks.h
@@ -139,7 +139,12 @@ public:
 
 	virtual FDisplayClusterProcessLatency_RenderThread& OnDisplayClusterProcessLatency_RenderThread() override
 	{
-		return FDisplayClusterProcessLatencyEvent;
+		return DisplayClusterProcessLatencyEvent;
+	}
+
+	virtual FDisplayClusterPostResolveOverridden_RenderThread& OnDisplayClusterPostResolveOverridden_RenderThread() override
+	{
+		return DisplayClusterPostResolveOverriddenEvent;
 	}
 
 	virtual FDisplayClusterPostFrameRender_RenderThread& OnDisplayClusterPostFrameRender_RenderThread() override
@@ -159,7 +164,7 @@ public:
 
 	virtual FDisplayClusterPreProcessIcvfx_RenderThread& OnDisplayClusterPreProcessIcvfx_RenderThread() override
 	{
-		return FDisplayClusterPreProcessIcvfxEvent;
+		return DisplayClusterPreProcessIcvfxEvent;
 	}
 
 	virtual FDisplayClusterUpdateViewportMediaState& OnDisplayClusterUpdateViewportMediaState() override
@@ -196,9 +201,10 @@ private:
 	FDisplayClusterPostCrossGpuTransfer_RenderThread DisplayClusterPostCrossGpuTransferEvent;
 	FDisplayClusterPassthroughMediaCapture_RenderThread DisplayClusterPassthroughMediaCaptureEvent;
 	FDisplayClusterPassthroughMediaInput_RenderThread   DisplayClusterPassthroughMediaInputEvent;
-	FDisplayClusterProcessLatency_RenderThread       FDisplayClusterProcessLatencyEvent;
+	FDisplayClusterProcessLatency_RenderThread       DisplayClusterProcessLatencyEvent;
+	FDisplayClusterPostResolveOverridden_RenderThread DisplayClusterPostResolveOverriddenEvent;
 	FDisplayClusterPostFrameRender_RenderThread      DisplayClusterPostFrameRenderEvent;
 	FDisplayClusterPostBackbufferUpdate_RenderThread DisplayClusterPostBackbufferUpdateEvent;
 	FDisplayClusterPostBackbufferUpdated_RenderThread DisplayClusterPostBackbufferUpdatedEvent;
-	FDisplayClusterPreProcessIcvfx_RenderThread      FDisplayClusterPreProcessIcvfxEvent;
+	FDisplayClusterPreProcessIcvfx_RenderThread      DisplayClusterPreProcessIcvfxEvent;
 };
diff --git a/Engine/Plugins/Runtime/nDisplay/Source/DisplayCluster/Private/Render/Viewport/DisplayClusterViewportManagerProxy.cpp b/Engine/Plugins/Runtime/nDisplay/Source/DisplayCluster/Private/Render/Viewport/DisplayClusterViewportManagerProxy.cpp
--- a/Engine/Plugins/Runtime/nDisplay/Source/DisplayCluster/Private/Render/Viewport/DisplayClusterViewportManagerProxy.cpp
+++ b/Engine/Plugins/Runtime/nDisplay/Source/DisplayCluster/Private/Render/Viewport/DisplayClusterViewportManagerProxy.cpp
@@ -400,6 +400,16 @@ void FDisplayClusterViewportManagerProxy::UpdateFrameResources_RenderThread(FRHI
 		if (ViewportProxyIt.IsValid() && ViewportProxyIt->GetRenderSettings_RenderThread().bVisible)
 		{
 			ViewportProxyIt->PostResolveViewport_RenderThread(RHICmdList);
+
+			// It's possible that some viewports don't need to request scene rendering, but get their final textures
+			// from another viewports. In such cases, no rendering callbacks like PostRenderViewFamily get called,
+			// therefore we can't media-capture such viewports from normal pipeline. To allow them to be captured,
+			// we explicitly trigger another callback for all viewports that are overridden.
+			const bool bOverridden = ViewportProxyIt->GetRenderSettings_RenderThread().IsViewportOverridden();
+			if (bOverridden)
+			{
+				IDisplayCluster::Get().GetCallbacks().OnDisplayClusterPostResolveOverridden_RenderThread().Broadcast(RHICmdList, ViewportProxyIt.Get());
+			}
 		}
 	}
 
diff --git a/Engine/Plugins/Runtime/nDisplay/Source/DisplayCluster/Public/IDisplayClusterCallbacks.h b/Engine/Plugins/Runtime/nDisplay/Source/DisplayCluster/Public/IDisplayClusterCallbacks.h
--- a/Engine/Plugins/Runtime/nDisplay/Source/DisplayCluster/Public/IDisplayClusterCallbacks.h
+++ b/Engine/Plugins/Runtime/nDisplay/Source/DisplayCluster/Public/IDisplayClusterCallbacks.h
@@ -135,6 +135,10 @@ public:
 	DECLARE_EVENT_ThreeParams(IDisplayClusterCallbacks, FDisplayClusterProcessLatency_RenderThread, FRHICommandListImmediate&, const IDisplayClusterViewportManagerProxy*, FViewport*);
 	virtual FDisplayClusterProcessLatency_RenderThread& OnDisplayClusterProcessLatency_RenderThread() = 0;
 
+	/** Called for every viewport that has configured content override **/
+	DECLARE_EVENT_TwoParams(IDisplayClusterCallbacks, FDisplayClusterPostResolveOverridden_RenderThread, FRHICommandListImmediate&, const IDisplayClusterViewportProxy*);
+	virtual FDisplayClusterPostResolveOverridden_RenderThread& OnDisplayClusterPostResolveOverridden_RenderThread() = 0;
+
 	/** Called in the end of the nD rendering chain, right before updating the backbuffer **/
 	DECLARE_EVENT_ThreeParams(IDisplayClusterCallbacks, FDisplayClusterPostFrameRender_RenderThread, FRHICommandListImmediate&, const IDisplayClusterViewportManagerProxy*, FViewport*);
 	virtual FDisplayClusterPostFrameRender_RenderThread& OnDisplayClusterPostFrameRender_RenderThread() = 0;
diff --git a/Engine/Plugins/Runtime/nDisplay/Source/DisplayClusterMedia/Private/Capture/DisplayClusterMediaCaptureViewportBase.cpp b/Engine/Plugins/Runtime/nDisplay/Source/DisplayClusterMedia/Private/Capture/DisplayClusterMediaCaptureViewportBase.cpp
--- a/Engine/Plugins/Runtime/nDisplay/Source/DisplayClusterMedia/Private/Capture/DisplayClusterMediaCaptureViewportBase.cpp
+++ b/Engine/Plugins/Runtime/nDisplay/Source/DisplayClusterMedia/Private/Capture/DisplayClusterMediaCaptureViewportBase.cpp
@@ -51,6 +51,9 @@ bool FDisplayClusterMediaCaptureViewportBase::StartCapture()
 	// Passthrough media capture
 	IDisplayCluster::Get().GetCallbacks().OnDisplayClusterPassthroughMediaCapture_RenderThread().AddRaw(this, &FDisplayClusterMediaCaptureViewportBase::OnPassthroughMediaCapture_RenderThread);
 
+	// PostResolveOverridden capture
+	IDisplayCluster::Get().GetCallbacks().OnDisplayClusterPostResolveOverridden_RenderThread().AddRaw(this, &FDisplayClusterMediaCaptureViewportBase::OnPostResolveOverridden_RenderThread);
+
 	// Start capture
 	const bool bStarted = FDisplayClusterMediaCaptureBase::StartCapture();
 
@@ -64,6 +67,7 @@ void FDisplayClusterMediaCaptureViewportBase::StopCapture()
 	IDisplayCluster::Get().GetCallbacks().OnDisplayClusterPostTonemapPass_RenderThread().RemoveAll(this);
 	IDisplayCluster::Get().GetCallbacks().OnDisplayClusterUpdateViewportMediaState().RemoveAll(this);
 	IDisplayCluster::Get().GetCallbacks().OnDisplayClusterPassthroughMediaCapture_RenderThread().RemoveAll(this);
+	IDisplayCluster::Get().GetCallbacks().OnDisplayClusterPostResolveOverridden_RenderThread().RemoveAll(this);
 
 	// Stop capturing
 	FDisplayClusterMediaCaptureBase::StopCapture();
@@ -285,3 +289,39 @@ void FDisplayClusterMediaCaptureViewportBase::OnPassthroughMediaCapture_RenderTh
 		}
 	}
 }
+
+void FDisplayClusterMediaCaptureViewportBase::OnPostResolveOverridden_RenderThread(FRHICommandListImmediate& RHICmdList, const IDisplayClusterViewportProxy* ViewportProxy)
+{
+	checkSlow(ViewportProxy);
+
+	if (!ViewportProxy)
+	{
+		return;
+	}
+
+	const FString CallbackViewportId = ViewportProxy->GetId();
+	const FString RequestedViewportId = GetViewportId();
+	if (!CallbackViewportId.Equals(RequestedViewportId, ESearchCase::IgnoreCase))
+	{
+		return;
+	}
+
+	TArray<FRHITexture*> Textures;
+	TArray<FIntRect>     Regions;
+
+	// Get RHI texture and pass it to the media capture pipeline
+	if (ViewportProxy->GetResourcesWithRects_RenderThread(EDisplayClusterViewportResourceType::AfterWarpBlendTargetableResource, Textures, Regions))
+	{
+		if (Textures.Num() > 0 && Regions.Num() > 0 && Textures[0])
+		{
+			FRDGBuilder GraphBuilder(RHICmdList);
+
+			FRDGTextureRef SrcTextureRef = RegisterExternalTexture(GraphBuilder, Textures[0], TEXT("DCMediaOutViewportTex"));
+
+			FMediaOutputTextureInfo TextureInfo{ SrcTextureRef, Regions[0] };
+			ExportMediaData_RenderThread(GraphBuilder, TextureInfo);
+
+			GraphBuilder.Execute();
+		}
+	}
+}
diff --git a/Engine/Plugins/Runtime/nDisplay/Source/DisplayClusterMedia/Private/Capture/DisplayClusterMediaCaptureViewportBase.h b/Engine/Plugins/Runtime/nDisplay/Source/DisplayClusterMedia/Private/Capture/DisplayClusterMediaCaptureViewportBase.h
--- a/Engine/Plugins/Runtime/nDisplay/Source/DisplayClusterMedia/Private/Capture/DisplayClusterMediaCaptureViewportBase.h
+++ b/Engine/Plugins/Runtime/nDisplay/Source/DisplayClusterMedia/Private/Capture/DisplayClusterMediaCaptureViewportBase.h
@@ -79,6 +79,9 @@ private:
 	/** PassthroughMediaCapture callback handler is used to capture in the passthrough scenarios */
 	void OnPassthroughMediaCapture_RenderThread(FRHICommandListImmediate& RHICmdList, const IDisplayClusterViewportManagerProxy* ViewportManagerProxy);
 
+	/** PostResolveOverridden callback for the viewports that grab textures from somewhere and don't emit PostRender events */
+	void OnPostResolveOverridden_RenderThread(FRHICommandListImmediate& RHICmdList, const IDisplayClusterViewportProxy* ViewportProxy);
+
 private:
 
 	/** Returns size of the viewport bound to this media */
