nDisplay: reworked preview rendering pipeline

ViewFamilies must be created and initialized in the same time period when BeginRenderingViewFamily() is called.

We cannot initialize, cache, and use ViewFamilies in a different frame or even in a different time period within a frame.
A known issue in the DLLS implementation of IUpscalerModularFeature expects a certain order of execution.

Removed FDisplayClusterViewportManagerPreview::ViewportsViewFamily array with pre-initialized ViewFamilies that were rendered over multiple frames.

Reworked FDisplayClusterViewportManagerPreview::RenderViewportPreview():
Uses PreviewRenderFrame as a container with viewports for rendering.
During rendering, modifies the PreviewRenderFrame to remove the processed viewport data.
When all viewports are rendered (bShouldRenderViewport=true), the final composition logic is called.
Finally, returns true, which signals the external function to initialize the PreviewRenderFrame for the next cluster node.

diff --git a/Engine/Plugins/Runtime/nDisplay/Source/DisplayCluster/Private/Render/Viewport/Preview/DisplayClusterViewportManagerPreview.cpp b/Engine/Plugins/Runtime/nDisplay/Source/DisplayCluster/Private/Render/Viewport/Preview/DisplayClusterViewportManagerPreview.cpp
--- a/Engine/Plugins/Runtime/nDisplay/Source/DisplayCluster/Private/Render/Viewport/Preview/DisplayClusterViewportManagerPreview.cpp
+++ b/Engine/Plugins/Runtime/nDisplay/Source/DisplayCluster/Private/Render/Viewport/Preview/DisplayClusterViewportManagerPreview.cpp
@@ -59,7 +59,6 @@ FString FDisplayClusterViewportManagerPreview::GetClusterNodeId(bool& bOutNextLo
 void FDisplayClusterViewportManagerPreview::ResetEntireClusterPreviewRendering()
 {
 	PreviewRenderFrame.Reset();
-	ViewportsViewFamily.Empty();
 	bEntireClusterRendered = false;
 }
 
diff --git a/Engine/Plugins/Runtime/nDisplay/Source/DisplayCluster/Private/Render/Viewport/Preview/DisplayClusterViewportManagerPreview.h b/Engine/Plugins/Runtime/nDisplay/Source/DisplayCluster/Private/Render/Viewport/Preview/DisplayClusterViewportManagerPreview.h
--- a/Engine/Plugins/Runtime/nDisplay/Source/DisplayCluster/Private/Render/Viewport/Preview/DisplayClusterViewportManagerPreview.h
+++ b/Engine/Plugins/Runtime/nDisplay/Source/DisplayCluster/Private/Render/Viewport/Preview/DisplayClusterViewportManagerPreview.h
@@ -95,6 +95,9 @@ protected:
 	/** Render preview frustums. */
 	void RenderPreviewFrustums();
 
+	/** Render ViewFamily with viewport.	*/
+	bool RenderViewportPreview(FViewport* InViewport, FCanvas* InSceneCanvas);
+
 public:
 	// Configuration of the current cluster node
 	const TSharedRef<FDisplayClusterViewportConfiguration, ESPMode::ThreadSafe> Configuration;
@@ -114,7 +117,4 @@ private:
 
 	// This render frame created for a cluster node
 	TUniquePtr<FDisplayClusterRenderFrame> PreviewRenderFrame;
-
-	// View families for each viewport for the current cluster node.
-	TArray<TSharedRef<FSceneViewFamilyContext>> ViewportsViewFamily;
 };
diff --git a/Engine/Plugins/Runtime/nDisplay/Source/DisplayCluster/Private/Render/Viewport/Preview/DisplayClusterViewportManagerPreview_Render.cpp b/Engine/Plugins/Runtime/nDisplay/Source/DisplayCluster/Private/Render/Viewport/Preview/DisplayClusterViewportManagerPreview_Render.cpp
--- a/Engine/Plugins/Runtime/nDisplay/Source/DisplayCluster/Private/Render/Viewport/Preview/DisplayClusterViewportManagerPreview_Render.cpp
+++ b/Engine/Plugins/Runtime/nDisplay/Source/DisplayCluster/Private/Render/Viewport/Preview/DisplayClusterViewportManagerPreview_Render.cpp
@@ -43,70 +43,64 @@
 ////////////////////////////////////////////////////////////////////////////////////////
 int32 FDisplayClusterViewportManagerPreview::RenderClusterNodePreview(const int32 InViewportsAmmount, FViewport* InViewport, FCanvas* InSceneCanvas)
 {
+	if (!InViewportsAmmount)
+	{
+		return 0;
+	}
+
+	if (!PreviewRenderFrame.IsValid())
+	{
+		return InViewportsAmmount;
+	}
+
 	// InViewportsAmmount - Total ammount of viewports that should be rendered on this frame.
 	// INDEX_NONE means render all viewports of this node
-	int32 OutViewportsAmount = (InViewportsAmmount == INDEX_NONE) ? (ViewportsViewFamily.Num() + 1) : InViewportsAmmount;
-
-	UWorld* CurrentWorld = Configuration->GetCurrentWorld();
-	FDisplayClusterViewportManager* ViewportManager = Configuration->GetViewportManagerImpl();
-	if (PreviewRenderFrame.IsValid() && OutViewportsAmount > 0 && CurrentWorld && ViewportManager)
+	int32 RenderedViewportsAmount = 0;
+	while (RenderViewportPreview(InViewport, InSceneCanvas))
 	{
+		RenderedViewportsAmount++;
 
-		// Render all viewports
-		while (!ViewportsViewFamily.IsEmpty() && OutViewportsAmount > 0)
+		// Stop if all viewports rendered
+		if (InViewportsAmmount >= 0 && RenderedViewportsAmount >= InViewportsAmmount)
 		{
-			FSceneViewFamily& ViewFamily = *ViewportsViewFamily[0];
-
-			if (InSceneCanvas)
-			{
-				GetRendererModule().BeginRenderingViewFamily(InSceneCanvas, &ViewFamily);
-			}
-			else
-			{
-				const ERHIFeatureLevel::Type FeatureLevel = CurrentWorld ? CurrentWorld->GetFeatureLevel() : GMaxRHIFeatureLevel;
-				FCanvas Canvas((FRenderTarget*)ViewportsViewFamily[0]->RenderTarget, nullptr, CurrentWorld, FeatureLevel, FCanvas::CDM_DeferDrawing /*FCanvas::CDM_ImmediateDrawing*/, 1.0f);
-				Canvas.Clear(FLinearColor::Black);
-
-				GetRendererModule().BeginRenderingViewFamily(&Canvas, &ViewFamily);
-			}
-
-			if (GNumExplicitGPUsForRendering > 1)
-			{
-				const FRHIGPUMask SubmitGPUMask = ViewFamily.Views.Num() == 1 ? ViewFamily.Views[0]->GPUMask : FRHIGPUMask::All();
-				ENQUEUE_RENDER_COMMAND(UDisplayClusterViewportClient_SubmitCommandList)(
-					[SubmitGPUMask](FRHICommandListImmediate& RHICmdList)
-					{
-						SCOPED_GPU_MASK(RHICmdList, SubmitGPUMask);
-						RHICmdList.SubmitCommandsHint();
-					});
-			}
+			return 0;
+		}
+	}
 
-			ViewportsViewFamily.RemoveAt(0);
-			OutViewportsAmount--;
+	if (!PreviewRenderFrame.IsValid())
+	{
+		// If error, ignore that render frame
+		return InViewportsAmmount;
+	}
 
-		}
+	// Entire cluster node rendered
+	if (PreviewRenderFrame->RenderTargets.IsEmpty())
+	{
+		// Current node render is completed
+		PreviewRenderFrame.Reset();
+		RenderedViewportsAmount++;
 
-		// After all viewports we render compose
-		if (OutViewportsAmount > 0)
+		if (FDisplayClusterViewportManager* ViewportManager = Configuration->GetViewportManagerImpl())
 		{
-			OutViewportsAmount--;
-
 			// Handle special viewports game-thread logic at frame end
 			// custom postprocess single frame flag must be removed at frame end on game thread
 			ViewportManager->FinalizeNewFrame();
 
 			// After all render target rendered call nDisplay frame rendering:
 			ViewportManager->RenderFrame(InViewport);
+		}
 
-			// Current node render is completed
-			PreviewRenderFrame.Reset();
+		// Send event about cluster node rendering is finished
+		OnClusterNodePreviewGenerated.ExecuteIfBound(Configuration->GetClusterNodeId());
+	}
 
-			// Send event about cluster node rendering is finished
-			OnClusterNodePreviewGenerated.ExecuteIfBound(Configuration->GetClusterNodeId());
-		}
+	// Returns the remaining number of viewports to render.
+	if (InViewportsAmmount >= 0)
+	{
+		return FMath::Max(0, InViewportsAmmount - RenderedViewportsAmount);
 	}
 
-	return OutViewportsAmount;
+	return 0;
 }
 
 bool FDisplayClusterViewportManagerPreview::InitializeClusterNodePreview(const EDisplayClusterRenderFrameMode InRenderMode, UWorld* InWorld, const FString& InClusterNodeId, FViewport* InViewport)
@@ -114,7 +108,6 @@ bool FDisplayClusterViewportManagerPreview::InitializeClusterNodePreview(const E
 	check(InWorld);
 
 	PreviewRenderFrame.Reset();
-	ViewportsViewFamily.Empty();
 
 	if (FDisplayClusterViewportManager* ViewportManager = Configuration->GetViewportManagerImpl())
 	{
@@ -131,91 +124,160 @@ bool FDisplayClusterViewportManagerPreview::InitializeClusterNodePreview(const E
 				// Initialize frame for render
 				ViewportManager->InitializeNewFrame();
 
-				// Create viewfamilies for all viewports
-				FSceneInterface* PreviewScene = InWorld->Scene;
-				FEngineShowFlags EngineShowFlags = FEngineShowFlags(EShowFlagInitMode::ESFIM_Game);
-
-				for (const FDisplayClusterRenderFrameTarget& RenderTargetIt : PreviewRenderFrame->RenderTargets)
-				{
-					// Special flag, allow clear RTT surface only for first family
-					bool bAdditionalViewFamily = false;
-					for (const FDisplayClusterRenderFrameTargetViewFamily& ViewFamiliesIt : RenderTargetIt.ViewFamilies)
-					{
-						// Create the view family for rendering the world scene to the viewport's render target
-						TSharedRef<FSceneViewFamilyContext> ViewFamily = MakeShared<FSceneViewFamilyContext>(ViewportManager->CreateViewFamilyConstructionValues(
-							RenderTargetIt,
-							PreviewScene,
-							EngineShowFlags,
-							bAdditionalViewFamily
-						));
-
-						ViewportManager->ConfigureViewFamily(RenderTargetIt, ViewFamiliesIt, *ViewFamily);
-
-						if (RenderTargetIt.CaptureMode == EDisplayClusterViewportCaptureMode::Default && Configuration->GetRenderFrameSettings().IsPostProcessDisabled())
-						{
-							if (ViewFamily->EngineShowFlags.TemporalAA)
-							{
-								ViewFamily->EngineShowFlags.SetTemporalAA(false);
-								ViewFamily->EngineShowFlags.SetAntiAliasing(true);
-							}
-						}
-
-						TArray<FSceneView*> Views;
-						for (const FDisplayClusterRenderFrameTargetView& ViewIt : ViewFamiliesIt.Views)
-						{
-							FDisplayClusterViewport* ViewportPtr = static_cast<FDisplayClusterViewport*>(ViewIt.Viewport.Get());
-							if (!ViewportPtr)
-							{
-								continue;
-							}
-
-							check(ViewportPtr->GetContexts().IsValidIndex(ViewIt.ContextNum));
-
-							// Always call CalcSceneView() because this function also starts VE and after that calls FDisplayClusterViewport::SetupSceneView() -> OCIO
-							// uvLC viewport is not rendered, but late OCIO is used in FDisplayClusterViewportProxy::ApplyOCIO_RenderThread()
-							FSceneView* View = ViewportPtr->ViewportPreview->CalcSceneView(*ViewFamily, ViewIt.ContextNum);
-
-							if (View == nullptr)
-							{
-								// This viewport is not rendering. Release all textures.
-								// The preview mesh materials associated with them will revert to default materials.
-								// In this case, all invalid viewports will be displayed gray or black.
-								ViewportPtr->ReleaseTextures();
-							}
-
-							// Remove the view that is not used for rendering.
-							if (View && (!ViewIt.IsViewportContextCanBeRendered() || ViewFamily->RenderTarget == nullptr))
-							{
-								ViewFamily->Views.Remove(View);
-
-								delete View;
-								View = nullptr;
-							}
-
-							if (View)
-							{
-								Views.Add(View);
-							}
-						}
-
-						if (!ViewFamily->Views.IsEmpty())
-						{
-							ViewportManager->PostConfigureViewFamily(RenderTargetIt, ViewFamiliesIt, *ViewFamily, Views);
-							ViewportsViewFamily.Add(ViewFamily);
-						}
-					}
-				}
-
 				// Update viewport preview instances if preview is used and DCRA supports previews:
 				Update();
 
-				return ViewportsViewFamily.Num() > 0;
+				return true;
 			}
 		}
 	}
 
 	PreviewRenderFrame.Reset();
-	ViewportsViewFamily.Empty();
 
 	return false;
 }
+
+bool FDisplayClusterViewportManagerPreview::RenderViewportPreview(FViewport* InViewport, FCanvas* InSceneCanvas)
+{
+	/**
+	* Uses PreviewRenderFrame as a container with viewports for rendering.
+	* During rendering, modifies the PreviewRenderFrame to remove the processed viewport data.
+	* When all viewports are rendered (bShouldRenderViewport=true), the final composition logic is called.
+	* Finally, returns true, which signals the external function to initialize the PreviewRenderFrame for the next cluster node.
+	*/
+
+	FDisplayClusterViewportManager* ViewportManager = Configuration->GetViewportManagerImpl();
+	UWorld* PreviewWorld = Configuration->GetCurrentWorld();
+	if (!ViewportManager || !PreviewWorld)
+	{
+		// Remove invalid render frame
+		PreviewRenderFrame.Reset();
+
+		return false;
+	}
+
+	FSceneInterface* PreviewScene = PreviewWorld->Scene;
+	FEngineShowFlags EngineShowFlags = FEngineShowFlags(EShowFlagInitMode::ESFIM_Game);
+
+	while(!PreviewRenderFrame->RenderTargets.IsEmpty())
+	{
+		if (PreviewRenderFrame->RenderTargets[0].ViewFamilies.IsEmpty())
+		{
+			PreviewRenderFrame->RenderTargets.RemoveAt(0);
+			continue;
+		}
+
+		// Special flag, allow clear RTT surface only for first family
+		bool bAdditionalViewFamily = false;
+		while(!PreviewRenderFrame->RenderTargets[0].ViewFamilies.IsEmpty())
+		{
+			if (PreviewRenderFrame->RenderTargets[0].ViewFamilies[0].Views.IsEmpty())
+			{
+				PreviewRenderFrame->RenderTargets[0].ViewFamilies.RemoveAt(0);
+				continue;
+			}
+
+			// Process ViewFamily
+			const FDisplayClusterRenderFrameTarget& RenderTargetIt = PreviewRenderFrame->RenderTargets[0];
+			const FDisplayClusterRenderFrameTargetViewFamily& ViewFamiliesIt = PreviewRenderFrame->RenderTargets[0].ViewFamilies[0];
+
+			// Create the view family for rendering the world scene to the viewport's render target
+			FSceneViewFamilyContext ViewFamily = ViewportManager->CreateViewFamilyConstructionValues(
+				RenderTargetIt,
+				PreviewScene,
+				EngineShowFlags,
+				bAdditionalViewFamily
+			);
+
+			ViewportManager->ConfigureViewFamily(RenderTargetIt, ViewFamiliesIt, ViewFamily);
+
+			if (RenderTargetIt.CaptureMode == EDisplayClusterViewportCaptureMode::Default && Configuration->GetRenderFrameSettings().IsPostProcessDisabled())
+			{
+				if (ViewFamily.EngineShowFlags.TemporalAA)
+				{
+					ViewFamily.EngineShowFlags.SetTemporalAA(false);
+					ViewFamily.EngineShowFlags.SetAntiAliasing(true);
+				}
+			}
+
+			TArray<FSceneView*> Views;
+			for(const FDisplayClusterRenderFrameTargetView& ViewIt : ViewFamiliesIt.Views)
+			{
+				FDisplayClusterViewport* ViewportPtr = static_cast<FDisplayClusterViewport*>(ViewIt.Viewport.Get());
+				if (!ViewportPtr)
+				{
+					continue;
+				}
+
+				check(ViewportPtr->GetContexts().IsValidIndex(ViewIt.ContextNum));
+
+				// Always call CalcSceneView() because this function also starts VE and after that calls FDisplayClusterViewport::SetupSceneView() -> OCIO
+				// uvLC viewport is not rendered, but late OCIO is used in FDisplayClusterViewportProxy::ApplyOCIO_RenderThread()
+				FSceneView* View = ViewportPtr->ViewportPreview->CalcSceneView(ViewFamily, ViewIt.ContextNum);
+
+				if (View == nullptr)
+				{
+					// This viewport is not rendering. Release all textures.
+					// The preview mesh materials associated with them will revert to default materials.
+					// In this case, all invalid viewports will be displayed gray or black.
+					ViewportPtr->ReleaseTextures();
+				}
+
+				// Remove the view that is not used for rendering.
+				if (View && (!ViewIt.IsViewportContextCanBeRendered() || ViewFamily.RenderTarget == nullptr))
+				{
+					ViewFamily.Views.Remove(View);
+
+					delete View;
+					View = nullptr;
+				}
+
+				if (View)
+				{
+					Views.Add(View);
+				}
+			}
+
+			const bool bShouldRenderViewport = !ViewFamily.Views.IsEmpty();
+			if (bShouldRenderViewport)
+			{
+				ViewportManager->PostConfigureViewFamily(RenderTargetIt, ViewFamiliesIt, ViewFamily, Views);
+
+				if (InSceneCanvas)
+				{
+					GetRendererModule().BeginRenderingViewFamily(InSceneCanvas, &ViewFamily);
+				}
+				else
+				{
+					const ERHIFeatureLevel::Type FeatureLevel = PreviewWorld ? PreviewWorld->GetFeatureLevel() : GMaxRHIFeatureLevel;
+					FCanvas Canvas((FRenderTarget*)ViewFamily.RenderTarget, nullptr, PreviewWorld, FeatureLevel, FCanvas::CDM_DeferDrawing /*FCanvas::CDM_ImmediateDrawing*/, 1.0f);
+					Canvas.Clear(FLinearColor::Black);
+
+					GetRendererModule().BeginRenderingViewFamily(&Canvas, &ViewFamily);
+				}
+
+				if (GNumExplicitGPUsForRendering > 1)
+				{
+					const FRHIGPUMask SubmitGPUMask = ViewFamily.Views.Num() == 1 ? ViewFamily.Views[0]->GPUMask : FRHIGPUMask::All();
+					ENQUEUE_RENDER_COMMAND(UDisplayClusterViewportClient_SubmitCommandList)(
+						[SubmitGPUMask](FRHICommandListImmediate& RHICmdList)
+						{
+							SCOPED_GPU_MASK(RHICmdList, SubmitGPUMask);
+							RHICmdList.SubmitCommandsHint();
+						});
+				}
+			}
+
+			// Remove processed Views from the FrameData.
+			PreviewRenderFrame->RenderTargets[0].ViewFamilies[0].Views.Reset();
+
+			// Return true if a single viewport was drawn.
+			if (bShouldRenderViewport)
+			{
+				return true;
+			}
+		}
+	}
+
+	return false;
+}
\ No newline at end of file
