nDisplay: Added features related to projection policy and viewpoint.

Note: This CL works with existing projection policies.

UDisplayClusterCameraComponent - added new property "Follow Camera"
IDisplayClusterViewport: added new CalculateViewData() function, implementing a new approach to recalculating view data.
FDisplayClusterViewport_Context modified: added a lot of data for math (OriginToWorld, RootActorToWorld, RenderLocation, etc.)

IDisplayClusterProjectionPolicy: added a new functions to the API:
- GetGeometryToMeters() - how many geometry units are in one meter.
- GetOriginLocationOffsetInGeometryUnits() - Origin's location in geometry units.
- UpdateClippingPlanes() - Projection policy can update clipping planes.
- GetProjectionPolicyFlags() - Return Projection policy flags EDisplayClusterProjectionPolicyFlags
* UseEyeTracking - use ViewPoint instead of camera
* UseLocalSpace - projection policy input args in local space
* EnableFollowCameraFeature - is projection polisy supports this feature.

diff --git a/Engine/Plugins/Runtime/nDisplay/Source/DisplayCluster/Private/Game/EngineClasses/Scene/DisplayClusterCameraComponent.cpp b/Engine/Plugins/Runtime/nDisplay/Source/DisplayCluster/Private/Game/EngineClasses/Scene/DisplayClusterCameraComponent.cpp
--- a/Engine/Plugins/Runtime/nDisplay/Source/DisplayCluster/Private/Game/EngineClasses/Scene/DisplayClusterCameraComponent.cpp
+++ b/Engine/Plugins/Runtime/nDisplay/Source/DisplayCluster/Private/Game/EngineClasses/Scene/DisplayClusterCameraComponent.cpp
@@ -133,6 +133,17 @@ UCameraComponent* UDisplayClusterCameraComponent::GetTargetCameraComponent(const
 	return CfgCameraComponent.GetExternalCineCameraActorComponent();
 }
 
+bool UDisplayClusterCameraComponent::CanOverrideEyePosition() const
+{
+	// ICVFX camera can be used as a ViewPoint
+	if (IsICVFXCameraBeingUsed())
+	{
+		return bUseICVFXCameraComponentTracking;
+	}
+
+	return false;
+}
+
 bool UDisplayClusterCameraComponent::IsViewPointOverrideCameraPosition() const
 {
 	// If the ICFX camera component is used, it can override the viewpoint position.
@@ -141,6 +152,12 @@ bool UDisplayClusterCameraComponent::IsViewPointOverrideCameraPosition() const
 		return !bUseICVFXCameraComponentTracking;
 	}
 
+	// Use external camera for render
+	if (ShouldFollowCameraLocation())
+	{
+		return false;
+	}
+
 	// By default, ViewPoint is always used as the camera.
 	return true;
 }
@@ -162,45 +179,31 @@ bool UDisplayClusterCameraComponent::GetTargetCameraDesiredViewInternal(const ID
 		OutCustomNearClippingPlane = nullptr;
 	}
 
+	bool bViewInfoChanged = false;
+
 	// 1. Active engine camera
 	if (CfgCameraComponent.IsActiveEngineCameraBeingUsed())
 	{
-		// Get PostProcess from the Game camera.
-		if (IDisplayClusterViewport::GetPlayerCameraView(InViewportConfiguration.GetCurrentWorld(), bUseCameraPostprocess, InOutViewInfo))
-		{
-			if (IsViewPointOverrideCameraPosition())
-			{
-				// Use this component as a camera
-				InOutViewInfo.Location = GetComponentLocation();
-				InOutViewInfo.Rotation = GetComponentRotation();
-			}
-
-			return true;
-		}
+		// Get ViewInfo from the Game camera.
+		bViewInfoChanged = IDisplayClusterViewport::GetPlayerCameraView(InViewportConfiguration.GetCurrentWorld(), bUseCameraPostprocess, InOutViewInfo);
 	}
-	else
+	else if (UCameraComponent* SceneCameraComponent = GetTargetCameraComponent(InViewportConfiguration))
 	{
-		if (UCameraComponent* SceneCameraComponent = GetTargetCameraComponent(InViewportConfiguration))
-		{
-			if (IDisplayClusterViewport::GetCameraComponentView(SceneCameraComponent, InViewportConfiguration.GetRootActorWorldDeltaSeconds(), bUseCameraPostprocess, InOutViewInfo, OutCustomNearClippingPlane))
-			{
-				if (IsViewPointOverrideCameraPosition())
-				{
-					// Use this component as a camera
-					InOutViewInfo.Location = GetComponentLocation();
-					InOutViewInfo.Rotation = GetComponentRotation();
-				}
+		// Get ViewInfo from the camera component
+		bViewInfoChanged = IDisplayClusterViewport::GetCameraComponentView(SceneCameraComponent, InViewportConfiguration.GetRootActorWorldDeltaSeconds(), bUseCameraPostprocess, InOutViewInfo, OutCustomNearClippingPlane);
+	}
 
-				return true;
-			}
-		}
+	if (IsViewPointOverrideCameraPosition())
+	{
+		// ViewPoint override camera position
+		InOutViewInfo.Location = GetComponentLocation();
+		InOutViewInfo.Rotation = GetComponentRotation();
 	}
 
-	// The target camera is not found, so we can't use the custom camera view.
-	return false;
+	return bViewInfoChanged;
 }
 
-void UDisplayClusterCameraComponent::GetDesiredView(IDisplayClusterViewportConfiguration& InViewportConfiguration, FMinimalViewInfo& InOutViewInfo, float* OutCustomNearClippingPlane)
+void UDisplayClusterCameraComponent::GetDesiredView(const IDisplayClusterViewportConfiguration& InViewportConfiguration, FMinimalViewInfo& InOutViewInfo, float* OutCustomNearClippingPlane) const
 {
 	if (GetTargetCameraDesiredViewInternal(InViewportConfiguration, InOutViewInfo, OutCustomNearClippingPlane))
 	{
@@ -224,19 +227,23 @@ void UDisplayClusterCameraComponent::GetDesiredView(IDisplayClusterViewportConfi
 
 void UDisplayClusterCameraComponent::GetEyePosition(const IDisplayClusterViewportConfiguration& InViewportConfiguration, FVector& OutViewLocation, FRotator& OutViewRotation)
 {
-	FMinimalViewInfo ViewInfo;
-	if (GetTargetCameraDesiredViewInternal(InViewportConfiguration, ViewInfo))
-	{
-		OutViewLocation = ViewInfo.Location;
-		OutViewRotation = ViewInfo.Rotation;
-	}
-	else
+	// Use camera position as a ViewPoint
+	if (CanOverrideEyePosition())
 	{
-		// By default this component is used as ViewPoint
-		// Use this component as a camera
-		OutViewLocation = GetComponentLocation();
-		OutViewRotation = GetComponentRotation();
+		FMinimalViewInfo ViewInfo;
+		if (GetTargetCameraDesiredViewInternal(InViewportConfiguration, ViewInfo))
+		{
+			OutViewLocation = ViewInfo.Location;
+			OutViewRotation = ViewInfo.Rotation;
+
+			return;
+		}
 	}
+
+	// By default this component is used as ViewPoint
+	// Use this component as a camera
+	OutViewLocation = GetComponentLocation();
+	OutViewRotation = GetComponentRotation();
 }
 
 #if WITH_EDITOR
@@ -342,6 +349,24 @@ bool UDisplayClusterCameraComponent::IsExternalCameraBeingUsed() const
 	return TargetCameraType == EDisplayClusterTargetCameraType::ExternalCineCameraActor && !IsICVFXCameraBeingUsed() && ExternalCineCameraActor.IsValid();
 }
 
+bool UDisplayClusterCameraComponent::ShouldFollowCameraLocation() const
+{
+	if (bFollowCameraPosition)
+	{
+		if (IsExternalCameraBeingUsed())
+		{
+			return true;
+		}
+
+		if (IsActiveEngineCameraBeingUsed())
+		{
+			return true;
+		}
+	}
+
+	return false;
+}
+
 EDisplayClusterViewportCameraPostProcessFlags UDisplayClusterCameraComponent::GetCameraPostProcessFlags() const
 {
 	EDisplayClusterViewportCameraPostProcessFlags OutPostProcessFlags = EDisplayClusterViewportCameraPostProcessFlags::None;
diff --git a/Engine/Plugins/Runtime/nDisplay/Source/DisplayCluster/Private/Render/Viewport/DisplayClusterViewport.cpp b/Engine/Plugins/Runtime/nDisplay/Source/DisplayCluster/Private/Render/Viewport/DisplayClusterViewport.cpp
--- a/Engine/Plugins/Runtime/nDisplay/Source/DisplayCluster/Private/Render/Viewport/DisplayClusterViewport.cpp
+++ b/Engine/Plugins/Runtime/nDisplay/Source/DisplayCluster/Private/Render/Viewport/DisplayClusterViewport.cpp
@@ -531,8 +531,9 @@ bool FDisplayClusterViewport::UpdateFrameContexts(const uint32 InStereoViewIndex
 			ContextIt.bDisableRender = true;
 			ContextIt.FrameTargetRect = FDisplayClusterViewportHelpers::GetValidViewportRect(DesiredFrameTargetRect, GetId(), TEXT("Context Frame"));
 
-			// Reset context state, and invalidate caches from previous frame.
-			ContextIt.ContextState = EDisplayClusterViewportContextState::None;
+			// Force context to be recalculated.
+			ContextIt.bCalculated = false;
+			ContextIt.bValidData = false;
 		}
 
 		// Release only part of the resources, leaving resources that can be used by other viewports (viewport override feature)
diff --git a/Engine/Plugins/Runtime/nDisplay/Source/DisplayCluster/Private/Render/Viewport/DisplayClusterViewport.h b/Engine/Plugins/Runtime/nDisplay/Source/DisplayCluster/Private/Render/Viewport/DisplayClusterViewport.h
--- a/Engine/Plugins/Runtime/nDisplay/Source/DisplayCluster/Private/Render/Viewport/DisplayClusterViewport.h
+++ b/Engine/Plugins/Runtime/nDisplay/Source/DisplayCluster/Private/Render/Viewport/DisplayClusterViewport.h
@@ -128,6 +128,8 @@ public:
 	virtual bool    CalculateView(const uint32 InContextNum, FVector& InOutViewLocation, FRotator& InOutViewRotation, const float WorldToMeters) override;
 	virtual bool    GetProjectionMatrix(const uint32 InContextNum, FMatrix& OutPrjMatrix)  override;
 
+	virtual bool CalculateViewData(const uint32 InContextNum, const FVector& InViewLocation, const FRotator& InViewRotation, const float WorldToMeters) override;
+
 	virtual bool SetupViewPoint(const uint32 InContextNum, FMinimalViewInfo& InOutViewInfo) override;
 	virtual float GetStereoEyeOffsetDistance(const uint32 InContextNum) override;
 	virtual class UDisplayClusterCameraComponent* GetViewPointCameraComponent(const EDisplayClusterRootActorType InRootActorType) const override;
@@ -452,6 +454,17 @@ private:
 	*/
 	void ImplGatherActiveExtensions(const int32 ViewIndex, FViewport* InViewport, TArray<FSceneViewExtensionRef>& OutExtensions) const;
 
+	/** Initialize View data
+	 *
+	 * @param InContextNum      - (in) viewport eye context index
+	 * @param InOutViewLocation - (in) View location
+	 * @param InOutViewRotation - (in) View rotator
+	 * @param WorldToMeters     - UE world scale
+	 *
+	 * @return true on success
+	*/
+	bool PreCalculateViewData(const uint32 InContextNum, const FVector& InViewLocation, const FRotator& InViewRotation, const float WorldToMeters);
+
 public:
 	// Configuration of the current cluster node
 	const TSharedRef<FDisplayClusterViewportConfiguration, ESPMode::ThreadSafe> Configuration;
diff --git a/Engine/Plugins/Runtime/nDisplay/Source/DisplayCluster/Private/Render/Viewport/DisplayClusterViewport_Math.cpp b/Engine/Plugins/Runtime/nDisplay/Source/DisplayCluster/Private/Render/Viewport/DisplayClusterViewport_Math.cpp
--- a/Engine/Plugins/Runtime/nDisplay/Source/DisplayCluster/Private/Render/Viewport/DisplayClusterViewport_Math.cpp
+++ b/Engine/Plugins/Runtime/nDisplay/Source/DisplayCluster/Private/Render/Viewport/DisplayClusterViewport_Math.cpp
@@ -24,23 +24,16 @@
 namespace UE::DisplayCluster::Viewport::Math
 {
 	/**
-	 * Calculates the ViewOffset for the eye from the view location and changes the value of the view location to the eye position.
-	 * 
-	 * @param PassOffsetSwap - Distance to the eye from the midpoint between the eyes.
-	 * @param InOutViewLocation - (in, out) view location
-	 * @param InViewRotation (in) - rotation of the view (from this rotator we get the direction to the eye)
-	 * 
+	 * Calculates the ViewOffset for the eye from the view location
+	 *
+	 * @param InViewRotation - (in) rotation of the view (from this rotator we get the direction to the eye)
+	 * @param PassOffsetSwap - (in) Distance to the eye from the midpoint between the eyes.
+	 *
 	 * @return - the distance to the eye from the original ViewLocation
 	 */
-	static inline FVector ImplGetViewOffset(const double PassOffsetSwap, FVector& InOutViewLocation, const FRotator& InViewRotation)
+	static inline FVector GetStereoEyeOffset(const FRotator& InViewRotation, const float StereoEyeOffsetDistance)
 	{
-		// Apply computed offset to the view location
-		const FQuat EyeQuat = InViewRotation.Quaternion();
-		FVector ViewOffset = EyeQuat.RotateVector(FVector(0.0f, PassOffsetSwap, 0.0f));
-		
-		InOutViewLocation += ViewOffset;
-
-		return ViewOffset;
+		return  InViewRotation.Quaternion().RotateVector(FVector(0.0f, StereoEyeOffsetDistance, 0.0f));
 	}
 
 	/** check frustum. */
@@ -59,24 +52,26 @@ namespace UE::DisplayCluster::Viewport::Math
 		}
 	}
 };
+
 ///////////////////////////////////////////////////////////////////////////////////////
 //          FDisplayClusterViewport
 ///////////////////////////////////////////////////////////////////////////////////////
 FVector2D FDisplayClusterViewport::GetClippingPlanes() const
 {
-	const float NCP = GNearClippingPlane;
-	const float FCP = NCP; // nDisplay does not use the far plane of the clipping
+	float NCP = (CustomNearClippingPlane >= 0)
+		? CustomNearClippingPlane
+		: GNearClippingPlane;
+
+	// nDisplay does not use the far plane of the clipping
+	float FCP = NCP;
 
-	// Supports custom near clipping plane
-	float ZNear = NCP;
-	float ZFar = FCP;
-	if (CustomNearClippingPlane >= 0)
+	// Projection policies can change the clipping planes (domeprojection).
+	if (ProjectionPolicy)
 	{
-		ZNear = CustomNearClippingPlane;
-		ZFar = (NCP == FCP) ? ZNear : ZFar;
+		ProjectionPolicy->ApplyClippingPlanesOverrides(this, NCP, FCP);
 	}
 
-	return FVector2D(ZNear, ZFar);
+	return FVector2D(NCP, FCP);
 }
 
 bool FDisplayClusterViewport::GetViewPointCameraEye(const uint32 InContextNum, FVector& OutViewLocation, FRotator& OutViewRotation, FVector& OutViewOffset)
@@ -90,106 +85,183 @@ bool FDisplayClusterViewport::GetViewPointCameraEye(const uint32 InContextNum, F
 		SceneCameraComponent->GetEyePosition(GetConfiguration(), OutViewLocation, OutViewRotation);
 
 		// Calculate stereo ViewOffset:
-		OutViewOffset = ImplGetViewOffset(GetStereoEyeOffsetDistance(InContextNum), OutViewLocation, OutViewRotation);
+		OutViewOffset = GetStereoEyeOffset(OutViewRotation, GetStereoEyeOffsetDistance(InContextNum));
+
+		// Add stereo eye offset to the output view location.
+		OutViewLocation += OutViewOffset;
 
 		return true;
 	}
 
 	return false;
 }
-
 bool FDisplayClusterViewport::CalculateView(const uint32 InContextNum, FVector& InOutViewLocation, FRotator& InOutViewRotation, const float WorldToMeters)
 {
-	using namespace UE::DisplayCluster::Viewport::Math;
-
 	if (Contexts.IsValidIndex(InContextNum))
 	{
-		if (!EnumHasAnyFlags(Contexts[InContextNum].ContextState, EDisplayClusterViewportContextState::InvalidViewPoint))
+		// The function can be called several times per frame.
+		// Each time it must return the same values. For optimization purposes, after the first call this function
+		// stores the result in the context variables 'ViewLocation' and 'ViewRotation'.
+		// Finally, raises this flag for subsequent calls in the current frame.
+
+		// Calculate once per frame.
+		if (!Contexts[InContextNum].bCalculated)
 		{
-			// The function can be called several times per frame.
-			// Each time it must return the same values. For optimization purposes, after the first call this function
-			// stores the result in the context variables 'ViewLocation' and 'ViewRotation'.
-			// Finally, raises this flag for subsequent calls in the current frame.
-
-			if (EnumHasAnyFlags(Contexts[InContextNum].ContextState, EDisplayClusterViewportContextState::HasCalculatedViewPoint))
-			{
-				// Use calculated values
-				// Since this function can be called several times from LocalPlayer.cpp, the cached values are used on repeated calls.
-				// This should give a performance boost for 'mesh', 'mpcdi' projections with a large number of vertices in the geometry or large warp texture size.
-				InOutViewLocation = Contexts[InContextNum].ViewLocation;
-				InOutViewRotation = Contexts[InContextNum].ViewRotation;
-
-				return true;
-			}
-
-			// Calculate stereo ViewOffset:
-			const FVector ViewOffset = ImplGetViewOffset(GetStereoEyeOffsetDistance(InContextNum), InOutViewLocation, InOutViewRotation);
-			const FVector2D ClipingPlanes = GetClippingPlanes();
-			if (ProjectionPolicy.IsValid() && ProjectionPolicy->CalculateView(this, InContextNum, InOutViewLocation, InOutViewRotation, ViewOffset, WorldToMeters, ClipingPlanes.X, ClipingPlanes.Y))
-			{
-				Contexts[InContextNum].WorldToMeters = WorldToMeters;
-
-				// Save the calculated values and update the state of the context
-				Contexts[InContextNum].ViewLocation = InOutViewLocation;
-				Contexts[InContextNum].ViewRotation = InOutViewRotation;
-				EnumAddFlags(Contexts[InContextNum].ContextState, EDisplayClusterViewportContextState::HasCalculatedViewPoint);
-
-				return true;
-			}
-
-			// ProjectionPolicy->CalculateView() returns false, this view is invalid
-			EnumAddFlags(Contexts[InContextNum].ContextState, EDisplayClusterViewportContextState::InvalidViewPoint);
+			CalculateViewData(InContextNum, InOutViewLocation, InOutViewRotation, WorldToMeters);
+		}
+
+		if (Contexts[InContextNum].bValidData)
+		{
+			// Use calculated values
+			// Since this function can be called several times from LocalPlayer.cpp, the cached values are used on repeated calls.
+			// This should give a performance boost for 'mesh', 'mpcdi' projections with a large number of vertices in the geometry or large warp texture size.
+			InOutViewLocation = Contexts[InContextNum].ViewData.RenderLocation;
+			InOutViewRotation = Contexts[InContextNum].ViewData.RenderRotation;
+
+			return true;
 		}
 	}
 
 	return false;
 }
 
+bool FDisplayClusterViewport::CalculateViewData(const uint32 InContextNum, const FVector& InViewLocation, const FRotator& InViewRotation, const float WorldToMeters)
+{
+	using namespace UE::DisplayCluster::Viewport::Math;
+
+	if (!PreCalculateViewData(InContextNum, InViewLocation, InViewRotation, WorldToMeters))
+	{
+		return false;
+	}
+
+	FDisplayClusterViewport_Context& OutContext = Contexts[InContextNum];
+
+	// Gets the projection policy flags that control the logic below:
+	const EDisplayClusterProjectionPolicyFlags ProjectionPolicyFlags = ProjectionPolicy->GetProjectionPolicyFlags(this, InContextNum);
+
+	// Evaluate View
+	FVector ViewLocation(InViewLocation);
+	FRotator ViewRotation(InViewRotation);
+
+	// Store Camera in world space
+	OutContext.ViewData.CameraOrigin = ViewLocation;
+	OutContext.ViewData.CameraRotation = ViewRotation;
+
+	// Uses the eye position (ViewPoint) instead of the camera
+	if (EnumHasAnyFlags(ProjectionPolicyFlags, EDisplayClusterProjectionPolicyFlags::UseEyeTracking))
+	{
+		// Here we use the ViewPoint component as the eye position
+		if (UDisplayClusterCameraComponent* SceneCameraComponent = GetViewPointCameraComponent(EDisplayClusterRootActorType::Scene))
+		{
+			// Get eye position from the ViewPoint component.
+			SceneCameraComponent->GetEyePosition(GetConfiguration(), ViewLocation, ViewRotation);
+		}
+	}
+
+	// Store Eye in world space
+	OutContext.ViewData.EyeOrigin = ViewLocation;
+	OutContext.ViewData.EyeRotation = ViewRotation;
+
+	// Stereo offset distance
+	const float StereoEyeOffsetDistance = GetStereoEyeOffsetDistance(InContextNum);
+	FVector ViewOffset = GetStereoEyeOffset(ViewRotation, StereoEyeOffsetDistance);
+
+	// Projection policy expects view location with stereo offset.
+	ViewLocation += ViewOffset;
+
+	// Convert eye to local space and to Geometry units
+	if (EnumHasAnyFlags(ProjectionPolicyFlags, EDisplayClusterProjectionPolicyFlags::UseLocalSpace))
+	{
+		ViewRotation = OutContext.OriginToWorld.InverseTransformRotation(ViewRotation.Quaternion()).Rotator();
+
+		// Scale units from Geometry to UE World
+		const float ScaleWorldToGeometry = OutContext.GeometryToMeters / OutContext.WorldToMeters;
+
+		ViewLocation = OutContext.OriginToWorld.InverseTransformPosition(ViewLocation) * ScaleWorldToGeometry;
+		ViewOffset = OutContext.OriginToWorld.InverseTransformPosition(ViewOffset) * ScaleWorldToGeometry;
+	}
+
+	// Calculate View
+	if (!ProjectionPolicy->CalculateView(this, InContextNum, ViewLocation, ViewRotation, ViewOffset, WorldToMeters, OutContext.ProjectionData.ZNear, OutContext.ProjectionData.ZFar))
+	{
+		// CalculateView() failed.
+		return false;
+	}
+
+	// Convert eye to world space and units
+	if (EnumHasAnyFlags(ProjectionPolicyFlags, EDisplayClusterProjectionPolicyFlags::UseLocalSpace))
+	{
+		ViewRotation = OutContext.OriginToWorld.TransformRotation(ViewRotation.Quaternion()).Rotator();
+
+		// Scale units from Geometry to UE World
+		const float ScaleGeometryToWorld = OutContext.WorldToMeters / OutContext.GeometryToMeters;
+
+		ViewLocation = OutContext.OriginToWorld.TransformPosition(ViewLocation * ScaleGeometryToWorld);
+	}
+
+	// Save view: DCRA related view
+	OutContext.ViewLocation = ViewLocation;
+	OutContext.ViewRotation = ViewRotation;
+
+	// Follow camera feature:
+	// At this point, let's change the values for rendering only.
+	if (EnumHasAnyFlags(ProjectionPolicyFlags, EDisplayClusterProjectionPolicyFlags::EnableFollowCameraFeature))
+	{
+		const UDisplayClusterCameraComponent* ConfigurationCameraComponent = GetViewPointCameraComponent(EDisplayClusterRootActorType::Configuration);
+		const UDisplayClusterCameraComponent* SceneCameraComponent = GetViewPointCameraComponent(EDisplayClusterRootActorType::Scene);
+		if (SceneCameraComponent && ConfigurationCameraComponent && ConfigurationCameraComponent->ShouldFollowCameraLocation())
+		{
+			FMinimalViewInfo CameraViewInfo;
+			SceneCameraComponent->GetDesiredView(*Configuration, CameraViewInfo);
+
+			const FTransform CameraToWorld(CameraViewInfo.Rotation, CameraViewInfo.Location);
+
+			// Follow camera position: use camera position as a DCRA position without moving of an actor.
+			// World -> [DCRA local space = CAMERA local space] -> World
+			const FVector RootActorSpaceLocation = OutContext.RootActorToWorld.InverseTransformPosition(ViewLocation);
+			const FVector NewViewLocation = CameraToWorld.TransformPosition(RootActorSpaceLocation);
+
+			const FQuat    RootActorSpaceRotation = OutContext.RootActorToWorld.InverseTransformRotation(ViewRotation.Quaternion());
+			const FRotator NewViewRotation = CameraToWorld.TransformRotation(RootActorSpaceRotation).Rotator();
+
+			// Update calculated values and update the state of the context
+			ViewLocation = NewViewLocation;
+			ViewRotation = NewViewRotation;
+		}
+	}
+
+	// save render view
+	OutContext.ViewData.RenderLocation = ViewLocation;
+	OutContext.ViewData.RenderRotation = ViewRotation;
+
+	// Calculate Projection matrix
+	FMatrix ProjectionMatrix;
+	if (!ProjectionPolicy->GetProjectionMatrix(this, InContextNum, ProjectionMatrix))
+	{
+		// GetProjectionMatrix() failed.
+		return false;
+	}
+
+	OutContext.ProjectionMatrix = ProjectionMatrix;
+
+	// Set this context data as valid.
+	OutContext.bValidData = true;
+
+	return true;
+}
+
 bool FDisplayClusterViewport::GetProjectionMatrix(const uint32 InContextNum, FMatrix& OutPrjMatrix)
 {
 	if (Contexts.IsValidIndex(InContextNum))
 	{
-		if (!EnumHasAnyFlags(Contexts[InContextNum].ContextState, EDisplayClusterViewportContextState::InvalidProjectionMatrix))
+		if (Contexts[InContextNum].bValidData)
 		{
-			// The function can also be called several times per frame.
-			// stores the result in the context variables 'ProjectionMatrix' and 'OverscanProjectionMatrix'.
-			// Finally, raises this flag for subsequent calls in the current frame.
-
-			if (EnumHasAnyFlags(Contexts[InContextNum].ContextState, EDisplayClusterViewportContextState::HasCalculatedProjectionMatrix))
-			{
-				// use already calculated values
-				if (EnumHasAnyFlags(Contexts[InContextNum].ContextState, EDisplayClusterViewportContextState::HasCalculatedOverscanProjectionMatrix))
-				{
-					// use overscan proj matrix for rendering
-					OutPrjMatrix = Contexts[InContextNum].OverscanProjectionMatrix;
-				}
-				else
-				{
-					OutPrjMatrix = Contexts[InContextNum].ProjectionMatrix;
-				}
-
-				return true;
-			}
-			else if (ProjectionPolicy.IsValid() && ProjectionPolicy->GetProjectionMatrix(this, InContextNum, OutPrjMatrix))
-			{
-				// Save the calculated values and update the state of the context
-				Contexts[InContextNum].ProjectionMatrix = OutPrjMatrix;
-				EnumAddFlags(Contexts[InContextNum].ContextState, EDisplayClusterViewportContextState::HasCalculatedProjectionMatrix);
-
-				if (OverscanRuntimeSettings.bIsEnabled)
-				{
-					// use overscan proj matrix for rendering
-					OutPrjMatrix = Contexts[InContextNum].OverscanProjectionMatrix;
-					EnumAddFlags(Contexts[InContextNum].ContextState, EDisplayClusterViewportContextState::HasCalculatedOverscanProjectionMatrix);
-				}
-
-				return true;
-			}
-			else
-			{
-				// ProjectionPolicy->GetProjectionMatrix() returns false, this projection matrix is invalid
-				EnumAddFlags(Contexts[InContextNum].ContextState, EDisplayClusterViewportContextState::InvalidProjectionMatrix);
-			}
+			// Use calculated values
+			OutPrjMatrix = Contexts[InContextNum].ProjectionData.bUseOverscan
+				? Contexts[InContextNum].OverscanProjectionMatrix
+				: Contexts[InContextNum].ProjectionMatrix;
+
+			return true;
 		}
 	}
 
@@ -262,10 +334,15 @@ void FDisplayClusterViewport::CalculateProjectionMatrix(const uint32 InContextNu
 
 	// Update cached projection data:
 	FDisplayClusterViewport_Context::FCachedProjectionData& CachedProjectionData = Contexts[InContextNum].ProjectionData;
+	CachedProjectionData.bValid = true;
+
+	// Store projection angles
 	CachedProjectionData.ProjectionAngles = FVector4(l, r, t, b);
+
+	// Update clipping planes
 	CachedProjectionData.ZNear = n;
 	CachedProjectionData.ZFar = f;
-	CachedProjectionData.bValid = true;
+
 
 	if (FDisplayClusterViewport_OverscanRuntimeSettings::UpdateProjectionAngles(OverscanRuntimeSettings, Contexts[InContextNum].RenderTargetRect.Size(), l, r, t, b))
 	{
diff --git a/Engine/Plugins/Runtime/nDisplay/Source/DisplayCluster/Private/Render/Viewport/DisplayClusterViewport_ViewPoint.cpp b/Engine/Plugins/Runtime/nDisplay/Source/DisplayCluster/Private/Render/Viewport/DisplayClusterViewport_ViewPoint.cpp
--- a/Engine/Plugins/Runtime/nDisplay/Source/DisplayCluster/Private/Render/Viewport/DisplayClusterViewport_ViewPoint.cpp
+++ b/Engine/Plugins/Runtime/nDisplay/Source/DisplayCluster/Private/Render/Viewport/DisplayClusterViewport_ViewPoint.cpp
@@ -27,7 +27,6 @@ namespace UE::DisplayCluster::Viewport::ViewPoint
 		COUNT
 	};
 };
-using namespace UE::DisplayCluster::Viewport::ViewPoint;
 
 ///////////////////////////////////////////////////////////////////////////////////////
 // IDisplayClusterViewport
@@ -209,8 +208,59 @@ bool FDisplayClusterViewport::SetupViewPoint(const uint32 InContextNum, FMinimal
 	return false;
 }
 
+bool FDisplayClusterViewport::PreCalculateViewData(const uint32 InContextNum, const FVector& InViewLocation, const FRotator& InViewRotation, const float WorldToMeters)
+{
+	using namespace UE::DisplayCluster::Viewport::ViewPoint;
+
+	check(Contexts.IsValidIndex(InContextNum));
+
+	FDisplayClusterViewport_Context& OutContext = Contexts[InContextNum];
+
+	// Mark the data as computed to make sure we only compute once per frame.
+	OutContext.bValidData = false;
+	OutContext.bCalculated = true;
+
+	// RootActor2World transform
+	if (const ADisplayClusterRootActor* SceneRootActor = GetConfiguration().GetRootActor(EDisplayClusterRootActorType::Scene))
+	{
+		OutContext.RootActorToWorld = SceneRootActor->GetTransform();
+	}
+	else
+	{
+		// Exit: missed required actor.
+		return false;
+	}
+
+	// Origin2World transform
+	if (USceneComponent* const OriginComponent = ProjectionPolicy->GetOriginComponent())
+	{
+		// Get Origin component transform
+		OutContext.OriginToWorld = OriginComponent->GetComponentTransform();
+	}
+
+	// Geometry units
+	OutContext.WorldToMeters = WorldToMeters;
+	OutContext.GeometryToMeters = ProjectionPolicy->GetGeometryToMeters(this, WorldToMeters);
+
+	// Clipping planes
+	const FVector2D ClipingPlanes = GetClippingPlanes();
+	OutContext.ProjectionData.ZNear = ClipingPlanes.X;
+	OutContext.ProjectionData.ZFar  = ClipingPlanes.Y;
+
+	// Scale units from Geometry to UE World
+	const float ScaleGeometryToWorld = OutContext.WorldToMeters / OutContext.GeometryToMeters;
+
+	// Also add origin location offset
+	const FVector OriginLocationOffsetInGeometryUnits = ProjectionPolicy->GetOriginLocationOffsetInGeometryUnits(this);
+	OutContext.OriginToWorld.AddToTranslation(OriginLocationOffsetInGeometryUnits * ScaleGeometryToWorld);
+
+	return true;
+}
+
 float FDisplayClusterViewport::GetStereoEyeOffsetDistance(const uint32 InContextNum)
 {
+	using namespace UE::DisplayCluster::Viewport::ViewPoint;
+
 	float StereoEyeOffsetDistance = 0.f;
 
 	if (UDisplayClusterCameraComponent* ConfigurationCameraComponent = GetViewPointCameraComponent(EDisplayClusterRootActorType::Configuration))
@@ -256,4 +306,4 @@ float FDisplayClusterViewport::GetStereoEyeOffsetDistance(const uint32 InContext
 	}
 
 	return StereoEyeOffsetDistance;
-}
+}
\ No newline at end of file
diff --git a/Engine/Plugins/Runtime/nDisplay/Source/DisplayCluster/Private/Render/Viewport/Preview/DisplayClusterViewportManagerPreview_Frustum.cpp b/Engine/Plugins/Runtime/nDisplay/Source/DisplayCluster/Private/Render/Viewport/Preview/DisplayClusterViewportManagerPreview_Frustum.cpp
--- a/Engine/Plugins/Runtime/nDisplay/Source/DisplayCluster/Private/Render/Viewport/Preview/DisplayClusterViewportManagerPreview_Frustum.cpp
+++ b/Engine/Plugins/Runtime/nDisplay/Source/DisplayCluster/Private/Render/Viewport/Preview/DisplayClusterViewportManagerPreview_Frustum.cpp
@@ -98,7 +98,7 @@ void FDisplayClusterViewportManagerPreview::RenderFrustumPreviewForViewport(IDis
 
 	// Preview rendered only in mono
 	const TArray<FDisplayClusterViewport_Context>& Contexts = InViewport->GetContexts();
-	if (!(Contexts.Num() == 1 && EnumHasAllFlags(Contexts[0].ContextState, EDisplayClusterViewportContextState::HasCalculatedProjectionMatrix | EDisplayClusterViewportContextState::HasCalculatedViewPoint)))
+	if (!(Contexts.Num() == 1 && Contexts[0].bValidData))
 	{
 		return;
 	}
diff --git a/Engine/Plugins/Runtime/nDisplay/Source/DisplayCluster/Public/Components/DisplayClusterCameraComponent.h b/Engine/Plugins/Runtime/nDisplay/Source/DisplayCluster/Public/Components/DisplayClusterCameraComponent.h
--- a/Engine/Plugins/Runtime/nDisplay/Source/DisplayCluster/Public/Components/DisplayClusterCameraComponent.h
+++ b/Engine/Plugins/Runtime/nDisplay/Source/DisplayCluster/Public/Components/DisplayClusterCameraComponent.h
@@ -80,7 +80,7 @@ public:
 	 * @param InOutViewInfo - ViewInfo data
 	 * @param OutCustomNearClippingPlane - Custom NCP, or a value less than zero if not defined.
 	 */
-	virtual void GetDesiredView(IDisplayClusterViewportConfiguration& InViewportConfiguration, FMinimalViewInfo& InOutViewInfo, float* OutCustomNearClippingPlane = nullptr);
+	virtual void GetDesiredView(const IDisplayClusterViewportConfiguration& InViewportConfiguration, FMinimalViewInfo& InOutViewInfo, float* OutCustomNearClippingPlane = nullptr) const;
 
 	/** Returns the position of the observer's eyes in the Stage. */
 	virtual void GetEyePosition(const IDisplayClusterViewportConfiguration& InViewportConfiguration, FVector& OutViewLocation, FRotator& OutViewRotation);
@@ -142,6 +142,12 @@ public:
 	/** Returns true if the given viewpoint position should be used as a camera position. */
 	virtual bool IsViewPointOverrideCameraPosition() const;
 
+	/** Returns true if the eye position (ViewPoint) inside the DCRa can be overriden. */
+	virtual bool CanOverrideEyePosition() const;
+
+	/** Returns true if ViewPoint should use the camera position instead of the DCRA position when transforming. */
+	virtual bool ShouldFollowCameraLocation() const;
+
 protected:
 	/** Get view from the referenced camera. */
 	virtual bool GetTargetCameraDesiredViewInternal(const IDisplayClusterViewportConfiguration& InViewportConfiguration, FMinimalViewInfo& InOutViewInfo, float* OutCustomNearClippingPlane = nullptr) const;
@@ -309,6 +315,10 @@ public:
 	UPROPERTY(EditAnywhere, BlueprintReadWrite, Category = "Camera Settings", meta = (DisplayName = "Follow ICVFX Camera", EditCondition = "TargetCameraType==EDisplayClusterTargetCameraType::ICVFXCameraComponent", EditConditionHides))
 	uint8 bUseICVFXCameraComponentTracking: 1 = 1;
 
+	/** Use camera position for rendering instead of RootActor. */
+	UPROPERTY(EditAnywhere, BlueprintReadWrite, Category = "Camera Settings", meta = (DisplayName = "Follow Camera", EditCondition = "TargetCameraType!=EDisplayClusterTargetCameraType::ICVFXCameraComponent && TargetCameraType!=EDisplayClusterTargetCameraType::None", EditConditionHides))
+	uint8 bFollowCameraPosition : 1 = 0;
+
 	/** Use a specific actor camera instead of a game camera. */
 	UPROPERTY(EditAnywhere, BlueprintReadWrite, Category = "Camera Settings", meta = (DisplayName = "Camera Actor", EditCondition = "TargetCameraType==EDisplayClusterTargetCameraType::ExternalCineCameraActor", EditConditionHides))
 	TSoftObjectPtr<ACineCameraActor> ExternalCineCameraActor;
diff --git a/Engine/Plugins/Runtime/nDisplay/Source/DisplayCluster/Public/Render/Projection/IDisplayClusterProjectionPolicy.h b/Engine/Plugins/Runtime/nDisplay/Source/DisplayCluster/Public/Render/Projection/IDisplayClusterProjectionPolicy.h
--- a/Engine/Plugins/Runtime/nDisplay/Source/DisplayCluster/Public/Render/Projection/IDisplayClusterProjectionPolicy.h
+++ b/Engine/Plugins/Runtime/nDisplay/Source/DisplayCluster/Public/Render/Projection/IDisplayClusterProjectionPolicy.h
@@ -14,6 +14,67 @@ class USceneComponent;
 struct FDisplayClusterConfigurationProjection;
 struct FMinimalViewInfo;
 
+
+/**
+ * The type of media usage for the viewport.
+ */
+enum class EDisplayClusterProjectionPolicyFlags : uint8
+{
+	// No flags
+	None = 0,
+
+	/** This projection policy use eye tracking (ViewPoint).
+	* 
+	* This means that the InOutViewLocation and InOutViewRotation arguments in CalculateView()
+	* will contain the position of the viewpoint in world space, not the camera.
+	*
+	* Camera position is the value passed from IStereoRenderingDevice.
+	* This value is usually based on the active engine camera,
+	* which can be overridden in ISceneViewExtension::SetupViewPoint().
+	* (see FDisplayClusterViewportManagerViewPointExtension::SetupViewPoint())
+	* 
+	* Most nDisplay projection policies (simple, mesh, easyblend, etc.) ignore this camera position
+	* and use their own values or DCRA components to compute the view.
+	* In this case, this flag replaces the IDisplayClusterProjectionPolicy::CalculateView() input arguments with the position of the
+	* of the DCRA ViewPoint component instead of the camera position.
+	* And the projection policy may not implement this logic.
+	* (see IDisplayClusterViewport::GetViewPointCameraEye())
+	*
+	* Note:
+	* On the nDisplay side, the camera position can already be overridden using the ViewPoint position for some cases:
+	* UDisplayClusterCameraComponent::GetDesiredView()
+	*  -> GetTargetCameraDesiredViewInternal()
+	*  -> IsViewPointOverrideCameraPosition()
+	*/
+	UseEyeTracking = 1 << 0,
+
+	/** Projection policies can use local space and custom geometric units.
+	* 
+	* Therefore, InOutViewLocation and InOutViewRotation must be converted from world space
+	* to the projection policy space and back again.
+	* 
+	* The projection policy space is formed from the Origin component and scaling units.
+	* see GetOriginComponent(), GetGeometryToMeters(), GetOriginLocationOffsetInGeometryUnits().
+	*
+	* This conversion will look like this:
+	*   World -> ProjectionPolicy
+	*     ScaleWorldToGeometry = GeometryToMeters / WorldToMeters
+	*     PrjPolicyLocation = OriginToWorld.InverseTransformPosition(WorldSpaceLocation) * ScaleWorldToGeometry
+	*     PrjPolicyRotation = OriginToWorld.InverseTransformRotation(WorldSpaceRotation)
+	* and vice versa.
+	*/
+	UseLocalSpace = 1 << 1,
+
+	/** This projection policy supports "Follow Camera" feature.
+	* To support this feature, the projection policy must return an output ViewLocation in the DCRA zone.
+	* 
+	* Usually these values are calculated based on objects such as these:
+	* ViewPoint component, ICVFXCamera component, CineCamera actor referenced by ICVFXCamera, etc.
+	*/
+	EnableFollowCameraFeature = 1 << 2,
+};
+ENUM_CLASS_FLAGS(EDisplayClusterProjectionPolicyFlags);
+
 /**
  * nDisplay projection policy
  */
@@ -33,15 +94,6 @@ public:
 	*/
 	virtual const FString& GetType() const = 0;
 
-	/**
-	* Return projection policy type
-	*/
-	UE_DEPRECATED(5.1, "This function has been deprecated. Please use 'GetType'.")
-	virtual const FString GetTypeId() const
-	{
-		return GetType();
-	}
-
 	/**
 	* Return projection policy configuration
 	*/
@@ -57,6 +109,35 @@ public:
 		return nullptr;
 	}
 
+	/**
+	* If the projection policy uses a custom scale internally, this function returns that value.
+	* The value means how many geometry units are in one meter.
+	* 
+	* @param InViewport      - (in) a owner viewport
+	* @param InWorldToMeters - (in) UE world units per meter
+	*/
+	virtual float GetGeometryToMeters(const IDisplayClusterViewport* InViewport, const float InWorldToMeters) const
+	{
+		// UE world units are used by default.
+		return InWorldToMeters;
+	}
+
+	/** Returns Origin's location in geometry units.
+	* 
+	* The Origin component must be configured in UE units, but the calibration tools use their own units.
+	* Therefore, the user must perform calculations to correctly set the location of the Origin component within the UE.
+	* Typically, this approach leads to mistakes in the DCRA setup.
+	* To avoid this, the projection policy can receive the Origin Location parameter in geometric units.
+	* This function returns the value of this parameter.
+	* 
+	* @param InViewport - (in) a owner viewport
+	*/
+	virtual FVector GetOriginLocationOffsetInGeometryUnits(const IDisplayClusterViewport* InViewport) const
+	{
+		// A zero vector means it is not used.
+		return FVector::ZeroVector;
+	}
+
 	/**
 	* Send projection policy game thread data to render thread proxy
 	* called once per frame from FDisplayClusterViewportManager::FinalizeNewFrame
@@ -196,6 +277,27 @@ public:
 	virtual void UpdatePostProcessSettings(IDisplayClusterViewport* InViewport)
 	{ }
 
+	/**
+	* Projection policy can modify the clipping planes.
+	* 
+	* @param InViewport - (in) a owner viewport
+	* @param InOutNCP   - (in, out) Near clipping plane
+	* @param InOutFCP   - (in, out) Far clipping plane
+	*/
+	virtual void ApplyClippingPlanesOverrides(const IDisplayClusterViewport* InViewport, float& InOutNCP, float& InOutFCP) const
+	{ }
+
+	/**
+	* Return Projection policy flags
+	*
+	* @param InViewport   - (in) a owner viewport
+	* @param InContextNum - Index of view that is being processed for this viewport
+	*/
+	virtual EDisplayClusterProjectionPolicyFlags GetProjectionPolicyFlags(const IDisplayClusterViewport* InViewport, const uint32 InContextNum) const
+	{
+		return EDisplayClusterProjectionPolicyFlags::None;
+	}
+
 	/** Calculate view projection data
 	* 
 	* @param InViewport        - a owner viewport
diff --git a/Engine/Plugins/Runtime/nDisplay/Source/DisplayCluster/Public/Render/Viewport/Containers/DisplayClusterViewport_Context.h b/Engine/Plugins/Runtime/nDisplay/Source/DisplayCluster/Public/Render/Viewport/Containers/DisplayClusterViewport_Context.h
--- a/Engine/Plugins/Runtime/nDisplay/Source/DisplayCluster/Public/Render/Viewport/Containers/DisplayClusterViewport_Context.h
+++ b/Engine/Plugins/Runtime/nDisplay/Source/DisplayCluster/Public/Render/Viewport/Containers/DisplayClusterViewport_Context.h
@@ -20,6 +20,13 @@ public:
 		, StereoViewIndex(InStereoViewIndex)
 	{ }
 
+public:
+	// True if this ViewData have a valid data.
+	bool bValidData : 1 = 0;
+
+	// Is this data have been calculated.
+	bool bCalculated : 1 = 0;
+
 public:
 	// Index of context (eye #)
 	const uint32            ContextNum;
@@ -28,18 +35,46 @@ public:
 	EStereoscopicPass StereoscopicPass;
 	int32             StereoViewIndex;
 
-	// Special flags for this context
-	EDisplayClusterViewportContextState ContextState = EDisplayClusterViewportContextState::None;
-
-	// ViewPoint: Camera location and orientation
-	FVector  ViewLocation = FVector::ZeroVector;
-	FRotator ViewRotation = FRotator::ZeroRotator;
-
-	// Projection Matrix
-	FMatrix ProjectionMatrix = FMatrix::Identity;
+	struct FCachedViewData
+	{
+		/**
+		* Here's the order of data transformation:
+		*
+		* IStereoRenderingDevice
+		* -> {CameraOrigin,CameraRotation}
+		*    IDisplayClusterViewport::CalculateViewData()
+		*      []->{EyeLocation, EyeRotation}
+		* 		   IDisplayClusterProjectionPolicy::CalculateView()
+		*      []-> {ViewLocation, ViewRotation}
+		* <- {ViewLocation, ViewRotation}
+		* IStereoRenderingDevice
+		*/
+
+		// Camera location in the world space
+		// This value receives nDisplay from IStereoRenderingDevice
+		FVector CameraOrigin = FVector::ZeroVector;
+
+		// Camera rotation in the world space
+		// This value receives nDisplay from IStereoRenderingDevice
+		FRotator CameraRotation = FRotator::ZeroRotator;
+
+		// The eye location  in the world space
+		// This value is used as an argument to IDisplayClusterProjectionPolicy::CalculateView()
+		FVector EyeOrigin = FVector::ZeroVector;
+
+		// The eye rotation in the world space
+		// This value is used as an argument to IDisplayClusterProjectionPolicy::CalculateView()
+		FRotator EyeRotation = FRotator::ZeroRotator;
+
+		// Render view location used in world space
+		FVector  RenderLocation = FVector::ZeroVector;
+
+		// Render view rotation used in world space
+		FRotator RenderRotation = FRotator::ZeroRotator;
+	};
 
-	// Overscan Projection Matrix (internal use)
-	FMatrix OverscanProjectionMatrix = FMatrix::Identity;
+	// View data for this context
+	FCachedViewData ViewData;
 
 	/** Cached projection data */
 	struct FCachedProjectionData
@@ -55,15 +90,40 @@ public:
 		// Projection angles for Overscan [Left, Right, Top, Bottom]
 		FVector4 OverscanProjectionAngles;
 
-		// Projection planes
+		// Clipping planes
 		double ZNear = 0.f;
 		double ZFar = 0.f;
 	};
-	
+
 	// Cached projection values
 	// This values updated from function FDisplayClusterViewport::CalculateProjectionMatrix()
 	FCachedProjectionData ProjectionData;
 
+	// UE scale: how many UE Units in one meter
+	float WorldToMeters = 100.f;
+
+	// Geometry scale: how many Geometry Units in one meter
+	float GeometryToMeters = 100.f;
+
+	// Origin component to world transform
+	// This component is defined in the projection policy
+	FTransform OriginToWorld = FTransform::Identity;
+
+	// RootActor to world transform
+	FTransform RootActorToWorld = FTransform::Identity;
+
+	// View location used in world space
+	FVector  ViewLocation = FVector::ZeroVector;
+
+	// View rotation used in world space
+	FRotator ViewRotation = FRotator::ZeroRotator;
+
+	// Projection Matrix
+	FMatrix ProjectionMatrix = FMatrix::Identity;
+
+	// Overscan Projection Matrix (internal use)
+	FMatrix OverscanProjectionMatrix = FMatrix::Identity;
+	
 	/** Additional data for the Depth of Field (DoF). */
 	struct FDepthOfFieldSettings
 	{
@@ -75,15 +135,12 @@ public:
 
 	} DepthOfField;
 
-	// World scale
-	float WorldToMeters = 100.f;
-
 	// GPU index for this context render target
 	int32 GPUIndex = INDEX_NONE;
 
 	// Enables nDisplay's native implementation of cross-GPU transfer.
 	// This disables cross-GPU transfer by default for nDisplay viewports in FSceneViewFamily structure.
-	bool bOverrideCrossGPUTransfer = false;
+	bool bOverrideCrossGPUTransfer : 1 = 0;
 
 	// Location and size on a render target texture
 	FIntRect RenderTargetRect;
@@ -104,7 +161,7 @@ public:
 	int32 NumMips = 1;
 
 	// Disable render for this viewport (Overlay)
-	bool bDisableRender = false;
+	bool bDisableRender : 1 = 0;
 
 	/**
 	* Viewport context data for rendering thread
diff --git a/Engine/Plugins/Runtime/nDisplay/Source/DisplayCluster/Public/Render/Viewport/Containers/DisplayClusterViewport_Enums.h b/Engine/Plugins/Runtime/nDisplay/Source/DisplayCluster/Public/Render/Viewport/Containers/DisplayClusterViewport_Enums.h
--- a/Engine/Plugins/Runtime/nDisplay/Source/DisplayCluster/Public/Render/Viewport/Containers/DisplayClusterViewport_Enums.h
+++ b/Engine/Plugins/Runtime/nDisplay/Source/DisplayCluster/Public/Render/Viewport/Containers/DisplayClusterViewport_Enums.h
@@ -142,33 +142,6 @@ enum class EDisplayClusterViewport_FrustumUnit: uint8
 	Percent
 };
 
-/**
- * Special flags for nDisplay Viewport context
- */
-enum class EDisplayClusterViewportContextState : uint8
-{
-	None = 0,
-
-	// The FDisplayClusterViewport::CalculateView() function can be called several times per frame.
-	// Each time it must return the same values. For optimization purposes, after the first call this function
-	// stores the result in the context variables 'ViewLocation' and 'ViewRotation'.
-	// Finally, raises this flag for subsequent calls in the current frame.
-	HasCalculatedViewPoint = 1 << 0,
-
-	// Viewpoint is not valid for this viewport (cannot be calculated)
-	InvalidViewPoint = 1 << 1,
-
-	// The FDisplayClusterViewport::GetProjectionMatrix() function can also be called several times per frame.
-	// stores the result in the context variables 'ProjectionMatrix' and 'OverscanProjectionMatrix'.
-	// Finally, raises this flag for subsequent calls in the current frame.
-	HasCalculatedProjectionMatrix = 1 << 2,
-	HasCalculatedOverscanProjectionMatrix = 1 << 3,
-
-	// The projection matrix is not valid (cannot be calculated)
-	InvalidProjectionMatrix = 1 << 4,
-};
-ENUM_CLASS_FLAGS(EDisplayClusterViewportContextState);
-
 /**
 * Type of DCRA used by configuration
 */
diff --git a/Engine/Plugins/Runtime/nDisplay/Source/DisplayCluster/Public/Render/Viewport/IDisplayClusterViewport.h b/Engine/Plugins/Runtime/nDisplay/Source/DisplayCluster/Public/Render/Viewport/IDisplayClusterViewport.h
--- a/Engine/Plugins/Runtime/nDisplay/Source/DisplayCluster/Public/Render/Viewport/IDisplayClusterViewport.h
+++ b/Engine/Plugins/Runtime/nDisplay/Source/DisplayCluster/Public/Render/Viewport/IDisplayClusterViewport.h
@@ -104,6 +104,21 @@ public:
 	 */
 	virtual bool GetProjectionMatrix(const uint32 InContextNum, FMatrix& OutPrjMatrix) = 0;
 
+	/** Calculate View data
+	 * Updates data within the FDisplayClusterViewport_Context.
+	 * This function calls IDisplayClusterProjectionPolicy::CalculateView() and IDisplayClusterProjectionPolicy::GetProjectionMatrix().
+	 * 
+	 * Note: Called once per frame from IDisplayClusterViewport::CalculateView().
+	 *
+	 * @param InContextNum      - (in) viewport eye context index
+	 * @param InOutViewLocation - (in) View location
+	 * @param InOutViewRotation - (in) View rotator
+	 * @param WorldToMeters     - UE world scale
+	 *
+	 * @return true on success
+	*/
+	virtual bool CalculateViewData(const uint32 InContextNum, const FVector& InViewLocation, const FRotator& InViewRotation, const float WorldToMeters) = 0;
+
 	/** Setup viewpoint for this viewport
 	 *
 	 * @param InOutViewInfo - [in\out] viewinfo
