nDisplay: Reworked projection policies for CL 43797173

- Camera: fixed bug: FReversedZPerspectiveMatrix can't be used when NCP==FCP.
added FollowCamera feature for the ICVFX cameras.

- Dome: added new parameter GeometryScale instead of hardcoded one.
added GeometryUnits, FollowCamera, UseEyeTracking and UseLocalSpace.

- EasyBlend: added new parameter Origin Location with values in geometry units
added GeometryUnits, FollowCamera, UseEyeTracking and UseLocalSpace.

- MPCDI: added FollowCamera, UseEyeTracking and UseLocalSpace.

- Reference: added new functions that call the same functions from the parent viewport.

- Simple: added FollowCamera and UseEyeTracking.

- VIOSO: added GeometryUnits, FollowCamera, UseEyeTracking and UseLocalSpace.

diff --git a/Engine/Plugins/Runtime/nDisplay/Source/DisplayClusterConfigurator/Private/Views/Details/Policies/DisplayClusterConfiguratorPolicyDetailCustomization.cpp b/Engine/Plugins/Runtime/nDisplay/Source/DisplayClusterConfigurator/Private/Views/Details/Policies/DisplayClusterConfiguratorPolicyDetailCustomization.cpp
--- a/Engine/Plugins/Runtime/nDisplay/Source/DisplayClusterConfigurator/Private/Views/Details/Policies/DisplayClusterConfiguratorPolicyDetailCustomization.cpp
+++ b/Engine/Plugins/Runtime/nDisplay/Source/DisplayClusterConfigurator/Private/Views/Details/Policies/DisplayClusterConfiguratorPolicyDetailCustomization.cpp
@@ -605,12 +605,18 @@ void FDisplayClusterConfiguratorProjectionCustomization::CreateEasyBlendPolicy(U
 		TArray<FString>{"pol*", "ol*"}));
 
 	CustomPolicyParameters.Add(MakeShared<FPolicyParameterInfoComponentCombo>(
-		"Origin",
+		"Origin Component",
 		DisplayClusterProjectionStrings::cfg::easyblend::Origin,
 		Blueprint,
 		ConfigurationViewports,
 		TArray<TSubclassOf<UActorComponent>>{ USceneComponent::StaticClass() }));
 
+	CustomPolicyParameters.Add(MakeShared<FPolicyParameterInfoVector>(
+		"Origin Location",
+		DisplayClusterProjectionStrings::cfg::easyblend::OriginOffset,
+		Blueprint,
+		ConfigurationViewports));
+
 	CustomPolicyParameters.Add(MakeShared<FPolicyParameterInfoNumber<float>>(
 		"Scale",
 		DisplayClusterProjectionStrings::cfg::easyblend::Scale,
diff --git a/Engine/Plugins/Runtime/nDisplay/Source/DisplayClusterConfigurator/Private/Views/Details/Policies/DisplayClusterConfiguratorPolicyParameterCustomization.cpp b/Engine/Plugins/Runtime/nDisplay/Source/DisplayClusterConfigurator/Private/Views/Details/Policies/DisplayClusterConfiguratorPolicyParameterCustomization.cpp
--- a/Engine/Plugins/Runtime/nDisplay/Source/DisplayClusterConfigurator/Private/Views/Details/Policies/DisplayClusterConfiguratorPolicyParameterCustomization.cpp
+++ b/Engine/Plugins/Runtime/nDisplay/Source/DisplayClusterConfigurator/Private/Views/Details/Policies/DisplayClusterConfiguratorPolicyParameterCustomization.cpp
@@ -792,7 +792,7 @@ FPolicyParameterInfo4x4Matrix::FPolicyParameterInfo4x4Matrix(const FString& InDi
 			check(ConfigurationViewport != nullptr);
 
 			FString* ParameterValue = ConfigurationViewport->ProjectionPolicy.Parameters.Find(GetParameterKey());
-			check(ParameterValue)
+			check(ParameterValue);
 
 			const FText TextValue = FText::FromString(*ParameterValue);
 			const FMatrix Matrix = DisplayClusterTypesConverter::template FromString<FMatrix>(TextValue.ToString());
@@ -950,7 +950,7 @@ FPolicyParameterInfoRotator::FPolicyParameterInfoRotator(const FString& InDispla
 			check(ConfigurationViewport != nullptr);
 
 			FString* ParameterValue = ConfigurationViewport->ProjectionPolicy.Parameters.Find(GetParameterKey());
-			check(ParameterValue)
+			check(ParameterValue);
 
 			const FText TextValue = FText::FromString(*ParameterValue);
 			const FRotator Rotation = DisplayClusterTypesConverter::template FromString<FRotator>(TextValue.ToString());
@@ -979,7 +979,7 @@ void FPolicyParameterInfoRotator::CreateCustomRowWidget(IDetailChildrenBuilder&
 	[
 		SNew(STextBlock)
 		.Font(IDetailLayoutBuilder::GetDetailFont())
-		.Text(LOCTEXT("RotationLabel", "Rotation"))
+		.Text(GetParameterDisplayName())
 		.ToolTipText(this, &FPolicyParameterInfo::GetParameterTooltip)
 	]
 	.ValueContent()
@@ -1019,6 +1019,95 @@ void FPolicyParameterInfoRotator::FormatTextAndUpdateParameter()
 	UpdateCustomParameterValueText(RotationString);
 }
 
+
+const FString FPolicyParameterInfoVector::BaseVectorString = DisplayClusterTypesConverter::template ToString<FVector>(FVector::ZeroVector);
+
+FPolicyParameterInfoVector::FPolicyParameterInfoVector(const FString& InDisplayName, const FString& InKey,
+	UDisplayClusterBlueprint* InBlueprint,
+	const TArray<TWeakObjectPtr<UDisplayClusterConfigurationViewport>>& InConfigurationViewports) :
+	FPolicyParameterInfoFloatReference(InDisplayName, InKey, InBlueprint, InConfigurationViewports, &BaseVectorString),
+	CachedVectorX(MakeShared<TOptional<float>>()),
+	CachedVectorY(MakeShared<TOptional<float>>()),
+	CachedVectorZ(MakeShared<TOptional<float>>())
+{
+
+	if (!DoParametersMatchForAllViewports())
+	{
+		for (const TWeakObjectPtr<UDisplayClusterConfigurationViewport>& Viewport : ConfigurationViewports)
+		{
+			UDisplayClusterConfigurationViewport* ConfigurationViewport = Viewport.Get();
+			check(ConfigurationViewport != nullptr);
+
+			FString* ParameterValue = ConfigurationViewport->ProjectionPolicy.Parameters.Find(GetParameterKey());
+			check(ParameterValue);
+
+			const FText TextValue = FText::FromString(*ParameterValue);
+			const FVector Location = DisplayClusterTypesConverter::template FromString<FVector>(TextValue.ToString());
+
+			ResetOrSetCachedValue(CachedVectorX, Location.X);
+			ResetOrSetCachedValue(CachedVectorY, Location.Y);
+			ResetOrSetCachedValue(CachedVectorZ, Location.Z);
+		}
+	}
+	else
+	{
+		const FText TextValue = GetOrAddCustomParameterValueText();
+		const FVector Location = DisplayClusterTypesConverter::template FromString<FVector>(TextValue.ToString());
+
+		*CachedVectorX = Location.X;
+		*CachedVectorY = Location.Y;
+		*CachedVectorZ = Location.Z;
+	}
+}
+
+void FPolicyParameterInfoVector::CreateCustomRowWidget(IDetailChildrenBuilder& InDetailWidgetRow)
+{
+	InDetailWidgetRow.AddCustomRow(GetParameterDisplayName())
+		.Visibility(MakeAttributeRaw(this, &FPolicyParameterInfoVector::IsParameterVisible))
+		.NameContent()
+		[
+			SNew(STextBlock)
+				.Font(IDetailLayoutBuilder::GetDetailFont())
+				.Text(GetParameterDisplayName())
+				.ToolTipText(this, &FPolicyParameterInfo::GetParameterTooltip)
+		]
+		.ValueContent()
+		.MinDesiredWidth(375.0f)
+		.MaxDesiredWidth(375.0f)
+		[
+			SNew(SHorizontalBox)
+				+ SHorizontalBox::Slot()
+				.Padding(FMargin(0.0f, 2.0f, 3.0f, 2.0f))
+				[
+					MakeFloatInputWidget(CachedVectorX, LOCTEXT("LocationX", "X"), true, FLinearColor::White, SNumericEntryBox<float>::RedLabelBackgroundColor)
+				]
+			+ SHorizontalBox::Slot()
+				.Padding(FMargin(0.0f, 2.0f, 3.0f, 2.0f))
+				[
+					MakeFloatInputWidget(CachedVectorY, LOCTEXT("LocationY", "Y"), true, FLinearColor::White, SNumericEntryBox<float>::GreenLabelBackgroundColor)
+				]
+			+ SHorizontalBox::Slot()
+				.Padding(FMargin(0.0f, 2.0f, 3.0f, 2.0f))
+				[
+					MakeFloatInputWidget(CachedVectorZ, LOCTEXT("LocationZ", "Z"), true, FLinearColor::White, SNumericEntryBox<float>::BlueLabelBackgroundColor)
+				]
+		];
+}
+
+void FPolicyParameterInfoVector::FormatTextAndUpdateParameter()
+{
+	const bool bAllValuesSet = CachedVectorX->IsSet() && CachedVectorY->IsSet() && CachedVectorZ->IsSet();
+	if (!bAllValuesSet)
+	{
+		return;
+	}
+
+	const FVector Location(*CachedVectorX.Get(), *CachedVectorY.Get(), *CachedVectorZ.Get());
+
+	const FString LocationString = DisplayClusterTypesConverter::template ToString(Location);
+	UpdateCustomParameterValueText(LocationString);
+}
+
 const FString FPolicyParameterInfoFrustumAngle::BaseFrustumPlanesString = FString("l=-30, r=30, t=30, b=-30");
 
 FPolicyParameterInfoFrustumAngle::FPolicyParameterInfoFrustumAngle(const FString& InDisplayName, const FString& InKey,
diff --git a/Engine/Plugins/Runtime/nDisplay/Source/DisplayClusterConfigurator/Private/Views/Details/Policies/DisplayClusterConfiguratorPolicyParameterCustomization.h b/Engine/Plugins/Runtime/nDisplay/Source/DisplayClusterConfigurator/Private/Views/Details/Policies/DisplayClusterConfiguratorPolicyParameterCustomization.h
--- a/Engine/Plugins/Runtime/nDisplay/Source/DisplayClusterConfigurator/Private/Views/Details/Policies/DisplayClusterConfiguratorPolicyParameterCustomization.h
+++ b/Engine/Plugins/Runtime/nDisplay/Source/DisplayClusterConfigurator/Private/Views/Details/Policies/DisplayClusterConfiguratorPolicyParameterCustomization.h
@@ -549,6 +549,33 @@ private:
 	static const FString BaseRotatorString;
 };
 
+/**
+ * Policy info for vector.
+ */
+class FPolicyParameterInfoVector final : public FPolicyParameterInfoFloatReference
+{
+public:
+	FPolicyParameterInfoVector(
+		const FString& InDisplayName,
+		const FString& InKey,
+		UDisplayClusterBlueprint* InBlueprint,
+		const TArray<TWeakObjectPtr<UDisplayClusterConfigurationViewport>>& InConfigurationViewports);
+
+	// FPolicyParameterInfo
+	virtual void CreateCustomRowWidget(IDetailChildrenBuilder& InDetailWidgetRow) override;
+	// ~FPolicyParameterInfo
+
+private:
+	virtual void FormatTextAndUpdateParameter() override;
+
+private:
+	mutable TSharedRef<TOptional<float>> CachedVectorX;
+	mutable TSharedRef<TOptional<float>> CachedVectorY;
+	mutable TSharedRef<TOptional<float>> CachedVectorZ;
+
+	static const FString BaseVectorString;
+};
+
 /**
  * Policy info for frustum angle.
  */
diff --git a/Engine/Plugins/Runtime/nDisplay/Source/DisplayClusterProjection/Private/Policy/Camera/DisplayClusterProjectionCameraPolicy.cpp b/Engine/Plugins/Runtime/nDisplay/Source/DisplayClusterProjection/Private/Policy/Camera/DisplayClusterProjectionCameraPolicy.cpp
--- a/Engine/Plugins/Runtime/nDisplay/Source/DisplayClusterProjection/Private/Policy/Camera/DisplayClusterProjectionCameraPolicy.cpp
+++ b/Engine/Plugins/Runtime/nDisplay/Source/DisplayClusterProjection/Private/Policy/Camera/DisplayClusterProjectionCameraPolicy.cpp
@@ -83,6 +83,17 @@ UCameraComponent* FDisplayClusterProjectionCameraPolicy::GetCameraComponent() co
 	return nullptr;
 }
 
+EDisplayClusterProjectionPolicyFlags FDisplayClusterProjectionCameraPolicy::GetProjectionPolicyFlags(const IDisplayClusterViewport* InViewport, const uint32 InContextNum) const
+{
+	if (InViewport && EnumHasAnyFlags(InViewport->GetRenderSettingsICVFX().RuntimeFlags, EDisplayClusterViewportRuntimeICVFXFlags::InCamera))
+	{
+		// The ICVFX camera is always in the DCRA area and can follow the camera.
+		return EDisplayClusterProjectionPolicyFlags::EnableFollowCameraFeature;
+	}
+
+	return EDisplayClusterProjectionPolicyFlags::None;
+}
+
 bool FDisplayClusterProjectionCameraPolicy::CalculateView(IDisplayClusterViewport* InViewport, const uint32 InContextNum, FVector& InOutViewLocation, FRotator& InOutViewRotation, const FVector& ViewOffset, const float WorldToMeters, const float NCP, const float FCP)
 {
 	check(IsInGameThread());
@@ -111,27 +122,19 @@ bool FDisplayClusterProjectionCameraPolicy::GetProjectionMatrix(IDisplayClusterV
 		UE_LOG(LogDisplayClusterProjectionCamera, Warning, TEXT("CameraFOV clamped: '%f' -> '%f'. (FieldOfView='%f', FOVMultiplier='%f'"), ScaledCameraFOVDegrees, ClampedCameraFOVDegrees, CameraFOVDegrees, CameraSettings.FOVMultiplier);
 	}
 
+	// Support inner camera custom frustum
+	const float HalfFOVH = ZNear * FMath::Tan(FMath::DegreesToRadians(0.5 * ClampedCameraFOVDegrees));
+	const float HalfFOVV = HalfFOVH / CameraAspectRatio;
+
 	if (InViewport)
 	{
-		// Support inner camera custom frustum
-		const float HalfFOVH = ZNear * FMath::Tan(FMath::DegreesToRadians(0.5 * ClampedCameraFOVDegrees));
-		const float HalfFOVV = HalfFOVH / CameraAspectRatio;
-
 		InViewport->CalculateProjectionMatrix(InContextNum, -HalfFOVH, HalfFOVH, HalfFOVV, -HalfFOVV, ZNear, ZFar, false);
 		OutPrjMatrix = InViewport->GetContexts()[InContextNum].ProjectionMatrix;
 	}
 	else
 	{
-		const float HalfFOV = 0.5 * FMath::DegreesToRadians(ClampedCameraFOVDegrees);
-
-		OutPrjMatrix = FReversedZPerspectiveMatrix(
-			HalfFOV,
-			HalfFOV,
-			/* XAxisMultiplier = */ 1.0,
-			/* YAxisMultiplier = */ CameraAspectRatio,
-			GNearClippingPlane,
-			GNearClippingPlane
-		);
+		// FReversedZPerspectiveMatrix() can't be used because it returns the wrong matrix when ZNear==ZFar (nDisplay encoding for infinite far clipping plane).
+		OutPrjMatrix = IDisplayClusterViewport::MakeProjectionMatrix(-HalfFOVH, HalfFOVH, HalfFOVV, -HalfFOVV, ZNear, ZFar);
 	}
 
 	if (!CameraSettings.OffCenterProjectionOffset.IsZero())
diff --git a/Engine/Plugins/Runtime/nDisplay/Source/DisplayClusterProjection/Private/Policy/Camera/DisplayClusterProjectionCameraPolicy.h b/Engine/Plugins/Runtime/nDisplay/Source/DisplayClusterProjection/Private/Policy/Camera/DisplayClusterProjectionCameraPolicy.h
--- a/Engine/Plugins/Runtime/nDisplay/Source/DisplayClusterProjection/Private/Policy/Camera/DisplayClusterProjectionCameraPolicy.h
+++ b/Engine/Plugins/Runtime/nDisplay/Source/DisplayClusterProjection/Private/Policy/Camera/DisplayClusterProjectionCameraPolicy.h
@@ -42,6 +42,9 @@ public:
 		// Disable the use of PP settings from the ViewPoint component.
 		return false;
 	}
+
+	virtual EDisplayClusterProjectionPolicyFlags GetProjectionPolicyFlags(const IDisplayClusterViewport* InViewport, const uint32 InContextNum) const override;
+
 	//~ End IDisplayClusterProjectionPolicy
 
 public:
diff --git a/Engine/Plugins/Runtime/nDisplay/Source/DisplayClusterProjection/Private/Policy/Domeprojection/DisplayClusterProjectionDomeprojectionPolicyBase.cpp b/Engine/Plugins/Runtime/nDisplay/Source/DisplayClusterProjection/Private/Policy/Domeprojection/DisplayClusterProjectionDomeprojectionPolicyBase.cpp
--- a/Engine/Plugins/Runtime/nDisplay/Source/DisplayClusterProjection/Private/Policy/Domeprojection/DisplayClusterProjectionDomeprojectionPolicyBase.cpp
+++ b/Engine/Plugins/Runtime/nDisplay/Source/DisplayClusterProjection/Private/Policy/Domeprojection/DisplayClusterProjectionDomeprojectionPolicyBase.cpp
@@ -45,7 +45,7 @@ bool FDisplayClusterProjectionDomeprojectionPolicyBase::HandleStartScene(IDispla
 
 	// Read domeprojection config data from nDisplay config file
 	FString File;
-	if (!ReadConfigData(InViewport, File, OriginCompId, DomeprojectionChannel))
+	if (!ReadConfigData(InViewport, File, OriginCompId, DomeprojectionChannel, GeometryScale))
 	{
 		if (!IsEditorOperationMode(InViewport))
 		{
@@ -100,32 +100,30 @@ void FDisplayClusterProjectionDomeprojectionPolicyBase::HandleEndScene(IDisplayC
 	ReleaseOriginComponent();
 }
 
+void FDisplayClusterProjectionDomeprojectionPolicyBase::ApplyClippingPlanesOverrides(const IDisplayClusterViewport* InViewport, float& InOutNCP, float& InOutFCP) const
+{
+	// Disable NearClippingPlane for DomeProjection.
+	if (!GDisplayClusterProjectionDomeprojectionPolicyEnableNearClippingPlane)
+	{
+		InOutNCP = 1;
+		
+		// Infinite FCP
+		InOutFCP = InOutNCP;
+	}
+}
 
 bool FDisplayClusterProjectionDomeprojectionPolicyBase::CalculateView(IDisplayClusterViewport* InViewport, const uint32 InContextNum, FVector& InOutViewLocation, FRotator& InOutViewRotation, const FVector& ViewOffset, const float WorldToMeters, const float InNCP, const float InFCP)
 {
 	check(IsInGameThread());
 
-	if (!ViewAdapter.IsValid())
+	if (!InViewport || !InViewport->GetContexts().IsValidIndex(InContextNum) || !ViewAdapter.IsValid())
 	{
 		return false;
 	}
 
-	const float NCP = GDisplayClusterProjectionDomeprojectionPolicyEnableNearClippingPlane ? InNCP : 1;
-	const float FCP = GDisplayClusterProjectionDomeprojectionPolicyEnableNearClippingPlane ? InFCP : 1;
-
-	// Get origin component
-	const USceneComponent* const OriginComp = GetOriginComponent();
-	check(OriginComp);
-
-	// Get world-origin matrix
-	const FTransform& World2LocalTransform = (OriginComp != nullptr ? OriginComp->GetComponentTransform() : FTransform::Identity);
-
-	// Calculate view location in origin space
-	FVector OriginSpaceViewLocation = World2LocalTransform.InverseTransformPosition(InOutViewLocation);
-
 	// Forward data to the RHI dependend Domeprojection implementation
-	FRotator OriginSpaceViewRotation = FRotator::ZeroRotator;
-	if (!ViewAdapter->CalculateView(InViewport, InContextNum, DomeprojectionChannel, OriginSpaceViewLocation, OriginSpaceViewRotation, ViewOffset, WorldToMeters, NCP, FCP))
+	InOutViewRotation = FRotator::ZeroRotator;
+	if (!ViewAdapter->CalculateView(InViewport, InContextNum, DomeprojectionChannel, InOutViewLocation, InOutViewRotation, ViewOffset, WorldToMeters, InNCP, InFCP))
 	{
 		if (!IsEditorOperationMode(InViewport))
 		{
@@ -135,12 +133,6 @@ bool FDisplayClusterProjectionDomeprojectionPolicyBase::CalculateView(IDisplayCl
 		return false;
 	}
 
-	// Convert rotation back from origin to world space
-	InOutViewRotation = World2LocalTransform.TransformRotation(OriginSpaceViewRotation.Quaternion()).Rotator();
-
-	// Convert location back from origin to world space
-	InOutViewLocation = World2LocalTransform.TransformPosition(OriginSpaceViewLocation);
-
 	return true;
 }
 
@@ -180,7 +172,7 @@ void FDisplayClusterProjectionDomeprojectionPolicyBase::ApplyWarpBlend_RenderThr
 //////////////////////////////////////////////////////////////////////////////////////////////
 // FDisplayClusterProjectionDomeprojectionPolicy
 //////////////////////////////////////////////////////////////////////////////////////////////
-bool FDisplayClusterProjectionDomeprojectionPolicyBase::ReadConfigData(IDisplayClusterViewport* InViewport, FString& OutFile, FString& OutOrigin, uint32& OutChannel)
+bool FDisplayClusterProjectionDomeprojectionPolicyBase::ReadConfigData(IDisplayClusterViewport* InViewport, FString& OutFile, FString& OutOrigin, uint32& OutChannel, float& OutGeometryScale)
 {
 	check(InViewport);
 
@@ -229,5 +221,15 @@ bool FDisplayClusterProjectionDomeprojectionPolicyBase::ReadConfigData(IDisplayC
 		}
 	}
 
+	// UnitsToMeter (optional)
+	if(DisplayClusterHelpers::map::ExtractValueFromString(GetParameters(), DisplayClusterProjectionStrings::cfg::domeprojection::UnitsToMeter, OutGeometryScale))
+	{
+		UE_LOG(LogDisplayClusterProjectionDomeprojection, Verbose, TEXT("Viewport <%s>: Projection parameter '%s' - '%f'"), *InViewportId, DisplayClusterProjectionStrings::cfg::domeprojection::UnitsToMeter, OutGeometryScale);
+	}
+	else
+	{
+		//@TODO: add LOG when this parameter is added to the domeprojection configuration user interface
+	}
+
 	return true;
 }
diff --git a/Engine/Plugins/Runtime/nDisplay/Source/DisplayClusterProjection/Private/Policy/Domeprojection/DisplayClusterProjectionDomeprojectionPolicyBase.h b/Engine/Plugins/Runtime/nDisplay/Source/DisplayClusterProjection/Private/Policy/Domeprojection/DisplayClusterProjectionDomeprojectionPolicyBase.h
--- a/Engine/Plugins/Runtime/nDisplay/Source/DisplayClusterProjection/Private/Policy/Domeprojection/DisplayClusterProjectionDomeprojectionPolicyBase.h
+++ b/Engine/Plugins/Runtime/nDisplay/Source/DisplayClusterProjection/Private/Policy/Domeprojection/DisplayClusterProjectionDomeprojectionPolicyBase.h
@@ -35,6 +35,21 @@ public:
 	{
 		return true;
 	}
+
+	virtual float GetGeometryToMeters(const IDisplayClusterViewport* InViewport, const float InWorldToMeters) const override
+	{
+		return GeometryScale;
+	}
+
+	virtual void ApplyClippingPlanesOverrides(const IDisplayClusterViewport* InViewport, float& InOutNCP, float& InOutFCP) const override;
+
+	virtual EDisplayClusterProjectionPolicyFlags GetProjectionPolicyFlags(const IDisplayClusterViewport* InViewport, const uint32 InContextNum) const override
+	{
+		return EDisplayClusterProjectionPolicyFlags::EnableFollowCameraFeature
+			| EDisplayClusterProjectionPolicyFlags::UseEyeTracking
+			| EDisplayClusterProjectionPolicyFlags::UseLocalSpace;
+	}
+
 	//~End IDisplayClusterProjectionPolicy
 
 protected:
@@ -43,12 +58,16 @@ protected:
 
 private:
 	// Parse Domeprojection related data from the nDisplay config file
-	bool ReadConfigData(IDisplayClusterViewport* InViewport, FString& OutFile, FString& OutOrigin, uint32& OutChannel);
+	bool ReadConfigData(IDisplayClusterViewport* InViewport, FString& OutFile, FString& OutOrigin, uint32& OutChannel, float& OutGeometryScale);
 
 private:
 	FString OriginCompId;
 	uint32 DomeprojectionChannel = 0;
 
+	// Dome geometry units. Default it is in mm.
+	//@TODO: add this parameter to the domeprojection configuration UI
+	float GeometryScale = 1000.f;
+
 	// RHI depended view adapter (different RHI require different DLL/API etc.)
 	TUniquePtr<FDisplayClusterProjectionDomeprojectionViewAdapterBase> ViewAdapter;
 };
diff --git a/Engine/Plugins/Runtime/nDisplay/Source/DisplayClusterProjection/Private/Policy/EasyBlend/DisplayClusterProjectionEasyBlendPolicy.cpp b/Engine/Plugins/Runtime/nDisplay/Source/DisplayClusterProjection/Private/Policy/EasyBlend/DisplayClusterProjectionEasyBlendPolicy.cpp
--- a/Engine/Plugins/Runtime/nDisplay/Source/DisplayClusterProjection/Private/Policy/EasyBlend/DisplayClusterProjectionEasyBlendPolicy.cpp
+++ b/Engine/Plugins/Runtime/nDisplay/Source/DisplayClusterProjection/Private/Policy/EasyBlend/DisplayClusterProjectionEasyBlendPolicy.cpp
@@ -260,20 +260,9 @@ bool FDisplayClusterProjectionEasyBlendPolicy::CalculateView(IDisplayClusterView
 	
 	FDisplayClusterProjectionEasyBlendPolicyViewInfo& InOutViewInfo = PolicyViewInfo[InContextNum];
 
-	const float WorldScale = WorldToMeters / 100.f;
-	const float EasyBlendScale = EasyBlendConfiguration.GeometryScale / 100.f;
-
-
-	// Get origin component
-	const USceneComponent* const OriginComp = GetOriginComponent();
-	check(OriginComp);
-
-	// Get world-origin matrix
-	const FTransform& World2LocalTransform = (OriginComp != nullptr ? OriginComp->GetComponentTransform() : FTransform::Identity);
-
 	// Calculate view location in origin space
-	const FVector InViewLocation = World2LocalTransform.InverseTransformPosition(InOutViewLocation) * EasyBlendScale;
-	const FRotator InViewRotation = FRotator::ZeroRotator;
+	const FVector InViewLocation = InOutViewLocation;
+	const FRotator InViewRotation = InOutViewRotation;
 
 	// Convert to EasyBlend coordinate system
 	InOutViewInfo.ViewLocation = FVector(
@@ -305,11 +294,7 @@ bool FDisplayClusterProjectionEasyBlendPolicy::CalculateView(IDisplayClusterView
 	);
 
 	// Convert rotation back from origin to world space
-	InOutViewRotation = World2LocalTransform.TransformRotation(OutViewRotation.Quaternion()).Rotator();
-
-	// Save clipping planes:
-	ZNear = NCP;
-	ZFar = FCP;
+	InOutViewRotation = OutViewRotation;
 
 	return true;
 }
@@ -318,18 +303,25 @@ bool FDisplayClusterProjectionEasyBlendPolicy::GetProjectionMatrix(IDisplayClust
 {
 	check(IsInGameThread());
 
-	if (!InViewport || !PolicyViewData.IsValid() || !PolicyViewInfo.IsValidIndex(InContextNum))
+	if (!PolicyViewData.IsValid() || !PolicyViewInfo.IsValidIndex(InContextNum))
 	{
 		return false;
 	}
 
+	if (!InViewport || InViewport->GetContexts().IsValidIndex(InContextNum))
+	{
+		return false;
+	}
+
+	const FDisplayClusterViewport_Context& InContext = InViewport->GetContexts()[InContextNum];
+
 	InViewport->CalculateProjectionMatrix(
 		InContextNum,
 		PolicyViewInfo[InContextNum].FrustumAngles.X,
 		PolicyViewInfo[InContextNum].FrustumAngles.Y,
 		PolicyViewInfo[InContextNum].FrustumAngles.Z,
 		PolicyViewInfo[InContextNum].FrustumAngles.W,
-		ZNear, ZFar, true
+		InContext.ProjectionData.ZNear, InContext.ProjectionData.ZFar, true
 	);
 
 	OutPrjMatrix = InViewport->GetContexts()[InContextNum].ProjectionMatrix;
diff --git a/Engine/Plugins/Runtime/nDisplay/Source/DisplayClusterProjection/Private/Policy/EasyBlend/DisplayClusterProjectionEasyBlendPolicy.h b/Engine/Plugins/Runtime/nDisplay/Source/DisplayClusterProjection/Private/Policy/EasyBlend/DisplayClusterProjectionEasyBlendPolicy.h
--- a/Engine/Plugins/Runtime/nDisplay/Source/DisplayClusterProjection/Private/Policy/EasyBlend/DisplayClusterProjectionEasyBlendPolicy.h
+++ b/Engine/Plugins/Runtime/nDisplay/Source/DisplayClusterProjection/Private/Policy/EasyBlend/DisplayClusterProjectionEasyBlendPolicy.h
@@ -35,6 +35,24 @@ public:
 
 	virtual void ApplyWarpBlend_RenderThread(FRHICommandListImmediate& RHICmdList, const IDisplayClusterViewportProxy* InViewportProxy) override;
 
+	virtual float GetGeometryToMeters(const IDisplayClusterViewport* InViewport, const float InWorldToMeters) const override
+	{
+		// Return how many geometry units in 1 meter
+		return EasyBlendConfiguration.GeometryScale;
+	}
+
+	virtual EDisplayClusterProjectionPolicyFlags GetProjectionPolicyFlags(const IDisplayClusterViewport* InViewport, const uint32 InContextNum) const override
+	{
+		return EDisplayClusterProjectionPolicyFlags::EnableFollowCameraFeature
+			| EDisplayClusterProjectionPolicyFlags::UseEyeTracking
+			| EDisplayClusterProjectionPolicyFlags::UseLocalSpace;
+	}
+
+	virtual FVector GetOriginLocationOffsetInGeometryUnits(const IDisplayClusterViewport* InViewport) const override
+	{
+		return EasyBlendConfiguration.OriginOffset;
+	}
+
 	// Request additional targetable resources for easyblend external warpblend
 	virtual bool ShouldUseAdditionalTargetableResource(const IDisplayClusterViewport* InViewport) const override
 	{
@@ -70,8 +88,4 @@ private:
 
 	// Viewport RHI
 	FRHIViewport* RHIViewportProxy = nullptr;
-
-	// Store clipping planes
-	float ZNear;
-	float ZFar;
 };
diff --git a/Engine/Plugins/Runtime/nDisplay/Source/DisplayClusterProjection/Private/Policy/EasyBlend/DisplayClusterProjectionEasyBlendPolicyConfiguration.cpp b/Engine/Plugins/Runtime/nDisplay/Source/DisplayClusterProjection/Private/Policy/EasyBlend/DisplayClusterProjectionEasyBlendPolicyConfiguration.cpp
--- a/Engine/Plugins/Runtime/nDisplay/Source/DisplayClusterProjection/Private/Policy/EasyBlend/DisplayClusterProjectionEasyBlendPolicyConfiguration.cpp
+++ b/Engine/Plugins/Runtime/nDisplay/Source/DisplayClusterProjection/Private/Policy/EasyBlend/DisplayClusterProjectionEasyBlendPolicyConfiguration.cpp
@@ -104,6 +104,19 @@ bool FDisplayClusterProjectionEasyBlendPolicyConfiguration::Initialize(const TMa
 		}
 	}
 
+	// Origin offset
+	if (DisplayClusterHelpers::map::template ExtractValueFromString(InParameters, DisplayClusterProjectionStrings::cfg::easyblend::OriginOffset, OriginOffset))
+	{
+		UE_LOG(LogDisplayClusterProjectionEasyBlend, Verbose, TEXT("Viewport <%s>: Projection parameter '%s' - %s"), *InViewportId, DisplayClusterProjectionStrings::cfg::easyblend::OriginOffset, *OriginOffset.ToString());
+	}
+	else
+	{
+		if (bEnableExtraLog)
+		{
+			UE_LOG(LogDisplayClusterProjectionEasyBlend, Log, TEXT("Viewport <%s>: Projection parameter '%s' not found"), *InViewportId, DisplayClusterProjectionStrings::cfg::easyblend::OriginOffset);
+		}
+	}
+
 	bInvalidConfiguration = false;
 
 	return true;
diff --git a/Engine/Plugins/Runtime/nDisplay/Source/DisplayClusterProjection/Private/Policy/EasyBlend/DisplayClusterProjectionEasyBlendPolicyConfiguration.h b/Engine/Plugins/Runtime/nDisplay/Source/DisplayClusterProjection/Private/Policy/EasyBlend/DisplayClusterProjectionEasyBlendPolicyConfiguration.h
--- a/Engine/Plugins/Runtime/nDisplay/Source/DisplayClusterProjection/Private/Policy/EasyBlend/DisplayClusterProjectionEasyBlendPolicyConfiguration.h
+++ b/Engine/Plugins/Runtime/nDisplay/Source/DisplayClusterProjection/Private/Policy/EasyBlend/DisplayClusterProjectionEasyBlendPolicyConfiguration.h
@@ -35,6 +35,9 @@ public:
 	// EasyBlend unit scale
 	float GeometryScale = 1.f;
 
+	// Origin's location in geometry units.
+	FVector OriginOffset = FVector::ZeroVector;
+
 	// Allow to use preview mesh
 	bool bIsPreviewMeshEnabled = false;
 
diff --git a/Engine/Plugins/Runtime/nDisplay/Source/DisplayClusterProjection/Private/Policy/MPCDI/DisplayClusterProjectionMPCDIPolicy.cpp b/Engine/Plugins/Runtime/nDisplay/Source/DisplayClusterProjection/Private/Policy/MPCDI/DisplayClusterProjectionMPCDIPolicy.cpp
--- a/Engine/Plugins/Runtime/nDisplay/Source/DisplayClusterProjection/Private/Policy/MPCDI/DisplayClusterProjectionMPCDIPolicy.cpp
+++ b/Engine/Plugins/Runtime/nDisplay/Source/DisplayClusterProjection/Private/Policy/MPCDI/DisplayClusterProjectionMPCDIPolicy.cpp
@@ -224,7 +224,7 @@ bool FDisplayClusterProjectionMPCDIPolicy::CalculateView(IDisplayClusterViewport
 {
 	check(IsInGameThread());
 
-	if (!InViewport)
+	if (!InViewport || InViewport->GetContexts().IsValidIndex(InContextNum))
 	{
 		return false;
 	}
@@ -239,29 +239,20 @@ bool FDisplayClusterProjectionMPCDIPolicy::CalculateView(IDisplayClusterViewport
 		return false;
 	}
 
-	// Override viewpoint
-	// MPCDI always expects the location of the viewpoint component (eye location from the real world)
-	FVector ViewOffset = FVector::ZeroVector;
-	if (!InViewport->GetViewPointCameraEye(InContextNum, InOutViewLocation, InOutViewRotation, ViewOffset))
-	{
-		return false;
-	}
-
 	// World scale multiplier
 	const float WorldScale = WorldToMeters / 100.f;
 
-	// Get view location in local space
-	const USceneComponent* const OriginComp = GetOriginComponent();
-
 	// Initialize frustum
 	TSharedPtr<FDisplayClusterWarpEye, ESPMode::ThreadSafe> WarpEye = MakeShared<FDisplayClusterWarpEye, ESPMode::ThreadSafe>(InViewport->ToSharedPtr(), InContextNum);
 
-	WarpEye->World2LocalTransform = (OriginComp ? OriginComp->GetComponentTransform() : FTransform::Identity);
+	const FDisplayClusterViewport_Context& InContext = InViewport->GetContexts()[InContextNum];
+
+	WarpEye->World2LocalTransform = InContext.OriginToWorld;
 
 	// Get our base camera location and view offset in local space (MPCDI space)
-	WarpEye->ViewPoint.Location  = WarpEye->World2LocalTransform.InverseTransformPosition(InOutViewLocation - ViewOffset);
-	WarpEye->ViewPoint.EyeOffset = WarpEye->World2LocalTransform.InverseTransformPosition(InOutViewLocation) - WarpEye->ViewPoint.Location;
-	WarpEye->ViewPoint.Rotation  = WarpEye->World2LocalTransform.InverseTransformRotation(InOutViewRotation.Quaternion()).Rotator();
+	WarpEye->ViewPoint.Location  = InOutViewLocation - InViewOffset;
+	WarpEye->ViewPoint.EyeOffset = InViewOffset;
+	WarpEye->ViewPoint.Rotation  = InOutViewRotation;
 
 	WarpEye->WorldScale = WorldScale;
 
@@ -284,8 +275,8 @@ bool FDisplayClusterProjectionMPCDIPolicy::CalculateView(IDisplayClusterViewport
 	WarpBlendContexts[InContextNum].Origin2WorldTransform = WarpEye->World2LocalTransform;
 
 	// Transform viewpoint back to world space
-	InOutViewRotation = WarpEye->World2LocalTransform.TransformRotation(WarpBlendContexts[InContextNum].Rotation.Quaternion()).Rotator();
-	InOutViewLocation = WarpEye->World2LocalTransform.TransformPosition(WarpBlendContexts[InContextNum].Location);
+	InOutViewRotation = WarpBlendContexts[InContextNum].Rotation;
+	InOutViewLocation = WarpBlendContexts[InContextNum].Location;
 
 	WarpBlendContexts[InContextNum].bIsValid = true;
 
diff --git a/Engine/Plugins/Runtime/nDisplay/Source/DisplayClusterProjection/Private/Policy/MPCDI/DisplayClusterProjectionMPCDIPolicy.h b/Engine/Plugins/Runtime/nDisplay/Source/DisplayClusterProjection/Private/Policy/MPCDI/DisplayClusterProjectionMPCDIPolicy.h
--- a/Engine/Plugins/Runtime/nDisplay/Source/DisplayClusterProjection/Private/Policy/MPCDI/DisplayClusterProjectionMPCDIPolicy.h
+++ b/Engine/Plugins/Runtime/nDisplay/Source/DisplayClusterProjection/Private/Policy/MPCDI/DisplayClusterProjectionMPCDIPolicy.h
@@ -57,6 +57,13 @@ public:
 		return true;
 	}
 
+	virtual EDisplayClusterProjectionPolicyFlags GetProjectionPolicyFlags(const IDisplayClusterViewport* InViewport, const uint32 InContextNum) const override
+	{
+		return EDisplayClusterProjectionPolicyFlags::EnableFollowCameraFeature
+			| EDisplayClusterProjectionPolicyFlags::UseEyeTracking
+			| EDisplayClusterProjectionPolicyFlags::UseLocalSpace;
+	}
+
 	virtual void UpdateProxyData(IDisplayClusterViewport* InViewport) override;
 
 	virtual void SetWarpPolicy(IDisplayClusterWarpPolicy* InWarpPolicy) override;
diff --git a/Engine/Plugins/Runtime/nDisplay/Source/DisplayClusterProjection/Private/Policy/Reference/DisplayClusterProjectionReferencePolicy.cpp b/Engine/Plugins/Runtime/nDisplay/Source/DisplayClusterProjection/Private/Policy/Reference/DisplayClusterProjectionReferencePolicy.cpp
--- a/Engine/Plugins/Runtime/nDisplay/Source/DisplayClusterProjection/Private/Policy/Reference/DisplayClusterProjectionReferencePolicy.cpp
+++ b/Engine/Plugins/Runtime/nDisplay/Source/DisplayClusterProjection/Private/Policy/Reference/DisplayClusterProjectionReferencePolicy.cpp
@@ -98,6 +98,20 @@ bool FDisplayClusterProjectionReferencePolicy::ShouldUseAdditionalTargetableReso
 	return false;
 }
 
+EDisplayClusterProjectionPolicyFlags FDisplayClusterProjectionReferencePolicy::GetProjectionPolicyFlags(const IDisplayClusterViewport* InViewport, const uint32 InContextNum) const
+{
+	if (IDisplayClusterViewport* SourceViewport = GetSourceViewport(InViewport))
+	{
+		// Only ICVFX requires composition.
+		if (SourceViewport->GetProjectionPolicy().IsValid() && SourceViewport->GetProjectionPolicy()->ShouldSupportICVFX(SourceViewport))
+		{
+			return SourceViewport->GetProjectionPolicy()->GetProjectionPolicyFlags(SourceViewport, InContextNum);
+		}
+	}
+
+	return EDisplayClusterProjectionPolicyFlags::None;
+}
+
 bool FDisplayClusterProjectionReferencePolicy::IsWarpBlendSupported(const IDisplayClusterViewport* InViewport) const
 {
 	if (IDisplayClusterViewport* SourceViewport = GetSourceViewport(InViewport))
diff --git a/Engine/Plugins/Runtime/nDisplay/Source/DisplayClusterProjection/Private/Policy/Reference/DisplayClusterProjectionReferencePolicy.h b/Engine/Plugins/Runtime/nDisplay/Source/DisplayClusterProjection/Private/Policy/Reference/DisplayClusterProjectionReferencePolicy.h
--- a/Engine/Plugins/Runtime/nDisplay/Source/DisplayClusterProjection/Private/Policy/Reference/DisplayClusterProjectionReferencePolicy.h
+++ b/Engine/Plugins/Runtime/nDisplay/Source/DisplayClusterProjection/Private/Policy/Reference/DisplayClusterProjectionReferencePolicy.h
@@ -49,6 +49,8 @@ public:
 	virtual bool ShouldUseSourceTextureWithMips(const IDisplayClusterViewport* InViewport) const override;
 	virtual bool ShouldUseAdditionalTargetableResource(const IDisplayClusterViewport* InViewport) const override;
 
+	virtual EDisplayClusterProjectionPolicyFlags GetProjectionPolicyFlags(const IDisplayClusterViewport* InViewport, const uint32 InContextNum) const override;
+
 	virtual void UpdateProxyData(IDisplayClusterViewport* InViewport) override;
 
 	virtual bool CalculateView(
diff --git a/Engine/Plugins/Runtime/nDisplay/Source/DisplayClusterProjection/Private/Policy/Simple/DisplayClusterProjectionSimplePolicy.h b/Engine/Plugins/Runtime/nDisplay/Source/DisplayClusterProjection/Private/Policy/Simple/DisplayClusterProjectionSimplePolicy.h
--- a/Engine/Plugins/Runtime/nDisplay/Source/DisplayClusterProjection/Private/Policy/Simple/DisplayClusterProjectionSimplePolicy.h
+++ b/Engine/Plugins/Runtime/nDisplay/Source/DisplayClusterProjection/Private/Policy/Simple/DisplayClusterProjectionSimplePolicy.h
@@ -18,9 +18,8 @@ public:
 	FDisplayClusterProjectionSimplePolicy(const FString& ProjectionPolicyId, const FDisplayClusterConfigurationProjection* InConfigurationProjectionPolicy);
 
 public:
-	//////////////////////////////////////////////////////////////////////////////////////////////
-	// IDisplayClusterProjectionPolicy
-	//////////////////////////////////////////////////////////////////////////////////////////////
+	//~BEGIN  IDisplayClusterProjectionPolicy
+	
 	virtual const FString& GetType() const override;
 
 	virtual bool HandleStartScene(IDisplayClusterViewport* InViewport) override;
@@ -35,6 +34,12 @@ public:
 	}
 
 	virtual UMeshComponent* GetOrCreatePreviewMeshComponent(IDisplayClusterViewport* InViewport, bool& bOutIsRootActorComponent) override;
+	virtual EDisplayClusterProjectionPolicyFlags GetProjectionPolicyFlags(const IDisplayClusterViewport* InViewport, const uint32 InContextNum) const override
+	{
+		return EDisplayClusterProjectionPolicyFlags::EnableFollowCameraFeature
+			| EDisplayClusterProjectionPolicyFlags::UseEyeTracking;
+	}
+	//~~END IDisplayClusterProjectionPolicy
 
 protected:
 	bool InitializeMeshData(IDisplayClusterViewport* InViewport);
diff --git a/Engine/Plugins/Runtime/nDisplay/Source/DisplayClusterProjection/Private/Policy/VIOSO/Windows/DisplayClusterProjectionVIOSOPolicy.cpp b/Engine/Plugins/Runtime/nDisplay/Source/DisplayClusterProjection/Private/Policy/VIOSO/Windows/DisplayClusterProjectionVIOSOPolicy.cpp
--- a/Engine/Plugins/Runtime/nDisplay/Source/DisplayClusterProjection/Private/Policy/VIOSO/Windows/DisplayClusterProjectionVIOSOPolicy.cpp
+++ b/Engine/Plugins/Runtime/nDisplay/Source/DisplayClusterProjection/Private/Policy/VIOSO/Windows/DisplayClusterProjectionVIOSOPolicy.cpp
@@ -100,18 +100,9 @@ bool FDisplayClusterProjectionVIOSOPolicy::CalculateView(IDisplayClusterViewport
 		return false;
 	}
 
-	// Get view location in local space
-	const USceneComponent* const OriginComp = GetOriginComponent();
-	const FTransform& World2LocalTransform = (OriginComp ? OriginComp->GetComponentTransform() : FTransform::Identity);
-
-	// Get our base camera location and view offset in local space (MPCDI space)
-	const FVector  LocalOrigin    = World2LocalTransform.InverseTransformPosition(InOutViewLocation - ViewOffset);
-	const FVector  LocalEyeOrigin = World2LocalTransform.InverseTransformPosition(InOutViewLocation);
-	const FRotator LocalRotator   = World2LocalTransform.InverseTransformRotation(InOutViewRotation.Quaternion()).Rotator();
-
 	// Get view prj data from VIOSO
 	FScopeLock lock(&DllAccessCS);
-	if (!CurrentView->UpdateVIOSO(InViewport, InContextNum, LocalEyeOrigin, LocalRotator, WorldToMeters, InNCP, InFCP))
+	if (!CurrentView->UpdateVIOSO(InViewport, InContextNum, InOutViewLocation, InOutViewRotation, WorldToMeters, InNCP, InFCP))
 	{
 		if (!FDisplayClusterProjectionPolicyBase::IsEditorOperationMode(InViewport))
 		{
@@ -123,8 +114,8 @@ bool FDisplayClusterProjectionVIOSOPolicy::CalculateView(IDisplayClusterViewport
 	}
 
 	// Transform rotation to world space
-	InOutViewRotation = World2LocalTransform.TransformRotation(CurrentView->ViewRotation.Quaternion()).Rotator();
-	InOutViewLocation = World2LocalTransform.TransformPosition(CurrentView->ViewLocation);
+	InOutViewRotation = CurrentView->ViewRotation;
+	InOutViewLocation = CurrentView->ViewLocation;
 
 	return true;
 }
diff --git a/Engine/Plugins/Runtime/nDisplay/Source/DisplayClusterProjection/Private/Policy/VIOSO/Windows/DisplayClusterProjectionVIOSOPolicy.h b/Engine/Plugins/Runtime/nDisplay/Source/DisplayClusterProjection/Private/Policy/VIOSO/Windows/DisplayClusterProjectionVIOSOPolicy.h
--- a/Engine/Plugins/Runtime/nDisplay/Source/DisplayClusterProjection/Private/Policy/VIOSO/Windows/DisplayClusterProjectionVIOSOPolicy.h
+++ b/Engine/Plugins/Runtime/nDisplay/Source/DisplayClusterProjection/Private/Policy/VIOSO/Windows/DisplayClusterProjectionVIOSOPolicy.h
@@ -44,6 +44,18 @@ public:
 	virtual bool HasPreviewMesh(IDisplayClusterViewport* InViewport) override;
 	virtual class UMeshComponent* GetOrCreatePreviewMeshComponent(IDisplayClusterViewport* InViewport, bool& bOutIsRootActorComponent) override;
 
+	virtual float GetGeometryToMeters(const IDisplayClusterViewport* InViewport, const float InWorldToMeters) const override
+	{
+		return ViosoConfigData.UnitsInMeter;
+	}
+
+	virtual EDisplayClusterProjectionPolicyFlags GetProjectionPolicyFlags(const IDisplayClusterViewport* InViewport, const uint32 InContextNum) const override
+	{
+		return EDisplayClusterProjectionPolicyFlags::EnableFollowCameraFeature
+			| EDisplayClusterProjectionPolicyFlags::UseEyeTracking
+			| EDisplayClusterProjectionPolicyFlags::UseLocalSpace;
+	}
+
 	//~~END IDisplayClusterProjectionPolicy
 
 protected:
diff --git a/Engine/Plugins/Runtime/nDisplay/Source/DisplayClusterProjection/Public/DisplayClusterProjectionStrings.h b/Engine/Plugins/Runtime/nDisplay/Source/DisplayClusterProjection/Public/DisplayClusterProjectionStrings.h
--- a/Engine/Plugins/Runtime/nDisplay/Source/DisplayClusterProjection/Public/DisplayClusterProjectionStrings.h
+++ b/Engine/Plugins/Runtime/nDisplay/Source/DisplayClusterProjection/Public/DisplayClusterProjectionStrings.h
@@ -21,6 +21,8 @@ namespace DisplayClusterProjectionStrings
 			static constexpr const TCHAR* File   = TEXT("file");
 			static constexpr const TCHAR* Origin = TEXT("origin");
 			static constexpr const TCHAR* Scale  = TEXT("scale");
+
+			static constexpr const TCHAR* OriginOffset = TEXT("origin_loc");
 		}
 
 		namespace VIOSO
@@ -164,6 +166,7 @@ namespace DisplayClusterProjectionStrings
 			static constexpr const TCHAR* File    = TEXT("file");
 			static constexpr const TCHAR* Channel = TEXT("channel");
 			static constexpr const TCHAR* Origin  = TEXT("origin");
+			static constexpr const TCHAR* UnitsToMeter = TEXT("scale");
 		}
 	}
 
