nDisplay: Expose to blueprint a couple of functions in the ICVFXCamera component for convenience.

#rnx

diff --git a/Engine/Plugins/Runtime/nDisplay/Source/DisplayCluster/Private/Game/EngineClasses/Scene/DisplayClusterICVFXCameraComponent.cpp b/Engine/Plugins/Runtime/nDisplay/Source/DisplayCluster/Private/Game/EngineClasses/Scene/DisplayClusterICVFXCameraComponent.cpp
--- a/Engine/Plugins/Runtime/nDisplay/Source/DisplayCluster/Private/Game/EngineClasses/Scene/DisplayClusterICVFXCameraComponent.cpp
+++ b/Engine/Plugins/Runtime/nDisplay/Source/DisplayCluster/Private/Game/EngineClasses/Scene/DisplayClusterICVFXCameraComponent.cpp
@@ -245,6 +245,32 @@ void UDisplayClusterICVFXCameraComponent::TickComponent(float DeltaTime, ELevelT
 	}
 }
 
+bool UDisplayClusterICVFXCameraComponent::IsICVFXEnabled() const
+{
+	if (!GDisplayCluster)
+	{
+		return false;
+	}
+
+	const ADisplayClusterRootActor* const RootActor = Cast<ADisplayClusterRootActor>(GetOwner());
+
+	if (!RootActor)
+	{
+		return false;
+	}
+
+	const UDisplayClusterConfigurationData* const ConfigData = RootActor->GetConfigData();
+
+	if (!ConfigData)
+	{
+		return false;
+	}
+
+	const FString NodeId = GDisplayCluster->GetPrivateClusterMgr()->GetNodeId();
+
+	return GetCameraSettingsICVFX().IsICVFXEnabled(*ConfigData, NodeId);
+}
+
 const FDisplayClusterConfigurationICVFX_CameraSettings& UDisplayClusterICVFXCameraComponent::GetCameraSettingsICVFX() const
 {
 	return CameraSettings;
diff --git a/Engine/Plugins/Runtime/nDisplay/Source/DisplayCluster/Public/Components/DisplayClusterICVFXCameraComponent.h b/Engine/Plugins/Runtime/nDisplay/Source/DisplayCluster/Public/Components/DisplayClusterICVFXCameraComponent.h
--- a/Engine/Plugins/Runtime/nDisplay/Source/DisplayCluster/Public/Components/DisplayClusterICVFXCameraComponent.h
+++ b/Engine/Plugins/Runtime/nDisplay/Source/DisplayCluster/Public/Components/DisplayClusterICVFXCameraComponent.h
@@ -66,14 +66,12 @@ public:
 	 * Return the actual source camera, e.g. the camera component of the referenced cine camera.
 	 * Use GetCameraView() function to get viewinfo with actual camera position, postprocess and ICVFX postprocess
 	 */
+	UFUNCTION(BlueprintCallable, BlueprintPure, Category = "ICVFX Camera")
 	UCineCameraComponent* GetActualCineCameraComponent();
 
-	// Returns true if this camera is active
-	UE_DEPRECATED(5.4, "This function has been deprecated. Please use 'GetCameraSettingsICVFX().IsICVFXEnabled()'.")
-	bool IsICVFXEnabled() const
-	{
-		return false;
-	}
+	/** Returns true if this camera is active in the current node. */
+	UFUNCTION(BlueprintCallable, BlueprintPure, Category = "ICVFX Camera")
+	bool IsICVFXEnabled() const;
 
 	// Return unique camera name
 	FString GetCameraUniqueId() const;
