Add a reset sync command to tell the nDisplay interceptor that it is OK to re-initiation sync'd interceptor. nDisplay users can issue this command over switchboard connection to tell the module that the disconnect was intentional and it is OK to restart when entering into a new Multi-user session connection.

diff --git a/Engine/Plugins/Runtime/nDisplay/Source/DisplayClusterMessageInterceptor/Private/DisplayClusterMessageInterceptionModule.cpp b/Engine/Plugins/Runtime/nDisplay/Source/DisplayClusterMessageInterceptor/Private/DisplayClusterMessageInterceptionModule.cpp
--- a/Engine/Plugins/Runtime/nDisplay/Source/DisplayClusterMessageInterceptor/Private/DisplayClusterMessageInterceptionModule.cpp
+++ b/Engine/Plugins/Runtime/nDisplay/Source/DisplayClusterMessageInterceptor/Private/DisplayClusterMessageInterceptionModule.cpp
@@ -111,9 +111,15 @@ public:
 		StopMessageSyncCommand = MakeUnique<FAutoConsoleCommand>(
 			TEXT("nDisplay.MessageBusSync.Stop"),
 			TEXT("Stop MessageBus syncing"),
-			FConsoleCommandDelegate::CreateRaw(this, &FDisplayClusterMessageInterceptionModule::StopInterception)
+			FConsoleCommandDelegate::CreateRaw(this, &FDisplayClusterMessageInterceptionModule::StopInterception));
+
+		ResetMultiUserSyncCommand = MakeUnique<FAutoConsoleCommand>(
+			TEXT("nDisplay.MessageBusSync.ResetSyncFlags"),
+			TEXT("Reset interceptor sync flags for Multi-user session connections."),
+			FConsoleCommandDelegate::CreateRaw(this, &FDisplayClusterMessageInterceptionModule::ResetSyncFlags)
 			);
 
+
 #if WITH_EDITOR
 		if (ISettingsModule* SettingsModule = FModuleManager::GetModulePtr<ISettingsModule>("Settings"))
 		{
@@ -169,10 +175,16 @@ public:
 		Interceptor.Reset();
 		StartMessageSyncCommand.Reset();
 		StopMessageSyncCommand.Reset();
+		ResetMultiUserSyncCommand.Reset();
 	}
 
 private:
-
+	void ResetSyncFlags()
+	{
+		// Reset the was ever disconnected flag to allow sync to happen again on when joining a session. 
+		bWasEverDisconnected = false;
+	}
+		
 	void SetupForMultiUser()
 	{
 #if defined(WITH_CONCERT)
@@ -581,9 +593,10 @@ private:
 	/** Cluster event listener delegate */
 	FOnClusterEventJsonListener ListenerDelegate;
 
-	/** Console commands handle. */
+	/** Handles to console commands that we manage. */
 	TUniquePtr<FAutoConsoleCommand> StartMessageSyncCommand;
 	TUniquePtr<FAutoConsoleCommand> StopMessageSyncCommand;
+	TUniquePtr<FAutoConsoleCommand> ResetMultiUserSyncCommand;
 
 	/** Request to start interception. Cached to be done once synchronization is done. */
 	bool bStartInterceptionRequested;
