[nDisplay] Blending the edges of tile viewports with each other.

- Added edge blending for tile viewports.

- Enhanced DisplayClusterShaderTextureUtils: Implemented blending with alpha feathering on edges for smoother transitions.

- Introduced TMarginSet, a flexible template class for storing and managing four directional values.

- Added new Tile settings to the  FDisplayClusterConfigurationTile_Overscan
* bBlendEntireOverscan - If true, blend across the entire overscan edge region instead of using OverscanBlendWidth. This applies the full overscan width/height for blending at tile edges.
* OverscanBlendWidth - Width of the blended edge region, in pixels.

- Code cleanup.

- Added new CVar for tile viewports:
"nDisplay.render.tile.BlendMode" 2
Blending mode control (0 - None, 1 - linear, 2 - smooth S-shaped)

diff --git a/Engine/Plugins/Runtime/nDisplay/Shaders/Private/ResourceUtils.usf b/Engine/Plugins/Runtime/nDisplay/Shaders/Private/ResourceUtils.usf
--- a/Engine/Plugins/Runtime/nDisplay/Shaders/Private/ResourceUtils.usf
+++ b/Engine/Plugins/Runtime/nDisplay/Shaders/Private/ResourceUtils.usf
@@ -29,16 +29,98 @@ float3 DisplayGamma;
 // Definition for the ColorPremultiply
 // ColorPremultiply.X = Input texture Color premultiply
 // ColorPremultiply.Y = Output texture Color premultiply
+// ColorPremultiply.Z = SideFeather
 
 #define EColorPremultiply_None          0
 #define EColorPremultiply_Alpha         1
 #define EColorPremultiply_InvertedAlpha 2
 
+#define ESideFeather_None      0
+#define ESideFeather_Linear    1
+#define ESideFeather_Smooth    2
+
 // Color modifier
 uint3 ColorPremultiply;
 
+// Feather margins
+//
+//  Outer LT              RT
+//  +---------------------+
+//  |                     |
+//  |  Inner LT      RT   |
+//  |     +----------+    |
+//  |     |          |    |
+//  |     +----------+    |
+//  |     LB         RB   |
+//  +---------------------+
+//  LB                    RB
+//
+// L,R,T,B = Left, Right, Top, Bottom
+//
+// InnerFeatherMargins = <InnerL, 1-InnerR, InnerT, 1-InnerB>  // Fully opaque region
+// OuterFeatherMargins = <OuterL, 1-OuterR, OuterT, 1-OuterB>  // Fully transparent region
+float4 InnerFeatherMargins;
+float4 OuterFeatherMargins;
+
+/**
+ * Returns a linearly feathered alpha based on UV position and side margins.
+ * The feathering fades from fully transparent (Max margins) to fully opaque (Min margins).
+ */
+float GetSideFeatherLinear(float2 UV)
+{
+	float2 InvUV = 1.0 - UV;
+	float4 uvMargins = float4(UV.x, InvUV.x, UV.y, InvUV.y);
+	float4 FeatherRange = InnerFeatherMargins - OuterFeatherMargins;
+
+    // Prevent division by zero by replacing 0 range with 1.0
+	float4 SafeFeatherRange = max(FeatherRange, 1e-5);
+
+    // Compute normalized feather
+	float4 NormalizedFeather = saturate((uvMargins - OuterFeatherMargins) / SafeFeatherRange);
+	
+	// Small threshold for treating near-zero floats as zero, to suppress noise.
+	const float NearZeroThreshold = 1e-5f;
+	
+    // Mask: if FeatherRange is 0, use 1.0 (i.e. fully opaque / no feather)
+	float4 RangeValidMask = step(NearZeroThreshold, FeatherRange);
+	
+	// RangeValidMask = 0 if < 1e-5, 1 if >= 1e-5
+	NormalizedFeather = lerp(1.0, NormalizedFeather, RangeValidMask);
+
+	// Return combined alpha (multiplicative blend)
+	return NormalizedFeather.x * NormalizedFeather.y * NormalizedFeather.z * NormalizedFeather.w;
+}
+
+/**
+ * Computes a feathered alpha value based on side-specific margins using smoothstep transitions.
+ *
+ * Logic:
+ * - Inside the inner margins (`InnerFeatherMargins`): alpha = 1 (fully opaque)
+ * - Outside the outer margins (`OuterFeatherMargins`): alpha = 0 (fully transparent)
+ * - Between Min and Max: alpha fades smoothly using smoothstep
+
+ * @param UV   UV coordinates (float2), expected in [0, 1] range
+ * @return       Alpha value in [0, 1] with smooth feathering on each side
+ */
+float GetSideFeatherSmooth(float2 UV)
+{
+	// Remap UVs for mirrored sides (Right and Bottom)
+	float2 InvUV = 1.0 - UV;
+
+    // Combine UVs into one vector: <Left, Right, Top, Bottom>
+	float4 uvMargins = float4(UV.x, InvUV.x, UV.y, InvUV.y);
+
+    // Compute smooth feather for all sides: Left, Right, Top, Bottom
+	float4 Feather = smoothstep(OuterFeatherMargins, InnerFeatherMargins, uvMargins);
+
+    // Return combined alpha (multiplicative blend)
+	return Feather.x * Feather.y * Feather.z * Feather.w;
+}
+
 /** Implements color encoding. */
-float4 DisplayClusterEncodeColorAndAlpha(float4 InColorArg)
+	float4 DisplayClusterEncodeColorAndAlpha(
+	FScreenVertexOutput Input,
+	float4 InColorArg)
 {
 	float4 OutColor = InColorArg;
 
@@ -70,6 +152,23 @@ float4 DisplayClusterEncodeColorAndAlpha(float4 InColorArg)
     {
         OutColor.a = 0.0f;
     }
+	
+	// Add Side Feather color
+	BRANCH
+    if(ColorPremultiply.z == ESideFeather_Linear)
+    {
+	    float Feather = GetSideFeatherLinear(Input.UV.xy);
+	
+		// Premultiply color and alpha
+		OutColor *= Feather;
+    }
+	else if(ColorPremultiply.z == ESideFeather_Smooth)
+	{  
+		float Feather = GetSideFeatherSmooth(Input.UV.xy);
+	
+		// Premultiply color and alpha
+		OutColor *= Feather;
+	}
 #endif
 
 #if ENCODE_INPUT || ENCODE_OUTPUT
@@ -87,9 +186,8 @@ float4 DisplayClusterEncodeColorAndAlpha(float4 InColorArg)
 	}
 	else if(Encodings.x == ESourceEncoding_MediaPQ) // MediaPQ -> Linear
     {
-		// Pixel shader for Linear-PQ conversion
-		float3 colorPQ = LinearToST2084(OutColor.rgb);
-        OutColor.rgb = saturate(colorPQ);
+		// Pixel shader for PQ-Linear conversion
+		OutColor.rgb = ST2084ToLinear(OutColor.rgb);
 	}
 
 #endif // ENCODE_INPUT
@@ -107,8 +205,9 @@ float4 DisplayClusterEncodeColorAndAlpha(float4 InColorArg)
     }
 	else if(Encodings.y == ESourceEncoding_MediaPQ) // Linear -> MediaPQ
     {
-		// Pixel shader for PQ-Linear conversion
-		OutColor.rgb = ST2084ToLinear(OutColor.rgb);
+		// Pixel shader for Linear-PQ conversion
+		float3 ColorPQ = LinearToST2084(OutColor.rgb);
+		OutColor.rgb = saturate(ColorPQ);
     }
 #endif // ENCODE_OUTPUT
 
@@ -132,10 +231,9 @@ float4 DisplayClusterEncodeColorAndAlpha(float4 InColorArg)
 
 void Main(
 	FScreenVertexOutput Input,
-	out float4 OutColor : SV_Target0
-	)
+	out float4 OutColor : SV_Target0)
 {
 	float4 InColor = Texture2DSample(InputTexture, InputTextureSampler, Input.UV);
 
-	OutColor = DisplayClusterEncodeColorAndAlpha(InColor);
+	OutColor = DisplayClusterEncodeColorAndAlpha(Input, InColor);
 }
\ No newline at end of file
diff --git a/Engine/Plugins/Runtime/nDisplay/Source/DisplayCluster/Internal/Math/MarginSet.h b/Engine/Plugins/Runtime/nDisplay/Source/DisplayCluster/Internal/Math/MarginSet.h
new file mode 100644
--- /dev/null
+++ b/Engine/Plugins/Runtime/nDisplay/Source/DisplayCluster/Internal/Math/MarginSet.h
@@ -0,0 +1,967 @@
+// Copyright Epic Games, Inc. All Rights Reserved.
+
+#pragma once
+
+#include "CoreTypes.h"
+
+#include "Misc/AssertionMacros.h"
+#include "Math/IntPoint.h"
+#include "Math/UnrealMathUtility.h"
+#include "Math/IntRect.h"
+
+namespace UE::Math
+{
+
+/**
+* Enumerates the sides of a margin set by index.
+* Used for indexed access to individual sides (Left, Right, Top, Bottom).
+*/
+enum class EMarginSetSideIndex : uint8
+{
+	Left = 0,
+	Right = 1,
+	Top = 2,
+	Bottom = 3,
+};
+
+/**
+* TMarginSet is a template struct that stores a set of four directional values:
+* Left, Right, Top, and Bottom.
+*/
+template<typename T>
+struct alignas(16) TMarginSet
+{
+public:
+	/*
+	*    +--------------------------+
+	*    |           Top            |
+	*    |        +-------+         |
+	*    |   Left |       |  Right  |
+	*    |        +-------+         |
+	*    |         Bottom           |
+	*    +--------------------------+
+	*/
+	T Left;
+	T Right;
+	T Top;
+	T Bottom;
+
+private:
+	static constexpr int32 NumSides = 4;
+
+public:
+	/**
+	* Default constructor.
+	*
+	* The default margin size is zero on all four sides.
+	*/
+	TMarginSet()
+		: Left(T(0))
+		, Right(T(0))
+		, Top(T(0))
+		, Bottom(T(0))
+	{ }
+
+	/**
+	 * Constructs a floating-point TMarginSet from a single arithmetic value U.
+	 *
+	 * Enabled only when T is floating-point and U is a built-in numeric type.
+	 * Conversion is explicit if U is not implicitly convertible to T.
+	 */
+	template <typename U>
+	requires (std::is_floating_point_v<T> && std::is_arithmetic_v<std::remove_cvref_t<U>>)
+	explicit (!std::is_convertible_v<U, T>)
+	TMarginSet(const U InAll)
+		: Left(static_cast<T>(InAll))
+		, Right(static_cast<T>(InAll))
+		, Top(static_cast<T>(InAll))
+		, Bottom(static_cast<T>(InAll))
+	{ }
+
+	/**
+	 * Constructs an integral TMarginSet from a single integral value U.
+	 *
+	 * Enabled only when both T and U are integral types.
+	 * Conversion is explicit if U is not implicitly convertible to T.
+	 */
+	template <typename IntType>
+	requires (std::is_integral_v<T> && std::is_integral_v<IntType>)
+	explicit(!std::is_convertible_v<IntType, T>)
+	TMarginSet(const IntType InAll)
+		: Left(static_cast<T>(InAll))
+		, Right(static_cast<T>(InAll))
+		, Top(static_cast<T>(InAll))
+		, Bottom(static_cast<T>(InAll))
+	{ }
+
+	/**
+	 * Constructs a floating-point TMarginSet from four values of type U.
+	 *
+	 * Enabled only when T is floating-point and U is a built-in numeric type.
+	 * Conversion is explicit if U is not implicitly convertible to T.
+	 */
+	template <typename U>
+	requires (std::is_floating_point_v<T> && std::is_arithmetic_v<std::remove_cvref_t<U>>)
+	explicit (!std::is_convertible_v<U, T>)
+	TMarginSet(
+		const U InLeft,
+		const U InRight,
+		const U InTop,
+		const U InBottom)
+		: Left   (static_cast<T>(InLeft))
+		, Right  (static_cast<T>(InRight))
+		, Top    (static_cast<T>(InTop))
+		, Bottom (static_cast<T>(InBottom))
+	{ }
+
+	/**
+	 * Constructs an integral TMarginSet from four integral values of type U.
+	 *
+	 * Enabled only when both T and U are integral types.
+	 * Conversion is explicit if U is not implicitly convertible to T.
+	 */
+	template <typename IntType>
+	requires (std::is_integral_v<T> && std::is_integral_v<IntType>)
+	explicit (!std::is_convertible_v<IntType, T>)
+	TMarginSet(
+		const IntType InLeft,
+		const IntType InRight,
+		const IntType InTop,
+		const IntType InBottom)
+		: Left   (static_cast<T>(InLeft))
+		, Right  (static_cast<T>(InRight))
+		, Top    (static_cast<T>(InTop))
+		, Bottom (static_cast<T>(InBottom))
+	{ }
+
+	/**
+	 * Constructs a floating-point TMarginSet from another TMarginSet<U>.
+	 *
+	 * Enabled only when T is floating-point and U is a built-in numeric type.
+	 * Conversion is explicit if U is not implicitly convertible to T.
+	 */
+	template <typename U>
+	requires (std::is_floating_point_v<T> && std::is_arithmetic_v<std::remove_cvref_t<U>>)
+	explicit (!std::is_convertible_v<U, T>)
+	TMarginSet(const TMarginSet<U>& Other)
+		: Left   (static_cast<T>(Other.Left))
+		, Right  (static_cast<T>(Other.Right))
+		, Top    (static_cast<T>(Other.Top))
+		, Bottom (static_cast<T>(Other.Bottom))
+	{ }
+
+	/**
+	 * Constructs an integral TMarginSet from another TMarginSet<U>.
+	 *
+	 * Enabled only when both T and U are integral types.
+	 * Conversion is explicit if U is not implicitly convertible to T.
+	 */
+	template <typename IntType>
+	requires (std::is_integral_v<T> && std::is_integral_v<IntType>)
+	explicit(!std::is_convertible_v<IntType, T>)
+	TMarginSet(const TMarginSet<IntType>& Other)
+		: Left   (static_cast<T>(Other.Left))
+		, Right  (static_cast<T>(Other.Right))
+		, Top    (static_cast<T>(Other.Top))
+		, Bottom (static_cast<T>(Other.Bottom))
+	{ }
+
+	/**
+	 * Constructs a floating-point TMarginSet from a TVector4<U>.
+	 *
+	 * Enabled only when T is floating-point and U is a built-in numeric type.
+	 * Conversion is explicit if U is not implicitly convertible to T.
+	 */
+	template <typename U>
+	requires (std::is_floating_point_v<T> && std::is_arithmetic_v<std::remove_cvref_t<U>>)
+	explicit (!std::is_convertible_v<U, T>)
+	TMarginSet(const TVector4<U>& Vec)
+		: Left   (static_cast<T>(Vec.X))
+		, Right  (static_cast<T>(Vec.Y))
+		, Top    (static_cast<T>(Vec.Z))
+		, Bottom (static_cast<T>(Vec.W))
+	{ }
+
+	/**
+	 * Constructs a TMarginSet from a TIntVector4<U> with integral elements.
+	 *
+	 * Enabled only when U is an integral type.
+	 * Conversion is explicit if U is not implicitly convertible to T.
+	 */
+	template <typename IntType>
+	requires (std::is_integral_v<IntType>)
+	explicit (!std::is_convertible_v<IntType, T>)
+	TMarginSet(const TIntVector4<IntType>& Vec)
+		: Left   (static_cast<T>(Vec.X))
+		, Right  (static_cast<T>(Vec.Y))
+		, Top    (static_cast<T>(Vec.Z))
+		, Bottom (static_cast<T>(Vec.W))
+	{ }
+
+public:
+	/**
+	 * Converts a floating-point TMarginSet to TVector4<U>.
+	 *
+	 * Enabled only when T is floating-point.
+	 * Marked [[nodiscard]] to warn if the result is ignored.
+	 */
+	template<typename U = T>
+	requires (std::is_floating_point_v<U>)
+	[[nodiscard]]
+	inline operator TVector4<U>() const
+	{
+		return TVector4<U>(
+			Left,
+			Right,
+			Top,
+			Bottom);
+	}
+
+	/**
+	 * Converts this integral TMarginSet<T> to a TIntVector4<U>.
+	 *
+	 * Enabled only when both T (this type) and U (defaults to T) are integral.
+	 * [[nodiscard]] warns if the result is unused.
+	 */
+	template<typename IntType = T>
+	requires (std::is_integral_v<IntType> && std::is_integral_v<T>)
+	[[nodiscard]]
+	inline operator TIntVector4<IntType>() const
+	{
+		return TIntVector4<IntType>(
+			Left,
+			Right,
+			Top,
+			Bottom);
+	}
+
+public:
+	/**
+	 * Assigns margins from a TVector4<U> to a floating-point TMarginSet<T>.
+	 *
+	 * Enabled only when:
+	 * - T is floating-point
+	 * - U is a built-in numeric type
+	 * - U is convertible to T
+	 */
+	template <typename U>
+	requires (std::is_floating_point_v<T> &&
+		std::is_arithmetic_v<std::remove_cvref_t<U>> &&
+		std::is_convertible_v<U, T>)
+	inline TMarginSet& AssignMargins(const TVector4<U>& InVector)
+	{
+		Left   = static_cast<T>(InVector.X);
+		Right  = static_cast<T>(InVector.Y);
+		Top    = static_cast<T>(InVector.Z);
+		Bottom = static_cast<T>(InVector.W);
+
+		return *this;
+	}
+
+	/**
+	 * Assigns margins from a TIntVector4<IntType> to this TMarginSet<T>.
+	 *
+	 * Enabled only when:
+	 * - IntType is an integral type
+	 * - IntType is convertible to T
+	 *
+	 * @param InVector  The TIntVector4<IntType> containing the new margin values.
+	 * @return          Reference to this TMarginSet<T> after assignment.
+	 */
+	template <typename IntType>
+	requires (std::is_integral_v<IntType> && std::is_convertible_v<IntType, T>)
+	inline TMarginSet& AssignMargins(const TIntVector4<IntType>& InVector)
+	{
+		Left   = static_cast<T>(InVector.X);
+		Right  = static_cast<T>(InVector.Y);
+		Top    = static_cast<T>(InVector.Z);
+		Bottom = static_cast<T>(InVector.W);
+
+		return *this;
+	}
+
+	/**
+	 * Assigns margins from another TMarginSet<U> to a floating-point TMarginSet<T>.
+	 *
+	 * Enabled only when:
+	 * - T is floating-point
+	 * - U is a built-in numeric type
+	 * - U is convertible to T
+	 */
+	template <typename U>
+	requires (std::is_floating_point_v<T> &&
+		std::is_arithmetic_v<std::remove_cvref_t<U>> &&
+		std::is_convertible_v<U, T>)
+	inline TMarginSet<T>& AssignMargins(const TMarginSet<U>& Other)
+	{
+		Left   = static_cast<T>(Other.Left);
+		Right  = static_cast<T>(Other.Right);
+		Top    = static_cast<T>(Other.Top);
+		Bottom = static_cast<T>(Other.Bottom);
+
+		return *this;
+	}
+
+	/**
+	 * Assigns margins from another TMarginSet<IntType> to an integral TMarginSet<T>.
+	 *
+	 * Enabled only when:
+	 * - T and IntType are integral types
+	 * - IntType is convertible to T
+	 */
+	template <typename IntType>
+	requires (std::is_integral_v<T> &&
+		std::is_integral_v<IntType> &&
+		std::is_convertible_v<IntType, T>)
+	inline TMarginSet<T>& AssignMargins(const TMarginSet<IntType>& Other)
+	{
+		Left   = static_cast<T>(Other.Left);
+		Right  = static_cast<T>(Other.Right);
+		Top    = static_cast<T>(Other.Top);
+		Bottom = static_cast<T>(Other.Bottom);
+
+		return *this;
+	}
+
+	/**
+	 * Assigns margins from four numeric values of type U to a floating-point TMarginSet<T>.
+	 *
+	 * Enabled only when:
+	 * - T is floating-point
+	 * - U is a built-in numeric type
+	 * - U is convertible to T
+	 */
+	template <typename U>
+	requires (std::is_floating_point_v<T> &&
+		std::is_arithmetic_v<std::remove_cvref_t<U>> &&
+		std::is_convertible_v<U, T>)
+	inline TMarginSet<T>& AssignMargins(const U InLeft, const U InRight, const  U InTop, const U InBottom)
+	{
+		Left   = static_cast<T>(InLeft);
+		Right  = static_cast<T>(InRight);
+		Top    = static_cast<T>(InTop);
+		Bottom = static_cast<T>(InBottom);
+
+		return *this;
+	}
+
+	/**
+	 * Assigns margins from four integral values of type IntType to an integral TMarginSet<T>.
+	 *
+	 * Enabled only when:
+	 * - T and IntType are integral types
+	 * - IntType is convertible to T
+	 */
+	template <typename IntType>
+	requires (std::is_integral_v<T> &&
+		std::is_integral_v<IntType> &&
+		std::is_convertible_v<IntType, T>)
+	inline TMarginSet<T>& AssignMargins(
+		const IntType InLeft,
+		const IntType InRight,
+		const IntType InTop,
+		const IntType InBottom)
+	{
+		Left   = static_cast<T>(InLeft);
+		Right  = static_cast<T>(InRight);
+		Top    = static_cast<T>(InTop);
+		Bottom = static_cast<T>(InBottom);
+
+		return *this;
+	}
+
+public:
+	/**
+	 * Multiplies each margin of a TMarginSet<T> by an integral scalar value.
+	 *
+	 * @param InScale  Integral multiplier applied to all margins.
+	 * @return         A TMarginSet<T> where each margin is scaled by InScale.
+	 */
+	template <typename IntType>
+	requires (std::is_integral_v<IntType> &&
+		std::is_convertible_v<IntType, T>)
+	[[nodiscard]]
+	inline TMarginSet<T> operator*(const IntType InScale) const
+	{
+		const T Scale = static_cast<T>(InScale);
+		return TMarginSet<T>(
+			Left   * Scale,
+			Right  * Scale,
+			Top    * Scale,
+			Bottom * Scale
+		);
+	}
+
+	/**
+	 * Multiplies each margin of a floating-point TMarginSet<T> by a scalar value.
+	 *
+	 * @param InScale  Scalar multiplier applied to all margins.
+	 * @return         A TMarginSet<T> where each margin is scaled by InScale.
+	 */
+	template <typename U>
+	requires (std::is_floating_point_v<T> &&
+		!std::is_integral_v<U> && 
+		std::is_arithmetic_v<std::remove_cvref_t<U>>)
+	[[nodiscard]]
+	inline TMarginSet<T> operator*(const U InScale) const
+	{
+		const T Scale = static_cast<T>(InScale);
+		return TMarginSet<T>(
+			Left   * Scale,
+			Right  * Scale,
+			Top    * Scale,
+			Bottom * Scale
+		);
+	}
+
+	/**
+	 * Multiplies each margin of a TMarginSet<T> by the corresponding axis of a TIntPoint<IntType>.
+	 * - Left and Right are multiplied by Scale.X
+	 * - Top and Bottom are multiplied by Scale.Y
+	 *
+	 * @param Scale   The TIntPoint<IntType> scale factors (X for horizontal, Y for vertical).
+	 * @return        A TMarginSet<T> where each margin is multiplied by the corresponding factor.
+	 */
+	template <typename IntType>
+	requires (std::is_integral_v<IntType>&&
+		std::is_convertible_v<IntType, T>)
+	[[nodiscard]]
+	inline TMarginSet<T> operator*(const TIntPoint<IntType>& Scale) const
+	{
+		const T ScaleX = static_cast<T>(Scale.X);
+		const T ScaleY = static_cast<T>(Scale.Y);
+
+		return TMarginSet<T>(
+			Left   * ScaleX,
+			Right  * ScaleX,
+			Top    * ScaleY,
+			Bottom * ScaleY
+		);
+	}
+
+	/**
+	 * Multiplies each margin of a TMarginSet<T> by the corresponding axis of a TIntVector2<IntType>.
+	 * - Left and Right are multiplied by Scale.X
+	 * - Top and Bottom are multiplied by Scale.Y
+	 *
+	 * @param Scale   The TIntVector2<IntType> scale factors (X for horizontal, Y for vertical).
+	 * @return        A TMarginSet<T> where each margin is scaled by the corresponding factor.
+	 */
+	template <typename IntType>
+	requires (std::is_integral_v<IntType> &&
+		std::is_convertible_v<IntType, T>)
+	[[nodiscard]]
+	inline TMarginSet<T> operator*(const TIntVector2<IntType>& Scale) const
+	{
+		const T ScaleX = static_cast<T>(Scale.X);
+		const T ScaleY = static_cast<T>(Scale.Y);
+
+		return TMarginSet<T>(
+			Left   * ScaleX,
+			Right  * ScaleX,
+			Top    * ScaleY,
+			Bottom * ScaleY
+		);
+	}
+
+	/**
+	 * Multiplies each margin of a TMarginSet<T> by the corresponding axis of a TVector2<U>.
+	 * - Left and Right are multiplied by Scale.X
+	 * - Top and Bottom are multiplied by Scale.Y
+	 *
+	 * @param Scale   The TVector2<U> scale factors (X for horizontal, Y for vertical).
+	 * @return        A TMarginSet<T> where each margin is scaled by the corresponding factor.
+	 */
+	template <typename U>
+	requires (std::is_floating_point_v<T> &&
+		std::is_floating_point_v<U> &&
+		std::is_convertible_v<U, T>)
+	[[nodiscard]]
+	inline TMarginSet<T> operator*(const TVector2<U>& Scale) const
+	{
+		const T ScaleX = static_cast<T>(Scale.X);
+		const T ScaleY = static_cast<T>(Scale.Y);
+
+		return TMarginSet<T>(
+			Left   * ScaleX,
+			Right  * ScaleX,
+			Top    * ScaleY,
+			Bottom * ScaleY
+		);
+	}
+
+public:
+	/**
+	 * Divides each margin of a floating-point TMarginSet<T> by a scalar value.
+	 *
+	 * @param InDivisor  Scalar divisor applied to all margins.
+	 * @return           A TMarginSet<T> where each margin is divided by InDivisor.
+	 */
+	template <typename U>
+	requires (std::is_floating_point_v<T> &&
+		std::is_arithmetic_v<std::remove_cvref_t<U>>)
+	[[nodiscard]]
+	inline TMarginSet<T> operator/(const U InDivisor) const
+	{
+		const T Divisor = static_cast<T>(InDivisor);
+
+		// Assumes caller ensures Divisor are nonzero
+		check(Divisor);
+
+		return UE::Math::TMarginSet<T>(
+			Left   / Divisor,
+			Right  / Divisor,
+			Top    / Divisor,
+			Bottom / Divisor
+		);
+	}
+	
+	/**
+	 * Divides each margin of a TMarginSet<T> by the corresponding axis of a TIntPoint<IntType>.
+	 * - Left and Right are divided by Divisor.X
+	 * - Top and Bottom are divided by Divisor.Y
+	 *
+	 * @param Divisor  The TIntPoint<IntType> divisors (X for horizontal, Y for vertical).
+	 * @return         A TMarginSet<T> where each margin is divided by the corresponding factor.
+	 */
+	template <typename IntType>
+	requires (std::is_floating_point_v<T>&& 
+		std::is_integral_v<IntType>&&
+		std::is_convertible_v<IntType, T>)
+	[[nodiscard]]
+	inline TMarginSet<T> operator/(const TIntPoint<IntType>& Divisor) const
+	{
+		const T DivisorX = static_cast<T>(Divisor.X);
+		const T DivisorY = static_cast<T>(Divisor.Y);
+
+		// Assumes caller ensures DivisorX and DivisorY are nonzero
+		check(DivisorX && DivisorY);
+
+		return TMarginSet<T>(
+			Left   / DivisorX,
+			Right  / DivisorX,
+			Top    / DivisorY,
+			Bottom / DivisorY
+		);
+	}
+
+	/**
+	 * Multiplies each margin of a TMarginSet<T> by the corresponding axis of a TIntVector2<IntType>.
+	 * - Left and Right are multiplied by Divisor.X
+	 * - Top and Bottom are multiplied by Divisor.Y
+	 *
+	 * @param Divisor  The TIntVector2<IntType> factors (X for horizontal, Y for vertical).
+	 * @return         A TMarginSet<T> where each margin is multiplied by the corresponding factor.
+	 */
+	template <typename IntType>
+	requires (std::is_floating_point_v<T> && 
+		std::is_integral_v<IntType> &&
+		std::is_convertible_v<IntType, T>)
+	[[nodiscard]]
+	inline TMarginSet<T> operator*(const TIntVector2<IntType>& Divisor) const
+	{
+		const T DivisorX = static_cast<T>(Divisor.X);
+		const T DivisorY = static_cast<T>(Divisor.Y);
+
+		// Assumes caller ensures DivisorX and DivisorY are nonzero
+		check(DivisorX && DivisorY);
+
+		return TMarginSet<T>(
+			Left   / DivisorX,
+			Right  / DivisorX,
+			Top    / DivisorY,
+			Bottom / DivisorY
+		);
+	}
+
+	/**
+	 * Divides each margin of a TMarginSet<T> by the corresponding axis of an FIntPoint.
+	 * - Left/Right are divided by Divisor.X
+	 * - Top/Bottom are divided by Divisor.Y
+	 *
+	 * @param Divisor       The FIntPoint divisor (X for horizontal, Y for vertical).
+	 * @return              A TMarginSet<T> where each margin is divided by the corresponding divisor.
+	 */
+	template <typename U>
+	requires (std::is_floating_point_v<T>&&
+		std::is_floating_point_v<U>&&
+		std::is_convertible_v<U, T>)
+	[[nodiscard]]
+	inline TMarginSet<T> operator/(const TVector2<U>& Divisor)
+	{
+		const T DivisorX = static_cast<T>(Divisor.X);
+		const T DivisorY = static_cast<T>(Divisor.Y);
+
+		// Assumes caller ensures Divisor.X and Divisor.Y are nonzero
+		check(DivisorX && DivisorY);
+
+		return UE::Math::TMarginSet<T>(
+			Left   / DivisorX,
+			Right  / DivisorX,
+			Top    / DivisorY,
+			Bottom / DivisorY
+		);
+	}
+
+public:
+	/**
+	 * Expands the given rectangle in-place by the values in this margin set.
+	 *
+	 * Each margin value increases the corresponding edge outward:
+	 * - Left decreases Min.X
+	 * - Right increases Max.X
+	 * - Top decreases Min.Y
+	 * - Bottom increases Max.Y
+	 *
+	 * This operation can result in negative-sized rectangles if large margins are used.
+	 *
+	 * @param InOutRect [in,out] The rectangle to expand. Modified in place.
+	 */
+	template <typename IntType>
+	requires (std::is_integral_v<T> &&
+		std::is_integral_v<IntType> &&
+		std::is_convertible_v<IntType, T>)
+	void ApplyExpandToRect(TIntRect<IntType>& InOutRect) const
+	{
+		InOutRect.Min.X -= Left;
+		InOutRect.Max.X += Right;
+		InOutRect.Min.Y -= Top;
+		InOutRect.Max.Y += Bottom;
+
+		// Clamp to prevent negative size (Min cannot exceed Max).
+		InOutRect.Min.X = FMath::Min(InOutRect.Min.X, InOutRect.Max.X);
+		InOutRect.Min.Y = FMath::Min(InOutRect.Min.Y, InOutRect.Max.Y);
+	}
+
+	/**
+	 * Shrinks the given rectangle in-place by the values in this margin set.
+	 * Each margin value moves the corresponding edge inward:
+	 *   - Left   increases Min.X
+	 *   - Right  decreases Max.X
+	 *   - Top    increases Min.Y
+	 *   - Bottom decreases Max.Y
+	 *
+	 * If the total margin for any axis exceeds the rectangle's size, the rectangle will collapse to a line or point.
+	 * Margin values are clamped as needed to prevent the rectangle from inverting or collapsing to negative size.
+	 *
+	 * @param InOutRect [in,out] The rectangle to be contracted. Modified in place.
+	 */
+	template <typename IntType>
+		requires (std::is_integral_v<T>&&
+		std::is_integral_v<IntType>&&
+		std::is_convertible_v<IntType, T>)
+	void ApplyClampedInsetToRect(TIntRect<IntType>& InOutRect) const
+	{
+		const FIntPoint RectSize = InOutRect.Size();
+
+		// Clamp margins to prevent invalid rectangle
+		T ClampedLeft   = FMath::Clamp(Left,   0, RectSize.X);
+		T ClampedRight  = FMath::Clamp(Right,  0, RectSize.X - ClampedLeft);
+		T ClampedTop    = FMath::Clamp(Top,    0, RectSize.Y);
+		T ClampedBottom = FMath::Clamp(Bottom, 0, RectSize.Y - ClampedTop);
+
+		// Shrink rect by clamped margins
+		InOutRect.Min.X += ClampedLeft;
+		InOutRect.Max.X -= ClampedRight;
+		InOutRect.Min.Y += ClampedTop;
+		InOutRect.Max.Y -= ClampedBottom;
+
+		// Ensure Min does not surpass Max (should already be safe)
+		InOutRect.Min.X = FMath::Min(InOutRect.Min.X, InOutRect.Max.X);
+		InOutRect.Min.Y = FMath::Min(InOutRect.Min.Y, InOutRect.Max.Y);
+	}
+
+public:
+	/**
+	 * Provides indexed access to a margin side value.
+	 *
+	 * @param SideIndex  Enum value specifying which margin side to access.
+	 * @return           Reference to the corresponding side value.
+	 */
+	[[nodiscard]] inline T& operator[](EMarginSetSideIndex SideIndex)
+	{
+		switch (SideIndex)
+		{
+		case EMarginSetSideIndex::Left:   return Left;
+		case EMarginSetSideIndex::Right:  return Right;
+		case EMarginSetSideIndex::Top:    return Top;
+		case EMarginSetSideIndex::Bottom: return Bottom;
+		default:			
+			break;
+		}
+
+		// Unreachable
+		checkNoEntry();
+
+		return Left;
+	}
+
+	/**
+	 * Provides  read-only indexed access to a margin side value.
+	 *
+	 * @param SideIndex  Enum value specifying which margin side to access.
+	 * @return           Const reference to the corresponding side value.
+	 */
+	[[nodiscard]] inline const T& operator[](EMarginSetSideIndex SideIndex) const
+	{
+		switch (SideIndex)
+		{
+		case EMarginSetSideIndex::Left:   return Left;
+		case EMarginSetSideIndex::Right:  return Right;
+		case EMarginSetSideIndex::Top:    return Top;
+		case EMarginSetSideIndex::Bottom: return Bottom;
+		default:
+			break;
+		}
+
+		// Unreachable
+		checkNoEntry();
+
+		return Left;
+	}
+
+	/**
+	 * Provides indexed access to a margin side value by numeric index.
+	 *
+	 * @param Index  Zero-based index of the side to access:
+	 *               0 = Left, 1 = Top, 2 = Right, 3 = Bottom.
+	 * @return       Reference to the corresponding side value.
+	 */
+	[[nodiscard]] inline T& operator[](int32 Index)
+	{	
+		check(Index >= 0 && Index < NumSides);
+		const EMarginSetSideIndex SideIndex = static_cast<EMarginSetSideIndex>(Index);
+
+		return operator[](SideIndex);
+	}
+
+	/**
+	 * Provides  read-only indexed access to a margin side value by numeric index.
+	 *
+	 * @param Index  Zero-based index of the side to access:
+	 *               0 = Left, 1 = Top, 2 = Right, 3 = Bottom.
+	 * @return       Const reference to the corresponding side value.
+	 */
+	[[nodiscard]] inline const T& operator[](int32 Index) const
+	{
+		check(Index >= 0 && Index < NumSides);
+		const EMarginSetSideIndex SideIndex = static_cast<EMarginSetSideIndex>(Index);
+
+		return operator[](SideIndex);
+	}
+
+	/** Equality operator */
+	[[nodiscard]] inline bool operator==(const TMarginSet& Other) const
+	{
+		return Left == Other.Left && Right == Other.Right &&
+			Top == Other.Top && Bottom == Other.Bottom;
+	}
+
+	/** Inequality operator */
+	[[nodiscard]] inline bool operator!=(const TMarginSet& Other) const
+	{
+		return !(*this == Other);
+	}
+
+public:
+	/**
+	* Returns a copy of this margin set where all margins floors to the nearest lower int32.
+	* 
+	* @return A TMarginSet<int32> with each value floored.
+	*/
+	[[nodiscard]] inline TMarginSet<int32> GetFloorToInt() const
+	{
+		return TMarginSet<int32>(
+			FMath::FloorToInt32(Left),
+			FMath::FloorToInt32(Right),
+			FMath::FloorToInt32(Top),
+			FMath::FloorToInt32(Bottom)
+		);
+	}
+
+	/**
+	* Returns a copy of this margin set where all margins rounds to the nearest int32.
+	* 
+	* @return A TMarginSet<int32> with all values rounded.
+	*/
+	[[nodiscard]] inline TMarginSet<int32> GetRoundToInt() const
+	{
+		return TMarginSet<int32>(
+			FMath::RoundToInt32(Left),
+			FMath::RoundToInt32(Right),
+			FMath::RoundToInt32(Top),
+			FMath::RoundToInt32(Bottom)
+		);
+	}
+
+	/**
+	* Returns a copy of this margin set where all margins ceils to the nearest greater or equal int32.
+	* 
+	* @return A TMarginSet<int32> with each value ceiled.
+	*/
+	[[nodiscard]] inline TMarginSet<int32> GetCeilToInt() const
+	{
+		return TMarginSet<int32>(
+			FMath::CeilToInt32(Left),
+			FMath::CeilToInt32(Right),
+			FMath::CeilToInt32(Top),
+			FMath::CeilToInt32(Bottom)
+		);
+	}
+
+public:
+	/**
+	* Returns a copy of this margin set where all margins clamped the specified [Min, Max] range.
+	*
+	* @param Min           The minimum allowed value for each margin.
+	* @param Max           The maximum allowed value for each margin.
+	* @return              A clamped TMarginSet with all values between Min and Max.
+	*/
+	[[nodiscard]] inline TMarginSet<T> GetClamped(const T Min, const T Max) const
+	{
+		return TMarginSet<T>(
+			FMath::Clamp(Left,   Min, Max),
+			FMath::Clamp(Right,  Min, Max),
+			FMath::Clamp(Top,    Min, Max),
+			FMath::Clamp(Bottom, Min, Max)
+		);
+	}
+
+	/**
+	 * Returns a copy of this margin set with each side individually clamped
+	 * between the corresponding sides of the provided minimum and maximum margins.
+	 *
+	 * For each side:
+	 *   - Left   = Clamp(Left,   MinMargins.Left,   MaxMargins.Left)
+	 *   - Right  = Clamp(Right,  MinMargins.Right,  MaxMargins.Right)
+	 *   - Top    = Clamp(Top,    MinMargins.Top,    MaxMargins.Top)
+	 *   - Bottom = Clamp(Bottom, MinMargins.Bottom, MaxMargins.Bottom)
+	 *
+	 * @param MinMargins Margin set defining the minimum values per side.
+	 * @param MaxMargins Margin set defining the maximum values per side.
+	 * @return A new TMarginSet with each side clamped to its respective range.
+	 */
+	[[nodiscard]] inline TMarginSet<T> GetClampedPerSide(const TMarginSet<T>& MinMargins, const TMarginSet& MaxMargins) const
+	{
+		return TMarginSet<T>(
+			FMath::Clamp(Left,   MinMargins.Left,   MaxMargins.Left),
+			FMath::Clamp(Right,  MinMargins.Right,  MaxMargins.Right),
+			FMath::Clamp(Top,    MinMargins.Top,    MaxMargins.Top),
+			FMath::Clamp(Bottom, MinMargins.Bottom, MaxMargins.Bottom)
+		);
+	}
+
+	/**
+	* Returns true if all margins are zero.
+	*/
+	[[nodiscard]] inline bool IsZero() const
+	{
+		return Left == T(0)
+			&& Right == T(0)
+			&& Top == T(0)
+			&& Bottom == T(0);
+	}
+
+	/**
+	 * Returns the minimum value among all four margins (Left, Right, Top, Bottom).
+	 *
+	 * @return Smallest margin value.
+	 */
+	[[nodiscard]] inline T GetMin() const
+	{
+		const T Value1 = FMath::Min(Left, Right);
+		const T Value2 = FMath::Min(Top, Bottom);
+
+		return FMath::Min(Value1, Value2);
+	}
+
+	/**
+	 * Returns the maximum value among all four margins (Left, Right, Top, Bottom).
+	 *
+	 * @return Largest margin value.
+	 */
+	[[nodiscard]] inline T GetMax() const
+	{
+		const T Value1 = FMath::Max(Left, Right);
+		const T Value2 = FMath::Max(Top, Bottom);
+
+		return FMath::Max(Value1, Value2);
+	}
+
+	/**
+	 * Returns a copy of this margin set where all margins is the minimum of this margin and another.
+	 *
+	 * @param In  The margin set to compare against
+	 * @return    A new TMarginSet with the smallest values for each side
+	 */
+	[[nodiscard]] inline TMarginSet<T> GetMinPerSide(const TMarginSet<T>& In) const
+	{
+		return TMarginSet<T>(
+			FMath::Min(Left,   In.Left),   // Minimum of Left sides
+			FMath::Min(Right,  In.Right),  // Minimum of Right sides
+			FMath::Min(Top,    In.Top),    // Minimum of Top sides
+			FMath::Min(Bottom, In.Bottom)  // Minimum of Bottom sides
+		);
+	}
+
+	/**
+	 * Returns a copy of this margin set where all margins is the maximum of this margin and another.
+	 *
+	 * @param In  The margin set to compare against
+	 * @return    A new TMarginSet with the largest values for each side
+	 */
+	[[nodiscard]] inline TMarginSet<T> GetMaxPerSide(const TMarginSet<T>& In) const
+	{
+		return TMarginSet<T>(
+			FMath::Max(Left,   In.Left),   // Maximum of Left sides
+			FMath::Max(Right,  In.Right),  // Maximum of Right sides
+			FMath::Max(Top,    In.Top),    // Maximum of Top sides
+			FMath::Max(Bottom, In.Bottom)  // Maximum of Bottom sides
+		);
+	}
+
+	/**
+	 * Returns a new margin set where each side is the maximum of:
+	 * - the current value, or
+	 * - the input value, if it is non-zero (i.e., treated as defined).
+	 *
+	 * This allows selectively overriding margins while skipping undefined (zero) inputs.
+	 *
+	 * @param In The input margin set, where zero values are treated as "undefined".
+	 * @return A new margin set with per-side conditional maximums.
+	 */
+	[[nodiscard]] inline TMarginSet<T> GetMaxPerSideIfDefined(const TMarginSet<T>& In) const
+	{
+		return TMarginSet<T>(
+			In.Left   != static_cast<T>(0) ? FMath::Max(Left,   In.Left)   : Left,
+			In.Right  != static_cast<T>(0) ? FMath::Max(Right,  In.Right)  : Right,
+			In.Top    != static_cast<T>(0) ? FMath::Max(Top,    In.Top)    : Top,
+			In.Bottom != static_cast<T>(0) ? FMath::Max(Bottom, In.Bottom) : Bottom
+		);
+	}
+
+	/**
+	 * Returns a new margin set where each side is the minimum of:
+	 * - the current value, or
+	 * - the input value, if it is non-zero (i.e., treated as defined).
+	 *
+	 * This allows selectively overriding margins while skipping undefined (zero) inputs.
+	 *
+	 * @param In The input margin set, where zero values are treated as "undefined".
+	 * @return A new margin set with per-side conditional minimum.
+	 */
+	[[nodiscard]] inline TMarginSet<T> GetMinPerSideIfDefined(const TMarginSet<T>& In) const
+	{
+		return TMarginSet<T>(
+			In.Left   != static_cast<T>(0) ? FMath::Min(Left,   In.Left)   : Left,
+			In.Right  != static_cast<T>(0) ? FMath::Min(Right,  In.Right)  : Right,
+			In.Top    != static_cast<T>(0) ? FMath::Min(Top,    In.Top)    : Top,
+			In.Bottom != static_cast<T>(0) ? FMath::Min(Bottom, In.Bottom) : Bottom
+		);
+	}
+};
+
+} // namespace UE::Math
+
+UE_DECLARE_LWC_TYPE(MarginSet, 4);
+using FIntMarginSet = UE::Math::TMarginSet<int32>;
diff --git a/Engine/Plugins/Runtime/nDisplay/Source/DisplayCluster/Private/Render/Viewport/Configuration/DisplayClusterViewportConfigurationHelpers.cpp b/Engine/Plugins/Runtime/nDisplay/Source/DisplayCluster/Private/Render/Viewport/Configuration/DisplayClusterViewportConfigurationHelpers.cpp
--- a/Engine/Plugins/Runtime/nDisplay/Source/DisplayCluster/Private/Render/Viewport/Configuration/DisplayClusterViewportConfigurationHelpers.cpp
+++ b/Engine/Plugins/Runtime/nDisplay/Source/DisplayCluster/Private/Render/Viewport/Configuration/DisplayClusterViewportConfigurationHelpers.cpp
@@ -53,27 +53,29 @@ FDisplayClusterViewport_OverscanSettings FDisplayClusterViewportConfigurationHel
 
 	if (InOverscan.bEnabled)
 	{
+		const FMarginSet OverscanMargins(
+			InOverscan.Left,
+			InOverscan.Right,
+			InOverscan.Top,
+			InOverscan.Bottom
+		);
+
 		switch (InOverscan.Mode)
 		{
 		case EDisplayClusterConfigurationViewportOverscanMode::Percent:
 			OutOverscanSettings.bEnabled = InOverscan.bEnabled;
 			OutOverscanSettings.Unit = EDisplayClusterViewport_FrustumUnit::Percent;
-
+			
 			// Scale 0..100% to 0..1 range
-			OutOverscanSettings.Left = .01f * InOverscan.Left;
-			OutOverscanSettings.Right = .01f * InOverscan.Right;
-			OutOverscanSettings.Top = .01f * InOverscan.Top;
-			OutOverscanSettings.Bottom = .01f * InOverscan.Bottom;
+			OutOverscanSettings.AssignMargins(OverscanMargins * .01f);
+
 			break;
 
 		case EDisplayClusterConfigurationViewportOverscanMode::Pixels:
 			OutOverscanSettings.bEnabled = InOverscan.bEnabled;
 			OutOverscanSettings.Unit = EDisplayClusterViewport_FrustumUnit::Pixels;
 
-			OutOverscanSettings.Left = InOverscan.Left;
-			OutOverscanSettings.Right = InOverscan.Right;
-			OutOverscanSettings.Top = InOverscan.Top;
-			OutOverscanSettings.Bottom = InOverscan.Bottom;
+			OutOverscanSettings.AssignMargins(OverscanMargins);
 			break;
 
 		default:
diff --git a/Engine/Plugins/Runtime/nDisplay/Source/DisplayCluster/Private/Render/Viewport/Configuration/DisplayClusterViewportConfigurationHelpers_Tile.cpp b/Engine/Plugins/Runtime/nDisplay/Source/DisplayCluster/Private/Render/Viewport/Configuration/DisplayClusterViewportConfigurationHelpers_Tile.cpp
--- a/Engine/Plugins/Runtime/nDisplay/Source/DisplayCluster/Private/Render/Viewport/Configuration/DisplayClusterViewportConfigurationHelpers_Tile.cpp
+++ b/Engine/Plugins/Runtime/nDisplay/Source/DisplayCluster/Private/Render/Viewport/Configuration/DisplayClusterViewportConfigurationHelpers_Tile.cpp
@@ -18,10 +18,6 @@
 #include "DisplayClusterConfigurationTypes_Tile.h"
 #include "DisplayClusterProjectionStrings.h"
 
-////////////////////////////////////////////////////////////////////////
-// FDisplayClusterViewportConfigurationHelpers_Tile
-////////////////////////////////////////////////////////////////////////
-
 void FDisplayClusterViewportConfigurationHelpers_Tile::UpdateICVFXCameraViewportTileSettings(FDisplayClusterViewport& InSourceViewport, const FDisplayClusterConfigurationMediaICVFX& InCameraMediaSettings)
 {
 	if (!InSourceViewport.CanSplitIntoTiles() || !InCameraMediaSettings.ShouldMediaICVFXSplitIntoTiles())
@@ -31,16 +27,30 @@ void FDisplayClusterViewportConfigurationHelpers_Tile::UpdateICVFXCameraViewport
 		return;
 	}
 
-	// Get additional tile flags
-	const EDisplayClusterViewportTileFlags TileFlags = InCameraMediaSettings.GetMediaICVFXTileFlags(InSourceViewport.GetClusterNodeId());
-
 	// Generate overscan settings
 	const FDisplayClusterViewport_OverscanSettings OverscanSettings = FDisplayClusterViewportConfigurationHelpers_Tile::GetTileOverscanSettings(InCameraMediaSettings.TileOverscan);
 
+	// Get additional tile flags
+	EDisplayClusterViewportTileFlags TileFlags = InCameraMediaSettings.GetMediaICVFXTileFlags(InSourceViewport.GetClusterNodeId());
+
+	if (InCameraMediaSettings.TileOverscan.bOptimizeTileOverscan)
+	{
+		TileFlags |= EDisplayClusterViewportTileFlags::OptimizeTileOverscan;
+	}
+
+	if (OverscanSettings.bApplyTileEdgeBlend)
+	{
+		TileFlags |= EDisplayClusterViewportTileFlags::ApplyTileEdgeBlend;
+
+		if (OverscanSettings.bUseEntireOverscanForEdgeBlend)
+		{
+			TileFlags |= EDisplayClusterViewportTileFlags::UseEntireOverscanForEdgeBlend;
+		}
+	}
+
 	// Set this viewport as a source for tile rendering.
 	FDisplayClusterViewport_TileSettings& OutTileSettings = InSourceViewport.GetRenderSettingsImpl().TileSettings;
 	OutTileSettings = FDisplayClusterViewport_TileSettings(InCameraMediaSettings.TiledSplitLayout, OverscanSettings, TileFlags);
-	OutTileSettings.bOptimizeTileOverscan = InCameraMediaSettings.TileOverscan.bOptimizeTileOverscan;
 }
 
 FIntRect FDisplayClusterViewportConfigurationHelpers_Tile::GetDestRect(const FDisplayClusterViewport_TileSettings& InTileSettings, const FIntRect& InSourceRect)
@@ -165,28 +175,48 @@ FDisplayClusterViewport* FDisplayClusterViewportConfigurationHelpers_Tile::GetOr
 		// Disables the custom frustum for the tile viewport (it has already been applied to the original viewport).
 		InOutRenderSettings.CustomFrustumSettings = FDisplayClusterViewport_CustomFrustumSettings();
 
-		// Disables the overscan for the tile viewport (it has already been applied to the original viewport).
-		// Override by the overscan settings obtained by splitting the tile.
+		// Initialize from the original viewport tile overscan settings, then apply per-tile adjustments.
 		InOutRenderSettings.OverscanSettings = SourceRenderSettings.TileSettings.GetOverscanSettings();
 
-		// Optimize overscan values for edge tiles
-		if (SourceRenderSettings.TileSettings.bOptimizeTileOverscan)
+		// OverscanSettings Per-tile adjustments: Optimize overscan values for edge tiles.
+		// Overscan is not applied on the outer edges of the image — only where tiles meet.
+		// Example layout (tiles with overscan overlap marked as "++"):
+		//
+		//  +------------+-----------+
+		//  | Tile0_0  ++|++ Tile1_0 |
+		//  |     ++     |     ++    |
+		//  +------------+-----------+
+		//  |     ++     |     ++    |
+		//  | Tile0_1  ++|++ Tile1_1 |
+		//  +------------+-----------+
 		{
+			using EMarginSide = UE::Math::EMarginSetSideIndex;
+			auto OptimizeSize = [&](const EMarginSide InSideIndex)
+				{
+					if (SourceRenderSettings.TileSettings.HasAnyTileFlags(EDisplayClusterViewportTileFlags::OptimizeTileOverscan))
+					{
+						InOutRenderSettings.OverscanSettings[InSideIndex] = 0;
+					}
+
+					// Always optimize ExtraPixels
+					InOutRenderSettings.OverscanSettings.ExtraPixels[InSideIndex] = 0;
+				};
+
 			if (InTilePos.X == 0)
 			{
-				InOutRenderSettings.OverscanSettings.Left = 0;
+				OptimizeSize(EMarginSide::Left);
 			}
 			if (InTilePos.Y == 0)
 			{
-				InOutRenderSettings.OverscanSettings.Top = 0;
+				OptimizeSize(EMarginSide::Top);
 			}
 			if (InTilePos.X == (InTileSize.X - 1))
 			{
-				InOutRenderSettings.OverscanSettings.Right = 0;
+				OptimizeSize(EMarginSide::Right);
 			}
 			if (InTilePos.Y == (InTileSize.Y - 1))
 			{
-				InOutRenderSettings.OverscanSettings.Bottom = 0;
+				OptimizeSize(EMarginSide::Bottom);
 			}
 		}
 
@@ -229,32 +259,30 @@ FDisplayClusterViewport_OverscanSettings FDisplayClusterViewportConfigurationHel
 {
 	FDisplayClusterViewport_OverscanSettings OutOverscanSettings;
 
-	OutOverscanSettings.bEnabled = false;
-	OutOverscanSettings.bOversize = InTileOverscan.bOversize;
+	if (InTileOverscan.bOversize)
+	{
+		OutOverscanSettings.bOversize = true;
+	}
 
 	if (InTileOverscan.bEnabled)
 	{
 		switch (InTileOverscan.Mode)
 		{
 		case EDisplayClusterConfigurationViewportOverscanMode::Percent:
-			OutOverscanSettings.bEnabled = InTileOverscan.bEnabled;
+			OutOverscanSettings.bEnabled = true;
 			OutOverscanSettings.Unit = EDisplayClusterViewport_FrustumUnit::Percent;
 
 			// Scale 0..100% to 0..1 range
-			OutOverscanSettings.Left =
-			OutOverscanSettings.Right = 
-			OutOverscanSettings.Top = 
-			OutOverscanSettings.Bottom = .01f * InTileOverscan.AllSides;
+			OutOverscanSettings.AssignMargins(FMarginSet(.01f * InTileOverscan.AllSides));
+
 			break;
 
 		case EDisplayClusterConfigurationViewportOverscanMode::Pixels:
-			OutOverscanSettings.bEnabled = InTileOverscan.bEnabled;
+			OutOverscanSettings.bEnabled = true;
 			OutOverscanSettings.Unit = EDisplayClusterViewport_FrustumUnit::Pixels;
 
-			OutOverscanSettings.Left =
-			OutOverscanSettings.Right =
-			OutOverscanSettings.Top =
-			OutOverscanSettings.Bottom = InTileOverscan.AllSides;
+			OutOverscanSettings.AssignMargins(FMarginSet(InTileOverscan.AllSides));
+
 			break;
 
 		default:
@@ -262,5 +290,37 @@ FDisplayClusterViewport_OverscanSettings FDisplayClusterViewportConfigurationHel
 		}
 	}
 
+	// The tile edge blending feature now requires overscan to be enabled.
+	// This logic should be updated if FDisplayClusterConfigurationTile_Overscan rules change.
+	if (!OutOverscanSettings.bEnabled)
+	{
+		return OutOverscanSettings;
+	}
+
+	// Configure tile edge blending feature
+	if (InTileOverscan.bBlendEntireOverscan)
+	{
+		OutOverscanSettings.bApplyTileEdgeBlend = true;
+		OutOverscanSettings.bUseEntireOverscanForEdgeBlend = true;
+	}
+	// Set the minimum overscan values in pixels for tile blending
+	else if (InTileOverscan.OverscanBlendWidth > 0)
+	{
+		// The number of pixels in the blending area cannot be greater than this number.
+		constexpr int32 MaxPixelsToBlend = 512;
+
+		// Limit CVar value to safe range.
+		const int32 PixelsToBlend = FMath::Clamp(InTileOverscan.OverscanBlendWidth, 0, MaxPixelsToBlend);
+
+		// Apply tile edge blending only when OverscanBlendWidth > 0.
+		if(PixelsToBlend > 0)
+		{
+			OutOverscanSettings.bApplyTileEdgeBlend = true;
+
+			// Set ExtraPixels 
+			OutOverscanSettings.ExtraPixels = FIntMarginSet(PixelsToBlend);
+		}
+	}
+
 	return OutOverscanSettings;
 }
diff --git a/Engine/Plugins/Runtime/nDisplay/Source/DisplayCluster/Private/Render/Viewport/Containers/DisplayClusterViewport_CustomFrustumRuntimeSettings.cpp b/Engine/Plugins/Runtime/nDisplay/Source/DisplayCluster/Private/Render/Viewport/Containers/DisplayClusterViewport_CustomFrustumRuntimeSettings.cpp
--- a/Engine/Plugins/Runtime/nDisplay/Source/DisplayCluster/Private/Render/Viewport/Containers/DisplayClusterViewport_CustomFrustumRuntimeSettings.cpp
+++ b/Engine/Plugins/Runtime/nDisplay/Source/DisplayCluster/Private/Render/Viewport/Containers/DisplayClusterViewport_CustomFrustumRuntimeSettings.cpp
@@ -90,11 +90,12 @@ bool FDisplayClusterViewport_CustomFrustumRuntimeSettings::UpdateProjectionAngle
 	if (InRuntimeSettings.bIsEnabled)
 	{
 		const FVector2D SizeFOV(InOutRight - InOutLeft, InOutTop - InOutBottom);
+		const FMarginSet FovMargins = InRuntimeSettings.CustomFrustumPercent * SizeFOV;
 
-		InOutLeft   -= SizeFOV.X * InRuntimeSettings.CustomFrustumPercent.Left;
-		InOutRight  += SizeFOV.X * InRuntimeSettings.CustomFrustumPercent.Right;
-		InOutBottom -= SizeFOV.Y * InRuntimeSettings.CustomFrustumPercent.Bottom;
-		InOutTop    += SizeFOV.Y * InRuntimeSettings.CustomFrustumPercent.Top;
+		InOutLeft   -= FovMargins.Left;
+		InOutRight  += FovMargins.Right;
+		InOutBottom -= FovMargins.Bottom;
+		InOutTop    += FovMargins.Top;
 
 		return true;
 	}
@@ -119,28 +120,20 @@ void FDisplayClusterViewport_CustomFrustumRuntimeSettings::UpdateCustomFrustumSe
 		return;
 	}
 
+	const double MaxCustomFrustumValue = double(GDisplayClusterRenderCustomFrustumMaxValue) / 100;
 	switch (InCustomFrustumSettings.Unit)
 	{
 	case EDisplayClusterViewport_FrustumUnit::Percent:
 	{
 		InOutRuntimeSettings.bIsEnabled = true;
-
-		InOutRuntimeSettings.CustomFrustumPercent.Left   = CustomFrustumHelpers::ClampPercent(InCustomFrustumSettings.Left);
-		InOutRuntimeSettings.CustomFrustumPercent.Right  = CustomFrustumHelpers::ClampPercent(InCustomFrustumSettings.Right);
-		InOutRuntimeSettings.CustomFrustumPercent.Top    = CustomFrustumHelpers::ClampPercent(InCustomFrustumSettings.Top);
-		InOutRuntimeSettings.CustomFrustumPercent.Bottom = CustomFrustumHelpers::ClampPercent(InCustomFrustumSettings.Bottom);
+		InOutRuntimeSettings.CustomFrustumPercent.AssignMargins(InCustomFrustumSettings.GetClamped(-MaxCustomFrustumValue, MaxCustomFrustumValue));
 		break;
 	}
 
 	case EDisplayClusterViewport_FrustumUnit::Pixels:
 	{
 		InOutRuntimeSettings.bIsEnabled = true;
-
-		InOutRuntimeSettings.CustomFrustumPercent.Left   = CustomFrustumHelpers::ClampPercent(InCustomFrustumSettings.Left   / Size.X);
-		InOutRuntimeSettings.CustomFrustumPercent.Right  = CustomFrustumHelpers::ClampPercent(InCustomFrustumSettings.Right  / Size.X);
-		InOutRuntimeSettings.CustomFrustumPercent.Top    = CustomFrustumHelpers::ClampPercent(InCustomFrustumSettings.Top    / Size.Y);
-		InOutRuntimeSettings.CustomFrustumPercent.Bottom = CustomFrustumHelpers::ClampPercent(InCustomFrustumSettings.Bottom / Size.Y);
-
+		InOutRuntimeSettings.CustomFrustumPercent.AssignMargins((InCustomFrustumSettings / Size).GetClamped(-MaxCustomFrustumValue, MaxCustomFrustumValue));
 		break;
 	}
 
@@ -153,16 +146,10 @@ void FDisplayClusterViewport_CustomFrustumRuntimeSettings::UpdateCustomFrustumSe
 	CustomFrustumHelpers::AdjustOverscanValuesToEnforceMinimumSize(InOutRuntimeSettings.CustomFrustumPercent.Bottom, InOutRuntimeSettings.CustomFrustumPercent.Top);
 
 	// Calc pixels from percent
-	InOutRuntimeSettings.CustomFrustumPixels.Left   = FMath::RoundToInt(Size.X * InOutRuntimeSettings.CustomFrustumPercent.Left);
-	InOutRuntimeSettings.CustomFrustumPixels.Right  = FMath::RoundToInt(Size.X * InOutRuntimeSettings.CustomFrustumPercent.Right);
-	InOutRuntimeSettings.CustomFrustumPixels.Top    = FMath::RoundToInt(Size.Y * InOutRuntimeSettings.CustomFrustumPercent.Top);
-	InOutRuntimeSettings.CustomFrustumPixels.Bottom = FMath::RoundToInt(Size.Y * InOutRuntimeSettings.CustomFrustumPercent.Bottom);
-
+		InOutRuntimeSettings.CustomFrustumPixels.AssignMargins((InOutRuntimeSettings.CustomFrustumPercent * Size).GetRoundToInt());
+	
 	// Quantize the percentage to exactly fit the number of pixels.
-	InOutRuntimeSettings.CustomFrustumPercent.Left   = double(InOutRuntimeSettings.CustomFrustumPixels.Left)   / double(Size.X);
-	InOutRuntimeSettings.CustomFrustumPercent.Right  = double(InOutRuntimeSettings.CustomFrustumPixels.Right)  / double(Size.X);
-	InOutRuntimeSettings.CustomFrustumPercent.Top    = double(InOutRuntimeSettings.CustomFrustumPixels.Top)    / double(Size.Y);
-	InOutRuntimeSettings.CustomFrustumPercent.Bottom = double(InOutRuntimeSettings.CustomFrustumPixels.Bottom) / double(Size.Y);
+	InOutRuntimeSettings.CustomFrustumPercent.AssignMargins(FMarginSet(InOutRuntimeSettings.CustomFrustumPixels) / Size);
 
 	// Update RTT size for CustomFrustum when we need to scale target resolution
 	if (InCustomFrustumSettings.bAdaptResolution)
diff --git a/Engine/Plugins/Runtime/nDisplay/Source/DisplayCluster/Private/Render/Viewport/Containers/DisplayClusterViewport_CustomFrustumRuntimeSettings.h b/Engine/Plugins/Runtime/nDisplay/Source/DisplayCluster/Private/Render/Viewport/Containers/DisplayClusterViewport_CustomFrustumRuntimeSettings.h
--- a/Engine/Plugins/Runtime/nDisplay/Source/DisplayCluster/Private/Render/Viewport/Containers/DisplayClusterViewport_CustomFrustumRuntimeSettings.h
+++ b/Engine/Plugins/Runtime/nDisplay/Source/DisplayCluster/Private/Render/Viewport/Containers/DisplayClusterViewport_CustomFrustumRuntimeSettings.h
@@ -53,6 +53,7 @@ struct FDisplayClusterViewport_CustomFrustumRuntimeSettings
 	*   FrustumAngles.Bottom -= FrustumAngles.Height * Bottom.
 	*/
 	struct FCustomFrustumPercent
+		: public FMarginSet
 	{
 		/** Returns the aspect ratio multiplier. */
 		inline double GetAspectRatioMult() const
@@ -69,17 +70,13 @@ struct FDisplayClusterViewport_CustomFrustumRuntimeSettings
 
 			return WidthMult / HeightMult;
 		}
-
-		double Left = 0;
-		double Right = 0;
-		double Top = 0;
-		double Bottom = 0;
 	};
 
 	/**
 	* Values in pixels
 	*/
 	struct FCustomFrustumPixels
+		: public FIntMarginSet
 	{
 		inline FIntRect GetInnerRect(const FIntRect& InRect) const
 		{
@@ -93,11 +90,6 @@ struct FDisplayClusterViewport_CustomFrustumRuntimeSettings
 		{
 			return FIntPoint(Left + Right, Top + Bottom);
 		}
-
-		int32 Left = 0;
-		int32 Right = 0;
-		int32 Top = 0;
-		int32 Bottom = 0;
 	};
 
 	// Enable custom frustum
diff --git a/Engine/Plugins/Runtime/nDisplay/Source/DisplayCluster/Private/Render/Viewport/Containers/DisplayClusterViewport_OverscanRuntimeSettings.cpp b/Engine/Plugins/Runtime/nDisplay/Source/DisplayCluster/Private/Render/Viewport/Containers/DisplayClusterViewport_OverscanRuntimeSettings.cpp
--- a/Engine/Plugins/Runtime/nDisplay/Source/DisplayCluster/Private/Render/Viewport/Containers/DisplayClusterViewport_OverscanRuntimeSettings.cpp
+++ b/Engine/Plugins/Runtime/nDisplay/Source/DisplayCluster/Private/Render/Viewport/Containers/DisplayClusterViewport_OverscanRuntimeSettings.cpp
@@ -22,24 +22,6 @@ static FAutoConsoleVariableRef CVarDisplayClusterRenderOverscanMaxValue(
 	ECVF_Default
 );
 
-namespace UE::DisplayCluster::Viewport::OverscanHelpers
-{
-	/** Clamp percent for overscan settings. */
-	static inline double ClampPercent(double InValue)
-	{
-		const double MaxCustomFrustumValue = double(GDisplayClusterRenderOverscanMaxValue) / 100;
-
-		// We can't use negative overscan values.
-		// The idea behind the overscan is to add extra space on the sides of the RTT.
-		// Note: this only applies to regular viewports (Outers, etc.).
-		// 
-		// The inner frustum viewport has its own implementation for the overscan feature called "CustomFrustum".
-		// (see FDisplayClusterViewport_CustomFrustumRuntimeSettings)
-
-		return FMath::Clamp(InValue, 0, MaxCustomFrustumValue);
-	}
-};
-
 ///////////////////////////////////////////////////////////////////////////////////////////
 // FDisplayClusterViewport_OverscanRuntimeSettings
 ///////////////////////////////////////////////////////////////////////////////////////////
@@ -54,11 +36,12 @@ bool FDisplayClusterViewport_OverscanRuntimeSettings::UpdateProjectionAngles(
 	if (InOverscanRuntimeSettings.bIsEnabled)
 	{
 		const FVector2D SizeFOV(InOutRight - InOutLeft, InOutTop - InOutBottom);
+		const FMarginSet FovMargins = InOverscanRuntimeSettings.OverscanPercent * SizeFOV;
 
-		InOutLeft   -= SizeFOV.X * InOverscanRuntimeSettings.OverscanPercent.Left;
-		InOutRight  += SizeFOV.X * InOverscanRuntimeSettings.OverscanPercent.Right;
-		InOutBottom -= SizeFOV.Y * InOverscanRuntimeSettings.OverscanPercent.Bottom;
-		InOutTop    += SizeFOV.Y * InOverscanRuntimeSettings.OverscanPercent.Top;
+		InOutLeft   -= FovMargins.Left;
+		InOutRight  += FovMargins.Right;
+		InOutBottom -= FovMargins.Bottom;
+		InOutTop    += FovMargins.Top;
 
 		return true;
 	}
@@ -72,38 +55,37 @@ void FDisplayClusterViewport_OverscanRuntimeSettings::UpdateOverscanSettings(
 	FDisplayClusterViewport_OverscanRuntimeSettings& InOutOverscanRuntimeSettings,
 	FIntRect& InOutRenderTargetRect)
 {
-	using namespace UE::DisplayCluster::Viewport;
-
-	const FIntPoint Size = InOutRenderTargetRect.Size();
-
 	// Disable viewport overscan feature
 	if (!GDisplayClusterRenderOverscanEnable || !InOverscanSettings.bEnabled)
 	{
 		return;
 	}
 
+	const FIntPoint Size = InOutRenderTargetRect.Size();
+
+	/**
+	* We can't use negative overscan values.
+	* The idea behind the overscan is to add extra space on the sides of the RTT.
+	* Note: this only applies to regular viewports (Outers, etc.).
+	*
+	* The inner frustum viewport has its own implementation for the overscan feature called "CustomFrustum".
+	* (see FDisplayClusterViewport_CustomFrustumRuntimeSettings)
+	*/
+	const double OverscanMaxValue = FMath::Max(0, double(GDisplayClusterRenderOverscanMaxValue) / 100);
+
 	switch (InOverscanSettings.Unit)
 	{
 	case EDisplayClusterViewport_FrustumUnit::Percent:
 	{
 		InOutOverscanRuntimeSettings.bIsEnabled = true;
-
-		InOutOverscanRuntimeSettings.OverscanPercent.Left   = OverscanHelpers::ClampPercent(InOverscanSettings.Left);
-		InOutOverscanRuntimeSettings.OverscanPercent.Right  = OverscanHelpers::ClampPercent(InOverscanSettings.Right);
-		InOutOverscanRuntimeSettings.OverscanPercent.Top    = OverscanHelpers::ClampPercent(InOverscanSettings.Top);
-		InOutOverscanRuntimeSettings.OverscanPercent.Bottom = OverscanHelpers::ClampPercent(InOverscanSettings.Bottom);
+		InOutOverscanRuntimeSettings.OverscanPercent = InOverscanSettings.GetClamped(0, OverscanMaxValue);
 		break;
 	}
 
 	case EDisplayClusterViewport_FrustumUnit::Pixels:
 	{
 		InOutOverscanRuntimeSettings.bIsEnabled = true;
-
-		InOutOverscanRuntimeSettings.OverscanPercent.Left   = OverscanHelpers::ClampPercent(InOverscanSettings.Left   / Size.X);
-		InOutOverscanRuntimeSettings.OverscanPercent.Right  = OverscanHelpers::ClampPercent(InOverscanSettings.Right  / Size.X);
-		InOutOverscanRuntimeSettings.OverscanPercent.Top    = OverscanHelpers::ClampPercent(InOverscanSettings.Top    / Size.Y);
-		InOutOverscanRuntimeSettings.OverscanPercent.Bottom = OverscanHelpers::ClampPercent(InOverscanSettings.Bottom / Size.Y);
-
+		InOutOverscanRuntimeSettings.OverscanPercent = (InOverscanSettings / Size).GetClamped(0, OverscanMaxValue);
 		break;
 	}
 
@@ -114,54 +96,58 @@ void FDisplayClusterViewport_OverscanRuntimeSettings::UpdateOverscanSettings(
 	// Update RTT size for overscan
 	if (InOutOverscanRuntimeSettings.bIsEnabled)
 	{
+		// Store extra pixels
+		InOutOverscanRuntimeSettings.ExtraPixels.AssignMargins(InOverscanSettings.ExtraPixels);
+		InOutOverscanRuntimeSettings.BaseExtraPixels.AssignMargins(InOverscanSettings.ExtraPixels);
+
 		// Calc pixels from percent
-		InOutOverscanRuntimeSettings.OverscanPixels.Left   = FMath::RoundToInt(Size.X * InOutOverscanRuntimeSettings.OverscanPercent.Left);
-		InOutOverscanRuntimeSettings.OverscanPixels.Right  = FMath::RoundToInt(Size.X * InOutOverscanRuntimeSettings.OverscanPercent.Right);
-		InOutOverscanRuntimeSettings.OverscanPixels.Top    = FMath::RoundToInt(Size.Y * InOutOverscanRuntimeSettings.OverscanPercent.Top);
-		InOutOverscanRuntimeSettings.OverscanPixels.Bottom = FMath::RoundToInt(Size.Y * InOutOverscanRuntimeSettings.OverscanPercent.Bottom);
+		const FIntMarginSet OverscanPixels = (InOutOverscanRuntimeSettings.OverscanPercent * Size).GetRoundToInt();
+
+		// Sets each margin to at least the value in the ExtraPixels.
+		const FIntMarginSet ExpandedOverscanPixels = OverscanPixels.GetMaxPerSideIfDefined(InOverscanSettings.ExtraPixels);
+
+		// Update values
+		InOutOverscanRuntimeSettings.OverscanPixels.AssignMargins(ExpandedOverscanPixels);
+		InOutOverscanRuntimeSettings.BaseOverscanPixels.AssignMargins(ExpandedOverscanPixels);
 
 		// Quantize the overscan percentage to exactly fit the number of pixels.
 		// This will avoid a mismatch between the overscanned frustum calculated in UpdateProjectionAngles and the pixel crop in GetFinalContextRect.
-		InOutOverscanRuntimeSettings.OverscanPercent.Left   = double(InOutOverscanRuntimeSettings.OverscanPixels.Left)   / double(Size.X);
-		InOutOverscanRuntimeSettings.OverscanPercent.Right  = double(InOutOverscanRuntimeSettings.OverscanPixels.Right)  / double(Size.X);
-		InOutOverscanRuntimeSettings.OverscanPercent.Top    = double(InOutOverscanRuntimeSettings.OverscanPixels.Top)    / double(Size.Y);
-		InOutOverscanRuntimeSettings.OverscanPercent.Bottom = double(InOutOverscanRuntimeSettings.OverscanPixels.Bottom) / double(Size.Y);
+		InOutOverscanRuntimeSettings.OverscanPercent = FMarginSet(InOutOverscanRuntimeSettings.OverscanPixels) / Size;
 
 		const FIntPoint OverscanSize = Size + InOutOverscanRuntimeSettings.OverscanPixels.Size();
-		const FIntPoint ValidOverscanSize = FDisplayClusterViewportHelpers::GetValidViewportRect(FIntRect(FIntPoint(0, 0), OverscanSize), InViewportId, TEXT("Overscan")).Size();
-		
-		bool bOversize = InOverscanSettings.bOversize;
+		InOutOverscanRuntimeSettings.bOversize = InOverscanSettings.bOversize;
 
-		if (OverscanSize != ValidOverscanSize)
+		// Check if new RTT size is supported
+		if (InOutOverscanRuntimeSettings.bOversize)
 		{
-			// can't use overscan with extra size, disable oversize
-			bOversize = false;
+			const FIntPoint ValidOverscanSize = FDisplayClusterViewportHelpers::GetValidViewportRect(FIntRect(FIntPoint(0, 0), OverscanSize), InViewportId, TEXT("Overscan")).Size();
+			if (OverscanSize != ValidOverscanSize)
+			{
+				// can't use overscan with extra size, disable oversize
+				InOutOverscanRuntimeSettings.bOversize = false;
+			}
 		}
 
-		if (bOversize)
+		// Increase the RTT size if possible
+		if (InOutOverscanRuntimeSettings.bOversize)
 		{
 			InOutRenderTargetRect.Max = OverscanSize;
+
+			return;
 		}
-		else
-		{
-			const double ScaleX = double(Size.X) / double(OverscanSize.X);
-			const double ScaleY = double(Size.Y) / double(OverscanSize.Y);
-
-			InOutOverscanRuntimeSettings.OverscanPixels.Left   = FMath::RoundToInt(InOutOverscanRuntimeSettings.OverscanPixels.Left   * ScaleX);
-			InOutOverscanRuntimeSettings.OverscanPixels.Right  = FMath::RoundToInt(InOutOverscanRuntimeSettings.OverscanPixels.Right  * ScaleX);
-			InOutOverscanRuntimeSettings.OverscanPixels.Top    = FMath::RoundToInt(InOutOverscanRuntimeSettings.OverscanPixels.Top    * ScaleY);
-			InOutOverscanRuntimeSettings.OverscanPixels.Bottom = FMath::RoundToInt(InOutOverscanRuntimeSettings.OverscanPixels.Bottom * ScaleY);
-
-			const FIntPoint ScaledSize(
-				Size.X - (InOutOverscanRuntimeSettings.OverscanPixels.Left + InOutOverscanRuntimeSettings.OverscanPixels.Right),
-				Size.Y - (InOutOverscanRuntimeSettings.OverscanPixels.Top  + InOutOverscanRuntimeSettings.OverscanPixels.Bottom)
-			);
-
-			// Quantize the overscan percentage to exactly fit the number of pixels.
-			InOutOverscanRuntimeSettings.OverscanPercent.Left   = double(InOutOverscanRuntimeSettings.OverscanPixels.Left)   / double(ScaledSize.X);
-			InOutOverscanRuntimeSettings.OverscanPercent.Right  = double(InOutOverscanRuntimeSettings.OverscanPixels.Right)  / double(ScaledSize.X);
-			InOutOverscanRuntimeSettings.OverscanPercent.Top    = double(InOutOverscanRuntimeSettings.OverscanPixels.Top)    / double(ScaledSize.Y);
-			InOutOverscanRuntimeSettings.OverscanPercent.Bottom = double(InOutOverscanRuntimeSettings.OverscanPixels.Bottom) / double(ScaledSize.Y);
-		}
+
+		// Configure overcan rendering within the existing RTT size.
+
+		// 1. Update value of overscan pixels
+		const FMarginSet NewOverscanPixelsPercent = FMarginSet(InOutOverscanRuntimeSettings.OverscanPixels) / OverscanSize;
+		InOutOverscanRuntimeSettings.OverscanPixels.AssignMargins((NewOverscanPixelsPercent * Size).GetRoundToInt());
+
+		// 2. Update value of extra pixels
+		const FMarginSet NewExtraPixelsPercent = FMarginSet(InOutOverscanRuntimeSettings.ExtraPixels) / OverscanSize;
+		InOutOverscanRuntimeSettings.ExtraPixels.AssignMargins((NewExtraPixelsPercent * Size).GetRoundToInt());
+
+		// 3. Quantize the overscan percentage to exactly fit the number of pixels.
+		const FIntPoint ScaledSize = Size - InOutOverscanRuntimeSettings.OverscanPixels.Size();
+		InOutOverscanRuntimeSettings.OverscanPercent = FMarginSet(InOutOverscanRuntimeSettings.OverscanPixels) / ScaledSize;
 	}
 }
diff --git a/Engine/Plugins/Runtime/nDisplay/Source/DisplayCluster/Private/Render/Viewport/Containers/DisplayClusterViewport_OverscanRuntimeSettings.h b/Engine/Plugins/Runtime/nDisplay/Source/DisplayCluster/Private/Render/Viewport/Containers/DisplayClusterViewport_OverscanRuntimeSettings.h
--- a/Engine/Plugins/Runtime/nDisplay/Source/DisplayCluster/Private/Render/Viewport/Containers/DisplayClusterViewport_OverscanRuntimeSettings.h
+++ b/Engine/Plugins/Runtime/nDisplay/Source/DisplayCluster/Private/Render/Viewport/Containers/DisplayClusterViewport_OverscanRuntimeSettings.h
@@ -5,10 +5,13 @@
 #include "CoreMinimal.h"
 #include "Render/Viewport/Containers/DisplayClusterViewport_OverscanSettings.h"
 
+#include "Math/MarginSet.h"
+
 /**
 * Runtime overscan settings
 */
 struct FDisplayClusterViewport_OverscanRuntimeSettings
+	: public FMarginSet
 {
 	/** Update overscan settings
 	* 
@@ -30,20 +33,31 @@ struct FDisplayClusterViewport_OverscanRuntimeSettings
 	static bool UpdateProjectionAngles(const FDisplayClusterViewport_OverscanRuntimeSettings& InOverscanRuntimeSettings, const FIntPoint& InRenderTargetSize, double& InOutLeft, double& InOutRight, double& InOutTop, double& InOutBottom);
 
 	/**
-	* Overscan values in percent
-	*/
-	struct FOverscanPercent
-	{
-		double Left = 0;
-		double Right = 0;
-		double Top = 0;
-		double Bottom = 0;
-	};
-
-	/**
-	* Overscan values in pixels
+	* Overscan values in pixels:
+	*    +--------------------------+
+	*    |           Top            |
+	*    |        +-------+         |
+	*    |   Left |       |  Right  |
+	*    |        +-------+         |
+	*    |         Bottom           |
+	*    +--------------------------+
+	* 
+	* Logic:
+	*   OverscanPixels = OverscanPixels.MaxPerSideIfDefined(ExtraPixels)
+	* 
+	* This ensures each overscan side is at least as large as the corresponding
+	* ExtraPixels value, but only if ExtraPixels is positive.
+	* 
+	* Example (horizontal layout):
+	* 
+	*  Current OverscanPixels:   |<---- 8 px ---->|
+	*  ExtraPixels:              |<-------- 12 px -------->|
+	* 
+	* Since ExtraPixels (12) > OverscanPixels (8) and ExtraPixels > 0,
+	* OverscanPixels will be increased to 12 px on this side.
 	*/
 	struct FOverscanPixels
+		: public FIntMarginSet
 	{
 		inline FIntRect GetInnerRect(const FIntRect& InRect) const
 		{
@@ -57,19 +71,35 @@ struct FDisplayClusterViewport_OverscanRuntimeSettings
 		{
 			return FIntPoint(Left + Right, Top + Bottom);
 		}
-
-		int32 Left = 0;
-		int32 Right = 0;
-		int32 Top = 0;
-		int32 Bottom = 0;
 	};
 
+public:
 	// Enable overscan
-	bool bIsEnabled = false;
+	uint8 bIsEnabled : 1 = 0;
+
+	// Set to True to render at the overscan resolution, set to false to render at the resolution in the configuration and scale for overscan
+	uint8 bOversize : 1 = 0;
+
+public:
+	// Overscan size per side, expressed as a percentage of the viewport size.
+	// 0.0 = no overscan, 1.0 = 100% overscan (equal to the full viewport dimension).
+	FMarginSet OverscanPercent;
+
+	// Overscan size per side (in pixels) after applying resolution adjustments to match the overscan settings.
+	// See bOversize for details on how resolution adjustments are applied.
+	FOverscanPixels OverscanPixels;
+
+	// Runtime overscan values for FDisplayClusterViewport_OverscanSettings::ExtraPixels.
+	// Used for blending between adjacent tiles in the viewport.
+	FOverscanPixels ExtraPixels;
+
+public:
+	// The tile viewport may render at a lower resolution but will be stretched to its final size in the original viewport.
+	// Stores the original values before any size adjustments, serving as the reference for overscan size modifiers.
 
-	// Overscan sides in percent
-	FOverscanPercent OverscanPercent;
+	// Original overscan size per side (in pixels) before any resolution adjustments.
+	FOverscanPixels BaseOverscanPixels;
 
-	// Overscan sides in pixels
-	FOverscanPixels  OverscanPixels;
+	// Original (pre–resolution-adjustment) overscan values for FDisplayClusterViewport_OverscanSettings::ExtraPixels.
+	FOverscanPixels BaseExtraPixels;
 };
diff --git a/Engine/Plugins/Runtime/nDisplay/Source/DisplayCluster/Private/Render/Viewport/DisplayClusterViewportProxy_Resources.cpp b/Engine/Plugins/Runtime/nDisplay/Source/DisplayCluster/Private/Render/Viewport/DisplayClusterViewportProxy_Resources.cpp
--- a/Engine/Plugins/Runtime/nDisplay/Source/DisplayCluster/Private/Render/Viewport/DisplayClusterViewportProxy_Resources.cpp
+++ b/Engine/Plugins/Runtime/nDisplay/Source/DisplayCluster/Private/Render/Viewport/DisplayClusterViewportProxy_Resources.cpp
@@ -52,11 +52,57 @@ static FAutoConsoleVariableRef CVarDisplayClusterRenderTileBorder(
 	ECVF_RenderThreadSafe
 );
 
+/** Specifies the blending algorithm for tile edges. */
+enum class EDisplayClusterRenderTileEdgeBlendMode : uint8
+{
+	/** No blending is applied to tile edges. */
+	None = 0,
+
+	/** Linear fade between tiles — fast but can appear harsh. */
+	Linear = 1,
+
+	/** Smooth (non-linear) fade for softer transitions between tiles. */
+	Smooth = 2
+};
+
+/** CVar that specifies the blending algorithm for tile edges. */
+int32 GDisplayClusterRenderTileEdgeBlendMode = 2;
+static FAutoConsoleVariableRef CVarDisplayClusterRenderTileBlendMode(
+	TEXT("nDisplay.render.tile.BlendMode"),
+	GDisplayClusterRenderTileEdgeBlendMode,
+	TEXT("Tile blend mode (default = 2).\n")
+	TEXT(" 0 - none (Disable tile blending).\n")
+	TEXT(" 1 - linear.\n")
+	TEXT(" 2 - smooth (non-linear, s-shape).\n"),
+	ECVF_Default
+);
+
 ///////////////////////////////////////////////////////////////////////////////////////
 namespace UE::DisplayCluster::ViewportProxy
 {
 	// The viewport override has the maximum depth. This protects against a link cycle
 	static const int32 DisplayClusterViewportProxyResourcesOverrideRecursionDepthMax = 4;
+
+	/**
+	 * Get current tile blending mode.
+	 * Falls back to EDisplayClusterRenderTileEdgeBlendMode::None if the value is not a valid mode.
+	 */
+	static inline EDisplayClusterRenderTileEdgeBlendMode GetTileEdgeBlendMode()
+	{
+		switch (GDisplayClusterRenderTileEdgeBlendMode)
+		{
+		case static_cast<int32>(EDisplayClusterRenderTileEdgeBlendMode::Linear):
+			return EDisplayClusterRenderTileEdgeBlendMode::Linear;
+
+		case static_cast<int32>(EDisplayClusterRenderTileEdgeBlendMode::Smooth):
+			return EDisplayClusterRenderTileEdgeBlendMode::Smooth;
+
+		default:
+			break;
+		}
+
+		return EDisplayClusterRenderTileEdgeBlendMode::None;
+	}
 };
 
 void FDisplayClusterViewportProxy::FillTextureWithColor_RenderThread(FRHICommandListImmediate& RHICmdList, FRHITexture* InRenderTargetTexture, const FLinearColor& InColor)
@@ -274,11 +320,27 @@ bool FDisplayClusterViewportProxy::ImplResolveResources_RenderThread(FRHICommand
 }
 
 void FDisplayClusterViewportProxy::ImplResolveTileResource_RenderThread(FRHICommandListImmediate& RHICmdList, FDisplayClusterViewportProxy const* InDestViewportProxy) const
-	{
+{
 	check(InDestViewportProxy);
 
 	using namespace UE::DisplayCluster::ViewportProxy;
 
+	// When AdaptResolution is disabled, the input RTT texture may be smaller than the output rect.
+	// Use separate fields to store the input and output rect dimensions in this case.
+	// 
+	// Note: The AdaptResolution option in overscan settings controls whether the RTT size 
+	// is increased to match the overscan size. If disabled, only the inner rectangle 
+	// matching the original viewport size is used as the input.
+	FIntMarginSet InputRectMargins = static_cast<FIntMarginSet>(OverscanRuntimeSettings.OverscanPixels);
+	FIntMarginSet OutputRectMargins = static_cast<FIntMarginSet>(OverscanRuntimeSettings.BaseOverscanPixels);
+
+	// Get actual tile edge blending mode
+	const EDisplayClusterRenderTileEdgeBlendMode TileEdgeBlendMode = GetTileEdgeBlendMode();
+
+	// Check if we can apply tile edge blending
+	const bool bApplyTileEdgeBlend = RenderSettings.TileSettings.HasAnyTileFlags(EDisplayClusterViewportTileFlags::ApplyTileEdgeBlend)
+		&& TileEdgeBlendMode != EDisplayClusterRenderTileEdgeBlendMode::None;
+
 	const TArray<FDisplayClusterViewport_Context>& DestContexts = InDestViewportProxy->GetContexts_RenderThread();
 	const TArray<FDisplayClusterViewport_Context>& SrcContexts = Contexts;
 	TSharedRef<IDisplayClusterShadersTextureUtils> TextureUtils = GetShadersAPI().CreateTextureUtils_RenderThread(RHICmdList);
@@ -302,6 +364,13 @@ void FDisplayClusterViewportProxy::ImplResolveTileResource_RenderThread(FRHIComm
 
 				Output.Rect = SrcContexts[ContextNum].TileDestRect;
 
+				// By default, we just copy part of the tile texture into dest.
+				EDisplayClusterShaderTextureUtilsFlags Flags = EDisplayClusterShaderTextureUtilsFlags::None;
+
+				// Get tile position ans size
+				const FIntPoint InTilePos = RenderSettings.TileSettings.GetPos();
+				const FIntPoint InTileSize = RenderSettings.TileSettings.GetSize();
+
 				if (GDisplayClusterRenderTileBorder > 0)
 				{
 					// The maximum border is 1/4 of the minimum side of the rectangle.
@@ -310,19 +379,19 @@ void FDisplayClusterViewportProxy::ImplResolveTileResource_RenderThread(FRHIComm
 
 					// Shows borders for internal areas only
 					FIntRect SubRect(FIntPoint::ZeroValue, FIntPoint::ZeroValue);
-					if (RenderSettings.TileSettings.GetPos().X > 0)
+					if (InTilePos.X > 0)
 					{
 						SubRect.Min.X = TileBorderSize;
 					}
-					if (RenderSettings.TileSettings.GetPos().Y > 0)
+					if (InTilePos.Y > 0)
 					{
 						SubRect.Min.Y = TileBorderSize;
 					}
-					if ((RenderSettings.TileSettings.GetPos().X + 1) < RenderSettings.TileSettings.GetSize().X)
+					if ((InTilePos.X + 1) < InTileSize.X)
 					{
 						SubRect.Max.X = TileBorderSize;
 					}
-					if ((RenderSettings.TileSettings.GetPos().Y + 1) < RenderSettings.TileSettings.GetSize().Y)
+					if ((InTilePos.Y + 1) < InTileSize.Y)
 					{
 						SubRect.Max.Y = TileBorderSize;
 					}
@@ -338,18 +407,97 @@ void FDisplayClusterViewportProxy::ImplResolveTileResource_RenderThread(FRHIComm
 					Output.Rect.Max.X -= SubRect.Min.X + SubRect.Max.X;
 					Output.Rect.Max.Y -= SubRect.Min.Y + SubRect.Max.Y;
 				}
+				// When we don't need to show tile borders, we can blend them together
+				else if (bApplyTileEdgeBlend)
+				{
+					if (!RenderSettings.TileSettings.HasAnyTileFlags(EDisplayClusterViewportTileFlags::UseEntireOverscanForEdgeBlend))
+					{
+						// Use minimal margins from EdgeBlendWidth (ExtraPixels)
+						InputRectMargins = InputRectMargins.GetMinPerSide(OverscanRuntimeSettings.ExtraPixels);
+						OutputRectMargins = OutputRectMargins.GetMinPerSide(OverscanRuntimeSettings.BaseExtraPixels);
+					}
+
+					// Per-tile adjustments: Optimize overscan values for edge tiles.
+					{
+						using EMarginSide = UE::Math::EMarginSetSideIndex;
+						auto OptimizeSize = [&](const EMarginSide InSideIndex)
+							{
+								InputRectMargins[InSideIndex] = 0;
+								OutputRectMargins[InSideIndex] = 0;
+							};
+
+						if (InTilePos.X == 0)
+						{
+							OptimizeSize(EMarginSide::Left);
+						}
+						if (InTilePos.Y == 0)
+						{
+							OptimizeSize(EMarginSide::Top);
+						}
+						if (InTilePos.X == (InTileSize.X - 1))
+						{
+							OptimizeSize(EMarginSide::Right);
+						}
+						if (InTilePos.Y == (InTileSize.Y - 1))
+						{
+							OptimizeSize(EMarginSide::Bottom);
+						}
+					}
+
+					// Configure rendering flags that define how tile edge blending is applied
+					switch (TileEdgeBlendMode)
+					{
+					case EDisplayClusterRenderTileEdgeBlendMode::Smooth:
+						// Enables smooth (cubic Hermite) alpha feathering(fade - out) on all sides of the rectangle.
+						// When set, alpha transitions smoothly from opaque to transparent along each edge using a smoothstep(S - curve) function.
+						Flags |= EDisplayClusterShaderTextureUtilsFlags::EnableSmoothAlphaFeather;
+						break;
+
+					case EDisplayClusterRenderTileEdgeBlendMode::Linear:
+						// Applies a linear alpha fade(feathering) to all edges of the rectangle.
+						Flags |= EDisplayClusterShaderTextureUtilsFlags::EnableLinearAlphaFeather;
+						break;
+
+					default:
+						check(false); // This case indicates an invalid value and should not occur if bApplyTileEdgeBlend was validated beforehand.
+						break;
+					};
+
+					// Add margins to input and output rects
+					InputRectMargins.ApplyExpandToRect(Input.Rect);
+					OutputRectMargins.ApplyExpandToRect(Output.Rect);
+				}
 
 				// Copy texture region
 				TextureUtils->ResolveTextureContext(
-					EDisplayClusterShaderTextureUtilsFlags::DisableResampleShader,
+					{ Flags,  InputRectMargins * 2},
 					Input, Output);
 			});
 }
 
 void FDisplayClusterViewportProxy::CleanupResources_RenderThread(FRHICommandListImmediate& RHICmdList) const
 {
-	// Since the RTT is reused through frames, in case we need to show a black border between viewport tiles, we must fill the original viewport with this colour.
-	if (GDisplayClusterRenderTileBorder > 0 && RenderSettings.TileSettings.GetType() == EDisplayClusterViewportTileType::Source)
+	bool bShouldCleanupRTT = false;
+
+	// Rules for tile viewport
+	if (RenderSettings.TileSettings.GetType() == EDisplayClusterViewportTileType::Source)
+	{
+		if (GDisplayClusterRenderTileBorder > 0)
+		{
+			// Since the RTT is reused through frames, in case we need to show a black border between viewport tiles,
+			// we must fill the original viewport with this colour.
+			bShouldCleanupRTT = true;
+		}
+
+		if (RenderSettings.TileSettings.HasAnyTileFlags(EDisplayClusterViewportTileFlags::ApplyTileEdgeBlend))
+		{
+			// Blending assumes the underlying surface is black.
+			// Tiles are composited on top of each other across multiple render passes.
+			bShouldCleanupRTT = true;
+		}
+	}
+
+	if(bShouldCleanupRTT)
 	{
 		TArray<FRHITexture*> RenderTargets;
 		if (GetResources_RenderThread(EDisplayClusterViewportResourceType::InternalRenderTargetEntireRectResource, RenderTargets))
diff --git a/Engine/Plugins/Runtime/nDisplay/Source/DisplayCluster/Public/Render/Viewport/Containers/DisplayClusterViewport_Context.h b/Engine/Plugins/Runtime/nDisplay/Source/DisplayCluster/Public/Render/Viewport/Containers/DisplayClusterViewport_Context.h
--- a/Engine/Plugins/Runtime/nDisplay/Source/DisplayCluster/Public/Render/Viewport/Containers/DisplayClusterViewport_Context.h
+++ b/Engine/Plugins/Runtime/nDisplay/Source/DisplayCluster/Public/Render/Viewport/Containers/DisplayClusterViewport_Context.h
@@ -21,11 +21,22 @@ public:
 	{ }
 
 public:
+	/**
+	* Declare all bitfields together to avoid alignment and padding.
+	*/
+
 	// True if this ViewData have a valid data.
-	bool bValidData : 1 = 0;
+	uint8 bValidData : 1 = 0;
 
 	// Is this data have been calculated.
-	bool bCalculated : 1 = 0;
+	uint8 bCalculated : 1 = 0;
+
+	// Enables nDisplay's native implementation of cross-GPU transfer.
+	// This disables cross-GPU transfer by default for nDisplay viewports in FSceneViewFamily structure.
+	uint8 bOverrideCrossGPUTransfer : 1 = 0;
+
+	// Disable render for this viewport (Overlay)
+	uint8 bDisableRender : 1 = 0;
 
 public:
 	// Index of context (eye #)
@@ -138,10 +149,6 @@ public:
 	// GPU index for this context render target
 	int32 GPUIndex = INDEX_NONE;
 
-	// Enables nDisplay's native implementation of cross-GPU transfer.
-	// This disables cross-GPU transfer by default for nDisplay viewports in FSceneViewFamily structure.
-	bool bOverrideCrossGPUTransfer : 1 = 0;
-
 	// Location and size on a render target texture
 	FIntRect RenderTargetRect;
 
@@ -160,9 +167,6 @@ public:
 	// Mips number for additional MipsShader resources
 	int32 NumMips = 1;
 
-	// Disable render for this viewport (Overlay)
-	bool bDisableRender : 1 = 0;
-
 	/**
 	* Viewport context data for rendering thread
 	*/
diff --git a/Engine/Plugins/Runtime/nDisplay/Source/DisplayCluster/Public/Render/Viewport/Containers/DisplayClusterViewport_CustomFrustumSettings.h b/Engine/Plugins/Runtime/nDisplay/Source/DisplayCluster/Public/Render/Viewport/Containers/DisplayClusterViewport_CustomFrustumSettings.h
--- a/Engine/Plugins/Runtime/nDisplay/Source/DisplayCluster/Public/Render/Viewport/Containers/DisplayClusterViewport_CustomFrustumSettings.h
+++ b/Engine/Plugins/Runtime/nDisplay/Source/DisplayCluster/Public/Render/Viewport/Containers/DisplayClusterViewport_CustomFrustumSettings.h
@@ -4,24 +4,20 @@
 
 #include "CoreMinimal.h"
 #include "DisplayClusterViewport_Enums.h"
+#include "Math/MarginSet.h"
 
 /**
 * Custom frustum settings of viewport
 */
 struct FDisplayClusterViewport_CustomFrustumSettings
+	: public FMarginSet
 {
 	// Enable custom frustum
-	bool bEnabled = false;
+	uint8 bEnabled : 1 = 0;
 
 	// Enable adaptive resolution
-	bool bAdaptResolution = false;
+	uint8 bAdaptResolution : 1 = 0;
 
 	// Units type of values
 	EDisplayClusterViewport_FrustumUnit Unit = EDisplayClusterViewport_FrustumUnit::Pixels;
-
-	// custum frustum values
-	float Left = 0;
-	float Right = 0;
-	float Top = 0;
-	float Bottom = 0;
 };
diff --git a/Engine/Plugins/Runtime/nDisplay/Source/DisplayCluster/Public/Render/Viewport/Containers/DisplayClusterViewport_OverscanSettings.h b/Engine/Plugins/Runtime/nDisplay/Source/DisplayCluster/Public/Render/Viewport/Containers/DisplayClusterViewport_OverscanSettings.h
--- a/Engine/Plugins/Runtime/nDisplay/Source/DisplayCluster/Public/Render/Viewport/Containers/DisplayClusterViewport_OverscanSettings.h
+++ b/Engine/Plugins/Runtime/nDisplay/Source/DisplayCluster/Public/Render/Viewport/Containers/DisplayClusterViewport_OverscanSettings.h
@@ -5,23 +5,33 @@
 #include "CoreMinimal.h"
 #include "DisplayClusterViewport_Enums.h"
 
+#include "Math/MarginSet.h"
+
 /**
 * Overscan settings of viewport
 */
 struct FDisplayClusterViewport_OverscanSettings
+	: public FMarginSet
 {
 	// Enable overscan
-	bool bEnabled = false;
+	uint8 bEnabled : 1 = 0;
 
 	// Set to True to render at the overscan resolution, set to false to render at the resolution in the configuration and scale for overscan
-	bool bOversize = false;
+	uint8 bOversize : 1 = 0;
+
+	// Blends the overlapping (overscan) regions between adjacent tiles.
+	// The size (in pixels) must be less than or equal to the overscan size,
+	// ensuring a smooth transition with no visible seams.
+	uint8 bApplyTileEdgeBlend : 1 = 0;
+
+	// When enabled, uses the full overscan width/height for tile edge blending,
+	// ignoring the EdgeBlendWidth setting. Requires ApplyTileEdgeBlend to be active.
+	uint8 bUseEntireOverscanForEdgeBlend : 1 = 0;
 
 	// Units type of overscan values
 	EDisplayClusterViewport_FrustumUnit Unit = EDisplayClusterViewport_FrustumUnit::Pixels;
-
-	// Overscan values
-	float Left = 0;
-	float Right = 0;
-	float Top = 0;
-	float Bottom = 0;
+	
+	// Additional margins in pixels that are guaranteed to be rendered.
+	// Note: this parameter is used to blend the viewport tiles.
+	FIntMarginSet ExtraPixels;
 };
diff --git a/Engine/Plugins/Runtime/nDisplay/Source/DisplayCluster/Public/Render/Viewport/Containers/DisplayClusterViewport_TileSettings.h b/Engine/Plugins/Runtime/nDisplay/Source/DisplayCluster/Public/Render/Viewport/Containers/DisplayClusterViewport_TileSettings.h
--- a/Engine/Plugins/Runtime/nDisplay/Source/DisplayCluster/Public/Render/Viewport/Containers/DisplayClusterViewport_TileSettings.h
+++ b/Engine/Plugins/Runtime/nDisplay/Source/DisplayCluster/Public/Render/Viewport/Containers/DisplayClusterViewport_TileSettings.h
@@ -26,18 +26,54 @@ enum class EDisplayClusterViewportTileType: uint8
 };
 
 /**
- * Additional tile settings represetned as a bitmask
+ * Additional tile settings represented as a bitmask
  */
 enum class EDisplayClusterViewportTileFlags : uint8
 {
 	None = 0,
 
-	// Allow this tile to render when unbound
+	// [Media]  Allow this tile to render when unbound 
+	// 
+	// This flag is passed directly from the FDisplayClusterConfigurationMediaICVFX::ClusterNodesToRenderUnboundTiles.
+	// 
+	// Controls if current node was explicitly allowed to render unbound tiles.
+	// `Unbound tiles` are tile viewports that are not connected to MediaIO.
+	// 
+	// Here are comments on this logic from FDisplayClusterViewport::IsRenderEnabledByMedia()
+	//
+	// In | Out | Render (offscreen) | Render (onscreen)
+	// ------------------------------------------------------------------
+	//  0 |  0  |        0           | EDisplayClusterViewportTileFlags::AllowUnboundRender
+	//  0 |  1  |        1           | 1
+	//  1 |  0  |        0           | 0
+	//  1 |  1  |        0           | 0
 	AllowUnboundRender = 1 << 0,
+
+	// Optimize overscan values for edge tiles.
+	// 
+	// Overscan is not applied on the outer edges of the image — only where tiles meet.
+	// Example layout (tiles with overscan overlap marked as "++"):
+	//
+	//  +------------+-----------+
+	//  | Tile0_0  ++|++ Tile1_0 |
+	//  |     ++     |     ++    |
+	//  +------------+-----------+
+	//  |     ++     |     ++    |
+	//  | Tile0_1  ++|++ Tile1_1 |
+	//  +------------+-----------+
+	OptimizeTileOverscan = 1 << 1,
+
+	// Blends the overlapping (overscan) regions between adjacent tiles.
+	// The size (in pixels) must be less than or equal to the overscan size,
+	// ensuring a smooth transition with no visible seams.
+	ApplyTileEdgeBlend = 1 << 2,
+
+	// When enabled, uses the full overscan width/height for tile edge blending,
+	// ignoring the EdgeBlendWidth setting. Requires ApplyTileEdgeBlend to be active.
+	UseEntireOverscanForEdgeBlend = 1 << 3,
 };
 ENUM_CLASS_FLAGS(EDisplayClusterViewportTileFlags);
 
-
 /**
  * nDisplay viewport tile settings.
  * These are runtime settings, updated every frame from the cluster configuration.
@@ -126,10 +162,6 @@ public:
 		}
 	}
 
-public:
-	// Optimize overscan values for edge tiles.
-	bool bOptimizeTileOverscan = false;
-
 private:
 	// Define the viewport type for tile rendering.
 	EDisplayClusterViewportTileType Type = EDisplayClusterViewportTileType::None;
diff --git a/Engine/Plugins/Runtime/nDisplay/Source/DisplayClusterConfiguration/Public/DisplayClusterConfigurationTypes_Tile.h b/Engine/Plugins/Runtime/nDisplay/Source/DisplayClusterConfiguration/Public/DisplayClusterConfigurationTypes_Tile.h
--- a/Engine/Plugins/Runtime/nDisplay/Source/DisplayClusterConfiguration/Public/DisplayClusterConfigurationTypes_Tile.h
+++ b/Engine/Plugins/Runtime/nDisplay/Source/DisplayClusterConfiguration/Public/DisplayClusterConfigurationTypes_Tile.h
@@ -60,4 +60,19 @@ public:
 	/** Overscan value for all sides. */
 	UPROPERTY(EditAnywhere, BlueprintReadWrite, Category = "Tile Overscan", meta = (DisplayName = "Overscan Value", EditCondition = "bEnabled == true"))
 	float AllSides = 10;
+
+	/** If true, blend across the entire overscan edge region instead of using OverscanBlendWidth.
+	 *  This applies the full overscan width/height for blending at tile edges.
+	 */
+	UPROPERTY(EditAnywhere, BlueprintReadWrite, Category = "Tile Overscan", meta = (DisplayName = "Blend Entire Overscan", EditCondition = "bEnabled == true"))
+	bool bBlendEntireOverscan = false;
+
+	/** Width of the blended edge region, in pixels.
+	 * Used to smooth seams by blending the overscan areas where adjacent tiles meet.
+	 * This value should be less than or equal to the overscan size to prevent artifacts.
+	 * 
+	 * Note: Values larger than the actual overscan size (in pixels) will be clamped to the overscan size.
+	 */
+	UPROPERTY(EditAnywhere, BlueprintReadWrite, Category = "Tile Overscan", meta = (DisplayName = "Blend Width", EditCondition = "bBlendEntireOverscan == false && bEnabled == true", ClampMin = "0", ClampMax = "512", UIMin = "0", UIMax = "512"))
+	int32 OverscanBlendWidth = 16;
 };
diff --git a/Engine/Plugins/Runtime/nDisplay/Source/DisplayClusterShaders/Private/DisplayClusterShadersTextureUtils.cpp b/Engine/Plugins/Runtime/nDisplay/Source/DisplayClusterShaders/Private/DisplayClusterShadersTextureUtils.cpp
--- a/Engine/Plugins/Runtime/nDisplay/Source/DisplayClusterShaders/Private/DisplayClusterShadersTextureUtils.cpp
+++ b/Engine/Plugins/Runtime/nDisplay/Source/DisplayClusterShaders/Private/DisplayClusterShadersTextureUtils.cpp
@@ -526,7 +526,7 @@ void FDisplayClusterShadersTextureUtils::ImplForEachContextByPredicate(
 			if (InitializeTextureViewportForRenderPass(Output, TEXT("nDisplay.Output")))
 			{
 				FDisplayClusterShadersTextureViewportContext Input;
-				if(CloneTextureViewportResourceForRenderPass(Input, Output, TEXT("nDisplay.OutputClone")))
+				if (CloneTextureViewportResourceForRenderPass(Input, Output, TEXT("nDisplay.OutputClone")))
 				{
 					// Use input color encoding
 					Input.ColorEncoding = InputTextureParameters.ColorEncoding;
@@ -566,7 +566,7 @@ void FDisplayClusterShadersTextureUtils::ImplForEachContextByPredicate(
 		if (const FDisplayClusterShadersTextureViewport* InputTexturePtr = InputTextureParameters.TextureViewports.Find(InputContextNum))
 		{
 			FDisplayClusterShadersTextureViewportContext Output(OutputTextureIt.Value, OutputTextureParameters.ColorEncoding, OutputContextNum);
-			FDisplayClusterShadersTextureViewportContext Input(*InputTexturePtr,       InputTextureParameters.ColorEncoding,  InputContextNum);
+			FDisplayClusterShadersTextureViewportContext Input(*InputTexturePtr, InputTextureParameters.ColorEncoding, InputContextNum);
 
 			if (UpdateResourcesRectsForResolve(Input, Output, Settings)
 				&& InitializeTextureViewportForRenderPass(Input, TEXT("nDisplay.Input"))
@@ -621,6 +621,14 @@ bool FDisplayClusterShadersTextureUtils::ShouldUseResampleShader(
 {
 	using namespace UE::DisplayClusterShaders::Private;
 
+	if (Settings.HasAnyFlags(
+	    EDisplayClusterShaderTextureUtilsFlags::EnableLinearAlphaFeather
+	  | EDisplayClusterShaderTextureUtilsFlags::EnableSmoothAlphaFeather))
+	{
+		// Alpha blending and alpha feathering require a shader.
+		return true;
+	}
+
 	if (Settings.ColorMask != EColorWriteMask::CW_RGBA)
 	{
 		// Color mask requires a shader
diff --git a/Engine/Plugins/Runtime/nDisplay/Source/DisplayClusterShaders/Private/Shaders/DisplayClusterShadersCopyTexture.cpp b/Engine/Plugins/Runtime/nDisplay/Source/DisplayClusterShaders/Private/Shaders/DisplayClusterShadersCopyTexture.cpp
--- a/Engine/Plugins/Runtime/nDisplay/Source/DisplayClusterShaders/Private/Shaders/DisplayClusterShadersCopyTexture.cpp
+++ b/Engine/Plugins/Runtime/nDisplay/Source/DisplayClusterShaders/Private/Shaders/DisplayClusterShadersCopyTexture.cpp
@@ -77,6 +77,14 @@ namespace UE::DisplayClusterShaders::Private
 		Zero   = 3, // #define EOverrideAlpha_Zero   3
 	};
 
+	/** Enumerations that cover definitions in the shader. */
+	enum class ESideFeather : uint32
+	{
+		None   = 0,  // #define ESideFeather_None     0
+		Linear = 1,  // #define ESideFeather_Linear   1
+		Smooth = 2,  // #define ESideFeather_Smooth   2
+	};
+
 	namespace FColorEncodingCopyRectPSPermutation
 	{
 		// Shared permutation for picp warp
@@ -119,6 +127,8 @@ namespace UE::DisplayClusterShaders::Private
 		SHADER_PARAMETER(FUint32Vector, ColorPremultiply)
 		SHADER_PARAMETER(FUint32Vector, Encodings)
 		SHADER_PARAMETER(FVector3f, DisplayGamma)
+		SHADER_PARAMETER(FVector4f, InnerFeatherMargins)
+		SHADER_PARAMETER(FVector4f, OuterFeatherMargins)
 	END_SHADER_PARAMETER_STRUCT()
 
 	/**
@@ -141,6 +151,7 @@ namespace UE::DisplayClusterShaders::Private
 		/** Initialize shader parameters. */
 		static bool InitializeShaderParameters(
 			const FDisplayClusterShadersTextureViewportContext& Input,
+			const FIntPoint& InputTextureSize,
 			const FDisplayClusterShadersTextureViewportContext& Output,
 			const FDisplayClusterShadersTextureUtilsSettings& Settings,
 			FParameters& OutParameters,
@@ -156,6 +167,10 @@ namespace UE::DisplayClusterShaders::Private
 			OutParameters.Encodings = FUint32Vector::ZeroValue;
 			OutParameters.DisplayGamma = FVector3f(1.f, 1.f, 1.f);
 
+			// Default values
+			OutParameters.InnerFeatherMargins =
+				OutParameters.OuterFeatherMargins = FVector4f::Zero();
+
 			const bool UseGammaEncoding = (Settings.ColorMask & EColorWriteMask::CW_RGB)
 				&& !Input.ColorEncoding.IsEqualsGammaEncoding(Output.ColorEncoding);
 
@@ -276,7 +291,70 @@ namespace UE::DisplayClusterShaders::Private
 				OutParameters.Encodings.Z = static_cast<uint32>(EOverrideAlpha::Zero);
 				OutPermutationVector.Set<FColorEncodingCopyRectPSPermutation::FPermutationOverrideAlpha>(true);
 				break;
-			}	
+			}
+
+			// Side Feather
+			if (Settings.HasAnyFlags(
+				  EDisplayClusterShaderTextureUtilsFlags::EnableLinearAlphaFeather
+				| EDisplayClusterShaderTextureUtilsFlags::EnableSmoothAlphaFeather))
+			{
+				// Feathering logic is enabled in the pixel shader under this permutation.
+				OutPermutationVector.Set<FColorEncodingCopyRectPSPermutation::FPermutationOverrideAlpha>(true);
+
+				// Only one alpha feathering mode(linear or smooth) should be enabled at a time.
+				if (Settings.HasAnyFlags(EDisplayClusterShaderTextureUtilsFlags::EnableLinearAlphaFeather))
+				{
+					// ColorPremultiply.Z = SideFeather method
+					OutParameters.ColorPremultiply.Z = static_cast<uint32>(ESideFeather::Linear);
+				}
+				else if (Settings.HasAnyFlags(EDisplayClusterShaderTextureUtilsFlags::EnableSmoothAlphaFeather))
+				{
+					// ColorPremultiply.Z = SideFeather method
+					OutParameters.ColorPremultiply.Z = static_cast<uint32>(ESideFeather::Smooth);
+				}
+
+				// Feather margins
+				//
+				//  Outer LT              RT
+				//  +---------------------+
+				//  |                     |
+				//  |  Inner LT      RT   |
+				//  |     +----------+    |
+				//  |     |          |    |
+				//  |     +----------+    |
+				//  |     LB         RB   |
+				//  +---------------------+
+				//  LB                    RB
+				//
+				// L,R,T,B = Left, Right, Top, Bottom
+				//
+				// InnerFeatherMargins = <InnerL, 1-InnerR, InnerT, 1-InnerB>  // Fully opaque region
+				// OuterFeatherMargins = <OuterL, 1-OuterR, OuterT, 1-OuterB>  // Fully transparent region
+
+				FIntRect InnerRect = Input.Rect;
+				Settings.SideMargins.ApplyClampedInsetToRect(InnerRect);
+
+				const FIntRect OuterRect = Input.Rect;
+
+				const float Width  = static_cast<float>(InputTextureSize.X);
+				const float Height = static_cast<float>(InputTextureSize.Y);
+
+				// Left
+				OutParameters.InnerFeatherMargins.X = float(InnerRect.Min.X) / Width;
+				OutParameters.OuterFeatherMargins.X = float(OuterRect.Min.X) / Width;
+
+				// Right
+				OutParameters.InnerFeatherMargins.Y = 1.0f - (float(InnerRect.Max.X) / Width);
+				OutParameters.OuterFeatherMargins.Y = 1.0f - (float(OuterRect.Max.X) / Width);
+
+				// Top
+				OutParameters.InnerFeatherMargins.Z = float(InnerRect.Min.Y) / Height;
+				OutParameters.OuterFeatherMargins.Z = float(OuterRect.Min.Y) / Height;
+
+				// Bottom
+				OutParameters.InnerFeatherMargins.W = 1.0f - (float(InnerRect.Max.Y) / Height);
+				OutParameters.OuterFeatherMargins.W = 1.0f - (float(OuterRect.Max.Y) / Height);
+			}
 
 			// Check the permutation vectors. This should prevent a crash when no shader permutation is found.
 			if (!FColorEncodingCopyRectPSPermutation::ShouldCompileCommonPSPermutation(OutPermutationVector))
@@ -303,10 +381,13 @@ namespace UE::DisplayClusterShaders::Private
 				return false;
 			}
 
+			const FIntPoint& SrcTextureSize = SrcTexture->GetDesc().Extent;
+			const FIntPoint& DestTextureSize = DestTexture->GetDesc().Extent;
+
 			// Initialize shader parameters and permutation vectors
 			FParameters PixelShaderParameters;
 			FColorEncodingCopyRectPSPermutation::FCommonPSDomain PermutationVector;
-			if (!InitializeShaderParameters(InputContext, OutputContext, Settings, PixelShaderParameters, PermutationVector))
+			if (!InitializeShaderParameters(InputContext, SrcTextureSize, OutputContext, Settings, PixelShaderParameters, PermutationVector))
 			{
 				return false;
 			}
@@ -321,8 +402,6 @@ namespace UE::DisplayClusterShaders::Private
 				return false;
 			}
 
-			const FIntPoint& SrcTextureSize = SrcTexture->GetDesc().Extent;
-			const FIntPoint& DestTextureSize = DestTexture->GetDesc().Extent;
 			const FIntRect& SrcRect = InputContext.Rect;
 			const FIntRect& DestRect = OutputContext.Rect;
 
@@ -335,7 +414,18 @@ namespace UE::DisplayClusterShaders::Private
 				FGraphicsPipelineStateInitializer GraphicsPSOInit;
 				RHICmdList.ApplyCachedRenderTargets(GraphicsPSOInit);
 
-				GraphicsPSOInit.BlendState = GetBlendStateRHI(Settings.ColorMask);
+				if (Settings.HasAnyFlags(
+					EDisplayClusterShaderTextureUtilsFlags::EnableLinearAlphaFeather
+				| EDisplayClusterShaderTextureUtilsFlags::EnableSmoothAlphaFeather))
+				{
+					// Special blending mode for feather
+					GraphicsPSOInit.BlendState = TStaticBlendState <CW_RGBA, BO_Add, BF_One, BF_One, BO_Add, BF_One, BF_One>::GetRHI();
+				}
+				else
+				{
+					GraphicsPSOInit.BlendState = GetBlendStateRHI(Settings.ColorMask);
+				}
+
 				GraphicsPSOInit.RasterizerState = TStaticRasterizerState<>::GetRHI();
 				GraphicsPSOInit.DepthStencilState = TStaticDepthStencilState<false, CF_Always>::GetRHI();
 
diff --git a/Engine/Plugins/Runtime/nDisplay/Source/DisplayClusterShaders/Public/Containers/DisplayClusterShaderContainers_TextureUtils.h b/Engine/Plugins/Runtime/nDisplay/Source/DisplayClusterShaders/Public/Containers/DisplayClusterShaderContainers_TextureUtils.h
--- a/Engine/Plugins/Runtime/nDisplay/Source/DisplayClusterShaders/Public/Containers/DisplayClusterShaderContainers_TextureUtils.h
+++ b/Engine/Plugins/Runtime/nDisplay/Source/DisplayClusterShaders/Public/Containers/DisplayClusterShaderContainers_TextureUtils.h
@@ -9,8 +9,10 @@
 #include "RenderGraphFwd.h"
 #include "ScreenPass.h"
 
+#include "Math/MarginSet.h"
+
 /** Texture utiles flags. */
-enum class EDisplayClusterShaderTextureUtilsFlags : uint8
+enum class EDisplayClusterShaderTextureUtilsFlags : uint16
 {
 	None = 0,
 
@@ -33,11 +35,40 @@ enum class EDisplayClusterShaderTextureUtilsFlags : uint8
 	*/
 	DisableResize = 1 << 4,
 
-	/** Disable the use of the shader for resizing, color coding, etc. */
+	/**
+	* When enabled, bypasses all shader processing and performs a simple texture copy only.
+	*/
 	DisableResampleShader = 1 << 5,
 
-	/** Dont update resource rects, they are expected to be user defined. */
-	DisableUpdateResourcesRectsForResolve = 1 << 6
+	/**
+	* Dont update resource rects, they are expected to be user defined.
+	*/
+	DisableUpdateResourcesRectsForResolve = 1 << 6,
+
+	/**
+	* Applies a linear alpha fade (feathering) to all edges of the rectangle.
+	* Requires specifying the feather (attenuation) width in pixels per side.
+	* 
+	* NOTE: Only one alpha feathering mode(linear or smooth) should be enabled at a time.
+	*       Although these are bit flags, do not combine multiple alpha feather modes.
+	* 
+	* Note: Enabling this flag disables the use of EColorWriteMask,
+	*       so all color channels will be written regardless of mask settings.
+	*/
+	EnableLinearAlphaFeather = 1 << 8,
+
+	/**
+	* Enables smooth (cubic Hermite) alpha feathering (fade-out) on all sides of the rectangle.
+	* When set, alpha transitions smoothly from opaque to transparent along each edge using a smoothstep (S-curve) function.
+	* Requires additional parameters specifying the feather (attenuation) width in pixels for each side.
+	* 
+	* NOTE: Only one alpha feathering mode(linear or smooth) should be enabled at a time.
+	*       Although these are bit flags, do not combine multiple alpha feather modes.
+	* 
+	* Note: Enabling this flag disables the use of EColorWriteMask,
+	*       so all color channels will be written regardless of mask settings.
+	*/
+	EnableSmoothAlphaFeather = 1 << 9,
 };
 ENUM_CLASS_FLAGS(EDisplayClusterShaderTextureUtilsFlags);
 
@@ -180,10 +211,16 @@ struct FDisplayClusterShadersTextureUtilsSettings
 		const uint32 InDestSliceIndex = 0)
 		: ColorMask(InColorMask), Flags(InFlags), DestSliceIndex(InDestSliceIndex) { }
 
+	FDisplayClusterShadersTextureUtilsSettings(
+		const EDisplayClusterShaderTextureUtilsFlags InFlags,
+		const FIntMarginSet& InSideMargins)
+		: Flags(InFlags), SideMargins(InSideMargins) { }
+
 	FDisplayClusterShadersTextureUtilsSettings(
 		const EDisplayClusterShaderTextureUtilsFlags InFlags,
 		const uint32 InDestSliceIndex = 0)
-		: Flags(InFlags), DestSliceIndex(InDestSliceIndex) { }
+		: Flags(InFlags), DestSliceIndex(InDestSliceIndex) {
+	}
 
 	inline bool HasAnyFlags(const EDisplayClusterShaderTextureUtilsFlags InFlags) const
 	{
@@ -208,6 +245,9 @@ public:
 	/** Additional parameters for the FRHICopyTextureInfo. */
 	uint32 SourceSliceIndex = 0;
 	uint32 DestSliceIndex = 0;
+
+	/** Per-side margins in pixels. */
+	FIntMarginSet SideMargins;
 };
 
 /** This container is used for IDisplayClusterShadersTextureUtils::ForEachContextByPredicate() */
