[nDisplay] Fixed BugHawk issues

Issue #1/8
Engine/Plugins/Runtime/nDisplay/Source/DisplayCluster/Private/Render/Viewport/Containers/DisplayClusterViewport_CustomFrustumRuntimeSettings.cpp
Line: 93
UpdateProjectionAngles declares the InRenderTargetSize parameter and header comments indicate it is used, but the implementation never references it. This creates a documentation/code mismatch and can mislead maintainers or callers about the parameter???s effect.

Issue #2/8
Engine/Plugins/Runtime/nDisplay/Source/DisplayCluster/Private/Render/Viewport/Configuration/DisplayClusterViewportConfigurationHelpers_Tile.cpp
Line: 312
The inline comment states that a 'CVar' value is being clamped, but the code clamps a UPROPERTY-backed configuration value (OverscanBlendWidth). This is misleading and can confuse maintainers about where the value originates and how it should be controlled. While explicit runtime clamping is appropriate (UPROPERTY ClampMin/ClampMax are editor/UI hints and not guaranteed at runtime), the comment inaccurately labels the source as a console variable.

Issue #3/8
Engine/Plugins/Runtime/nDisplay/Source/DisplayCluster/Public/Render/Viewport/Containers/DisplayClusterViewport_TileSettings.h
Line: 71
Comment refers to 'EdgeBlendWidth' while the configuration property introduced for tile edge blending is 'OverscanBlendWidth', creating a naming inconsistency that can mislead maintainers and users.

Issue #4/8
Engine/Plugins/Runtime/nDisplay/Source/DisplayClusterConfiguration/Public/DisplayClusterConfigurationTypes_Tile.h
Line: 74
The property comment for OverscanBlendWidth states that values larger than the actual overscan size will be clamped to the overscan size, but the current runtime behavior expands the overscan to at least match the blend width (using ExtraPixels), effectively growing the overscan instead of clamping the blend width. This mismatch can mislead users and cause unexpected changes to viewport extents and performance.

Issue #5/8
Engine/Plugins/Runtime/nDisplay/Source/DisplayCluster/Internal/Math/MarginSet.h
Line: 555
Conflicting operator* overloads for TIntVector2: one multiplies components while a later, more constrained overload (selected for floating-point T) performs division. This violates expected operator* semantics and causes silently different behavior for floating vs non-floating T.
Why this is an issue:
Operator* is expected to implement multiplication semantics. The header defines an earlier operator* that multiplies by per-axis factors, but this later overload divides by the per-axis values. Because it???s more constrained (applies only for certain T) it can be preferred during overload resolution, meaning operator* will divide when T is floating-point and multiply otherwise, which is highly surprising and error-prone.

Issue #6/8
Engine/Plugins/Runtime/nDisplay/Source/DisplayCluster/Internal/Math/MarginSet.h
Line: 725
operator[](int32) index mapping comment contradicts actual implementation order, risking misuse by callers relying on documentation.
Code Context:
/* ... 0 = Left, 1 = Top, 2 = Right, 3 = Bottom. */ ... const EMarginSetSideIndex SideIndex = static_cast<EMarginSetSideIndex>(Index);

Issue #7/8
Engine/Plugins/Runtime/nDisplay/Source/DisplayCluster/Internal/Math/MarginSet.h
Line: 577
Misleading parameter type in documentation for operator/(const TVector2<U>&): comment mentions FIntPoint while the signature accepts TVector2<U>.
Note: Fixed typos and comments.

Issue #8/8
Engine/Plugins/Runtime/nDisplay/Source/DisplayCluster/Internal/Math/MarginSet.h
Line: 548
Comment for the second operator*(const TIntVector2<IntType>&) claims it multiplies margins, while the implementation divides by the per-axis values.

diff --git a/Engine/Plugins/Runtime/nDisplay/Source/DisplayCluster/Internal/Math/MarginSet.h b/Engine/Plugins/Runtime/nDisplay/Source/DisplayCluster/Internal/Math/MarginSet.h
--- a/Engine/Plugins/Runtime/nDisplay/Source/DisplayCluster/Internal/Math/MarginSet.h
+++ b/Engine/Plugins/Runtime/nDisplay/Source/DisplayCluster/Internal/Math/MarginSet.h
@@ -546,19 +546,19 @@ public:
 	}
 
 	/**
-	 * Multiplies each margin of a TMarginSet<T> by the corresponding axis of a TIntVector2<IntType>.
-	 * - Left and Right are multiplied by Divisor.X
-	 * - Top and Bottom are multiplied by Divisor.Y
+	 * Divides each margin of a TMarginSet<T> by the corresponding axis of an TIntVector2<IntType>.
+	 * - Left/Right are divided by Divisor.X
+	 * - Top/Bottom are divided by Divisor.Y
 	 *
-	 * @param Divisor  The TIntVector2<IntType> factors (X for horizontal, Y for vertical).
-	 * @return         A TMarginSet<T> where each margin is multiplied by the corresponding factor.
+	 * @param Divisor       The TIntVector2<IntType> divisor (X for horizontal, Y for vertical).
+	 * @return              A TMarginSet<T> where each margin is divided by the corresponding divisor.
 	 */
 	template <typename IntType>
 	requires (std::is_floating_point_v<T> && 
 		std::is_integral_v<IntType> &&
 		std::is_convertible_v<IntType, T>)
 	[[nodiscard]]
-	inline TMarginSet<T> operator*(const TIntVector2<IntType>& Divisor) const
+	inline TMarginSet<T> operator/(const TIntVector2<IntType>& Divisor) const
 	{
 		const T DivisorX = static_cast<T>(Divisor.X);
 		const T DivisorY = static_cast<T>(Divisor.Y);
@@ -575,11 +575,11 @@ public:
 	}
 
 	/**
-	 * Divides each margin of a TMarginSet<T> by the corresponding axis of an FIntPoint.
+	 * Divides each margin of a TMarginSet<T> by the corresponding axis of an TVector2<U>.
 	 * - Left/Right are divided by Divisor.X
 	 * - Top/Bottom are divided by Divisor.Y
 	 *
-	 * @param Divisor       The FIntPoint divisor (X for horizontal, Y for vertical).
+	 * @param Divisor       The TVector2<U> divisor (X for horizontal, Y for vertical).
 	 * @return              A TMarginSet<T> where each margin is divided by the corresponding divisor.
 	 */
 	template <typename U>
@@ -724,7 +724,7 @@ public:
 	 * Provides indexed access to a margin side value by numeric index.
 	 *
 	 * @param Index  Zero-based index of the side to access:
-	 *               0 = Left, 1 = Top, 2 = Right, 3 = Bottom.
+	 *               0 = Left, 1 = Right, 2 = Top, 3 = Bottom.
 	 * @return       Reference to the corresponding side value.
 	 */
 	[[nodiscard]] inline T& operator[](int32 Index)
@@ -739,7 +739,7 @@ public:
 	 * Provides  read-only indexed access to a margin side value by numeric index.
 	 *
 	 * @param Index  Zero-based index of the side to access:
-	 *               0 = Left, 1 = Top, 2 = Right, 3 = Bottom.
+	 *               0 = Left, 1 = Right, 2 = Top, 3 = Bottom.
 	 * @return       Const reference to the corresponding side value.
 	 */
 	[[nodiscard]] inline const T& operator[](int32 Index) const
diff --git a/Engine/Plugins/Runtime/nDisplay/Source/DisplayCluster/Private/Render/Viewport/Configuration/DisplayClusterViewportConfigurationHelpers_Tile.cpp b/Engine/Plugins/Runtime/nDisplay/Source/DisplayCluster/Private/Render/Viewport/Configuration/DisplayClusterViewportConfigurationHelpers_Tile.cpp
--- a/Engine/Plugins/Runtime/nDisplay/Source/DisplayCluster/Private/Render/Viewport/Configuration/DisplayClusterViewportConfigurationHelpers_Tile.cpp
+++ b/Engine/Plugins/Runtime/nDisplay/Source/DisplayCluster/Private/Render/Viewport/Configuration/DisplayClusterViewportConfigurationHelpers_Tile.cpp
@@ -306,11 +306,8 @@ FDisplayClusterViewport_OverscanSettings FDisplayClusterViewportConfigurationHel
 	// Set the minimum overscan values in pixels for tile blending
 	else if (InTileOverscan.OverscanBlendWidth > 0)
 	{
-		// The number of pixels in the blending area cannot be greater than this number.
-		constexpr int32 MaxPixelsToBlend = 512;
-
-		// Limit CVar value to safe range.
-		const int32 PixelsToBlend = FMath::Clamp(InTileOverscan.OverscanBlendWidth, 0, MaxPixelsToBlend);
+		// Get width of the blended edge region, in pixels.
+		const int32 PixelsToBlend = FMath::Max(0, InTileOverscan.OverscanBlendWidth);
 
 		// Apply tile edge blending only when OverscanBlendWidth > 0.
 		if(PixelsToBlend > 0)
diff --git a/Engine/Plugins/Runtime/nDisplay/Source/DisplayCluster/Private/Render/Viewport/Containers/DisplayClusterViewport_CustomFrustumRuntimeSettings.cpp b/Engine/Plugins/Runtime/nDisplay/Source/DisplayCluster/Private/Render/Viewport/Containers/DisplayClusterViewport_CustomFrustumRuntimeSettings.cpp
--- a/Engine/Plugins/Runtime/nDisplay/Source/DisplayCluster/Private/Render/Viewport/Containers/DisplayClusterViewport_CustomFrustumRuntimeSettings.cpp
+++ b/Engine/Plugins/Runtime/nDisplay/Source/DisplayCluster/Private/Render/Viewport/Containers/DisplayClusterViewport_CustomFrustumRuntimeSettings.cpp
@@ -81,7 +81,6 @@ namespace UE::DisplayCluster::Viewport::CustomFrustumHelpers
 ///////////////////////////////////////////////////////////////////////////////////////////
 bool FDisplayClusterViewport_CustomFrustumRuntimeSettings::UpdateProjectionAngles(
 	const FDisplayClusterViewport_CustomFrustumRuntimeSettings& InRuntimeSettings,
-	const FIntPoint& InRenderTargetSize,
 	double& InOutLeft,
 	double& InOutRight,
 	double& InOutTop,
diff --git a/Engine/Plugins/Runtime/nDisplay/Source/DisplayCluster/Private/Render/Viewport/Containers/DisplayClusterViewport_CustomFrustumRuntimeSettings.h b/Engine/Plugins/Runtime/nDisplay/Source/DisplayCluster/Private/Render/Viewport/Containers/DisplayClusterViewport_CustomFrustumRuntimeSettings.h
--- a/Engine/Plugins/Runtime/nDisplay/Source/DisplayCluster/Private/Render/Viewport/Containers/DisplayClusterViewport_CustomFrustumRuntimeSettings.h
+++ b/Engine/Plugins/Runtime/nDisplay/Source/DisplayCluster/Private/Render/Viewport/Containers/DisplayClusterViewport_CustomFrustumRuntimeSettings.h
@@ -29,7 +29,6 @@ struct FDisplayClusterViewport_CustomFrustumRuntimeSettings
 	/** Update projection angles by custom frustum
 	*
 	* @param InOverscanRuntimeSettings - the CustomFrustum runtime settings.
-	* @param InRenderTargetSize        - RTT size that used for custom frustum
 	* @param InOutLeft                 - the value of the left projection plane that you want to change
 	* @param InOutRight                - the value of the right projection plane that you want to change
 	* @param InOutTop                  - the value of the top projection plane that you want to change
@@ -37,7 +36,6 @@ struct FDisplayClusterViewport_CustomFrustumRuntimeSettings
 	*/
 	static bool UpdateProjectionAngles(
 		const FDisplayClusterViewport_CustomFrustumRuntimeSettings& InRuntimeSettings,
-		const FIntPoint& InRenderTargetSize,
 		double& InOutLeft,
 		double& InOutRight,
 		double& InOutTop,
diff --git a/Engine/Plugins/Runtime/nDisplay/Source/DisplayCluster/Private/Render/Viewport/Containers/DisplayClusterViewport_OverscanRuntimeSettings.cpp b/Engine/Plugins/Runtime/nDisplay/Source/DisplayCluster/Private/Render/Viewport/Containers/DisplayClusterViewport_OverscanRuntimeSettings.cpp
--- a/Engine/Plugins/Runtime/nDisplay/Source/DisplayCluster/Private/Render/Viewport/Containers/DisplayClusterViewport_OverscanRuntimeSettings.cpp
+++ b/Engine/Plugins/Runtime/nDisplay/Source/DisplayCluster/Private/Render/Viewport/Containers/DisplayClusterViewport_OverscanRuntimeSettings.cpp
@@ -27,7 +27,6 @@ static FAutoConsoleVariableRef CVarDisplayClusterRenderOverscanMaxValue(
 ///////////////////////////////////////////////////////////////////////////////////////////
 bool FDisplayClusterViewport_OverscanRuntimeSettings::UpdateProjectionAngles(
 	const FDisplayClusterViewport_OverscanRuntimeSettings& InOverscanRuntimeSettings,
-	const FIntPoint& InRenderTargetSize,
 	double& InOutLeft,
 	double& InOutRight,
 	double& InOutTop,
diff --git a/Engine/Plugins/Runtime/nDisplay/Source/DisplayCluster/Private/Render/Viewport/Containers/DisplayClusterViewport_OverscanRuntimeSettings.h b/Engine/Plugins/Runtime/nDisplay/Source/DisplayCluster/Private/Render/Viewport/Containers/DisplayClusterViewport_OverscanRuntimeSettings.h
--- a/Engine/Plugins/Runtime/nDisplay/Source/DisplayCluster/Private/Render/Viewport/Containers/DisplayClusterViewport_OverscanRuntimeSettings.h
+++ b/Engine/Plugins/Runtime/nDisplay/Source/DisplayCluster/Private/Render/Viewport/Containers/DisplayClusterViewport_OverscanRuntimeSettings.h
@@ -24,13 +24,12 @@ struct FDisplayClusterViewport_OverscanRuntimeSettings
 	/** Update projection angles by overscan
 	* 
 	* @param InOverscanRuntimeSettings - the overscan runtime settings.
-	* @param InRenderTargetSize         - RenderTarget texture size
 	* @param InOutLeft                 - the value of the left projection plane that you want to change
 	* @param InOutRight                - the value of the right projection plane that you want to change
 	* @param InOutTop                  - the value of the top projection plane that you want to change
 	* @parma InOutBottom               - the value of the bottom projection plane that you want to change
 	*/
-	static bool UpdateProjectionAngles(const FDisplayClusterViewport_OverscanRuntimeSettings& InOverscanRuntimeSettings, const FIntPoint& InRenderTargetSize, double& InOutLeft, double& InOutRight, double& InOutTop, double& InOutBottom);
+	static bool UpdateProjectionAngles(const FDisplayClusterViewport_OverscanRuntimeSettings& InOverscanRuntimeSettings, double& InOutLeft, double& InOutRight, double& InOutTop, double& InOutBottom);
 
 	/**
 	* Overscan values in pixels:
diff --git a/Engine/Plugins/Runtime/nDisplay/Source/DisplayCluster/Private/Render/Viewport/DisplayClusterViewportProxy_Resources.cpp b/Engine/Plugins/Runtime/nDisplay/Source/DisplayCluster/Private/Render/Viewport/DisplayClusterViewportProxy_Resources.cpp
--- a/Engine/Plugins/Runtime/nDisplay/Source/DisplayCluster/Private/Render/Viewport/DisplayClusterViewportProxy_Resources.cpp
+++ b/Engine/Plugins/Runtime/nDisplay/Source/DisplayCluster/Private/Render/Viewport/DisplayClusterViewportProxy_Resources.cpp
@@ -412,7 +412,7 @@ void FDisplayClusterViewportProxy::ImplResolveTileResource_RenderThread(FRHIComm
 				{
 					if (!RenderSettings.TileSettings.HasAnyTileFlags(EDisplayClusterViewportTileFlags::UseEntireOverscanForEdgeBlend))
 					{
-						// Use minimal margins from EdgeBlendWidth (ExtraPixels)
+						// Use minimal margins from OverscanBlendWidth (ExtraPixels)
 						InputRectMargins = InputRectMargins.GetMinPerSide(OverscanRuntimeSettings.ExtraPixels);
 						OutputRectMargins = OutputRectMargins.GetMinPerSide(OverscanRuntimeSettings.BaseExtraPixels);
 					}
diff --git a/Engine/Plugins/Runtime/nDisplay/Source/DisplayCluster/Private/Render/Viewport/DisplayClusterViewport_Math.cpp b/Engine/Plugins/Runtime/nDisplay/Source/DisplayCluster/Private/Render/Viewport/DisplayClusterViewport_Math.cpp
--- a/Engine/Plugins/Runtime/nDisplay/Source/DisplayCluster/Private/Render/Viewport/DisplayClusterViewport_Math.cpp
+++ b/Engine/Plugins/Runtime/nDisplay/Source/DisplayCluster/Private/Render/Viewport/DisplayClusterViewport_Math.cpp
@@ -307,7 +307,7 @@ void FDisplayClusterViewport::CalculateProjectionMatrix(const uint32 InContextNu
 
 	// Support custom frustum rendering
 	const double OrigValues[] = {l, r, t, b};
-	if (FDisplayClusterViewport_CustomFrustumRuntimeSettings::UpdateProjectionAngles(CustomFrustumRuntimeSettings, Contexts[InContextNum].RenderTargetRect.Size(), l, r, t, b))
+	if (FDisplayClusterViewport_CustomFrustumRuntimeSettings::UpdateProjectionAngles(CustomFrustumRuntimeSettings, l, r, t, b))
 	{
 		const bool bIsValidLimits =  FMath::IsWithin(l, -MaxValue, MaxValue)
 							&& FMath::IsWithin(r, -MaxValue, MaxValue)
@@ -344,7 +344,7 @@ void FDisplayClusterViewport::CalculateProjectionMatrix(const uint32 InContextNu
 	CachedProjectionData.ZFar = f;
 
 
-	if (FDisplayClusterViewport_OverscanRuntimeSettings::UpdateProjectionAngles(OverscanRuntimeSettings, Contexts[InContextNum].RenderTargetRect.Size(), l, r, t, b))
+	if (FDisplayClusterViewport_OverscanRuntimeSettings::UpdateProjectionAngles(OverscanRuntimeSettings, l, r, t, b))
 	{
 		if (FMath::IsWithin(l, -MaxValue, MaxValue) &&
 			FMath::IsWithin(r, -MaxValue, MaxValue) &&
diff --git a/Engine/Plugins/Runtime/nDisplay/Source/DisplayCluster/Public/Render/Viewport/Containers/DisplayClusterViewport_OverscanSettings.h b/Engine/Plugins/Runtime/nDisplay/Source/DisplayCluster/Public/Render/Viewport/Containers/DisplayClusterViewport_OverscanSettings.h
--- a/Engine/Plugins/Runtime/nDisplay/Source/DisplayCluster/Public/Render/Viewport/Containers/DisplayClusterViewport_OverscanSettings.h
+++ b/Engine/Plugins/Runtime/nDisplay/Source/DisplayCluster/Public/Render/Viewport/Containers/DisplayClusterViewport_OverscanSettings.h
@@ -25,7 +25,7 @@ struct FDisplayClusterViewport_OverscanSettings
 	uint8 bApplyTileEdgeBlend : 1 = 0;
 
 	// When enabled, uses the full overscan width/height for tile edge blending,
-	// ignoring the EdgeBlendWidth setting. Requires ApplyTileEdgeBlend to be active.
+	// ignoring the OverscanBlendWidth setting. Requires ApplyTileEdgeBlend to be active.
 	uint8 bUseEntireOverscanForEdgeBlend : 1 = 0;
 
 	// Units type of overscan values
diff --git a/Engine/Plugins/Runtime/nDisplay/Source/DisplayCluster/Public/Render/Viewport/Containers/DisplayClusterViewport_TileSettings.h b/Engine/Plugins/Runtime/nDisplay/Source/DisplayCluster/Public/Render/Viewport/Containers/DisplayClusterViewport_TileSettings.h
--- a/Engine/Plugins/Runtime/nDisplay/Source/DisplayCluster/Public/Render/Viewport/Containers/DisplayClusterViewport_TileSettings.h
+++ b/Engine/Plugins/Runtime/nDisplay/Source/DisplayCluster/Public/Render/Viewport/Containers/DisplayClusterViewport_TileSettings.h
@@ -69,7 +69,7 @@ enum class EDisplayClusterViewportTileFlags : uint8
 	ApplyTileEdgeBlend = 1 << 2,
 
 	// When enabled, uses the full overscan width/height for tile edge blending,
-	// ignoring the EdgeBlendWidth setting. Requires ApplyTileEdgeBlend to be active.
+	// ignoring the OverscanBlendWidth setting. Requires ApplyTileEdgeBlend to be active.
 	UseEntireOverscanForEdgeBlend = 1 << 3,
 };
 ENUM_CLASS_FLAGS(EDisplayClusterViewportTileFlags);
diff --git a/Engine/Plugins/Runtime/nDisplay/Source/DisplayClusterConfiguration/Public/DisplayClusterConfigurationTypes_Tile.h b/Engine/Plugins/Runtime/nDisplay/Source/DisplayClusterConfiguration/Public/DisplayClusterConfigurationTypes_Tile.h
--- a/Engine/Plugins/Runtime/nDisplay/Source/DisplayClusterConfiguration/Public/DisplayClusterConfigurationTypes_Tile.h
+++ b/Engine/Plugins/Runtime/nDisplay/Source/DisplayClusterConfiguration/Public/DisplayClusterConfigurationTypes_Tile.h
@@ -69,9 +69,8 @@ public:
 
 	/** Width of the blended edge region, in pixels.
 	 * Used to smooth seams by blending the overscan areas where adjacent tiles meet.
-	 * This value should be less than or equal to the overscan size to prevent artifacts.
 	 * 
-	 * Note: Values larger than the actual overscan size (in pixels) will be clamped to the overscan size.
+	 * Note: Values exceeding the actual overscan size (in pixels) will be used as the new overscan size.
 	 */
 	UPROPERTY(EditAnywhere, BlueprintReadWrite, Category = "Tile Overscan", meta = (DisplayName = "Blend Width", EditCondition = "bBlendEntireOverscan == false && bEnabled == true", ClampMin = "0", ClampMax = "512", UIMin = "0", UIMax = "512"))
 	int32 OverscanBlendWidth = 16;
