[nDisplay] Fixed BugHawk issues: fixed typo in comments.

Issue #1/2
Engine/Plugins/Runtime/nDisplay/Source/DisplayCluster/Private/Render/Viewport/Containers/DisplayClusterViewport_CustomFrustumRuntimeSettings.h
Line: 32
Doxygen [REDACTED_HANDLE] name mismatch: the comment documents a parameter named 'InOverscanRuntimeSettings' but the function signature uses 'InRuntimeSettings'. This mismatch can cause Doxygen to emit warnings and fail to associate the documentation with the actual parameter, leading to inaccurate or incomplete generated docs.

Issue #2/2
Engine/Plugins/Runtime/nDisplay/Source/DisplayCluster/Private/Render/Viewport/Containers/DisplayClusterViewport_CustomFrustumRuntimeSettings.h
Line: 32
Typo in Doxygen tag: '[REDACTED_HANDLE]' is used instead of '[REDACTED_HANDLE]' for the 'InOutBottom' parameter. Unrecognized commands are treated as normal text, so this parameter will not be documented properly and may trigger missing-parameter warnings.

diff --git a/Engine/Plugins/Runtime/nDisplay/Source/DisplayCluster/Private/Render/Viewport/Containers/DisplayClusterViewport_CustomFrustumRuntimeSettings.h b/Engine/Plugins/Runtime/nDisplay/Source/DisplayCluster/Private/Render/Viewport/Containers/DisplayClusterViewport_CustomFrustumRuntimeSettings.h
--- a/Engine/Plugins/Runtime/nDisplay/Source/DisplayCluster/Private/Render/Viewport/Containers/DisplayClusterViewport_CustomFrustumRuntimeSettings.h
+++ b/Engine/Plugins/Runtime/nDisplay/Source/DisplayCluster/Private/Render/Viewport/Containers/DisplayClusterViewport_CustomFrustumRuntimeSettings.h
@@ -28,14 +28,14 @@ struct FDisplayClusterViewport_CustomFrustumRuntimeSettings
 	
 	/** Update projection angles by custom frustum
 	*
-	* @param InOverscanRuntimeSettings - the CustomFrustum runtime settings.
+	* @param InCustomFrustumSettings   - the CustomFrustum runtime settings.
 	* @param InOutLeft                 - the value of the left projection plane that you want to change
 	* @param InOutRight                - the value of the right projection plane that you want to change
 	* @param InOutTop                  - the value of the top projection plane that you want to change
-	* @parma InOutBottom               - the value of the bottom projection plane that you want to change
+	* @param InOutBottom               - the value of the bottom projection plane that you want to change
 	*/
 	static bool UpdateProjectionAngles(
-		const FDisplayClusterViewport_CustomFrustumRuntimeSettings& InRuntimeSettings,
+		const FDisplayClusterViewport_CustomFrustumRuntimeSettings& InCustomFrustumSettings,
 		double& InOutLeft,
 		double& InOutRight,
 		double& InOutTop,
diff --git a/Engine/Plugins/Runtime/nDisplay/Source/DisplayCluster/Private/Render/Viewport/Containers/DisplayClusterViewport_OverscanRuntimeSettings.h b/Engine/Plugins/Runtime/nDisplay/Source/DisplayCluster/Private/Render/Viewport/Containers/DisplayClusterViewport_OverscanRuntimeSettings.h
--- a/Engine/Plugins/Runtime/nDisplay/Source/DisplayCluster/Private/Render/Viewport/Containers/DisplayClusterViewport_OverscanRuntimeSettings.h
+++ b/Engine/Plugins/Runtime/nDisplay/Source/DisplayCluster/Private/Render/Viewport/Containers/DisplayClusterViewport_OverscanRuntimeSettings.h
@@ -27,7 +27,7 @@ struct FDisplayClusterViewport_OverscanRuntimeSettings
 	* @param InOutLeft                 - the value of the left projection plane that you want to change
 	* @param InOutRight                - the value of the right projection plane that you want to change
 	* @param InOutTop                  - the value of the top projection plane that you want to change
-	* @parma InOutBottom               - the value of the bottom projection plane that you want to change
+	* @param InOutBottom               - the value of the bottom projection plane that you want to change
 	*/
 	static bool UpdateProjectionAngles(const FDisplayClusterViewport_OverscanRuntimeSettings& InOverscanRuntimeSettings, double& InOutLeft, double& InOutRight, double& InOutTop, double& InOutBottom);
 
