[nDisplay] When the upscaling method is set to "Default", nDisplay should always use TSR.

Adds FDisplayClusterUpscaler::InitializeNewFrame() to prepare all registered upscaler modular features at the beginning of each frame.

Explanation:
When UpscalingMethod = Default, UE::DisplayCluster::Upscaler::ForeachUpscaler() was not called, so its internal call to AddSceneViewExtensionIsActiveFunctor() also never executed, leaving upscalers unprepared before use.

diff --git a/Engine/Plugins/Runtime/nDisplay/Source/DisplayCluster/Private/Render/Upscaler/DisplayClusterUpscaler.cpp b/Engine/Plugins/Runtime/nDisplay/Source/DisplayCluster/Private/Render/Upscaler/DisplayClusterUpscaler.cpp
--- a/Engine/Plugins/Runtime/nDisplay/Source/DisplayCluster/Private/Render/Upscaler/DisplayClusterUpscaler.cpp
+++ b/Engine/Plugins/Runtime/nDisplay/Source/DisplayCluster/Private/Render/Upscaler/DisplayClusterUpscaler.cpp
@@ -52,12 +52,6 @@ namespace UE::DisplayCluster::Upscaler
 	static void ForeachUpscaler(TFunction<void(IUpscalerModularFeature&)> IteratorFunc)
 	{
 		static IModularFeatures& ModularFeatures = IModularFeatures::Get();
-		static TMap<FName, FGuid> SceneViewExtensionIsActiveFunctorsMap;
-
-		// If the modular function will be disabled at runtime, we should forget about it
-		// And then restore again when it is enabled.
-		TMap<FName, FGuid> NewSceneViewExtensionIsActiveFunctorsMap;
-
 		ModularFeatures.LockModularFeatureList();
 		ON_SCOPE_EXIT
 		{
@@ -73,64 +67,11 @@ namespace UE::DisplayCluster::Upscaler
 				continue;
 			}
 
-			// Register IsActive Functors
-			const FName& UpscalerName = UpscalerModularFeature->GetName();
-			if (SceneViewExtensionIsActiveFunctorsMap.Contains(UpscalerName))
-			{
-				// Use exist Guid
-				NewSceneViewExtensionIsActiveFunctorsMap.Emplace(UpscalerName, SceneViewExtensionIsActiveFunctorsMap[UpscalerName]);
-			}
-			else
-			{
-				FSceneViewExtensionIsActiveFunctor IsActiveFunctor;
-				IsActiveFunctor.IsActiveFunction = [UpscalerName](const ISceneViewExtension* SceneViewExtension, const FSceneViewExtensionContext& Context)
-				{
-					// Note: SceneViewExtension this pointer is always passed to functors from FSceneViewExtensionBase::IsActiveThisFrame, we confirm this remains so.
-					check(SceneViewExtension);
-
-					// Is nDisplay context
-					static const FDisplayClusterSceneViewExtensionContext DCViewExtensionContext;
-					if (Context.IsA(MoveTempIfPossible(DCViewExtensionContext)))
-					{
-						// Is viewport exist
-						const FDisplayClusterSceneViewExtensionContext& DisplayContext = static_cast<const FDisplayClusterSceneViewExtensionContext&>(Context);
-						if (DisplayContext.DisplayClusterViewport.IsValid())
-						{
-							const FDisplayClusterUpscalerSettings& UpscalerSettings = DisplayContext.DisplayClusterViewport->GetRenderSettings().UpscalerSettings;
-							if (!UpscalerSettings.bIsActive)
-							{
-								// Custom upscalers are disabled for this viewport
-								return TOptional<bool>(false);
-							}
-
-							if (UpscalerSettings.CustomUpscalerName != UpscalerName)
-							{
-								// This viewport uses a different upscaler.
-								return TOptional<bool>(false);
-							}
-
-							return TOptional<bool>(true);
-						}
-					}
-
-					return TOptional<bool>();
-				};
-
-				if (UpscalerModularFeature->AddSceneViewExtensionIsActiveFunctor(IsActiveFunctor))
-				{
-					NewSceneViewExtensionIsActiveFunctorsMap.Emplace(UpscalerName, IsActiveFunctor.GetGuid());
-				}
-			}
-
-
 			if (UpscalerModularFeature->IsFeatureEnabled())
 			{
 				IteratorFunc(*UpscalerModularFeature);
 			}
 		}
-
-		// Update
-		SceneViewExtensionIsActiveFunctorsMap = NewSceneViewExtensionIsActiveFunctorsMap;
 	}
 
 	/** Get global fraction. */
@@ -559,4 +500,72 @@ FName FDisplayClusterUpscaler::PostConfigureViewFamily(
 	InOutViewFamily.SecondaryViewFraction = UpscalerParam.SecondaryScreenPercentage;
 
 	return NAME_None;
-}
\ No newline at end of file
+}
+
+void FDisplayClusterUpscaler::InitializeNewFrame()
+{
+	using namespace UE::DisplayCluster::Upscaler;
+
+	// Store the upscaler name with function GUID
+	static TMap<FName, FGuid> SceneViewExtensionIsActiveFunctorsMap;
+
+	// If the modular function will be disabled at runtime, we should forget about it
+	// And then restore again when it is enabled.
+	TMap<FName, FGuid> NewSceneViewExtensionIsActiveFunctorsMap;
+
+	ForeachUpscaler(
+		[&](IUpscalerModularFeature& Upscaler)
+		{
+			// Register IsActive Functors
+			const FName& UpscalerName = Upscaler.GetName();
+			if (SceneViewExtensionIsActiveFunctorsMap.Contains(UpscalerName))
+			{
+				// Use exist Guid
+				NewSceneViewExtensionIsActiveFunctorsMap.Emplace(UpscalerName, SceneViewExtensionIsActiveFunctorsMap[UpscalerName]);
+			}
+			else
+			{
+				FSceneViewExtensionIsActiveFunctor IsActiveFunctor;
+				IsActiveFunctor.IsActiveFunction = [UpscalerName](const ISceneViewExtension* SceneViewExtension, const FSceneViewExtensionContext& Context)
+					{
+						// Note: SceneViewExtension this pointer is always passed to functors from FSceneViewExtensionBase::IsActiveThisFrame, we confirm this remains so.
+						check(SceneViewExtension);
+
+						// Is nDisplay context
+						static const FDisplayClusterSceneViewExtensionContext DCViewExtensionContext;
+						if (Context.IsA(MoveTempIfPossible(DCViewExtensionContext)))
+						{
+							// Is viewport exist
+							const FDisplayClusterSceneViewExtensionContext& DisplayContext = static_cast<const FDisplayClusterSceneViewExtensionContext&>(Context);
+							if (DisplayContext.DisplayClusterViewport.IsValid())
+							{
+								const FDisplayClusterUpscalerSettings& UpscalerSettings = DisplayContext.DisplayClusterViewport->GetRenderSettings().UpscalerSettings;
+								if (!UpscalerSettings.bIsActive)
+								{
+									// Custom upscalers are disabled for this viewport
+									return TOptional<bool>(false);
+								}
+
+								if (UpscalerSettings.CustomUpscalerName != UpscalerName)
+								{
+									// This viewport uses a different upscaler.
+									return TOptional<bool>(false);
+								}
+
+								return TOptional<bool>(true);
+							}
+						}
+
+						return TOptional<bool>();
+					};
+
+				if (Upscaler.AddSceneViewExtensionIsActiveFunctor(IsActiveFunctor))
+				{
+					NewSceneViewExtensionIsActiveFunctorsMap.Emplace(UpscalerName, IsActiveFunctor.GetGuid());
+				}
+			}
+		});
+
+	// Update
+	SceneViewExtensionIsActiveFunctorsMap = NewSceneViewExtensionIsActiveFunctorsMap;
+}
diff --git a/Engine/Plugins/Runtime/nDisplay/Source/DisplayCluster/Private/Render/Upscaler/DisplayClusterUpscaler.h b/Engine/Plugins/Runtime/nDisplay/Source/DisplayCluster/Private/Render/Upscaler/DisplayClusterUpscaler.h
--- a/Engine/Plugins/Runtime/nDisplay/Source/DisplayCluster/Private/Render/Upscaler/DisplayClusterUpscaler.h
+++ b/Engine/Plugins/Runtime/nDisplay/Source/DisplayCluster/Private/Render/Upscaler/DisplayClusterUpscaler.h
@@ -17,6 +17,12 @@ class FSceneView;
 class FDisplayClusterUpscaler
 {
 public:
+	/**
+	 * Prepares all registered upscaler modular features for the new frame.
+	 * Should be called at the start of each frame.
+	 */
+	static void InitializeNewFrame();
+
 	/** Configure upscaler and viewfamily for ertain settings.
 	* 
 	* @param InUpscalerSettings - (in) the settings to apply for viewport
diff --git a/Engine/Plugins/Runtime/nDisplay/Source/DisplayCluster/Private/Render/Viewport/DisplayClusterViewportManager.cpp b/Engine/Plugins/Runtime/nDisplay/Source/DisplayCluster/Private/Render/Viewport/DisplayClusterViewportManager.cpp
--- a/Engine/Plugins/Runtime/nDisplay/Source/DisplayCluster/Private/Render/Viewport/DisplayClusterViewportManager.cpp
+++ b/Engine/Plugins/Runtime/nDisplay/Source/DisplayCluster/Private/Render/Viewport/DisplayClusterViewportManager.cpp
@@ -711,6 +711,9 @@ void FDisplayClusterViewportManager::InitializeNewFrame()
 {
 	check(IsInGameThread());
 
+	// Prepares all registered upscaler modular features for the new frame.
+	FDisplayClusterUpscaler::InitializeNewFrame();
+
 	// Handle new frame for warp policies
 	FDisplayClusterWarpPolicyManager WarpPolicyManager;
 	WarpPolicyManager.HandleNewFrame(*this);
