[nDisplay] Fixed DCRA rendering issues:

- Fixed a bug that caused the DCRA preview in the editor to crash after deleting all cluster nodes.
- The cluster node name "All" is no longer reserved and is now available for use.

diff --git a/Engine/Plugins/Runtime/nDisplay/Source/DisplayCluster/Private/Game/EngineClasses/Scene/DisplayClusterRootActor.cpp b/Engine/Plugins/Runtime/nDisplay/Source/DisplayCluster/Private/Game/EngineClasses/Scene/DisplayClusterRootActor.cpp
--- a/Engine/Plugins/Runtime/nDisplay/Source/DisplayCluster/Private/Game/EngineClasses/Scene/DisplayClusterRootActor.cpp
+++ b/Engine/Plugins/Runtime/nDisplay/Source/DisplayCluster/Private/Game/EngineClasses/Scene/DisplayClusterRootActor.cpp
@@ -326,7 +326,7 @@ bool ADisplayClusterRootActor::IsPrimaryRootActorForPIE() const
 {
 #if WITH_EDITOR
 	// A preview node must be selected in DCRA so that it can be used in PIE mode
-	if (PreviewNodeId == DisplayClusterConfigurationStrings::gui::preview::PreviewNodeAll || PreviewNodeId == DisplayClusterConfigurationStrings::gui::preview::PreviewNodeNone)
+	if (PreviewNodeId == DisplayClusterConfigurationStrings::gui::preview::PreviewNodeNone)
 	{
 		return false;
 	}
diff --git a/Engine/Plugins/Runtime/nDisplay/Source/DisplayCluster/Private/Render/Viewport/Configuration/DisplayClusterViewportConfiguration.cpp b/Engine/Plugins/Runtime/nDisplay/Source/DisplayCluster/Private/Render/Viewport/Configuration/DisplayClusterViewportConfiguration.cpp
--- a/Engine/Plugins/Runtime/nDisplay/Source/DisplayCluster/Private/Render/Viewport/Configuration/DisplayClusterViewportConfiguration.cpp
+++ b/Engine/Plugins/Runtime/nDisplay/Source/DisplayCluster/Private/Render/Viewport/Configuration/DisplayClusterViewportConfiguration.cpp
@@ -375,9 +375,6 @@ bool FDisplayClusterViewportConfiguration::ImplUpdateConfiguration(EDisplayClust
 	FDisplayClusterViewportConfiguration_ICVFX            ConfigurationICVFX(*this);
 	FDisplayClusterViewportConfiguration_Tile             ConfigurationTile(*this);
 
-	// when InClusterNodeId==PreviewNodeAll, means that it is an undefined cluster node
-	const FString ClusterNodeId = InClusterNodeId == DisplayClusterConfigurationStrings::gui::preview::PreviewNodeAll ? TEXT("") : InClusterNodeId;
-
 	if (InViewportNames)
 	{
 		ConfigurationViewportManager.UpdateCustomViewports(*InViewportNames);
@@ -387,9 +384,13 @@ bool FDisplayClusterViewportConfiguration::ImplUpdateConfiguration(EDisplayClust
 	}
 	else
 	{
-		ConfigurationViewportManager.UpdateClusterNodeViewports(ClusterNodeId);
+		// Cluster node should be defined
+		check(!InClusterNodeId.IsEmpty());
+
+		ConfigurationViewportManager.UpdateClusterNodeViewports(InClusterNodeId);
+
 		// Use only valid values of cluster node id
-		SetClusterNodeId(ClusterNodeId);
+		SetClusterNodeId(InClusterNodeId);
 
 	}
 
@@ -406,7 +407,7 @@ bool FDisplayClusterViewportConfiguration::ImplUpdateConfiguration(EDisplayClust
 	if (!InViewportNames)
 	{
 		// Update postprocess for current cluster node
-		ConfigurationPostprocess.UpdateClusterNodePostProcess(ClusterNodeId);
+		ConfigurationPostprocess.UpdateClusterNodePostProcess(InClusterNodeId);
 	}
 
 	FDisplayClusterViewportConfigurationHelpers_RenderFrameSettings::PostUpdateRenderFrameConfiguration(*this);
diff --git a/Engine/Plugins/Runtime/nDisplay/Source/DisplayCluster/Private/Render/Viewport/Preview/DisplayClusterViewportManagerPreview.cpp b/Engine/Plugins/Runtime/nDisplay/Source/DisplayCluster/Private/Render/Viewport/Preview/DisplayClusterViewportManagerPreview.cpp
--- a/Engine/Plugins/Runtime/nDisplay/Source/DisplayCluster/Private/Render/Viewport/Preview/DisplayClusterViewportManagerPreview.cpp
+++ b/Engine/Plugins/Runtime/nDisplay/Source/DisplayCluster/Private/Render/Viewport/Preview/DisplayClusterViewportManagerPreview.cpp
@@ -19,9 +19,8 @@ FDisplayClusterViewportManagerPreview::FDisplayClusterViewportManagerPreview(con
 	: Configuration(InConfiguration)
 { }
 
-FString FDisplayClusterViewportManagerPreview::GetClusterNodeId(bool& bOutNextLoop) const
+FString FDisplayClusterViewportManagerPreview::GetClusterNodeId() const
 {
-	bOutNextLoop = false;
 	if (const UDisplayClusterConfigurationData* CurrentConfigData = Configuration->GetConfigurationData())
 	{
 		TArray<FString> ExistClusterNodesIDs;
@@ -45,8 +44,10 @@ FString FDisplayClusterViewportManagerPreview::GetClusterNodeId(bool& bOutNextLo
 			// next loop
 			if (!ExistClusterNodesIDs.IsValidIndex(NodeIndex))
 			{
-				bOutNextLoop = true;
 				NodeIndex = 0;
+
+				// When rendering of all nodes in the cluster is complete, generate events for subscribers
+				OnEntireClusterPreviewGenerated.ExecuteIfBound();
 			}
 
 			return ExistClusterNodesIDs[NodeIndex];
@@ -140,8 +141,16 @@ void FDisplayClusterViewportManagerPreview::OnPreviewRenderTick()
 
 	while (ViewportsAmmount > 0)
 	{
-		bool bNextLoop = false;
-		FString CurrentClusterNodeId = GetClusterNodeId(bNextLoop);
+		// The function GetClusterNodeId() returns an empty string if no nodes are found to render.
+		const FString CurrentClusterNodeId = GetClusterNodeId();
+		if (CurrentClusterNodeId.IsEmpty())
+		{
+			// No cluster nodes found.
+			Release();
+
+			return;
+		}
+
 		if (FirstClusterNodeId.IsEmpty())
 		{
 			FirstClusterNodeId = CurrentClusterNodeId;
@@ -161,8 +170,10 @@ void FDisplayClusterViewportManagerPreview::OnPreviewRenderTick()
 			return;
 		}
 
-		if (Configuration->GetClusterNodeId() != CurrentClusterNodeId || !PreviewRenderFrame.IsValid())
+		const FString& CurrentClusterNodeIdCfg = Configuration->GetClusterNodeId();
+		if (CurrentClusterNodeIdCfg != CurrentClusterNodeId || !PreviewRenderFrame.IsValid())
 		{
+			// Initialize new cluster node for render
 			if (!InitializeClusterNodePreview(Configuration->GetRenderFrameSettings().PreviewSettings.EntireClusterPreviewRenderMode, CurrentWorld, CurrentClusterNodeId, nullptr))
 			{
 				// this cluster node cannot be initialized, skip and try the next node
@@ -173,14 +184,9 @@ void FDisplayClusterViewportManagerPreview::OnPreviewRenderTick()
 		// Render viewports
 		ViewportsAmmount = RenderClusterNodePreview(ViewportsAmmount);
 
-		// Update cluster node
-		CurrentClusterNodeId = GetClusterNodeId(bNextLoop);
-
-		if (bNextLoop)
-		{
-			// When rendering of all nodes in the cluster is complete, generate events for subscribers
-			OnEntireClusterPreviewGenerated.ExecuteIfBound();
-		}
+		// Call this function after rendering to immediately perform
+		// OnEntireClusterPreviewGenerated() after rendering of all cluster nodes is complete.
+		GetClusterNodeId();
 	}
 }
 
diff --git a/Engine/Plugins/Runtime/nDisplay/Source/DisplayCluster/Private/Render/Viewport/Preview/DisplayClusterViewportManagerPreview.h b/Engine/Plugins/Runtime/nDisplay/Source/DisplayCluster/Private/Render/Viewport/Preview/DisplayClusterViewportManagerPreview.h
--- a/Engine/Plugins/Runtime/nDisplay/Source/DisplayCluster/Private/Render/Viewport/Preview/DisplayClusterViewportManagerPreview.h
+++ b/Engine/Plugins/Runtime/nDisplay/Source/DisplayCluster/Private/Render/Viewport/Preview/DisplayClusterViewportManagerPreview.h
@@ -90,7 +90,7 @@ protected:
 	void RenderFrustumPreviewForViewport(IDisplayClusterViewport* InViewport);
 
 	/** Getting the cluster node ID for preview. */
-	FString GetClusterNodeId(bool& bOutNextLoop) const;
+	FString GetClusterNodeId() const;
 
 	/** Render preview frustums. */
 	void RenderPreviewFrustums();
diff --git a/Engine/Plugins/Runtime/nDisplay/Source/DisplayClusterConfiguration/Public/DisplayClusterConfigurationStrings.h b/Engine/Plugins/Runtime/nDisplay/Source/DisplayClusterConfiguration/Public/DisplayClusterConfigurationStrings.h
--- a/Engine/Plugins/Runtime/nDisplay/Source/DisplayClusterConfiguration/Public/DisplayClusterConfigurationStrings.h
+++ b/Engine/Plugins/Runtime/nDisplay/Source/DisplayClusterConfiguration/Public/DisplayClusterConfigurationStrings.h
@@ -11,7 +11,6 @@ namespace DisplayClusterConfigurationStrings
 	{
 		namespace preview
 		{
-			static constexpr const TCHAR* PreviewNodeAll = TEXT("All");
 			static constexpr const TCHAR* PreviewNodeNone = TEXT("None");
 		}
 	}
