nDisplay: Fix crash when media capture is unable to get started.

#rnx

diff --git a/Engine/Plugins/Runtime/nDisplay/Source/DisplayClusterMedia/Private/Capture/DisplayClusterMediaCaptureBase.cpp b/Engine/Plugins/Runtime/nDisplay/Source/DisplayClusterMedia/Private/Capture/DisplayClusterMediaCaptureBase.cpp
--- a/Engine/Plugins/Runtime/nDisplay/Source/DisplayClusterMedia/Private/Capture/DisplayClusterMediaCaptureBase.cpp
+++ b/Engine/Plugins/Runtime/nDisplay/Source/DisplayClusterMedia/Private/Capture/DisplayClusterMediaCaptureBase.cpp
@@ -105,6 +105,10 @@ bool FDisplayClusterMediaCaptureBase::StartCapture()
 			bWasCaptureStarted = StartMediaCapture();
 			return bWasCaptureStarted;
 		}
+		else
+		{
+			UE_LOG(LogDisplayClusterMedia, Warning, TEXT("Failed to create MediaCapture '%s' from MediaOutput"), *GetMediaId());
+		}
 	}
 
 	return false;
@@ -129,6 +133,13 @@ void FDisplayClusterMediaCaptureBase::StopCapture()
 
 void FDisplayClusterMediaCaptureBase::ExportMediaData_RenderThread(FRDGBuilder& GraphBuilder, const FMediaOutputTextureInfo& TextureInfo)
 {
+	// Check if MediaCapture is available
+	if (!MediaCapture)
+	{
+		UE_LOG(LogDisplayClusterMedia, Verbose, TEXT("MediaCapture '%s' is null, unable to export on RT frame %lu"), *GetMediaId(), GFrameCounterRenderThread);
+		return;
+	}
+
 	// Check if request data is valid
 	if (!IsValidRequestData(TextureInfo))
 	{
