nDisplay: Fix UE_LOG format specifier of 64-bit frame counter in DisplayClusterMediaCaptureBase.cpp.

#rnx

diff --git a/Engine/Plugins/Runtime/nDisplay/Source/DisplayClusterMedia/Private/Capture/DisplayClusterMediaCaptureBase.cpp b/Engine/Plugins/Runtime/nDisplay/Source/DisplayClusterMedia/Private/Capture/DisplayClusterMediaCaptureBase.cpp
--- a/Engine/Plugins/Runtime/nDisplay/Source/DisplayClusterMedia/Private/Capture/DisplayClusterMediaCaptureBase.cpp
+++ b/Engine/Plugins/Runtime/nDisplay/Source/DisplayClusterMedia/Private/Capture/DisplayClusterMediaCaptureBase.cpp
@@ -136,14 +136,14 @@ void FDisplayClusterMediaCaptureBase::ExportMediaData_RenderThread(FRDGBuilder&
 	// Check if MediaCapture is available
 	if (!MediaCapture)
 	{
-		UE_LOG(LogDisplayClusterMedia, Verbose, TEXT("MediaCapture '%s' is null, unable to export on RT frame %lu"), *GetMediaId(), GFrameCounterRenderThread);
+		UE_LOG(LogDisplayClusterMedia, Verbose, TEXT("MediaCapture '%s' is null, unable to export on RT frame %llu"), *GetMediaId(), GFrameCounterRenderThread);
 		return;
 	}
 
 	// Check if request data is valid
 	if (!IsValidRequestData(TextureInfo))
 	{
-		UE_LOG(LogDisplayClusterMedia, Warning, TEXT("MediaCapture '%s': no capture performed on RT frame %lu"), *GetMediaId(), GFrameCounterRenderThread);
+		UE_LOG(LogDisplayClusterMedia, Warning, TEXT("MediaCapture '%s': no capture performed on RT frame %llu"), *GetMediaId(), GFrameCounterRenderThread);
 		return;
 	}
 
@@ -155,7 +155,7 @@ void FDisplayClusterMediaCaptureBase::ExportMediaData_RenderThread(FRDGBuilder&
 
 		LastSrcRegionSize = FIntSize(SrcRegionSize);
 
-		UE_LOG(LogDisplayClusterMedia, VeryVerbose, TEXT("MediaCapture '%s': Requested texture export [size=%dx%d, rect=%dx%d, format=%u] on RT frame '%lu'..."),
+		UE_LOG(LogDisplayClusterMedia, VeryVerbose, TEXT("MediaCapture '%s': Requested texture export [size=%dx%d, rect=%dx%d, format=%u] on RT frame '%llu'..."),
 			*GetMediaId(), SrcTextureSize.X, SrcTextureSize.Y, SrcRegionSize.X, SrcRegionSize.Y, (uint32)TextureInfo.Texture->Desc.Format, GFrameCounterRenderThread);
 	}
 
@@ -176,7 +176,7 @@ void FDisplayClusterMediaCaptureBase::ExportMediaData_RenderThread(FRDGBuilder&
 		FDisplayClusterShaderParameters_MediaPQ Parameters{ TextureInfo.Texture, TextureInfo.Region, TexturePQ, { {0, 0}, TexturePQ->Desc.Extent } };
 		IDisplayClusterShaders::Get().AddLinearToPQPass(GraphBuilder, Parameters);
 
-		UE_LOG(LogDisplayClusterMedia, VeryVerbose, TEXT("MediaCapture '%s': PQ exporting TexSize[%s], TexRect[%s], TexFormat[%u] on RT frame '%lu'..."),
+		UE_LOG(LogDisplayClusterMedia, VeryVerbose, TEXT("MediaCapture '%s': PQ exporting TexSize[%s], TexRect[%s], TexFormat[%u] on RT frame '%llu'..."),
 			*GetMediaId(),
 			*Parameters.OutputTexture->Desc.Extent.ToString(),
 			*Parameters.OutputRect.ToString(),
@@ -188,7 +188,7 @@ void FDisplayClusterMediaCaptureBase::ExportMediaData_RenderThread(FRDGBuilder&
 	}
 	else
 	{
-		UE_LOG(LogDisplayClusterMedia, VeryVerbose, TEXT("MediaCapture '%s': Direct exporting TexSize[%s], TexRect[%s], TexFormat[%u] on RT frame '%lu'..."),
+		UE_LOG(LogDisplayClusterMedia, VeryVerbose, TEXT("MediaCapture '%s': Direct exporting TexSize[%s], TexRect[%s], TexFormat[%u] on RT frame '%llu'..."),
 			*GetMediaId(),
 			*TextureInfo.Texture->Desc.Extent.ToString(),
 			*TextureInfo.Region.ToString(),
@@ -210,7 +210,7 @@ bool FDisplayClusterMediaCaptureBase::IsValidRequestData(const FMediaOutputTextu
 	// Check if source texture is valid
 	if (!TextureInfo.Texture)
 	{
-		UE_LOG(LogDisplayClusterMedia, Warning, TEXT("MediaCapture '%s': invalid source texture on RT frame %lu"), *GetMediaId(), GFrameCounterRenderThread);
+		UE_LOG(LogDisplayClusterMedia, Warning, TEXT("MediaCapture '%s': invalid source texture on RT frame %llu"), *GetMediaId(), GFrameCounterRenderThread);
 		return false;
 	}
 
@@ -226,7 +226,7 @@ bool FDisplayClusterMediaCaptureBase::IsValidRequestData(const FMediaOutputTextu
 
 	if (!bCorrectRegion)
 	{
-		UE_LOG(LogDisplayClusterMedia, Warning, TEXT("MediaCapture '%s': invalid source region on RT frame %lu"), *GetMediaId(), GFrameCounterRenderThread);
+		UE_LOG(LogDisplayClusterMedia, Warning, TEXT("MediaCapture '%s': invalid source region on RT frame %llu"), *GetMediaId(), GFrameCounterRenderThread);
 		return false;
 	}
 
