[nDisplay] Tick UDisplayClusterSyncTickComponent while game is paused
- This is required to unblock nDisplay syncrhonization pipeline when game is paused

#rnx

diff --git a/Engine/Plugins/Runtime/nDisplay/Source/DisplayCluster/Private/Game/EngineClasses/Scene/DisplayClusterSyncTickComponent.cpp b/Engine/Plugins/Runtime/nDisplay/Source/DisplayCluster/Private/Game/EngineClasses/Scene/DisplayClusterSyncTickComponent.cpp
--- a/Engine/Plugins/Runtime/nDisplay/Source/DisplayCluster/Private/Game/EngineClasses/Scene/DisplayClusterSyncTickComponent.cpp
+++ b/Engine/Plugins/Runtime/nDisplay/Source/DisplayCluster/Private/Game/EngineClasses/Scene/DisplayClusterSyncTickComponent.cpp
@@ -12,6 +12,7 @@ UDisplayClusterSyncTickComponent::UDisplayClusterSyncTickComponent(const FObject
 	: Super(ObjectInitializer)
 {
 	PrimaryComponentTick.bCanEverTick = true;
+	PrimaryComponentTick.bTickEvenWhenPaused = true;
 	PrimaryComponentTick.TickGroup = TG_PostPhysics;
 	PrimaryComponentTick.bStartWithTickEnabled = true;
 	bAutoActivate = true;
