[nDisplay] Blending the edges of tile viewports with each other.

reworked overscan blend options:
OverscanBlendMode = < Disable, 25%, 50%, 75%>, Default = 50%

Removed unused logic (ExtraPixels)

Code cleanup.

diff --git a/Engine/Plugins/Runtime/nDisplay/Source/DisplayCluster/Private/Render/Viewport/Configuration/DisplayClusterViewportConfigurationHelpers_Tile.cpp b/Engine/Plugins/Runtime/nDisplay/Source/DisplayCluster/Private/Render/Viewport/Configuration/DisplayClusterViewportConfigurationHelpers_Tile.cpp
--- a/Engine/Plugins/Runtime/nDisplay/Source/DisplayCluster/Private/Render/Viewport/Configuration/DisplayClusterViewportConfigurationHelpers_Tile.cpp
+++ b/Engine/Plugins/Runtime/nDisplay/Source/DisplayCluster/Private/Render/Viewport/Configuration/DisplayClusterViewportConfigurationHelpers_Tile.cpp
@@ -38,16 +38,6 @@ void FDisplayClusterViewportConfigurationHelpers_Tile::UpdateICVFXCameraViewport
 		TileFlags |= EDisplayClusterViewportTileFlags::OptimizeTileOverscan;
 	}
 
-	if (OverscanSettings.bApplyTileEdgeBlend)
-	{
-		TileFlags |= EDisplayClusterViewportTileFlags::ApplyTileEdgeBlend;
-
-		if (OverscanSettings.bUseEntireOverscanForEdgeBlend)
-		{
-			TileFlags |= EDisplayClusterViewportTileFlags::UseEntireOverscanForEdgeBlend;
-		}
-	}
-
 	// Set this viewport as a source for tile rendering.
 	FDisplayClusterViewport_TileSettings& OutTileSettings = InSourceViewport.GetRenderSettingsImpl().TileSettings;
 	OutTileSettings = FDisplayClusterViewport_TileSettings(InCameraMediaSettings.TiledSplitLayout, OverscanSettings, TileFlags);
@@ -197,9 +187,6 @@ FDisplayClusterViewport* FDisplayClusterViewportConfigurationHelpers_Tile::GetOr
 					{
 						InOutRenderSettings.OverscanSettings[InSideIndex] = 0;
 					}
-
-					// Always optimize ExtraPixels
-					InOutRenderSettings.OverscanSettings.ExtraPixels[InSideIndex] = 0;
 				};
 
 			if (InTilePos.X == 0)
@@ -297,26 +284,29 @@ FDisplayClusterViewport_OverscanSettings FDisplayClusterViewportConfigurationHel
 		return OutOverscanSettings;
 	}
 
-	// Configure tile edge blending feature
-	if (InTileOverscan.bBlendEntireOverscan)
+	switch (InTileOverscan.OverscanBlendMode)
 	{
+	case EDisplayClusterConfigurationViewportOverscanBlendMode::Percent25:
 		OutOverscanSettings.bApplyTileEdgeBlend = true;
-		OutOverscanSettings.bUseEntireOverscanForEdgeBlend = true;
-	}
-	// Set the minimum overscan values in pixels for tile blending
-	else if (InTileOverscan.OverscanBlendWidth > 0)
-	{
-		// Get width of the blended edge region, in pixels.
-		const int32 PixelsToBlend = FMath::Max(0, InTileOverscan.OverscanBlendWidth);
+		OutOverscanSettings.TileEdgeBlendPercentage = 0.25f;
+		break;
 
-		// Apply tile edge blending only when OverscanBlendWidth > 0.
-		if(PixelsToBlend > 0)
-		{
-			OutOverscanSettings.bApplyTileEdgeBlend = true;
+	case EDisplayClusterConfigurationViewportOverscanBlendMode::Percent50:
+		OutOverscanSettings.bApplyTileEdgeBlend = true;
+		OutOverscanSettings.TileEdgeBlendPercentage = 0.50f;
+		break;
+
+	case EDisplayClusterConfigurationViewportOverscanBlendMode::Percent75:
+		OutOverscanSettings.bApplyTileEdgeBlend = true;
+		OutOverscanSettings.TileEdgeBlendPercentage = 0.75f;
+		break;
+
+	case EDisplayClusterConfigurationViewportOverscanBlendMode::Disable:
+	default:
+		OutOverscanSettings.bApplyTileEdgeBlend = false;
+		OutOverscanSettings.TileEdgeBlendPercentage = 0.0f;
+		break;
 
-			// Set ExtraPixels 
-			OutOverscanSettings.ExtraPixels = FIntMarginSet(PixelsToBlend);
-		}
 	}
 
 	return OutOverscanSettings;
diff --git a/Engine/Plugins/Runtime/nDisplay/Source/DisplayCluster/Private/Render/Viewport/Containers/DisplayClusterViewport_OverscanRuntimeSettings.cpp b/Engine/Plugins/Runtime/nDisplay/Source/DisplayCluster/Private/Render/Viewport/Containers/DisplayClusterViewport_OverscanRuntimeSettings.cpp
--- a/Engine/Plugins/Runtime/nDisplay/Source/DisplayCluster/Private/Render/Viewport/Containers/DisplayClusterViewport_OverscanRuntimeSettings.cpp
+++ b/Engine/Plugins/Runtime/nDisplay/Source/DisplayCluster/Private/Render/Viewport/Containers/DisplayClusterViewport_OverscanRuntimeSettings.cpp
@@ -61,6 +61,13 @@ void FDisplayClusterViewport_OverscanRuntimeSettings::UpdateOverscanSettings(
 	}
 
 	const FIntPoint Size = InOutRenderTargetRect.Size();
+	if(Size.X <=0 || Size.Y <=0)
+	{
+		// If the Size is invalid (zero or negative), overscan is disabled.
+		InOutOverscanRuntimeSettings = FDisplayClusterViewport_OverscanRuntimeSettings();
+
+		return;
+	}
 
 	/**
 	* We can't use negative overscan values.
@@ -95,19 +102,12 @@ void FDisplayClusterViewport_OverscanRuntimeSettings::UpdateOverscanSettings(
 	// Update RTT size for overscan
 	if (InOutOverscanRuntimeSettings.bIsEnabled)
 	{
-		// Store extra pixels
-		InOutOverscanRuntimeSettings.ExtraPixels.AssignMargins(InOverscanSettings.ExtraPixels);
-		InOutOverscanRuntimeSettings.BaseExtraPixels.AssignMargins(InOverscanSettings.ExtraPixels);
-
 		// Calc pixels from percent
 		const FIntMarginSet OverscanPixels = (InOutOverscanRuntimeSettings.OverscanPercent * Size).GetRoundToInt();
 
-		// Sets each margin to at least the value in the ExtraPixels.
-		const FIntMarginSet ExpandedOverscanPixels = OverscanPixels.GetMaxPerSideIfDefined(InOverscanSettings.ExtraPixels);
-
 		// Update values
-		InOutOverscanRuntimeSettings.OverscanPixels.AssignMargins(ExpandedOverscanPixels);
-		InOutOverscanRuntimeSettings.BaseOverscanPixels.AssignMargins(ExpandedOverscanPixels);
+		InOutOverscanRuntimeSettings.OverscanPixels.AssignMargins(OverscanPixels);
+		InOutOverscanRuntimeSettings.BaseOverscanPixels.AssignMargins(OverscanPixels);
 
 		// Quantize the overscan percentage to exactly fit the number of pixels.
 		// This will avoid a mismatch between the overscanned frustum calculated in UpdateProjectionAngles and the pixel crop in GetFinalContextRect.
@@ -136,17 +136,49 @@ void FDisplayClusterViewport_OverscanRuntimeSettings::UpdateOverscanSettings(
 		}
 
 		// Configure overcan rendering within the existing RTT size.
+		{
+			// Calculate the effective overscan margins relative to the original texture size.
+			//
+			// Definitions:
+			//   Size                  = <W, H> (original texture resolution)
+			//   OverscanPixels        = {L, R, T, B} (absolute overscan in pixels)
+			//   OverscanPixels.Size() = <L + R, T + B>
+			//
+			// Steps:
+			//   1. OverscanSize = Size + OverscanPixels.Size()
+			//        "overscan-adjusted" total size, used only as a math reference
+			//
+			//   2. AdjustedOverscanPercent = OverscanPixels / OverscanSize
+			//        fraction of each margin relative to overscan-adjusted size
+			//
+			//   3. AdjustedOverscanPixels = Size * AdjustedOverscanPercent
+			//        overscan margins rescaled to match the original Size (in pixels)
+			//
+			//   4. InnerRegionSize = Size - AdjustedOverscanPixels.Size()
+			//        the effective inner region after overscan is applied
+			// 
+			//   5. Final OverscanPercent = AdjustedOverscanPixels / InnerRegionSize
+			//        overscan margins expressed as a fraction of the final usable region
+			//
+			const FMarginSet AdjustedOverscanPercent = FMarginSet(InOutOverscanRuntimeSettings.OverscanPixels) / OverscanSize;
+			const FOverscanPixels AdjustedOverscanPixels((AdjustedOverscanPercent * Size).GetRoundToInt());
+
+			// Compute the scaled inner region size based on overscan margins.
+			// Mathematically, InnerRegionSize > 0 as long as Size > 0.
+			const FIntPoint InnerRegionSize = Size - AdjustedOverscanPixels.Size();
+			if (InnerRegionSize.X <= 0 || InnerRegionSize.Y <= 0)
+			{
+				// If the InnerRegionSize is invalid (zero or negative), overscan is disabled.
+				InOutOverscanRuntimeSettings = FDisplayClusterViewport_OverscanRuntimeSettings();
 
-		// 1. Update value of overscan pixels
-		const FMarginSet NewOverscanPixelsPercent = FMarginSet(InOutOverscanRuntimeSettings.OverscanPixels) / OverscanSize;
-		InOutOverscanRuntimeSettings.OverscanPixels.AssignMargins((NewOverscanPixelsPercent * Size).GetRoundToInt());
+				return;
+			}
 
-		// 2. Update value of extra pixels
-		const FMarginSet NewExtraPixelsPercent = FMarginSet(InOutOverscanRuntimeSettings.ExtraPixels) / OverscanSize;
-		InOutOverscanRuntimeSettings.ExtraPixels.AssignMargins((NewExtraPixelsPercent * Size).GetRoundToInt());
+			// Quantize the overscan percentage to exactly fit the number of pixels.
+			InOutOverscanRuntimeSettings.OverscanPercent = FMarginSet(AdjustedOverscanPixels) / InnerRegionSize;
 
-		// 3. Quantize the overscan percentage to exactly fit the number of pixels.
-		const FIntPoint ScaledSize = Size - InOutOverscanRuntimeSettings.OverscanPixels.Size();
-		InOutOverscanRuntimeSettings.OverscanPercent = FMarginSet(InOutOverscanRuntimeSettings.OverscanPixels) / ScaledSize;
+			// Update OverscanPixels with new value
+			InOutOverscanRuntimeSettings.OverscanPixels = AdjustedOverscanPixels;
+		}
 	}
 }
diff --git a/Engine/Plugins/Runtime/nDisplay/Source/DisplayCluster/Private/Render/Viewport/Containers/DisplayClusterViewport_OverscanRuntimeSettings.h b/Engine/Plugins/Runtime/nDisplay/Source/DisplayCluster/Private/Render/Viewport/Containers/DisplayClusterViewport_OverscanRuntimeSettings.h
--- a/Engine/Plugins/Runtime/nDisplay/Source/DisplayCluster/Private/Render/Viewport/Containers/DisplayClusterViewport_OverscanRuntimeSettings.h
+++ b/Engine/Plugins/Runtime/nDisplay/Source/DisplayCluster/Private/Render/Viewport/Containers/DisplayClusterViewport_OverscanRuntimeSettings.h
@@ -40,20 +40,6 @@ struct FDisplayClusterViewport_OverscanRuntimeSettings
 	*    |        +-------+         |
 	*    |         Bottom           |
 	*    +--------------------------+
-	* 
-	* Logic:
-	*   OverscanPixels = OverscanPixels.MaxPerSideIfDefined(ExtraPixels)
-	* 
-	* This ensures each overscan side is at least as large as the corresponding
-	* ExtraPixels value, but only if ExtraPixels is positive.
-	* 
-	* Example (horizontal layout):
-	* 
-	*  Current OverscanPixels:   |<---- 8 px ---->|
-	*  ExtraPixels:              |<-------- 12 px -------->|
-	* 
-	* Since ExtraPixels (12) > OverscanPixels (8) and ExtraPixels > 0,
-	* OverscanPixels will be increased to 12 px on this side.
 	*/
 	struct FOverscanPixels
 		: public FIntMarginSet
@@ -88,17 +74,9 @@ public:
 	// See bOversize for details on how resolution adjustments are applied.
 	FOverscanPixels OverscanPixels;
 
-	// Runtime overscan values for FDisplayClusterViewport_OverscanSettings::ExtraPixels.
-	// Used for blending between adjacent tiles in the viewport.
-	FOverscanPixels ExtraPixels;
-
-public:
 	// The tile viewport may render at a lower resolution but will be stretched to its final size in the original viewport.
 	// Stores the original values before any size adjustments, serving as the reference for overscan size modifiers.
 
 	// Original overscan size per side (in pixels) before any resolution adjustments.
 	FOverscanPixels BaseOverscanPixels;
-
-	// Original (pre–resolution-adjustment) overscan values for FDisplayClusterViewport_OverscanSettings::ExtraPixels.
-	FOverscanPixels BaseExtraPixels;
 };
diff --git a/Engine/Plugins/Runtime/nDisplay/Source/DisplayCluster/Private/Render/Viewport/DisplayClusterViewportProxy_Resources.cpp b/Engine/Plugins/Runtime/nDisplay/Source/DisplayCluster/Private/Render/Viewport/DisplayClusterViewportProxy_Resources.cpp
--- a/Engine/Plugins/Runtime/nDisplay/Source/DisplayCluster/Private/Render/Viewport/DisplayClusterViewportProxy_Resources.cpp
+++ b/Engine/Plugins/Runtime/nDisplay/Source/DisplayCluster/Private/Render/Viewport/DisplayClusterViewportProxy_Resources.cpp
@@ -331,15 +331,70 @@ void FDisplayClusterViewportProxy::ImplResolveTileResource_RenderThread(FRHIComm
 	// Note: The AdaptResolution option in overscan settings controls whether the RTT size 
 	// is increased to match the overscan size. If disabled, only the inner rectangle 
 	// matching the original viewport size is used as the input.
-	FIntMarginSet InputRectMargins = static_cast<FIntMarginSet>(OverscanRuntimeSettings.OverscanPixels);
-	FIntMarginSet OutputRectMargins = static_cast<FIntMarginSet>(OverscanRuntimeSettings.BaseOverscanPixels);
+
+	// Get tile position ans size
+	const FIntPoint InTilePos = RenderSettings.TileSettings.GetPos();
+	const FIntPoint InTileSize = RenderSettings.TileSettings.GetSize();
 
 	// Get actual tile edge blending mode
 	const EDisplayClusterRenderTileEdgeBlendMode TileEdgeBlendMode = GetTileEdgeBlendMode();
 
-	// Check if we can apply tile edge blending
-	const bool bApplyTileEdgeBlend = RenderSettings.TileSettings.HasAnyTileFlags(EDisplayClusterViewportTileFlags::ApplyTileEdgeBlend)
+	// Since we are blending into the destination viewport, 
+	// we should use its settings for blending.
+	const FDisplayClusterViewport_OverscanSettings& DestOverscanSettings = InDestViewportProxy->GetRenderSettings_RenderThread().TileSettings.GetOverscanSettings();
+
+	// Check if we can apply tile edge blending for dest viewport
+	bool bApplyTileEdgeBlend = DestOverscanSettings.bApplyTileEdgeBlend
 		&& TileEdgeBlendMode != EDisplayClusterRenderTileEdgeBlendMode::None;
+	
+	// If so, get input and output margins.
+	FIntMarginSet InputRectEdgeBlendMargins;
+	FIntMarginSet OutputRectEdgeBlendMargins;
+	if (bApplyTileEdgeBlend)
+	{
+		// Clamp to 0%..100%: [0..1]
+		const float OverscanPercentage = FMath::Clamp(DestOverscanSettings.TileEdgeBlendPercentage, 0.0f, 1.0f);
+
+		// Use percentage of overscan
+		InputRectEdgeBlendMargins = (FMarginSet(static_cast<FIntMarginSet>(OverscanRuntimeSettings.OverscanPixels)) * OverscanPercentage).GetCeilToInt();
+		OutputRectEdgeBlendMargins = (FMarginSet(static_cast<FIntMarginSet>(OverscanRuntimeSettings.BaseOverscanPixels)) * OverscanPercentage).GetCeilToInt();
+
+		// Per-tile adjustments: Optimize overscan values for edge tiles.
+		{
+			using EMarginSide = UE::Math::EMarginSetSideIndex;
+			auto OptimizeSize = [&](const EMarginSide InSideIndex)
+				{
+					InputRectEdgeBlendMargins[InSideIndex] = 0;
+					OutputRectEdgeBlendMargins[InSideIndex] = 0;
+				};
+
+			if (InTilePos.X == 0)
+			{
+				OptimizeSize(EMarginSide::Left);
+			}
+			if (InTilePos.Y == 0)
+			{
+				OptimizeSize(EMarginSide::Top);
+			}
+			if (InTilePos.X == (InTileSize.X - 1))
+			{
+				OptimizeSize(EMarginSide::Right);
+			}
+			if (InTilePos.Y == (InTileSize.Y - 1))
+			{
+				OptimizeSize(EMarginSide::Bottom);
+			}
+		}
+
+		// check if margins not empty
+		if (InputRectEdgeBlendMargins.IsZero() || OutputRectEdgeBlendMargins.IsZero())
+		{
+			// clear margins and avoid edge blending
+			InputRectEdgeBlendMargins = 0;
+			OutputRectEdgeBlendMargins = 0;
+			bApplyTileEdgeBlend = false;
+		}
+	}
 
 	const TArray<FDisplayClusterViewport_Context>& DestContexts = InDestViewportProxy->GetContexts_RenderThread();
 	const TArray<FDisplayClusterViewport_Context>& SrcContexts = Contexts;
@@ -367,10 +422,6 @@ void FDisplayClusterViewportProxy::ImplResolveTileResource_RenderThread(FRHIComm
 				// By default, we just copy part of the tile texture into dest.
 				EDisplayClusterShaderTextureUtilsFlags Flags = EDisplayClusterShaderTextureUtilsFlags::None;
 
-				// Get tile position ans size
-				const FIntPoint InTilePos = RenderSettings.TileSettings.GetPos();
-				const FIntPoint InTileSize = RenderSettings.TileSettings.GetSize();
-
 				if (GDisplayClusterRenderTileBorder > 0)
 				{
 					// The maximum border is 1/4 of the minimum side of the rectangle.
@@ -396,54 +447,27 @@ void FDisplayClusterViewportProxy::ImplResolveTileResource_RenderThread(FRHIComm
 						SubRect.Max.Y = TileBorderSize;
 					}
 
+					// support when bAdaptResolution == false
+					// Input.Rect.Size != Output.Rect.Size()
+					const FVector2D Scale(
+						double(Output.Rect.Width()) / Input.Rect.Width(),
+						double(Output.Rect.Height()) / Input.Rect.Height()
+					);
+
 					// Set rect smaller to show gaps between tiles:
 					Input.Rect.Min.X += SubRect.Min.X;
 					Input.Rect.Min.Y += SubRect.Min.Y;
 					Input.Rect.Max.X -= SubRect.Min.X + SubRect.Max.X;
 					Input.Rect.Max.Y -= SubRect.Min.Y + SubRect.Max.Y;
 
-					Output.Rect.Min.X += SubRect.Min.X;
-					Output.Rect.Min.Y += SubRect.Min.Y;
-					Output.Rect.Max.X -= SubRect.Min.X + SubRect.Max.X;
-					Output.Rect.Max.Y -= SubRect.Min.Y + SubRect.Max.Y;
+					Output.Rect.Min.X += SubRect.Min.X * Scale.X;
+					Output.Rect.Min.Y += SubRect.Min.Y * Scale.Y;
+					Output.Rect.Max.X -= (SubRect.Min.X + SubRect.Max.X) * Scale.X;
+					Output.Rect.Max.Y -= (SubRect.Min.Y + SubRect.Max.Y) * Scale.Y;
 				}
 				// When we don't need to show tile borders, we can blend them together
 				else if (bApplyTileEdgeBlend)
 				{
-					if (!RenderSettings.TileSettings.HasAnyTileFlags(EDisplayClusterViewportTileFlags::UseEntireOverscanForEdgeBlend))
-					{
-						// Use minimal margins from OverscanBlendWidth (ExtraPixels)
-						InputRectMargins = InputRectMargins.GetMinPerSide(OverscanRuntimeSettings.ExtraPixels);
-						OutputRectMargins = OutputRectMargins.GetMinPerSide(OverscanRuntimeSettings.BaseExtraPixels);
-					}
-
-					// Per-tile adjustments: Optimize overscan values for edge tiles.
-					{
-						using EMarginSide = UE::Math::EMarginSetSideIndex;
-						auto OptimizeSize = [&](const EMarginSide InSideIndex)
-							{
-								InputRectMargins[InSideIndex] = 0;
-								OutputRectMargins[InSideIndex] = 0;
-							};
-
-						if (InTilePos.X == 0)
-						{
-							OptimizeSize(EMarginSide::Left);
-						}
-						if (InTilePos.Y == 0)
-						{
-							OptimizeSize(EMarginSide::Top);
-						}
-						if (InTilePos.X == (InTileSize.X - 1))
-						{
-							OptimizeSize(EMarginSide::Right);
-						}
-						if (InTilePos.Y == (InTileSize.Y - 1))
-						{
-							OptimizeSize(EMarginSide::Bottom);
-						}
-					}
-
 					// Configure rendering flags that define how tile edge blending is applied
 					switch (TileEdgeBlendMode)
 					{
@@ -464,13 +488,13 @@ void FDisplayClusterViewportProxy::ImplResolveTileResource_RenderThread(FRHIComm
 					};
 
 					// Add margins to input and output rects
-					InputRectMargins.ApplyExpandToRect(Input.Rect);
-					OutputRectMargins.ApplyExpandToRect(Output.Rect);
+					InputRectEdgeBlendMargins.ApplyExpandToRect(Input.Rect);
+					OutputRectEdgeBlendMargins.ApplyExpandToRect(Output.Rect);
 				}
 
 				// Copy texture region
 				TextureUtils->ResolveTextureContext(
-					{ Flags,  InputRectMargins * 2},
+					{ Flags,  InputRectEdgeBlendMargins * 2},
 					Input, Output);
 			});
 }
@@ -479,7 +503,11 @@ void FDisplayClusterViewportProxy::CleanupResources_RenderThread(FRHICommandList
 {
 	bool bShouldCleanupRTT = false;
 
-	// Rules for tile viewport
+	/**
+	 * Tile viewport rules:
+	 * The [EDisplayClusterViewportTileType::Source] tile viewport represents the destination surface where all tiles
+	 * are composed and blended together into the final frame.
+	 */
 	if (RenderSettings.TileSettings.GetType() == EDisplayClusterViewportTileType::Source)
 	{
 		if (GDisplayClusterRenderTileBorder > 0)
@@ -489,7 +517,8 @@ void FDisplayClusterViewportProxy::CleanupResources_RenderThread(FRHICommandList
 			bShouldCleanupRTT = true;
 		}
 
-		if (RenderSettings.TileSettings.HasAnyTileFlags(EDisplayClusterViewportTileFlags::ApplyTileEdgeBlend))
+		// Check if this viewport applies blending between tiles.
+		if (RenderSettings.TileSettings.GetOverscanSettings().bApplyTileEdgeBlend)
 		{
 			// Blending assumes the underlying surface is black.
 			// Tiles are composited on top of each other across multiple render passes.
diff --git a/Engine/Plugins/Runtime/nDisplay/Source/DisplayCluster/Public/Render/Viewport/Containers/DisplayClusterViewport_OverscanSettings.h b/Engine/Plugins/Runtime/nDisplay/Source/DisplayCluster/Public/Render/Viewport/Containers/DisplayClusterViewport_OverscanSettings.h
--- a/Engine/Plugins/Runtime/nDisplay/Source/DisplayCluster/Public/Render/Viewport/Containers/DisplayClusterViewport_OverscanSettings.h
+++ b/Engine/Plugins/Runtime/nDisplay/Source/DisplayCluster/Public/Render/Viewport/Containers/DisplayClusterViewport_OverscanSettings.h
@@ -24,14 +24,9 @@ struct FDisplayClusterViewport_OverscanSettings
 	// ensuring a smooth transition with no visible seams.
 	uint8 bApplyTileEdgeBlend : 1 = 0;
 
-	// When enabled, uses the full overscan width/height for tile edge blending,
-	// ignoring the OverscanBlendWidth setting. Requires ApplyTileEdgeBlend to be active.
-	uint8 bUseEntireOverscanForEdgeBlend : 1 = 0;
+	// Percentage of the overscan (0.0–1.0) to use for edge blending.
+	float TileEdgeBlendPercentage = 0.0f;
 
 	// Units type of overscan values
 	EDisplayClusterViewport_FrustumUnit Unit = EDisplayClusterViewport_FrustumUnit::Pixels;
-	
-	// Additional margins in pixels that are guaranteed to be rendered.
-	// Note: this parameter is used to blend the viewport tiles.
-	FIntMarginSet ExtraPixels;
 };
diff --git a/Engine/Plugins/Runtime/nDisplay/Source/DisplayCluster/Public/Render/Viewport/Containers/DisplayClusterViewport_TileSettings.h b/Engine/Plugins/Runtime/nDisplay/Source/DisplayCluster/Public/Render/Viewport/Containers/DisplayClusterViewport_TileSettings.h
--- a/Engine/Plugins/Runtime/nDisplay/Source/DisplayCluster/Public/Render/Viewport/Containers/DisplayClusterViewport_TileSettings.h
+++ b/Engine/Plugins/Runtime/nDisplay/Source/DisplayCluster/Public/Render/Viewport/Containers/DisplayClusterViewport_TileSettings.h
@@ -62,15 +62,6 @@ enum class EDisplayClusterViewportTileFlags : uint8
 	//  | Tile0_1  ++|++ Tile1_1 |
 	//  +------------+-----------+
 	OptimizeTileOverscan = 1 << 1,
-
-	// Blends the overlapping (overscan) regions between adjacent tiles.
-	// The size (in pixels) must be less than or equal to the overscan size,
-	// ensuring a smooth transition with no visible seams.
-	ApplyTileEdgeBlend = 1 << 2,
-
-	// When enabled, uses the full overscan width/height for tile edge blending,
-	// ignoring the OverscanBlendWidth setting. Requires ApplyTileEdgeBlend to be active.
-	UseEntireOverscanForEdgeBlend = 1 << 3,
 };
 ENUM_CLASS_FLAGS(EDisplayClusterViewportTileFlags);
 
diff --git a/Engine/Plugins/Runtime/nDisplay/Source/DisplayClusterConfiguration/Public/DisplayClusterConfigurationTypes_Enums.h b/Engine/Plugins/Runtime/nDisplay/Source/DisplayClusterConfiguration/Public/DisplayClusterConfigurationTypes_Enums.h
--- a/Engine/Plugins/Runtime/nDisplay/Source/DisplayClusterConfiguration/Public/DisplayClusterConfigurationTypes_Enums.h
+++ b/Engine/Plugins/Runtime/nDisplay/Source/DisplayClusterConfiguration/Public/DisplayClusterConfigurationTypes_Enums.h
@@ -212,6 +212,28 @@ enum class EDisplayClusterConfigurationViewportOverscanMode : uint8
 	Percent  UMETA(DisplayName = "Percent")
 };
 
+/**
+* Defines how the overscan blend width is interpreted.
+* Allows blending between tiles to be specified in
+* a percentage of the overscan size, or disabled entirely.
+*/
+UENUM()
+enum class EDisplayClusterConfigurationViewportOverscanBlendMode : uint8
+{
+	/** Disable blending across overscan region */
+	Disable     UMETA(DisplayName = "Disable"),
+
+	/** Blend width defined as 25% of the overscan size */
+	Percent25   UMETA(DisplayName = "25%"),
+
+	/** Blend width defined as 50% of the overscan size */
+	Percent50   UMETA(DisplayName = "50%"),
+
+	/** Blend width defined as 75% of the overscan size */
+	Percent75   UMETA(DisplayName = "75%"),
+};
+
+
 UENUM()
 enum class EDisplayClusterConfigurationRenderMode : uint8
 {
diff --git a/Engine/Plugins/Runtime/nDisplay/Source/DisplayClusterConfiguration/Public/DisplayClusterConfigurationTypes_Tile.h b/Engine/Plugins/Runtime/nDisplay/Source/DisplayClusterConfiguration/Public/DisplayClusterConfigurationTypes_Tile.h
--- a/Engine/Plugins/Runtime/nDisplay/Source/DisplayClusterConfiguration/Public/DisplayClusterConfigurationTypes_Tile.h
+++ b/Engine/Plugins/Runtime/nDisplay/Source/DisplayClusterConfiguration/Public/DisplayClusterConfigurationTypes_Tile.h
@@ -61,17 +61,11 @@ public:
 	UPROPERTY(EditAnywhere, BlueprintReadWrite, Category = "Tile Overscan", meta = (DisplayName = "Overscan Value", EditCondition = "bEnabled == true"))
 	float AllSides = 10;
 
-	/** If true, blend across the entire overscan edge region instead of using OverscanBlendWidth.
-	 *  This applies the full overscan width/height for blending at tile edges.
-	 */
-	UPROPERTY(EditAnywhere, BlueprintReadWrite, Category = "Tile Overscan", meta = (DisplayName = "Blend Entire Overscan", EditCondition = "bEnabled == true"))
-	bool bBlendEntireOverscan = false;
-
-	/** Width of the blended edge region, in pixels.
-	 * Used to smooth seams by blending the overscan areas where adjacent tiles meet.
-	 * 
-	 * Note: Values exceeding the actual overscan size (in pixels) will be used as the new overscan size.
-	 */
-	UPROPERTY(EditAnywhere, BlueprintReadWrite, Category = "Tile Overscan", meta = (DisplayName = "Blend Width", EditCondition = "bBlendEntireOverscan == false && bEnabled == true", ClampMin = "0", ClampMax = "512", UIMin = "0", UIMax = "512"))
-	int32 OverscanBlendWidth = 16;
+	 /**
+	  * Defines how the overscan blend width is interpreted.
+	  * Allows blending between tiles to be specified in
+	  * a percentage of the overscan size, or disabled entirely.
+	  */
+	UPROPERTY(EditAnywhere, BlueprintReadWrite, Category = "Tile Overscan", meta = (DisplayName = "Blend Mode", EditCondition = "bEnabled == true"))
+	EDisplayClusterConfigurationViewportOverscanBlendMode OverscanBlendMode = EDisplayClusterConfigurationViewportOverscanBlendMode::Percent50;
 };
