[nDisplay] Fixed BugHaws issues

- Potential division by zero
//Engine/Plugins/Runtime/nDisplay/Source/DisplayCluster/Private/Render/Viewport/DisplayClusterViewportProxy_Resources.cpp: 452

- Misleading comment
//Engine/Plugins/Runtime/nDisplay/Source/DisplayCluster/Private/Render/Viewport/Containers/DisplayClusterViewport_OverscanRuntimeSettings.cpp: 166

diff --git a/Engine/Plugins/Runtime/nDisplay/Source/DisplayCluster/Private/Render/Viewport/Containers/DisplayClusterViewport_OverscanRuntimeSettings.cpp b/Engine/Plugins/Runtime/nDisplay/Source/DisplayCluster/Private/Render/Viewport/Containers/DisplayClusterViewport_OverscanRuntimeSettings.cpp
--- a/Engine/Plugins/Runtime/nDisplay/Source/DisplayCluster/Private/Render/Viewport/Containers/DisplayClusterViewport_OverscanRuntimeSettings.cpp
+++ b/Engine/Plugins/Runtime/nDisplay/Source/DisplayCluster/Private/Render/Viewport/Containers/DisplayClusterViewport_OverscanRuntimeSettings.cpp
@@ -61,9 +61,9 @@ void FDisplayClusterViewport_OverscanRuntimeSettings::UpdateOverscanSettings(
 	}
 
 	const FIntPoint Size = InOutRenderTargetRect.Size();
-	if(Size.X <=0 || Size.Y <=0)
+	if(Size.GetMin() <= 0)
 	{
-		// If the Size is invalid (zero or negative), overscan is disabled.
+		// If the Size is invalid (on of values are zero or negative), overscan is disabled.
 		InOutOverscanRuntimeSettings = FDisplayClusterViewport_OverscanRuntimeSettings();
 
 		return;
@@ -164,11 +164,13 @@ void FDisplayClusterViewport_OverscanRuntimeSettings::UpdateOverscanSettings(
 			const FOverscanPixels AdjustedOverscanPixels((AdjustedOverscanPercent * Size).GetRoundToInt());
 
 			// Compute the scaled inner region size based on overscan margins.
-			// Mathematically, InnerRegionSize > 0 as long as Size > 0.
 			const FIntPoint InnerRegionSize = Size - AdjustedOverscanPixels.Size();
-			if (InnerRegionSize.X <= 0 || InnerRegionSize.Y <= 0)
+
+			// Mathematically, InnerRegionSize should be to be positive for valid input
+			// This check is a safeguard for unexpected or corrupted input values.
+			if (InnerRegionSize.GetMin() <= 0)
 			{
-				// If the InnerRegionSize is invalid (zero or negative), overscan is disabled.
+				// If it evaluates to zero or negative (invalid case), disable overscan.
 				InOutOverscanRuntimeSettings = FDisplayClusterViewport_OverscanRuntimeSettings();
 
 				return;
diff --git a/Engine/Plugins/Runtime/nDisplay/Source/DisplayCluster/Private/Render/Viewport/DisplayClusterViewportProxy_Resources.cpp b/Engine/Plugins/Runtime/nDisplay/Source/DisplayCluster/Private/Render/Viewport/DisplayClusterViewportProxy_Resources.cpp
--- a/Engine/Plugins/Runtime/nDisplay/Source/DisplayCluster/Private/Render/Viewport/DisplayClusterViewportProxy_Resources.cpp
+++ b/Engine/Plugins/Runtime/nDisplay/Source/DisplayCluster/Private/Render/Viewport/DisplayClusterViewportProxy_Resources.cpp
@@ -428,42 +428,36 @@ void FDisplayClusterViewportProxy::ImplResolveTileResource_RenderThread(FRHIComm
 					const int32 MaxBorderSize = FMath::Min(Input.Rect.Size().GetMin(), Output.Rect.Size().GetMin()) / 4;
 					const int32 TileBorderSize = FMath::Min(GDisplayClusterRenderTileBorder, MaxBorderSize);
 
-					// Shows borders for internal areas only
-					FIntRect SubRect(FIntPoint::ZeroValue, FIntPoint::ZeroValue);
+					FIntMarginSet InputBorderMargins = 0;
 					if (InTilePos.X > 0)
 					{
-						SubRect.Min.X = TileBorderSize;
+						InputBorderMargins.Left = TileBorderSize;
 					}
 					if (InTilePos.Y > 0)
 					{
-						SubRect.Min.Y = TileBorderSize;
+						InputBorderMargins.Top = TileBorderSize;
 					}
 					if ((InTilePos.X + 1) < InTileSize.X)
 					{
-						SubRect.Max.X = TileBorderSize;
+						InputBorderMargins.Right = TileBorderSize;
 					}
 					if ((InTilePos.Y + 1) < InTileSize.Y)
 					{
-						SubRect.Max.Y = TileBorderSize;
+						InputBorderMargins.Bottom = TileBorderSize;
 					}
 
-					// support when bAdaptResolution == false
-					// Input.Rect.Size != Output.Rect.Size()
-					const FVector2D Scale(
-						double(Output.Rect.Width()) / Input.Rect.Width(),
-						double(Output.Rect.Height()) / Input.Rect.Height()
-					);
+					// Scale border margins from input rect space to output rect space.
+					// Falls back to (1,1) if input size is zero (avoid division by zero).
+					const FVector2D Scale =
+						(Input.Rect.Size().GetMin() > 0)
+						? FVector2D(Output.Rect.Size()) / Input.Rect.Size()
+						: FVector2D(1.f, 1.f);
+
+					const FIntMarginSet OutputBorderMargins = (FMarginSet(InputBorderMargins) * Scale).GetRoundToInt();
 
 					// Set rect smaller to show gaps between tiles:
-					Input.Rect.Min.X += SubRect.Min.X;
-					Input.Rect.Min.Y += SubRect.Min.Y;
-					Input.Rect.Max.X -= SubRect.Min.X + SubRect.Max.X;
-					Input.Rect.Max.Y -= SubRect.Min.Y + SubRect.Max.Y;
-
-					Output.Rect.Min.X += SubRect.Min.X * Scale.X;
-					Output.Rect.Min.Y += SubRect.Min.Y * Scale.Y;
-					Output.Rect.Max.X -= (SubRect.Min.X + SubRect.Max.X) * Scale.X;
-					Output.Rect.Max.Y -= (SubRect.Min.Y + SubRect.Max.Y) * Scale.Y;
+					InputBorderMargins.ApplyClampedInsetToRect(Input.Rect);
+					OutputBorderMargins.ApplyClampedInsetToRect(Output.Rect);
 				}
 				// When we don't need to show tile borders, we can blend them together
 				else if (bApplyTileEdgeBlend)
