[nDisplay] Custom capture size for nDisplay media
- Now, if a corresponding UMediaOutput instance provides any valid size, it will be used as custom capture size. No matter what size the original viewport or camera have.

#rnx

diff --git a/Engine/Plugins/Runtime/nDisplay/Source/DisplayClusterMedia/Private/Capture/DisplayClusterMediaCaptureBase.cpp b/Engine/Plugins/Runtime/nDisplay/Source/DisplayClusterMedia/Private/Capture/DisplayClusterMediaCaptureBase.cpp
--- a/Engine/Plugins/Runtime/nDisplay/Source/DisplayClusterMedia/Private/Capture/DisplayClusterMediaCaptureBase.cpp
+++ b/Engine/Plugins/Runtime/nDisplay/Source/DisplayClusterMedia/Private/Capture/DisplayClusterMediaCaptureBase.cpp
@@ -35,6 +35,12 @@ FDisplayClusterMediaCaptureBase::FDisplayClusterMediaCaptureBase(
 	, SyncPolicy(InSyncPolicy)
 {
 	checkSlow(InMediaOutput);
+
+	// Here we expect to receive
+	// {0, 0}, if no capture size requested
+	// {X > 0, Y > 0}, if explicit capture size is specified
+	CustomResolution = InMediaOutput->GetRequestedSize();
+
 	MediaOutput = DuplicateObject(InMediaOutput, GetTransientPackage());
 	checkSlow(MediaOutput);
 
@@ -246,6 +252,8 @@ void FDisplayClusterMediaCaptureBase::OnPostClusterTick()
 		{
 			const FIntPoint LastSrcRegionIntPoint = LastSrcRegionSize.load().ToIntPoint();
 			const FIntPoint DesiredSize = MediaCapture->GetDesiredSize();
+			const bool bDesiredSizeMatchesLastRegion = LastSrcRegionIntPoint != FIntPoint::ZeroValue && DesiredSize == LastSrcRegionIntPoint;
+			const bool bDesiredSizeMatchesCustomRegion = CustomResolution.X > 0 && CustomResolution.Y > 0 && DesiredSize == CustomResolution;
 
 			// We don't restart if we haven't exported any textures yet (indicated by zero size LastSrcRegion)
 			// to avoid constant media restarts since in such case media is set to use GetCaptureSize() != (0.0).
@@ -254,7 +262,7 @@ void FDisplayClusterMediaCaptureBase::OnPostClusterTick()
 			// This is preferred to setting LastSrcRegion to GetCaptureSize() because that would not reflect the actual
 			// last captured size.
 
-			if ((DesiredSize != LastSrcRegionIntPoint) && (LastSrcRegionIntPoint != FIntPoint(0,0)))
+			if (!(bDesiredSizeMatchesLastRegion || bDesiredSizeMatchesCustomRegion))
 			{
 				UE_LOG(LogDisplayClusterMedia, Log, TEXT("Stopping MediaCapture '%s' because its DesiredSize (%d, %d) doesn't match the captured texture size (%d, %d)"), 
 					*GetMediaId(), DesiredSize.X, DesiredSize.Y, LastSrcRegionIntPoint.X, LastSrcRegionIntPoint.Y);
@@ -306,6 +314,14 @@ bool FDisplayClusterMediaCaptureBase::StartMediaCapture()
 	MediaCaptureOptions.bSkipFrameWhenRunningExpensiveTasks = false;
 	MediaCaptureOptions.OverrunAction = EMediaCaptureOverrunAction::Flush;
 
+	// When custom resolution is required, we should always scale the output to the required size
+	if (CustomResolution.X > 1 && CustomResolution.Y > 1)
+	{
+		Descriptor.ResourceSize = CustomResolution;
+		MediaCaptureOptions.Crop = EMediaCaptureCroppingType::None;
+		MediaCaptureOptions.ResizeMethod = EMediaCaptureResizeMethod::ResizeInRenderPass;
+	}
+
 	const bool bCaptureStarted = MediaCapture->CaptureRHITexture(Descriptor, MediaCaptureOptions);
 
 	if (bCaptureStarted)
diff --git a/Engine/Plugins/Runtime/nDisplay/Source/DisplayClusterMedia/Private/Capture/DisplayClusterMediaCaptureBase.h b/Engine/Plugins/Runtime/nDisplay/Source/DisplayClusterMedia/Private/Capture/DisplayClusterMediaCaptureBase.h
--- a/Engine/Plugins/Runtime/nDisplay/Source/DisplayClusterMedia/Private/Capture/DisplayClusterMediaCaptureBase.h
+++ b/Engine/Plugins/Runtime/nDisplay/Source/DisplayClusterMedia/Private/Capture/DisplayClusterMediaCaptureBase.h
@@ -114,6 +114,9 @@ private:
 	TObjectPtr<UDisplayClusterMediaOutputSynchronizationPolicy> SyncPolicy;
 	//~ End GC by AddReferencedObjects
 
+	/** Custom resolution to use on capture side */
+	FIntPoint CustomResolution = FIntPoint::ZeroValue;
+
 	/** Used to restart media capture in the case it falls in error */
 	bool bWasCaptureStarted = false;
 
