[nDisplay] Fixed bug where viewport output capture was initialized even with Media.bEnable==false

#rnx

diff --git a/Engine/Plugins/Runtime/nDisplay/Source/DisplayClusterMedia/Private/DisplayClusterMediaModule.cpp b/Engine/Plugins/Runtime/nDisplay/Source/DisplayClusterMedia/Private/DisplayClusterMediaModule.cpp
--- a/Engine/Plugins/Runtime/nDisplay/Source/DisplayClusterMedia/Private/DisplayClusterMediaModule.cpp
+++ b/Engine/Plugins/Runtime/nDisplay/Source/DisplayClusterMedia/Private/DisplayClusterMediaModule.cpp
@@ -309,29 +309,32 @@ void FDisplayClusterMediaModule::InitializeViewportOutput(const UDisplayClusterC
 	{
 		const FDisplayClusterConfigurationMediaViewport& MediaSettings = Viewport->RenderSettings.Media;
 
-		// Media capture
-		uint8 CaptureIdx = 0;
-		for (const FDisplayClusterConfigurationMediaOutput& MediaOutputItem : MediaSettings.MediaOutputs)
+		if (MediaSettings.bEnable)
 		{
-			if (IsValid(MediaOutputItem.MediaOutput))
+			// Media capture
+			uint8 CaptureIdx = 0;
+			for (const FDisplayClusterConfigurationMediaOutput& MediaOutputItem : MediaSettings.MediaOutputs)
 			{
-				const FString MediaCaptureId = DisplayClusterMediaHelpers::MediaId::GenerateMediaId(
-					DisplayClusterMediaHelpers::MediaId::EMediaDeviceType::Output,
-					DisplayClusterMediaHelpers::MediaId::EMediaOwnerType::Viewport,
-					ClusterNodeId, RootActorName, ViewportId, CaptureIdx);
+				if (IsValid(MediaOutputItem.MediaOutput))
+				{
+					const FString MediaCaptureId = DisplayClusterMediaHelpers::MediaId::GenerateMediaId(
+						DisplayClusterMediaHelpers::MediaId::EMediaDeviceType::Output,
+						DisplayClusterMediaHelpers::MediaId::EMediaOwnerType::Viewport,
+						ClusterNodeId, RootActorName, ViewportId, CaptureIdx);
 
-				UE_LOG(LogDisplayClusterMedia, Log, TEXT("Initializing viewport capture [%u]: '%s' for viewport '%s'"), CaptureIdx, *MediaCaptureId, *ViewportId);
+					UE_LOG(LogDisplayClusterMedia, Log, TEXT("Initializing viewport capture [%u]: '%s' for viewport '%s'"), CaptureIdx, *MediaCaptureId, *ViewportId);
 
-				TSharedPtr<FDisplayClusterMediaCaptureViewportFull> NewViewportCapture = MakeShared<FDisplayClusterMediaCaptureViewportFull>(
-					MediaCaptureId,
-					ClusterNodeId,
-					ViewportId,
-					MediaOutputItem.MediaOutput,
-					MediaOutputItem.OutputSyncPolicy);
+					TSharedPtr<FDisplayClusterMediaCaptureViewportFull> NewViewportCapture = MakeShared<FDisplayClusterMediaCaptureViewportFull>(
+						MediaCaptureId,
+						ClusterNodeId,
+						ViewportId,
+						MediaOutputItem.MediaOutput,
+						MediaOutputItem.OutputSyncPolicy);
 
-				CaptureDevices.Emplace(MediaCaptureId, MoveTemp(NewViewportCapture));
+					CaptureDevices.Emplace(MediaCaptureId, MoveTemp(NewViewportCapture));
 
-				++CaptureIdx;
+					++CaptureIdx;
+				}
 			}
 		}
 	}
