[nDisplay] A very small fix to the validation of custom resolution after AI review

#rnx

diff --git a/Engine/Plugins/Runtime/nDisplay/Source/DisplayClusterMedia/Private/Capture/DisplayClusterMediaCaptureBase.cpp b/Engine/Plugins/Runtime/nDisplay/Source/DisplayClusterMedia/Private/Capture/DisplayClusterMediaCaptureBase.cpp
--- a/Engine/Plugins/Runtime/nDisplay/Source/DisplayClusterMedia/Private/Capture/DisplayClusterMediaCaptureBase.cpp
+++ b/Engine/Plugins/Runtime/nDisplay/Source/DisplayClusterMedia/Private/Capture/DisplayClusterMediaCaptureBase.cpp
@@ -315,7 +315,7 @@ bool FDisplayClusterMediaCaptureBase::StartMediaCapture()
 	MediaCaptureOptions.OverrunAction = EMediaCaptureOverrunAction::Flush;
 
 	// When custom resolution is required, we should always scale the output to the required size
-	if (CustomResolution.X > 1 && CustomResolution.Y > 1)
+	if (CustomResolution.X > 0 && CustomResolution.Y > 0)
 	{
 		Descriptor.ResourceSize = CustomResolution;
 		MediaCaptureOptions.Crop = EMediaCaptureCroppingType::None;
