[nDisplay] Draw GUI overlay across nDisplay viewports

- Once the final SrcColor is rendered to the game viewport buffer, we substitute the buffer texture with a transparent one. This allows to make slate & debug canvas render into a separate texture. This texture is used later to draw GUI on top of every nDisplay viewport. In DisplayClusterDeviceBase, right before copying the viewport's buffer into the backbuffer, we restore the original viewport's texture to let SrcColor to be displayed as expected.

[RN] Draw UI overlay across nDisplay viewports

diff --git a/Engine/Plugins/Runtime/nDisplay/Shaders/Private/OverlayShaders.usf b/Engine/Plugins/Runtime/nDisplay/Shaders/Private/OverlayShaders.usf
new file mode 100644
--- /dev/null
+++ b/Engine/Plugins/Runtime/nDisplay/Shaders/Private/OverlayShaders.usf
@@ -0,0 +1,18 @@
+// Copyright Epic Games, Inc. All Rights Reserved.
+
+#include "/Engine/Private/Common.ush"
+
+Texture2D TextureOverlay;
+Texture2D TextureBase;
+SamplerState Sampler;
+
+float4 DrawOverlay_PS(float2 uv : TEXCOORD) : SV_Target
+{
+	float4 srcColor = TextureOverlay.Sample(Sampler, uv);
+	float4 dstColor = TextureBase.Sample(Sampler, uv);
+
+	// Source-over blending (source color is already alpha-premultiplied)
+	float3 outColor = dstColor.rgb * (1 - srcColor.a) + srcColor.rgb;
+
+	return float4(outColor, 1);
+}
diff --git a/Engine/Plugins/Runtime/nDisplay/Source/DisplayCluster/Private/Game/EngineClasses/Basics/DisplayClusterViewportClient.cpp b/Engine/Plugins/Runtime/nDisplay/Source/DisplayCluster/Private/Game/EngineClasses/Basics/DisplayClusterViewportClient.cpp
--- a/Engine/Plugins/Runtime/nDisplay/Source/DisplayCluster/Private/Game/EngineClasses/Basics/DisplayClusterViewportClient.cpp
+++ b/Engine/Plugins/Runtime/nDisplay/Source/DisplayCluster/Private/Game/EngineClasses/Basics/DisplayClusterViewportClient.cpp
@@ -2,56 +2,53 @@
 
 #include "DisplayClusterViewportClient.h"
 
+#include "AudioDevice.h"
+#include "AudioDeviceManager.h"
+#include "BufferVisualizationData.h"
+#include "EngineModule.h"
+#include "EngineStats.h"
+#include "EngineUtils.h"
+#include "FXSystem.h"
+#include "GlobalRenderResources.h"
+#include "IHeadMountedDisplay.h"
+#include "IXRTrackingSystem.h"
+#include "RenderGraph.h"
 #include "SceneManagement.h"
 #include "SceneView.h"
-#include "Engine/Canvas.h"
 #include "SceneViewExtension.h"
-#include "IHeadMountedDisplay.h"
-#include "IXRTrackingSystem.h"
-#include "BufferVisualizationData.h"
+#include "UnrealEngine.h"
+
+#include "Audio/AudioDebug.h"
+#include "Components/LineBatchComponent.h"
+#include "Debug/DebugDrawService.h"
 #include "Engine/Engine.h"
+#include "Engine/Canvas.h"
 #include "Engine/Console.h"
+#include "Engine/CoreSettings.h"
 #include "Engine/LocalPlayer.h"
-#include "UnrealEngine.h"
-#include "EngineUtils.h"
-#include "ShaderCore.h"
-
-#include "AudioDeviceManager.h"
-#include "AudioDevice.h"
-#include "Audio/AudioDebug.h"
-
+#include "Engine/TextureRenderTarget2D.h"
+#include "Framework/Application/SlateApplication.h"
+#include "GameFramework/GameUserSettings.h"
 #include "GameFramework/PlayerController.h"
-#include "Debug/DebugDrawService.h"
-#include "ContentStreaming.h"
-#include "EngineModule.h"
-#include "GlobalRenderResources.h"
-#include "FXSystem.h"
 #include "GameFramework/HUD.h"
-#include "SubtitleManager.h"
-#include "Components/LineBatchComponent.h"
-#include "GameFramework/GameUserSettings.h"
-#include "Framework/Application/SlateApplication.h"
 #include "Slate/SceneViewport.h"
+#include "UObject/Package.h"
+#include "Widgets/SWindow.h"
 
-#include "EngineStats.h"
-
-#include "Render/Device/IDisplayClusterRenderDevice.h"
-
 #include "DisplayClusterEnums.h"
 #include "DisplayClusterSceneViewExtensions.h"
-#include "Misc/DisplayClusterGlobals.h"
-
 #include "IDisplayCluster.h"
 #include "IDisplayClusterCallbacks.h"
+
+#include "Config/DisplayClusterConfigManager.h"
 #include "Game/IDisplayClusterGameManager.h"
-
-#include "Render/Viewport/IDisplayClusterViewportManager.h"
+#include "Misc/DisplayClusterGlobals.h"
+#include "Render/Device/IDisplayClusterRenderDevice.h"
+#include "Render/GUILayer/DisplayClusterGuiLayerController.h"
 #include "Render/Viewport/IDisplayClusterViewport.h"
+#include "Render/Viewport/IDisplayClusterViewportManager.h"
 #include "Render/Viewport/RenderFrame/DisplayClusterRenderFrame.h"
 
-#include "Config/DisplayClusterConfigManager.h"
-#include "UObject/Package.h"
 
-
 // Debug feature to synchronize and force all external resources to be transferred cross GPU at the end of graph execution.
 // May be useful for testing cross GPU synchronization logic.
@@ -410,13 +407,17 @@
 
 	//Get world for render
 	UWorld* const MyWorld = GetWorld();
+	if (MyWorld == nullptr)
+	{
+		return;
+	}
 
 	////////////////////////////////
 	// Otherwise we use our own version of the UGameViewportClient::Draw which is basically
 	// a simpler version of the original one but with multiple ViewFamilies support
 
-	// Valid SceneCanvas is required.  Make this explicit.
 	check(SceneCanvas);
+	check(GEngine);
 
 	OnBeginDraw().Broadcast();
 
@@ -443,6 +444,9 @@
 	{
 		DebugCanvas->SetScaledToRenderTarget(bStereoRendering);
 		DebugCanvas->SetStereoRendering(bStereoRendering);
+
+		const bool bGuiPropagationActive = FDisplayClusterGuiLayerController::Get().IsActiveThisFrame();
+		DebugCanvas->SetWriteDestinationAlpha(bGuiPropagationActive);
 	}
 	if (SceneCanvas)
 	{
@@ -810,8 +814,11 @@
 	// custom postprocess single frame flag must be removed at frame end on game thread
 	DCRenderDevice->FinalizeNewFrame();
 
-	// Update level streaming.
-	MyWorld->UpdateLevelStreaming();
+	if (!GUseUnifiedTimeBudgetForStreaming)
+	{
+		// Update level streaming.
+		MyWorld->UpdateLevelStreaming();
+	}
 
 	// Remove temporary debug lines.
 	constexpr const UWorld::ELineBatcherType LineBatchersToFlush[] = { UWorld::ELineBatcherType::World, UWorld::ELineBatcherType::Foreground };
@@ -823,6 +830,29 @@
 		MyWorld->FXSystem->DrawDebug(SceneCanvas);
 	}
 
+#if WITH_STATE_STREAM
+	if (IStateStreamManager* Manager = MyWorld->GetStateStreamManager())
 	{
+		class FStateStreamDebugRenderer : public IStateStreamDebugRenderer
+		{
+		public:
+			FStateStreamDebugRenderer(FViewport* V, UCanvas* C) : Viewport(V), Canvas(C) {}
+
+			virtual void DrawText(const FStringView& Text) override
+			{
+				Canvas->SetDrawColor(FColor::Orange);
+				Canvas->DrawText(GEngine->GetSmallFont(), Text, 100, Y);
+				Y += 20;
+			}
+
+			uint32 Y = 200;
+			FViewport* Viewport;
+			UCanvas* Canvas;
+		} StateStreamDebugRenderer(InViewport, DebugCanvasObject);
+		Manager->Game_DebugRender(StateStreamDebugRenderer);
+	}
+#endif
+
+	{
 		//ensure canvas has been flushed before rendering UI
 		SceneCanvas->Flush_GameThread();
@@ -843,28 +873,17 @@
 
 	if (DebugCanvas)
 	{
-		// Reset the debug canvas to be full-screen before drawing the console
-		// (the debug draw service above has messed with the viewport size to fit it to a single player's subregion)
-		DebugCanvasObject->Init(DebugCanvasSize.X, DebugCanvasSize.Y, NULL, DebugCanvas);
-
 		DrawStatsHUD(MyWorld, InViewport, DebugCanvas, DebugCanvasObject, DebugProperties, PlayerCameraLocation, PlayerCameraRotation);
 
-		if (GEngine->IsStereoscopic3D(InViewport))
-		{
-#if 0 //!UE_BUILD_SHIPPING
-			// TODO: replace implementation in OculusHMD with a debug renderer
-			if (GEngine->XRSystem.IsValid())
-			{
-				GEngine->XRSystem->DrawDebug(DebugCanvasObject);
-			}
-#endif
-		}
-
 		if (GDisplayClusterDebugDraw && !ViewFamilies.IsEmpty())
 		{
 			UDebugDrawService::Draw(ViewFamilies.Last()->EngineShowFlags, InViewport, const_cast<FSceneView*>(ViewFamilies.Last()->Views[0]), DebugCanvas, DebugCanvasObject);
 		}
 
+		// Reset the debug canvas to be full-screen before drawing the console
+		// (the debug draw service above has messed with the viewport size to fit it to a single player's subregion)
+		DebugCanvasObject->Init(DebugCanvasSize.X, DebugCanvasSize.Y, NULL, DebugCanvas);
+
 		// Render the console absolutely last because developer input is was matter the most.
 		if (ViewportConsole)
 		{
diff --git a/Engine/Plugins/Runtime/nDisplay/Source/DisplayCluster/Private/Render/Device/DisplayClusterDeviceBase.cpp b/Engine/Plugins/Runtime/nDisplay/Source/DisplayCluster/Private/Render/Device/DisplayClusterDeviceBase.cpp
--- a/Engine/Plugins/Runtime/nDisplay/Source/DisplayCluster/Private/Render/Device/DisplayClusterDeviceBase.cpp
+++ b/Engine/Plugins/Runtime/nDisplay/Source/DisplayCluster/Private/Render/Device/DisplayClusterDeviceBase.cpp
@@ -26,6 +26,7 @@
 #include "RHIStaticStates.h"
 #include "Slate/SceneViewport.h"
 
+#include "Render/GUILayer/DisplayClusterGuiLayerController.h"
 #include "Render/PostProcess/IDisplayClusterPostProcess.h"
 #include "Render/Presentation/DisplayClusterPresentationBase.h"
 #include "Render/Projection/IDisplayClusterProjectionPolicy.h"
@@ -404,6 +405,9 @@ void FDisplayClusterDeviceBase::RenderTexture_RenderThread(class FRDGBuilder& Gr
 	CopyInfo.Size.X = FMath::Min(SrcSize.X, DstSize.X);
 	CopyInfo.Size.Y = FMath::Min(SrcSize.Y, DstSize.Y);
 
+	// Allow custom GUI related jobs to be done first. After that, we'll use the texture it returns (can be the same).
+	SrcTexture = FDisplayClusterGuiLayerController::Get().ProcessFinalTexture_RenderThread(GraphBuilder, SrcTexture);
+
 	AddCopyTexturePass(GraphBuilder, SrcTexture, BackBuffer, CopyInfo);
 
 	if (RenderFrameMode == EDisplayClusterRenderFrameMode::Stereo)
diff --git a/Engine/Plugins/Runtime/nDisplay/Source/DisplayCluster/Private/Render/Device/DisplayClusterDeviceBase.h b/Engine/Plugins/Runtime/nDisplay/Source/DisplayCluster/Private/Render/Device/DisplayClusterDeviceBase.h
--- a/Engine/Plugins/Runtime/nDisplay/Source/DisplayCluster/Private/Render/Device/DisplayClusterDeviceBase.h
+++ b/Engine/Plugins/Runtime/nDisplay/Source/DisplayCluster/Private/Render/Device/DisplayClusterDeviceBase.h
@@ -62,7 +62,13 @@ public:
 	virtual void CalculateStereoViewOffset(const int32 ViewIndex, FRotator& ViewRotation, const float WorldToMeters, FVector& ViewLocation) override;
 	virtual FMatrix GetStereoProjectionMatrix(const int32 ViewIndex) const override;
 	virtual void RenderTexture_RenderThread(class FRDGBuilder& GraphBuilder, FRDGTextureRef BackBuffer, FRDGTextureRef SrcTexture, FVector2f WindowSize) const override;
-	
+
+	// So far, we don’t need any region bounds
+	virtual FVector2D GetTextSafeRegionBounds() const override
+	{
+		return FVector2D(1.f, 1.f);
+	}
+
 	virtual int32 GetDesiredNumberOfViews(bool bStereoRequested) const override
 	{
 		return FMath::Max(1, DesiredNumberOfViews);
diff --git a/Engine/Plugins/Runtime/nDisplay/Source/DisplayCluster/Private/Render/GUILayer/DisplayClusterGuiLayerController.cpp b/Engine/Plugins/Runtime/nDisplay/Source/DisplayCluster/Private/Render/GUILayer/DisplayClusterGuiLayerController.cpp
new file mode 100644
--- /dev/null
+++ b/Engine/Plugins/Runtime/nDisplay/Source/DisplayCluster/Private/Render/GUILayer/DisplayClusterGuiLayerController.cpp
@@ -0,0 +1,284 @@
+// Copyright Epic Games, Inc. All Rights Reserved.
+
+#include "Render/GUILayer/DisplayClusterGuiLayerController.h"
+
+#include "IDisplayClusterShaders.h"
+#include "RenderGraphBuilder.h"
+#include "RenderGraphUtils.h"
+
+#include "Async/Async.h"
+#include "Framework/Application/SlateApplication.h"
+#include "Engine/Console.h"
+#include "Engine/Engine.h"
+#include "Engine/GameViewportClient.h"
+#include "Misc/DisplayClusterGlobals.h"
+#include "Misc/DisplayClusterLog.h"
+#include "Render/Viewport/IDisplayClusterViewportProxy.h"
+#include "ShaderParameters/DisplayClusterShaderParameters_Overlay.h"
+#include "Slate/SceneViewport.h"
+
+
+TAutoConsoleVariable<bool> CVarPropagateGui(
+	TEXT("nDisplay.GUI.Propagate"),
+	true,
+	TEXT("Show GUI on viewports\n")
+	TEXT("0 : Disabled\n")
+	TEXT("1 : Enabled\n"),
+	ECVF_RenderThreadSafe
+);
+
+
+FDisplayClusterGuiLayerController::FDisplayClusterGuiLayerController()
+	: bEnabled(GDisplayCluster ? GDisplayCluster->GetOperationMode() == EDisplayClusterOperationMode::Cluster : false)
+{
+	if (bEnabled)
+	{
+		// GameThread only intialization
+		AsyncTask(ENamedThreads::GameThread, [this]()
+			{
+				FSlateApplication::Get().OnPreTick().AddRaw(this, &FDisplayClusterGuiLayerController::HandleSlatePreTick);
+			});
+	}
+}
+
+FDisplayClusterGuiLayerController& FDisplayClusterGuiLayerController::Get()
+{
+	static FDisplayClusterGuiLayerController Instance;
+	return Instance;
+}
+
+FRDGTextureRef FDisplayClusterGuiLayerController::ProcessFinalTexture_RenderThread(FRDGBuilder& GraphBuilder, FRDGTextureRef FinalTexture)
+{
+	checkSlow(IsInRenderingThread());
+	checkSlow(FinalTexture);
+
+	// Nothing to do if GUI layer propagation is not active
+	if (!bActiveThisFrameRT || !bEnabled)
+	{
+		return FinalTexture;
+	}
+
+	// Make sure our buffer duplicate is valid
+	if (!TexViewportOriginalRHI.IsValid())
+	{
+		return FinalTexture;
+	}
+
+	// Create an RDG texture reference to the original game viewport's texture that we stored before Slate rendering
+	FRDGTextureRef Output = RegisterExternalTexture(GraphBuilder, TexViewportOriginalRHI.GetReference(), TEXT("nD.TexViewportOriginalRHI"));
+	if (!Output)
+	{
+		return FinalTexture;
+	}
+
+	// If everything Ok, return the original viewport's buffer
+	return Output;
+}
+
+void FDisplayClusterGuiLayerController::ProcessPostRenderViewFamily_RenderThread(FRDGBuilder& GraphBuilder, const FSceneViewFamily& ViewFamiliy, const IDisplayClusterViewportProxy* ViewportProxy)
+{
+	checkSlow(IsInRenderingThread());
+	checkSlow(ViewportProxy);
+
+	// Nothing to do if GUI layer propagation is not active
+	if (!bActiveThisFrameRT || !bEnabled)
+	{
+		return;
+	}
+
+	// Validate both input and internal data
+	if (!ViewportProxy || !TexViewportOriginalRHI.IsValid() || !TexViewportGUILayerRHI.IsValid())
+	{
+		return;
+	}
+
+	TArray<FRHITexture*> Textures;
+	TArray<FIntRect>     Regions;
+
+	// Get nD viewport's buffer
+	if (ViewportProxy->GetResourcesWithRects_RenderThread(EDisplayClusterViewportResourceType::InternalRenderTargetEntireRectResource, Textures, Regions))
+	{
+		if (Textures.Num() > 0 && Regions.Num() > 0 && Textures[0])
+		{
+			const FIntPoint DCViewportSize = Regions[0].Size();
+
+			// Prepare textures for drawing
+			FRDGTextureRef TexViewportGUILayerRDG = RegisterExternalTexture(GraphBuilder, TexViewportGUILayerRHI, TEXT("nD.TexViewportGUILayerRHI"));
+			FRDGTextureRef TexDCViewportRDG       = RegisterExternalTexture(GraphBuilder, Textures[0], *Textures[0]->GetName().ToString());
+			FRDGTextureRef TexTempDCViewportRDG   = CreateTextureFrom_RenderThread(GraphBuilder, Textures[0], TEXT("nD.TexTempDCViewportRDG"), DCViewportSize);
+
+			// Make a copy to the input texture
+			FRHICopyTextureInfo CopyInfo;
+			CopyInfo.Size = FIntVector{ DCViewportSize.X, DCViewportSize.Y, 0 };
+			CopyInfo.SourcePosition = FIntVector{ Regions[0].Min.X, Regions[0].Min.Y, 0 };
+			AddCopyTexturePass(GraphBuilder, TexDCViewportRDG, TexTempDCViewportRDG, CopyInfo);
+
+			FDisplayClusterShaderParameters_Overlay Parameters;
+			Parameters.OverlayTexture = TexViewportGUILayerRDG; // GUI layer
+			Parameters.BaseTexture    = TexTempDCViewportRDG;   // Scene color (copy of nD viewport)
+			Parameters.OutputTexture  = TexDCViewportRDG;       // Scene color (original nD viewport)
+
+			// Draw the GUI layer on top of the nD viewport
+			IDisplayClusterShaders::Get().AddDrawOverlayPass(GraphBuilder, Parameters);
+		}
+	}
+}
+
+void FDisplayClusterGuiLayerController::HandleSlatePreTick(float InDeltaTime)
+{
+	checkSlow(IsInGameThread());
+
+	// Nothing to do if disabled
+	if (!bEnabled)
+	{
+		return;
+	}
+
+	// Read the cvar once per frame
+	bActiveThisFrame = CVarPropagateGui.GetValueOnGameThread();
+
+	ENQUEUE_RENDER_COMMAND(FlushCanvasRT)(
+		[this, bNewActiveThisFrame = bActiveThisFrame](FRHICommandListImmediate& RHICmdList)
+		{
+			// Update the corresponding flag, and leave if inactive
+			bActiveThisFrameRT = bNewActiveThisFrame;
+			if (!bActiveThisFrameRT)
+			{
+				return;
+			}
+
+			// Get game viewport's current render target
+			FTextureRHIRef ViewportRTT = GetViewportRTT_RenderThread();
+			if (!ViewportRTT.IsValid())
+			{
+				return;
+			}
+
+			// Create a buffer duplicate, but float16
+			UpdateTempTexture_RenderThread(RHICmdList, TexViewportGUILayerRHI, ViewportRTT, TEXT("nD.TexViewportGUILayerRHI"), PF_FloatRGBA);
+			if (!TexViewportGUILayerRHI.IsValid())
+			{
+				return;
+			}
+
+			// Store the original vieport's buffer, and substitute is with our float16 buffer
+			TexViewportOriginalRHI = ViewportRTT;
+			SetViewportRTT_RenderThread(TexViewportGUILayerRHI);
+
+			// Clear out texture to transparent to let Slate renderer produce a clear GUI layer
+			FRDGBuilder GraphBuilder(RHICmdList);
+			FRDGTextureRef TexViewportGUILayerRDG = RegisterExternalTexture(GraphBuilder, TexViewportGUILayerRHI.GetReference(), TEXT("nD.TexViewportGUILayerRDG"));
+			AddClearRenderTargetPass(GraphBuilder, TexViewportGUILayerRDG, FLinearColor::Transparent);
+
+			GraphBuilder.Execute();
+		});
+}
+
+FTextureRHIRef FDisplayClusterGuiLayerController::GetViewportRTT_RenderThread()
+{
+	checkSlow(IsInRenderingThread());
+
+	if (GEngine && GEngine->GameViewport && GEngine->GameViewport->Viewport)
+	{
+		return GEngine->GameViewport->Viewport->GetRenderTargetTexture();
+	}
+
+	return nullptr;
+}
+
+void FDisplayClusterGuiLayerController::SetViewportRTT_RenderThread(FTextureRHIRef& NewRTT)
+{
+	checkSlow(IsInRenderingThread());
+	checkSlow(NewRTT.IsValid());
+
+	if (NewRTT.IsValid())
+	{
+		if (GEngine && GEngine->GameViewport && GEngine->GameViewport->Viewport)
+		{
+			static_cast<FSceneViewport*>(GEngine->GameViewport->Viewport)->SetRenderTargetTextureRenderThread(NewRTT);
+		}
+	}
+}
+
+void FDisplayClusterGuiLayerController::UpdateTempTexture_RenderThread(
+	FRHICommandListImmediate& RHICmdList,
+	FTextureRHIRef& StorageTexture,
+	FTextureRHIRef& FromTexture,
+	const TCHAR* Name,
+	EPixelFormat PixelFormat
+)
+{
+	if (!FromTexture.IsValid())
+	{
+		return;
+	}
+
+	// Create new one if not exists
+	if (!StorageTexture.IsValid())
+	{
+		StorageTexture = CreateTextureFrom_RenderThread(RHICmdList, FromTexture, Name, PixelFormat);
+	}
+
+	// Re-create if different size
+	if (FromTexture->GetDesc().Extent != StorageTexture->GetDesc().Extent)
+	{
+		StorageTexture = CreateTextureFrom_RenderThread(RHICmdList, FromTexture, Name, PixelFormat);
+	}
+}
+
+FTextureRHIRef FDisplayClusterGuiLayerController::CreateTextureFrom_RenderThread(FRHICommandListImmediate& RHICmdList, FTextureRHIRef& FromTexture, const TCHAR* Name, EPixelFormat PixelFormat)
+{
+	if (!FromTexture.IsValid())
+	{
+		return nullptr;
+	}
+
+	// Use original format and size
+	const EPixelFormat Format = (PixelFormat != PF_Unknown ? PixelFormat : FromTexture->GetDesc().Format);
+	const FIntPoint Size = FromTexture->GetDesc().Extent;
+
+	// Prepare description
+	FRHITextureCreateDesc Desc =
+		FRHITextureCreateDesc::Create2D(Name, Size.X, Size.Y, Format)
+		.SetClearValue(FClearValueBinding::Transparent)
+		.SetNumMips(1)
+		.SetFlags(ETextureCreateFlags::ShaderResource | ETextureCreateFlags::RenderTargetable)
+		.SetInitialState(ERHIAccess::SRVMask);
+
+	// Create texture
+	return RHICreateTexture(Desc);
+}
+
+FRDGTextureRef FDisplayClusterGuiLayerController::CreateTextureFrom_RenderThread(
+	FRDGBuilder& GraphBuilder,
+	const TCHAR* Name,
+	const FIntPoint& Size,
+	EPixelFormat PixelFormat,
+	ETextureCreateFlags CreateFlags = TexCreate_ShaderResource | TexCreate_RenderTargetable,
+	const FClearValueBinding& ClearValueBinding = FClearValueBinding::Transparent)
+{
+	const FRDGTextureDesc Desc = FRDGTextureDesc::Create2D(Size, PixelFormat, ClearValueBinding, CreateFlags);
+	return GraphBuilder.CreateTexture(Desc, Name);
+}
+
+FRDGTextureRef FDisplayClusterGuiLayerController::CreateTextureFrom_RenderThread(FRDGBuilder& GraphBuilder, FRDGTextureRef& FromTexture, const TCHAR* Name, const FIntPoint& Size)
+{
+	if (!FromTexture)
+	{
+		return nullptr;
+	}
+
+	const FIntPoint NewSize = (Size == FIntPoint::ZeroValue ? FromTexture->Desc.Extent : Size);
+	return CreateTextureFrom_RenderThread(GraphBuilder, Name, NewSize, FromTexture->Desc.Format);
+}
+
+FRDGTextureRef FDisplayClusterGuiLayerController::CreateTextureFrom_RenderThread(FRDGBuilder& GraphBuilder, FRHITexture* FromTexture, const TCHAR* Name, const FIntPoint& Size)
+{
+	if (!FromTexture)
+	{
+		return nullptr;
+	}
+
+	const FIntPoint NewSize = (Size == FIntPoint::ZeroValue ? FromTexture->GetDesc().Extent : Size);
+	return CreateTextureFrom_RenderThread(GraphBuilder, Name, NewSize, FromTexture->GetFormat());
+}
diff --git a/Engine/Plugins/Runtime/nDisplay/Source/DisplayCluster/Private/Render/GUILayer/DisplayClusterGuiLayerController.h b/Engine/Plugins/Runtime/nDisplay/Source/DisplayCluster/Private/Render/GUILayer/DisplayClusterGuiLayerController.h
new file mode 100644
--- /dev/null
+++ b/Engine/Plugins/Runtime/nDisplay/Source/DisplayCluster/Private/Render/GUILayer/DisplayClusterGuiLayerController.h
@@ -0,0 +1,123 @@
+// Copyright Epic Games, Inc. All Rights Reserved.
+
+#pragma once
+
+#include "PixelFormat.h"
+#include "RHIResources.h"
+#include "UnrealClient.h"
+
+
+class FRDGBuilder;
+class FSceneViewFamily;
+class IDisplayClusterViewportProxy;
+
+
+/**
+ * GUI layer controller singleton
+ * 
+ * This class is responsible for propagation of the GUI layer (includes canvases + UMG) across
+ * all nDisplay viewports.
+ */
+class FDisplayClusterGuiLayerController
+{
+protected:
+	FDisplayClusterGuiLayerController();
+	~FDisplayClusterGuiLayerController() = default;
+
+public:
+	static FDisplayClusterGuiLayerController& Get();
+
+public:
+
+	/**
+	 * Returns true if this controller is enabled
+	 */
+	bool IsEnabled() const
+	{
+		return bEnabled;
+	}
+
+	/**
+	 * Returns true if this controller is active this frame
+	 */
+	bool IsActiveThisFrame() const
+	{
+		checkSlow(IsInGameThread() || IsInRenderingThread());
+		const bool bActive = (IsInRenderingThread() ? bActiveThisFrameRT : bActiveThisFrame);
+		return IsEnabled() && bActive;
+	}
+
+	/**
+	 * Called before copying the final output to the backbuffer. It's used to substitute the current
+	 * viewport's buffer to its original one stored on SlatePreTick (refer corresponding method).
+	 * 
+	 * @param FinalTexture - The final Slate outcome that is passed for copying to the backbuffer
+	 */
+	FRDGTextureRef ProcessFinalTexture_RenderThread(FRDGBuilder& GraphBuilder, FRDGTextureRef FinalTexture);
+
+	/**
+	 * Called for each nDisplay view family once rendering is finished. Responsible for drawing
+	 * the GUI layer on top of every nDisplay viewport.
+	 * 
+	 * @param ViewFamily - View family that was just rendered
+	 * @param ViewProxy  - nDisplay viewport (render thread obj)
+	 */
+	void ProcessPostRenderViewFamily_RenderThread(FRDGBuilder& GraphBuilder, const FSceneViewFamily& ViewFamiliy, const IDisplayClusterViewportProxy* ViewportProxy);
+
+private:
+
+	/**
+	 * Slate pre-Tick callback, used to substitute original game viewport buffer with our custom texture
+	 * to store whole Slate rendering output separately, and with high-precision alpha.
+	 */
+	void HandleSlatePreTick(float InDeltaTime);
+
+private:
+
+	/** Returns current game vieport buffer */
+	FTextureRHIRef GetViewportRTT_RenderThread();
+
+	/** Sets up new game viewport bufer (stays there until end of the frame) */
+	void SetViewportRTT_RenderThread(FTextureRHIRef& NewRTT);
+
+	/**
+	 * Creates or re-creates a texture duplicate from a referenced texture. Allows to override some texture parameters.
+	 * If some parameters of the referenced texture are different, a new internal texture will be created.
+	 * 
+	 * @param StorageTexture - An internal "duplicate" of the referenced texture
+	 * @param FromTexture    - The referenced texture
+	 * @param Name           - Name of the new texture
+	 * @param PixelFormat    - Pixel format of the new texture
+	 *
+	*/
+	void UpdateTempTexture_RenderThread(FRHICommandListImmediate& RHICmdList, FTextureRHIRef& StorageTexture, FTextureRHIRef& FromTexture, const TCHAR* Name, EPixelFormat PixelFormat = PF_Unknown);
+
+	/** Auxiliary function to create an RHI texture using another RHI texture as a template, overriding some of the reference parameters. */
+	FTextureRHIRef CreateTextureFrom_RenderThread(FRHICommandListImmediate& RHICmdList, FTextureRHIRef& FromTexture, const TCHAR* Name, EPixelFormat PixelFormat = PF_Unknown);
+
+	/** Auxiliary function to create an RDG texture. */
+	FRDGTextureRef CreateTextureFrom_RenderThread(FRDGBuilder& GraphBuilder, const TCHAR* Name, const FIntPoint& Size, EPixelFormat PixelFormat, ETextureCreateFlags CreateFlags, const FClearValueBinding& ClearValueBinding);
+
+	/** Auxiliary function to create an RDG texture from an RDG texture template, overriding some of the reference parameters. */
+	FRDGTextureRef CreateTextureFrom_RenderThread(FRDGBuilder& GraphBuilder, FRDGTextureRef& FromTexture, const TCHAR* Name, const FIntPoint& Size = FIntPoint::ZeroValue);
+
+	/** Auxiliary function to create an RDG texture from an RHI texture template, overriding some of the reference parameters. */
+	FRDGTextureRef CreateTextureFrom_RenderThread(FRDGBuilder& GraphBuilder, FRHITexture* FromTexture, const TCHAR* Name, const FIntPoint& Size = FIntPoint::ZeroValue);
+
+private:
+
+	/** Whether GUI layer propagation is active within current session */
+	const bool bEnabled;
+
+	/** Whether GUI layer propagation is active within current frame (game thread) */
+	bool bActiveThisFrame = false;
+
+	/** Whether GUI layer propagation is active within current frame (render thread) */
+	bool bActiveThisFrameRT = false;
+
+	/** The RHI texture reference to the original game viewport buffer  */
+	FTextureRHIRef TexViewportOriginalRHI;
+
+	/** Internal RHI texture used by Slate to draw UI. */
+	FTextureRHIRef TexViewportGUILayerRHI;
+};
diff --git a/Engine/Plugins/Runtime/nDisplay/Source/DisplayCluster/Private/Render/Viewport/DisplayClusterViewportProxy.cpp b/Engine/Plugins/Runtime/nDisplay/Source/DisplayCluster/Private/Render/Viewport/DisplayClusterViewportProxy.cpp
--- a/Engine/Plugins/Runtime/nDisplay/Source/DisplayCluster/Private/Render/Viewport/DisplayClusterViewportProxy.cpp
+++ b/Engine/Plugins/Runtime/nDisplay/Source/DisplayCluster/Private/Render/Viewport/DisplayClusterViewportProxy.cpp
@@ -1,50 +1,26 @@
 // Copyright Epic Games, Inc. All Rights Reserved.
 
 #include "Render/Viewport/DisplayClusterViewportProxy.h"
-#include "Render/Viewport/DisplayClusterViewportManagerProxy.h"
-#include "Render/Viewport/DisplayClusterViewport.h"
-#include "Render/Viewport/DisplayClusterViewportManager.h"
-#include "Render/Viewport/DisplayClusterViewportManagerViewExtension.h"
-
-#include "Render/Viewport/Containers/DisplayClusterViewport_PostRenderSettings.h"
-#include "Render/Viewport/Containers/DisplayClusterViewportProxyData.h"
 
-#include "Render/Viewport/RenderTarget/DisplayClusterRenderTargetResource.h"
-#include "Render/Viewport/RenderFrame/DisplayClusterRenderFrameSettings.h"
-
-#include "Render/Viewport/LightCard/DisplayClusterViewportLightCardManagerProxy.h"
-
-#include "Render/Projection/IDisplayClusterProjectionPolicy.h"
+#include "CommonRenderResources.h"
+#include "RenderGraphResources.h"
+#include "RenderGraphUtils.h"
+#include "SceneTextures.h"
 
-#include "Render/Containers/IDisplayClusterRender_MeshComponent.h"
-#include "Render/Containers/IDisplayClusterRender_MeshComponentProxy.h"
+#include "Engine/TextureRenderTarget.h"
 
 #include "IDisplayCluster.h"
 #include "IDisplayClusterCallbacks.h"
 #include "IDisplayClusterShaders.h"
 #include "IDisplayClusterShadersTextureUtils.h"
-#include "TextureResource.h"
-
-#include "Engine/TextureRenderTarget.h"
 
-#include "RHIStaticStates.h"
-
-#include "RenderResource.h"
-#include "RenderingThread.h"
-#include "CommonRenderResources.h"
-#include "PixelShaderUtils.h"
-
-#include "GlobalShader.h"
-#include "ShaderParameters.h"
-#include "ShaderParameterUtils.h"
-
-#include "ScreenRendering.h"
-#include "RenderGraphUtils.h"
-#include "RenderGraphResources.h"
-
-#include "ScreenPass.h"
-#include "SceneTextures.h"
-#include "PostProcess/PostProcessMaterialInputs.h"
+#include "Render/Containers/IDisplayClusterRender_MeshComponentProxy.h"
+#include "Render/GUILayer/DisplayClusterGuiLayerController.h"
+#include "Render/Projection/IDisplayClusterProjectionPolicy.h"
+#include "Render/Viewport/DisplayClusterViewport.h"
+#include "Render/Viewport/DisplayClusterViewportManagerProxy.h"
+#include "Render/Viewport/Containers/DisplayClusterViewportProxyData.h"
+#include "Render/Viewport/Resource/DisplayClusterViewportResource.h"
 
 
 ///////////////////////////////////////////////////////////////////////////////////////
@@ -418,6 +394,10 @@ void FDisplayClusterViewportProxy::OnPostRenderViewFamily_RenderThread(FRDGBuild
 		Contexts[InContextNum].RenderThreadData.EngineShowFlags = InSceneView.Family->EngineShowFlags;
 	}
 
+	// Allow the GUI controller to process output textures before PostRenderViewFamily is broadcasted. A dedicated
+	// call is required to avoid any potential issues caused by the event call order.
+	FDisplayClusterGuiLayerController::Get().ProcessPostRenderViewFamily_RenderThread(GraphBuilder, InViewFamily, this);
+
 	if (!InProxyContext.ViewFamilyProfileDescription.IsEmpty())
 	{
 		static IDisplayClusterCallbacks& DCCallbacksAPI = IDisplayCluster::Get().GetCallbacks();
diff --git a/Engine/Plugins/Runtime/nDisplay/Source/DisplayCluster/Public/DisplayClusterViewportClient.h b/Engine/Plugins/Runtime/nDisplay/Source/DisplayCluster/Public/DisplayClusterViewportClient.h
--- a/Engine/Plugins/Runtime/nDisplay/Source/DisplayCluster/Public/DisplayClusterViewportClient.h
+++ b/Engine/Plugins/Runtime/nDisplay/Source/DisplayCluster/Public/DisplayClusterViewportClient.h
@@ -21,6 +21,12 @@ public:
 
 	virtual FSceneViewport* CreateGameViewport(TSharedPtr<SViewport> InViewportWidget) override;
 
+	/** Always use non-scaled viewports with nDisplay */
+	virtual float GetDPIScale() const override
+	{
+		return 1.f;
+	}
+
 protected:
 #if WITH_EDITOR
 	bool Draw_PIE(FViewport* InViewport, FCanvas* SceneCanvas);
diff --git a/Engine/Plugins/Runtime/nDisplay/Source/DisplayClusterShaders/Private/DisplayClusterShadersModule.cpp b/Engine/Plugins/Runtime/nDisplay/Source/DisplayClusterShaders/Private/DisplayClusterShadersModule.cpp
--- a/Engine/Plugins/Runtime/nDisplay/Source/DisplayClusterShaders/Private/DisplayClusterShadersModule.cpp
+++ b/Engine/Plugins/Runtime/nDisplay/Source/DisplayClusterShaders/Private/DisplayClusterShadersModule.cpp
@@ -9,6 +9,7 @@
 
 #include "Shaders/DisplayClusterShadersGenerateMips.h"
 #include "Shaders/DisplayClusterShadersMedia.h"
+#include "Shaders/DisplayClusterShadersOverlay.h"
 #include "Shaders/DisplayClusterShadersPreprocess_UVLightCards.h"
 #include "Shaders/DisplayClusterShadersPostprocess_Blur.h"
 #include "Shaders/DisplayClusterShadersPostprocess_OutputRemap.h"
@@ -75,6 +76,11 @@ void FDisplayClusterShadersModule::AddPQToLinearPass(FRDGBuilder& GraphBuilder,
 	FDisplayClusterShadersMedia::AddPQToLinearPass(GraphBuilder, Parameters);
 }
 
+void FDisplayClusterShadersModule::AddDrawOverlayPass(FRDGBuilder& GraphBuilder, const FDisplayClusterShaderParameters_Overlay& Parameters) const
+{
+	FDisplayClusterShadersOverlay::AddOverlayBlendingPass(GraphBuilder, Parameters);
+}
+
 TSharedRef<IDisplayClusterShadersTextureUtils> FDisplayClusterShadersModule::CreateTextureUtils_RenderThread(FRHICommandListImmediate& RHICmdList) const
 {
 	return FDisplayClusterShadersTextureUtils::CreateTextureUtils_RenderThread(RHICmdList);
diff --git a/Engine/Plugins/Runtime/nDisplay/Source/DisplayClusterShaders/Private/DisplayClusterShadersModule.h b/Engine/Plugins/Runtime/nDisplay/Source/DisplayClusterShaders/Private/DisplayClusterShadersModule.h
--- a/Engine/Plugins/Runtime/nDisplay/Source/DisplayClusterShaders/Private/DisplayClusterShadersModule.h
+++ b/Engine/Plugins/Runtime/nDisplay/Source/DisplayClusterShaders/Private/DisplayClusterShadersModule.h
@@ -25,7 +25,8 @@ public:
 	virtual bool GenerateMips(FRHICommandListImmediate& RHICmdList, FRHITexture* InOutMipsTexture, const FDisplayClusterShaderParameters_GenerateMips& InSettings) const override;
 	virtual void AddLinearToPQPass(FRDGBuilder& GraphBuilder, const FDisplayClusterShaderParameters_MediaPQ& Parameters) const override;
 	virtual void AddPQToLinearPass(FRDGBuilder& GraphBuilder, const FDisplayClusterShaderParameters_MediaPQ& Parameters) const override;
-	
+	virtual void AddDrawOverlayPass(FRDGBuilder& GraphBuilder, const FDisplayClusterShaderParameters_Overlay& Parameters) const override;
+
 	virtual TSharedRef<IDisplayClusterShadersTextureUtils> CreateTextureUtils_RenderThread(FRHICommandListImmediate& RHICmdList) const override;
 	virtual TSharedRef<IDisplayClusterShadersTextureUtils> CreateTextureUtils_RenderThread(FRDGBuilder& GraphBuilder) const override;
 	//~~End IDisplayClusterShaders
diff --git a/Engine/Plugins/Runtime/nDisplay/Source/DisplayClusterShaders/Private/Shaders/DisplayClusterShadersOverlay.cpp b/Engine/Plugins/Runtime/nDisplay/Source/DisplayClusterShaders/Private/Shaders/DisplayClusterShadersOverlay.cpp
new file mode 100644
--- /dev/null
+++ b/Engine/Plugins/Runtime/nDisplay/Source/DisplayClusterShaders/Private/Shaders/DisplayClusterShadersOverlay.cpp
@@ -0,0 +1,90 @@
+// Copyright Epic Games, Inc. All Rights Reserved.
+
+#include "Shaders/DisplayClusterShadersOverlay.h"
+#include "ShaderParameters/DisplayClusterShaderParameters_Overlay.h"
+
+#include "GlobalShader.h"
+#include "PixelFormat.h"
+#include "PixelShaderUtils.h"
+#include "RenderGraphBuilder.h"
+#include "RenderGraphUtils.h"
+#include "RHIResources.h"
+#include "RHIStaticStates.h"
+#include "SceneView.h"
+#include "ScreenPass.h"
+#include "TextureResource.h"
+
+#include "DisplayClusterShadersLog.h"
+
+
+namespace UE::DisplayClusterShaders::Private
+{
+	static const FString OverlayShadersPath = TEXT("/Plugin/nDisplay/Private/OverlayShaders.usf");
+
+	///////////////////////////////////////////////////////////////////////////////////
+	// DrawOverlay PS
+	class FDrawOverlayPS : public FGlobalShader
+	{
+		DECLARE_SHADER_TYPE(FDrawOverlayPS, Global);
+		SHADER_USE_PARAMETER_STRUCT(FDrawOverlayPS, FGlobalShader);
+
+		BEGIN_SHADER_PARAMETER_STRUCT(FParameters, )
+			SHADER_PARAMETER_RDG_TEXTURE(Texture2D, TextureBase)
+			SHADER_PARAMETER_RDG_TEXTURE(Texture2D, TextureOverlay)
+			SHADER_PARAMETER_SAMPLER(SamplerState, Sampler)
+			RENDER_TARGET_BINDING_SLOTS()
+		END_SHADER_PARAMETER_STRUCT()
+
+	public:
+
+		/** A helper function to initialize shader parameters based on the draw request data */
+		FParameters* AllocateAndSetParameters(FRDGBuilder& InGraphBuilder, const FDisplayClusterShaderParameters_Overlay& InParameters)
+		{
+			FParameters* Parameters = InGraphBuilder.AllocParameters<FParameters>();
+
+			// Input
+			Parameters->TextureBase    = InParameters.BaseTexture;
+			Parameters->TextureOverlay = InParameters.OverlayTexture;
+			Parameters->Sampler        = TStaticSamplerState<SF_Point>::GetRHI();
+
+			// Output
+			Parameters->RenderTargets[0] = FRenderTargetBinding{ InParameters.OutputTexture, ERenderTargetLoadAction::ENoAction };
+
+			return Parameters;
+		}
+	};
+
+	IMPLEMENT_SHADER_TYPE(, FDrawOverlayPS, *OverlayShadersPath, TEXT("DrawOverlay_PS"), SF_Pixel);
+}
+
+
+void FDisplayClusterShadersOverlay::AddOverlayBlendingPass(FRDGBuilder& GraphBuilder, const FDisplayClusterShaderParameters_Overlay& InParameters)
+{
+	using namespace UE::DisplayClusterShaders::Private;
+
+	// Nothing to do if wrong input
+	if (!InParameters.IsValidData())
+	{
+		return;
+	}
+
+	// Initialize shaders
+	const FGlobalShaderMap* const GlobalShaderMap = GetGlobalShaderMap(GMaxRHIFeatureLevel);
+	const TShaderMapRef<FScreenPassVS>  VertexShader(GlobalShaderMap);
+	const TShaderMapRef<FDrawOverlayPS> PixelShader(GlobalShaderMap);
+
+	// Instantiate PS shader params
+	FDrawOverlayPS::FParameters* ShaderParams = PixelShader->AllocateAndSetParameters(GraphBuilder, InParameters);
+
+	// Add draw pass
+	AddDrawScreenPass(
+		GraphBuilder,
+		RDG_EVENT_NAME("nDisplay.DrawOverlay_PS"),
+		FScreenPassViewInfo(),
+		FScreenPassTextureViewport{ InParameters.BaseTexture,    FIntRect({0, 0}, InParameters.BaseTexture->Desc.Extent) },
+		FScreenPassTextureViewport{ InParameters.OverlayTexture, FIntRect({0, 0}, InParameters.OverlayTexture->Desc.Extent) },
+		VertexShader,
+		PixelShader,
+		ShaderParams
+	);
+}
diff --git a/Engine/Plugins/Runtime/nDisplay/Source/DisplayClusterShaders/Private/Shaders/DisplayClusterShadersOverlay.h b/Engine/Plugins/Runtime/nDisplay/Source/DisplayClusterShaders/Private/Shaders/DisplayClusterShadersOverlay.h
new file mode 100644
--- /dev/null
+++ b/Engine/Plugins/Runtime/nDisplay/Source/DisplayClusterShaders/Private/Shaders/DisplayClusterShadersOverlay.h
@@ -0,0 +1,20 @@
+// Copyright Epic Games, Inc. All Rights Reserved.
+
+#pragma once
+
+#include "CoreMinimal.h"
+
+class FRDGBuilder;
+struct FDisplayClusterShaderParameters_Overlay;
+
+
+/**
+ * API for overlay drawing
+ */
+class FDisplayClusterShadersOverlay
+{
+public:
+
+	/** Draws overlay on top of the base texture (refer function arguments) */
+	static void AddOverlayBlendingPass(FRDGBuilder& GraphBuilder, const FDisplayClusterShaderParameters_Overlay& Parameters);
+};
diff --git a/Engine/Plugins/Runtime/nDisplay/Source/DisplayClusterShaders/Public/IDisplayClusterShaders.h b/Engine/Plugins/Runtime/nDisplay/Source/DisplayClusterShaders/Public/IDisplayClusterShaders.h
--- a/Engine/Plugins/Runtime/nDisplay/Source/DisplayClusterShaders/Public/IDisplayClusterShaders.h
+++ b/Engine/Plugins/Runtime/nDisplay/Source/DisplayClusterShaders/Public/IDisplayClusterShaders.h
@@ -20,6 +20,7 @@ class FRenderTarget;
 
 struct FDisplayClusterShaderParameters_GenerateMips;
 struct FDisplayClusterShaderParameters_MediaPQ;
+struct FDisplayClusterShaderParameters_Overlay;
 struct FDisplayClusterShaderParameters_PostprocessBlur;
 struct FDisplayClusterShaderParameters_WarpBlend;
 struct FDisplayClusterShaderParameters_UVLightCards;
@@ -153,6 +154,14 @@ public:
 	*/
 	virtual void AddPQToLinearPass(FRDGBuilder& GraphBuilder, const FDisplayClusterShaderParameters_MediaPQ& Parameters) const = 0;
 
+	/**
+	* Adds overlay drawing pass
+	*
+	* @param GraphBuilder - RDG builder
+	* @param Parameters   - Render parameters
+	*/
+	virtual void AddDrawOverlayPass(FRDGBuilder& GraphBuilder, const FDisplayClusterShaderParameters_Overlay& Parameters) const = 0;
+
 	/**
 	* Return new instance of the resource utils.
 	* 
diff --git a/Engine/Plugins/Runtime/nDisplay/Source/DisplayClusterShaders/Public/ShaderParameters/DisplayClusterShaderParameters_Overlay.h b/Engine/Plugins/Runtime/nDisplay/Source/DisplayClusterShaders/Public/ShaderParameters/DisplayClusterShaderParameters_Overlay.h
new file mode 100644
--- /dev/null
+++ b/Engine/Plugins/Runtime/nDisplay/Source/DisplayClusterShaders/Public/ShaderParameters/DisplayClusterShaderParameters_Overlay.h
@@ -0,0 +1,31 @@
+// Copyright Epic Games, Inc. All Rights Reserved.
+
+#pragma once
+
+#include "CoreMinimal.h"
+
+class FRDGTexture;
+
+
+/*
+ * Holds overlay drawing parameters
+ */
+struct FDisplayClusterShaderParameters_Overlay
+{
+	/** Base texture (bottom layer) */
+	FRDGTexture* BaseTexture = nullptr;
+
+	/** Overlay texture (top layer) */
+	FRDGTexture* OverlayTexture = nullptr;
+
+	/** Output texture (the result of source-over blending) */
+	FRDGTexture* OutputTexture = nullptr;
+
+public:
+
+	/** Parameters validation */
+	bool IsValidData() const
+	{
+		return BaseTexture && OverlayTexture && OutputTexture;
+	}
+};
