[nDisplay] Fixed copy size  after bughawk review

- Replaced FRHICopyTextureInfo::Size.Z==0 to ==1 in all nDisplay occurrences

#rnx

diff --git a/Engine/Plugins/Runtime/nDisplay/Source/DisplayCluster/Private/Render/GUILayer/DisplayClusterGuiLayerController.cpp b/Engine/Plugins/Runtime/nDisplay/Source/DisplayCluster/Private/Render/GUILayer/DisplayClusterGuiLayerController.cpp
--- a/Engine/Plugins/Runtime/nDisplay/Source/DisplayCluster/Private/Render/GUILayer/DisplayClusterGuiLayerController.cpp
+++ b/Engine/Plugins/Runtime/nDisplay/Source/DisplayCluster/Private/Render/GUILayer/DisplayClusterGuiLayerController.cpp
@@ -109,7 +109,7 @@ void FDisplayClusterGuiLayerController::ProcessPostRenderViewFamily_RenderThread
 
 			// Make a copy to the input texture
 			FRHICopyTextureInfo CopyInfo;
-			CopyInfo.Size = FIntVector{ DCViewportSize.X, DCViewportSize.Y, 0 };
+			CopyInfo.Size = FIntVector{ DCViewportSize.X, DCViewportSize.Y, 1 };
 			CopyInfo.SourcePosition = FIntVector{ Regions[0].Min.X, Regions[0].Min.Y, 0 };
 			AddCopyTexturePass(GraphBuilder, TexDCViewportRDG, TexTempDCViewportRDG, CopyInfo);
 
diff --git a/Engine/Plugins/Runtime/nDisplay/Source/DisplayCluster/Private/Render/Viewport/DisplayClusterViewportManagerProxy.cpp b/Engine/Plugins/Runtime/nDisplay/Source/DisplayCluster/Private/Render/Viewport/DisplayClusterViewportManagerProxy.cpp
--- a/Engine/Plugins/Runtime/nDisplay/Source/DisplayCluster/Private/Render/Viewport/DisplayClusterViewportManagerProxy.cpp
+++ b/Engine/Plugins/Runtime/nDisplay/Source/DisplayCluster/Private/Render/Viewport/DisplayClusterViewportManagerProxy.cpp
@@ -625,7 +625,7 @@ bool FDisplayClusterViewportManagerProxy::ImplResolveFrameTargetToBackBuffer_Ren
 			CopyInfo.SourcePosition = FIntVector(SrcRect.Min.X, SrcRect.Min.Y, 0);
 			CopyInfo.DestPosition = FIntVector(DestRect.Min.X, DestRect.Min.Y, 0);
 
-			CopyInfo.Size = FIntVector(DestRect.Width(), DestRect.Height(), 0);
+			CopyInfo.Size = FIntVector(DestRect.Width(), DestRect.Height(), 1);
 
 			FRDGTextureRef SrcTexture = RegisterExternalTexture(GraphBuilder, FrameTexture->GetTexture2D(), TEXT("DisplayClusterStereoTexture"));
 
@@ -685,7 +685,7 @@ bool FDisplayClusterViewportManagerProxy::ResolveFrameTargetToBackBuffer_RenderT
 			CopyInfo.SourcePosition = FIntVector(SrcRect.Min.X, SrcRect.Min.Y, 0);
 			CopyInfo.DestPosition = FIntVector(DestRect.Min.X, DestRect.Min.Y, 0);
 
-			CopyInfo.Size = FIntVector(DestRect.Width(), DestRect.Height(), 0);
+			CopyInfo.Size = FIntVector(DestRect.Width(), DestRect.Height(), 1);
 
 			TransitionAndCopyTexture(RHICmdList, FrameTexture, DestTexture, CopyInfo);
 
diff --git a/Engine/Plugins/Runtime/nDisplay/Source/DisplayClusterMedia/Private/Input/DisplayClusterMediaInputBase.cpp b/Engine/Plugins/Runtime/nDisplay/Source/DisplayClusterMedia/Private/Input/DisplayClusterMediaInputBase.cpp
--- a/Engine/Plugins/Runtime/nDisplay/Source/DisplayClusterMedia/Private/Input/DisplayClusterMediaInputBase.cpp
+++ b/Engine/Plugins/Runtime/nDisplay/Source/DisplayClusterMedia/Private/Input/DisplayClusterMediaInputBase.cpp
@@ -213,7 +213,7 @@ void FDisplayClusterMediaInputBase::ImportMediaDataDirect_RenderThread(FRHIComma
 		FRHICopyTextureInfo CopyInfo;
 		CopyInfo.SourcePosition = FIntVector(SrcRect.Min.X, SrcRect.Min.Y, 0);
 		CopyInfo.DestPosition = FIntVector(DstRect.Min.X, DstRect.Min.Y, 0);
-		CopyInfo.Size = FIntVector(DstRect.Size().X, DstRect.Size().Y, 0);
+		CopyInfo.Size = FIntVector(DstRect.Size().X, DstRect.Size().Y, 1);
 
 		TransitionAndCopyTexture(RHICmdList, SrcTexture, DstTexture, CopyInfo);
 	}
diff --git a/Engine/Plugins/Runtime/nDisplay/Source/DisplayClusterMoviePipeline/Private/DisplayClusterMoviePipelineViewportPass.cpp b/Engine/Plugins/Runtime/nDisplay/Source/DisplayClusterMoviePipeline/Private/DisplayClusterMoviePipelineViewportPass.cpp
--- a/Engine/Plugins/Runtime/nDisplay/Source/DisplayClusterMoviePipeline/Private/DisplayClusterMoviePipelineViewportPass.cpp
+++ b/Engine/Plugins/Runtime/nDisplay/Source/DisplayClusterMoviePipeline/Private/DisplayClusterMoviePipelineViewportPass.cpp
@@ -47,7 +47,7 @@ namespace UE::DisplayCluster::MoviePipelineViewportPass
 							CopyInfo.Size = FIntVector(
 								RTTRect.Width() - OffsetMin.X - OffsetMax.X,
 								RTTRect.Height() - OffsetMin.Y - OffsetMax.Y,
-								0);
+								1);
 							{
 								// Copy RTTRect->WarpInputRect
 								CopyInfo.SourcePosition.X = RTTRect.Min.X + OffsetMin.X;
diff --git a/Engine/Plugins/Runtime/nDisplay/Source/DisplayClusterShaders/Private/Shaders/DisplayClusterShadersCopyTexture.cpp b/Engine/Plugins/Runtime/nDisplay/Source/DisplayClusterShaders/Private/Shaders/DisplayClusterShadersCopyTexture.cpp
--- a/Engine/Plugins/Runtime/nDisplay/Source/DisplayClusterShaders/Private/Shaders/DisplayClusterShadersCopyTexture.cpp
+++ b/Engine/Plugins/Runtime/nDisplay/Source/DisplayClusterShaders/Private/Shaders/DisplayClusterShadersCopyTexture.cpp
@@ -470,7 +470,7 @@ namespace UE::DisplayClusterShaders::Private
 		OutCopyInfo.SourcePosition = FIntVector(Input.Rect.Min.X, Input.Rect.Min.Y, 0);
 		OutCopyInfo.DestPosition = FIntVector(Output.Rect.Min.X, Output.Rect.Min.Y, 0);
 
-		OutCopyInfo.Size = FIntVector(Output.Rect.Width(), Output.Rect.Height(), 0);
+		OutCopyInfo.Size = FIntVector(Output.Rect.Width(), Output.Rect.Height(), 1);
 
 		return OutCopyInfo;
 	}
