[nDisplay] Fixed non-critical regression that forces some media capture to restart on the 2nd frame

#rnx

diff --git a/Engine/Plugins/Runtime/nDisplay/Source/DisplayClusterMedia/Private/Capture/DisplayClusterMediaCaptureBase.cpp b/Engine/Plugins/Runtime/nDisplay/Source/DisplayClusterMedia/Private/Capture/DisplayClusterMediaCaptureBase.cpp
--- a/Engine/Plugins/Runtime/nDisplay/Source/DisplayClusterMedia/Private/Capture/DisplayClusterMediaCaptureBase.cpp
+++ b/Engine/Plugins/Runtime/nDisplay/Source/DisplayClusterMedia/Private/Capture/DisplayClusterMediaCaptureBase.cpp
@@ -252,8 +252,9 @@ void FDisplayClusterMediaCaptureBase::OnPostClusterTick()
 		{
 			const FIntPoint LastSrcRegionIntPoint = LastSrcRegionSize.load().ToIntPoint();
 			const FIntPoint DesiredSize = MediaCapture->GetDesiredSize();
-			const bool bDesiredSizeMatchesLastRegion = LastSrcRegionIntPoint != FIntPoint::ZeroValue && DesiredSize == LastSrcRegionIntPoint;
-			const bool bDesiredSizeMatchesCustomRegion = CustomResolution.X > 0 && CustomResolution.Y > 0 && DesiredSize == CustomResolution;
+			const bool bHasExportedAlready = (LastSrcRegionIntPoint != FIntPoint::ZeroValue);
+			const bool bDesiredSizeMatchesLastRegion = (DesiredSize == LastSrcRegionIntPoint);
+			const bool bDesiredSizeMatchesCustomRegion = (CustomResolution.X > 0 && CustomResolution.Y > 0 && DesiredSize == CustomResolution);
 
 			// We don't restart if we haven't exported any textures yet (indicated by zero size LastSrcRegion)
 			// to avoid constant media restarts since in such case media is set to use GetCaptureSize() != (0.0).
@@ -262,7 +263,7 @@ void FDisplayClusterMediaCaptureBase::OnPostClusterTick()
 			// This is preferred to setting LastSrcRegion to GetCaptureSize() because that would not reflect the actual
 			// last captured size.
 
-			if (!(bDesiredSizeMatchesLastRegion || bDesiredSizeMatchesCustomRegion))
+			if (bHasExportedAlready && !(bDesiredSizeMatchesLastRegion || bDesiredSizeMatchesCustomRegion))
 			{
 				UE_LOG(LogDisplayClusterMedia, Log, TEXT("Stopping MediaCapture '%s' because its DesiredSize (%d, %d) doesn't match the captured texture size (%d, %d)"), 
 					*GetMediaId(), DesiredSize.X, DesiredSize.Y, LastSrcRegionIntPoint.X, LastSrcRegionIntPoint.Y);
