[nDisplay] Fixed race condition that crashes a cluster node during connections establishing

-The problem was caused by multiple networking sessions broadcasting SeesionOpened and SessionClosed events from their own threads (SessionOpened if crashes at the very beginning). This means the issue becomes more frequent as the number of nodes in the cluster increases. The fix makes SessionOpened and SessionClosed events broadcasting safe by throwing them under a critical section.

#rnx

diff --git a/Engine/Plugins/Runtime/nDisplay/Source/DisplayCluster/Private/Network/DisplayClusterServer.h b/Engine/Plugins/Runtime/nDisplay/Source/DisplayCluster/Private/Network/DisplayClusterServer.h
--- a/Engine/Plugins/Runtime/nDisplay/Source/DisplayCluster/Private/Network/DisplayClusterServer.h
+++ b/Engine/Plugins/Runtime/nDisplay/Source/DisplayCluster/Private/Network/DisplayClusterServer.h
@@ -76,6 +76,20 @@ public:
 		return SessionClosedEvent;
 	}
 
+	// Wraps SessionOpened to emit it safely from any thread.
+	virtual void NotifySessionOpened(const FDisplayClusterSessionInfo& SessionInfo) override
+	{
+		FScopeLock Lock(&SessionNotifyCritSec);
+		OnSessionOpened().Broadcast(SessionInfo);
+	}
+
+	// Wraps SessionClosed to emit it safely from any thread.
+	virtual void NotifySessionClosed(const FDisplayClusterSessionInfo& SessionInfo) override
+	{
+		FScopeLock Lock(&SessionNotifyCritSec);
+		OnSessionClosed().Broadcast(SessionInfo);
+	}
+
 	//~ End IDisplayClusterServer
 
 protected:
@@ -142,4 +156,12 @@ private:
 
 	/** Critical section to manipulate the connection sessions */
 	mutable FCriticalSection SessionsCritSec;
+
+	/**
+	 * Critical section for safe session notifications. The same critical section
+	 * is used for both SessionOpened and SessionClosed events because they don't
+	 * really interfere in runtime. If there are new session events at some point,
+	 * we might probably need separate critical sections for those.
+	 */
+	mutable FCriticalSection SessionNotifyCritSec;
 };
diff --git a/Engine/Plugins/Runtime/nDisplay/Source/DisplayCluster/Private/Network/IDisplayClusterServer.h b/Engine/Plugins/Runtime/nDisplay/Source/DisplayCluster/Private/Network/IDisplayClusterServer.h
--- a/Engine/Plugins/Runtime/nDisplay/Source/DisplayCluster/Private/Network/IDisplayClusterServer.h
+++ b/Engine/Plugins/Runtime/nDisplay/Source/DisplayCluster/Private/Network/IDisplayClusterServer.h
@@ -58,4 +58,10 @@ public:
 	/** Session closed event */
 	DECLARE_EVENT_OneParam(IDisplayClusterServer, FSessionClosedEvent, const FDisplayClusterSessionInfo&);
 	virtual FSessionClosedEvent& OnSessionClosed() = 0;
+
+	/** Safe SessionOpened notification from a session thread */
+	virtual void NotifySessionOpened(const FDisplayClusterSessionInfo& SessionInfo) = 0;
+
+	/** Safe SessionClosed notification from a session thread */
+	virtual void NotifySessionClosed(const FDisplayClusterSessionInfo& SessionInfo) = 0;
 };
diff --git a/Engine/Plugins/Runtime/nDisplay/Source/DisplayCluster/Private/Network/Session/DisplayClusterSession.h b/Engine/Plugins/Runtime/nDisplay/Source/DisplayCluster/Private/Network/Session/DisplayClusterSession.h
--- a/Engine/Plugins/Runtime/nDisplay/Source/DisplayCluster/Private/Network/Session/DisplayClusterSession.h
+++ b/Engine/Plugins/Runtime/nDisplay/Source/DisplayCluster/Private/Network/Session/DisplayClusterSession.h
@@ -107,7 +107,7 @@ protected:
 		SessionInfo.TimeStart = FPlatformTime::Seconds();
 
 		// Notify owner about new session
-		GetOwner().OnSessionOpened().Broadcast(SessionInfo);
+		GetOwner().NotifySessionOpened(SessionInfo);
 
 		// Process all incoming messages unless server is shutdown or remote host is up
 		while (FDisplayClusterSocketOperations::IsOpen())
@@ -140,7 +140,7 @@ protected:
 		UE_LOG(LogDisplayClusterNetwork, Log, TEXT("Session thread %s has finished"), *SessionInfo.SessionName);
 
 		// Since we left the cycle above, it means the session has been closed. We need to notify the owning server about that.
-		GetOwner().OnSessionClosed().Broadcast(SessionInfo);
+		GetOwner().NotifySessionClosed(SessionInfo);
 
 		return 0;
 	}
