[nDisplay] Added temporary workaround to avoid crashes under some circumstances

- Return dummy configuration objects if a DCRA instance has invalid cluster configuration

#rnx

diff --git a/Engine/Plugins/Runtime/nDisplay/Source/DisplayCluster/Private/Game/EngineClasses/Scene/DisplayClusterRootActor.cpp b/Engine/Plugins/Runtime/nDisplay/Source/DisplayCluster/Private/Game/EngineClasses/Scene/DisplayClusterRootActor.cpp
--- a/Engine/Plugins/Runtime/nDisplay/Source/DisplayCluster/Private/Game/EngineClasses/Scene/DisplayClusterRootActor.cpp
+++ b/Engine/Plugins/Runtime/nDisplay/Source/DisplayCluster/Private/Game/EngineClasses/Scene/DisplayClusterRootActor.cpp
@@ -361,16 +361,30 @@ bool ADisplayClusterRootActor::IsRunningDisplayCluster() const
 
 const FDisplayClusterConfigurationICVFX_StageSettings& ADisplayClusterRootActor::GetStageSettings() const
 {
-	check(CurrentConfigData);
-
-	return CurrentConfigData->StageSettings;
+	// [TEMPORARY FIX]
+	// There are various ways we can get to this function while still having an uninitialized 'CurrentConfigData'.
+	// Some of them have already been fixed, while others are still unclear how to reproduce.
+	// UE-297696, UE-218690, UE-222703
+	// Normally, CurrentConfigData should never be nullptr. As a quick work around, we always return
+	// some default settings to avoid nullptr crash. Debug builds will keep reporting the issue.
+	checkSlow(CurrentConfigData);
+	ensureMsgf(IsValid(CurrentConfigData), TEXT("Found invalid DCRA configuration"));
+	static const FDisplayClusterConfigurationICVFX_StageSettings DummyStageSettings;
+	return IsValid(CurrentConfigData) ? CurrentConfigData->StageSettings : DummyStageSettings;
 }
 
 const FDisplayClusterConfigurationRenderFrame& ADisplayClusterRootActor::GetRenderFrameSettings() const 
-{ 
-	check(CurrentConfigData);
-
-	return CurrentConfigData->RenderFrameSettings;
+{
+	// [TEMPORARY FIX]
+	// There are various ways we can get to this function while still having an uninitialized 'CurrentConfigData'.
+	// Some of them have already been fixed, while others are still unclear how to reproduce.
+	// UE-297696, UE-218690, UE-222703
+	// Normally, CurrentConfigData should never be nullptr. As a quick work around, we always return
+	// some default settings to avoid nullptr crash. Debug builds will keep reporting the issue.
+	checkSlow(CurrentConfigData);
+	ensureMsgf(IsValid(CurrentConfigData), TEXT("Found invalid DCRA configuration"));
+	static const FDisplayClusterConfigurationRenderFrame DummyRenderFrameSettings;
+	return IsValid(CurrentConfigData) ? CurrentConfigData->RenderFrameSettings : DummyRenderFrameSettings;
 }
 
 void ADisplayClusterRootActor::InitializeFromConfig(UDisplayClusterConfigurationData* ConfigData)
diff --git a/Engine/Plugins/Runtime/nDisplay/Source/DisplayCluster/Private/Render/Viewport/Configuration/DisplayClusterViewportConfiguration.cpp b/Engine/Plugins/Runtime/nDisplay/Source/DisplayCluster/Private/Render/Viewport/Configuration/DisplayClusterViewportConfiguration.cpp
--- a/Engine/Plugins/Runtime/nDisplay/Source/DisplayCluster/Private/Render/Viewport/Configuration/DisplayClusterViewportConfiguration.cpp
+++ b/Engine/Plugins/Runtime/nDisplay/Source/DisplayCluster/Private/Render/Viewport/Configuration/DisplayClusterViewportConfiguration.cpp
@@ -35,15 +35,8 @@ namespace UE::DisplayCluster::Configuration
 {
 	static inline ADisplayClusterRootActor* ImplGetRootActor(const FDisplayClusterActorRef& InConfigurationRootActorRef)
 	{
-		if (AActor* ActorPtr = InConfigurationRootActorRef.GetOrFindSceneActor())
-		{
-			if (ActorPtr->IsA<ADisplayClusterRootActor>())
-			{
-				return static_cast<ADisplayClusterRootActor*>(ActorPtr);
-			}
-		}
-
-		return nullptr;
+		AActor* const ActorPtr = InConfigurationRootActorRef.GetOrFindSceneActor();
+		return IsValid(ActorPtr) ? Cast<ADisplayClusterRootActor>(ActorPtr) : nullptr;
 	}
 
 	static inline bool ImplIsChangedRootActor(const ADisplayClusterRootActor* InRootActor, FDisplayClusterActorRef& InOutConfigurationRootActorRef)
diff --git a/Engine/Plugins/Runtime/nDisplay/Source/DisplayCluster/Private/Render/Viewport/Configuration/DisplayClusterViewportConfigurationHelpers_Postprocess.cpp b/Engine/Plugins/Runtime/nDisplay/Source/DisplayCluster/Private/Render/Viewport/Configuration/DisplayClusterViewportConfigurationHelpers_Postprocess.cpp
--- a/Engine/Plugins/Runtime/nDisplay/Source/DisplayCluster/Private/Render/Viewport/Configuration/DisplayClusterViewportConfigurationHelpers_Postprocess.cpp
+++ b/Engine/Plugins/Runtime/nDisplay/Source/DisplayCluster/Private/Render/Viewport/Configuration/DisplayClusterViewportConfigurationHelpers_Postprocess.cpp
@@ -564,7 +564,7 @@ FDisplayClusterViewport_CameraMotionBlur FDisplayClusterViewportConfigurationHel
 	switch (InCfgCameraSettings.CameraMotionBlur.MotionBlurMode)
 	{
 	case EDisplayClusterConfigurationCameraMotionBlurMode::Override:
-		if (ADisplayClusterRootActor* SceneRootActor = static_cast<ADisplayClusterRootActor*>(InSceneCameraComponent.GetOwner()))
+		if (ADisplayClusterRootActor* SceneRootActor = Cast<ADisplayClusterRootActor>(InSceneCameraComponent.GetOwner()))
 		{
 			UDisplayClusterCameraComponent* OuterCamera = SceneRootActor->GetDefaultCamera();
 			if (OuterCamera)
