[nDisplay] Added tooltip for nDisplay configuration asset

#rnx

diff --git a/Engine/Plugins/Runtime/nDisplay/Source/DisplayCluster/Public/Blueprints/DisplayClusterBlueprint.h b/Engine/Plugins/Runtime/nDisplay/Source/DisplayCluster/Public/Blueprints/DisplayClusterBlueprint.h
--- a/Engine/Plugins/Runtime/nDisplay/Source/DisplayCluster/Public/Blueprints/DisplayClusterBlueprint.h
+++ b/Engine/Plugins/Runtime/nDisplay/Source/DisplayCluster/Public/Blueprints/DisplayClusterBlueprint.h
@@ -10,6 +10,14 @@
 class USCS_Node;
 class UActorComponent;
 
+/**
+ * Display Cluster Configuration Asset
+ *
+ * Asset containing the full nDisplay cluster setup and settings.
+ * Can also contain any blueprint logic like any other actor blueprints.
+ *
+ * @see UBlueprint
+ */
 UCLASS(BlueprintType, DisplayName = "nDisplay Blueprint")
 class DISPLAYCLUSTER_API UDisplayClusterBlueprint : public UBlueprint
 {
