Do not use {} when clearing the TOptional state.   This will not properly reset the optional state.

#rnx

diff --git a/Engine/Plugins/Runtime/nDisplay/Source/DisplayClusterMessageInterceptor/Private/DisplayClusterMessageInterceptionModule.cpp b/Engine/Plugins/Runtime/nDisplay/Source/DisplayClusterMessageInterceptor/Private/DisplayClusterMessageInterceptionModule.cpp
--- a/Engine/Plugins/Runtime/nDisplay/Source/DisplayClusterMessageInterceptor/Private/DisplayClusterMessageInterceptionModule.cpp
+++ b/Engine/Plugins/Runtime/nDisplay/Source/DisplayClusterMessageInterceptor/Private/DisplayClusterMessageInterceptionModule.cpp
@@ -66,17 +66,17 @@ public:
 		FScopeLock Lock(&InternalsCS);
 
 		bValue = bInValue;
-		bPendingValue = {};
+		bPendingValue.Reset();
 	}
 
 	void UpdateOnPreTick()
 	{
 		FScopeLock Lock(&InternalsCS);
 
-		if (bPendingValue)
+		if (bPendingValue.IsSet())
 		{
 			bValue = *bPendingValue;
-			bPendingValue = {};
+			bPendingValue.Reset();
 		}
 	}
 
