[nDisplay] Fixed nullptr crash in DCMessageInterceptor during node shutdown

#rnx

diff --git a/Engine/Plugins/Runtime/nDisplay/Source/DisplayClusterMessageInterceptor/Private/DisplayClusterMessageInterceptor.cpp b/Engine/Plugins/Runtime/nDisplay/Source/DisplayClusterMessageInterceptor/Private/DisplayClusterMessageInterceptor.cpp
--- a/Engine/Plugins/Runtime/nDisplay/Source/DisplayClusterMessageInterceptor/Private/DisplayClusterMessageInterceptor.cpp
+++ b/Engine/Plugins/Runtime/nDisplay/Source/DisplayClusterMessageInterceptor/Private/DisplayClusterMessageInterceptor.cpp
@@ -181,6 +181,11 @@ void FDisplayClusterMessageInterceptor::HandleClusterNodeFailure(const FString&
 {
 	TArray<TSharedPtr<IMessageContext, ESPMode::ThreadSafe>> ContextToForward;
 
+	if (!ClusterManager)
+	{
+		return;
+	}
+
 	{
 		FScopeLock Lock(&ContextQueueCS);
 
