nDisplay: Improved stage geometry mesh generation to mitigate deformities when geometry edges are near the mesh equator

diff --git a/Engine/Plugins/Runtime/nDisplay/Shaders/Private/MeshProjectionNormalSmoothing.usf b/Engine/Plugins/Runtime/nDisplay/Shaders/Private/MeshProjectionNormalSmoothing.usf
--- a/Engine/Plugins/Runtime/nDisplay/Shaders/Private/MeshProjectionNormalSmoothing.usf
+++ b/Engine/Plugins/Runtime/nDisplay/Shaders/Private/MeshProjectionNormalSmoothing.usf
@@ -9,9 +9,14 @@ Texture2D SceneColor;
 Texture2D SceneDepth;
 Texture2D<uint2> SceneStencil;
 
-RWTexture2D<float4> RWColor;
-RWTexture2D<float> RWDepth;
-RWTexture2D<uint> RWStencil;
+RWTexture2D<float4> RWColorIn;
+RWTexture2D<float4> RWColorOut;
+RWTexture2D<float> RWDepthIn;
+RWTexture2D<float> RWDepthOut;
+RWTexture2D<uint> RWStencilIn;
+RWTexture2D<uint> RWStencilOut;
+RWTexture2D<uint> RWBlurStencilIn;
+RWTexture2D<uint> RWBlurStencilOut;
 
 float4x4 NormalCorrectionMatrix;
 
@@ -24,10 +29,28 @@ float2 SampleDirection;
 float2 SampleOffsetScale;
 float InvSigma;
 
-/** Adjusts the specified pixel position so as not to be outside of the viewport bounds */
-uint2 AdjustPixelPos(uint2 PixelPos)
+/** Determines if the contents of the geometry stencil buffer should be copied to the blur stencil buffer during the dilation pass */
+uint bCopyToBlurStencil;
+
+/** Returns true if the current pixel is beyond the circular azimuthal bounds and should be clipped */
+bool ClipToAzimuthalBounds(uint2 PixelPos)
 {
-	return uint2(clamp(PixelPos.x, 0.0, Input_Extent.x), clamp(PixelPos.y, 0.0, Input_Extent.y));
+	const float2 UV = (PixelPos - float2(0.5, 0.5)) * View.ViewSizeAndInvSize.zw;
+	const float Radius = length(UV - float2(0.5, 0.5));
+
+	return Radius >= 0.5;
+}
+
+/** Applies an offset to the specified pixel position, ensuring that the resulting offset position is within the azimuthal bounds */
+uint2 AdjustPixelPos(uint2 PixelPos, int2 PixelOffset)
+{
+	uint2 OffsetPos = PixelPos + PixelOffset;
+
+	// If the sample position would be outside the circular bounds of the azimuthal projection,
+	// return the original pixel position instead
+	uint2 SamplePos = select(ClipToAzimuthalBounds(OffsetPos), PixelPos, OffsetPos);
+	
+	return uint2(clamp(SamplePos.x, 0.0, Input_Extent.x - 1), clamp(SamplePos.y, 0.0, Input_Extent.y - 1));
 }
 
 /** Converts a vector in normal space to color space */
@@ -86,9 +109,18 @@ void CreateRWTexturesCS(uint2 DispatchThreadId : SV_DispatchThreadID)
 	float3 Normal = lerp(-RadialVector, NormalFromColorSpace(Color.rgb), (float)Stencil);
 	float3 LocalNormal = TransformToRadialBasis(Normal, RadialVector);
 	
-	RWColor[PixelPos] = float4(NormalToColorSpace(LocalNormal), 1.0);
-	RWDepth[PixelPos] = Depth;
-	RWStencil[PixelPos] = Stencil;
+	// Sphere-ize the empty space depth buffer. Due to using a perspective projection matrix, a uniform depth value does
+	// not create a sphere around the stage, which can cause artifacts during blurring around stage geometry edges. To
+	// compute the correct depth, compute an azimuthal projection for a maximum radius point (pulled from the projection's
+	// far plane value), and extract its depth value
+	float MaxRho = View.NearPlane * (1.0 - 1.0 / View.ViewToClip[2][2]);
+	float4 ClipSpacePos = mul(float4(UnitViewPos * MaxRho, 1.0), View.ViewToClip);
+	float MaxDepth = ClipSpacePos.z / ClipSpacePos.w;
+
+	RWColorOut[PixelPos] = float4(NormalToColorSpace(LocalNormal), 1.0);
+	RWDepthOut[PixelPos] = lerp(MaxDepth, Depth, (float)Stencil);
+	RWStencilOut[PixelPos] = Stencil;
+	RWBlurStencilOut[PixelPos] = Stencil;
 }
 
 /** A compute pass that dilates the depth buffer of overlapping geometry, allowing nearer geometry depth to dilate into further geometry depth */
@@ -99,12 +131,19 @@ void DepthDilationCS(uint2 DispatchThreadId : SV_DispatchThreadID)
 	
 	const int kRadius = (FILTER_KERNEL_SIZE - 1) / 2;
 
-	float InitialDepth = RWDepth[PixelPos];
-	uint InitialStencil = RWStencil[PixelPos];
+	float InitialDepth = RWDepthIn[PixelPos];
+	uint InitialStencil = RWStencilIn[PixelPos];
 
 	float FinalDepth = 0.0;
 	float W = 0.0;
 	
+	RWDepthOut[PixelPos] = RWDepthIn[PixelPos];
+
+	if (ClipToAzimuthalBounds(PixelPos))
+	{
+		return;
+	}
+	
 	// Only dilate onto geometry; if the current pixel's stencil is zero, the dilation can be skipped
 	if (InitialStencil > 0)
 	{
@@ -112,10 +151,10 @@ void DepthDilationCS(uint2 DispatchThreadId : SV_DispatchThreadID)
 		for (int kIndex = 0; kIndex < FILTER_KERNEL_SIZE; ++kIndex)
 		{
 			int2 PixelOffset = int2(kIndex - kRadius, kIndex - kRadius) * SampleOffsetScale * SampleDirection;
-			uint2 SamplePos = AdjustPixelPos(PixelPos + PixelOffset);
+			uint2 SamplePos = AdjustPixelPos(PixelPos, PixelOffset);
 			
-			float Depth = RWDepth[SamplePos];
-			uint Stencil = RWStencil[SamplePos];
+			float Depth = RWDepthIn[SamplePos];
+			uint Stencil = RWStencilIn[SamplePos];
 			
 			// Use a spatial factor here to ensure the depth closest to the edge of the geometry is dilated outwards
 			float SpatialFactor = GET_SCALAR_ARRAY_ELEMENT(InteriorSpatialKernel, kIndex);
@@ -129,7 +168,7 @@ void DepthDilationCS(uint2 DispatchThreadId : SV_DispatchThreadID)
 		
 		if (W > 0.0)
 		{
-			RWDepth[PixelPos] = FinalDepth / W;
+			RWDepthOut[PixelPos] = FinalDepth / W;
 		}
 	}
 }
@@ -142,15 +181,23 @@ void DilationCS(uint2 DispatchThreadId : SV_DispatchThreadID)
 	
 	const int kRadius = (FILTER_KERNEL_SIZE - 1) / 2;
 	
-	float4 InitialColor = RWColor[PixelPos];
-	float InitialDepth = RWDepth[PixelPos];
-	uint InitialStencil = RWStencil[PixelPos];
+	uint InitialStencil = RWStencilIn[PixelPos];
 	
 	float3 FinalColor = float3(0.0, 0.0, 0.0);
 	float FinalDepth = 0.0;
 	uint FinalStencil = 0;
 	
 	float W = 0.0;
+
+	RWColorOut[PixelPos] = RWColorIn[PixelPos];
+	RWDepthOut[PixelPos] = RWDepthIn[PixelPos];
+	RWStencilOut[PixelPos] = RWStencilIn[PixelPos];
+	RWBlurStencilOut[PixelPos] = RWBlurStencilIn[PixelPos];
+
+	if (ClipToAzimuthalBounds(PixelPos))
+	{
+		return;
+	}
 	
 	// Only dilate into empty space; if the current pixel's stencil is non-zero, the dilation can be skipped
 	if (InitialStencil < 1)
@@ -159,11 +206,11 @@ void DilationCS(uint2 DispatchThreadId : SV_DispatchThreadID)
 		for (int kIndex = 0; kIndex < FILTER_KERNEL_SIZE; ++kIndex)
 		{
 			int2 PixelOffset = int2(kIndex - kRadius, kIndex - kRadius) * SampleOffsetScale * SampleDirection;
-			uint2 SamplePos = AdjustPixelPos(PixelPos + PixelOffset);
+			uint2 SamplePos = AdjustPixelPos(PixelPos, PixelOffset);
 		
-			float3 Color = RWColor[SamplePos].rgb;
-			float Depth = RWDepth[SamplePos];
-			uint Stencil = RWStencil[SamplePos];
+			float3 Color = RWColorIn[SamplePos].rgb;
+			float Depth = RWDepthIn[SamplePos];
+			uint Stencil = RWStencilIn[SamplePos];
 			
 			// Use a spatial factor here to ensure the depth closest to the edge of the geometry is dilated outwards
 			float SpatialFactor = GET_SCALAR_ARRAY_ELEMENT(SpatialKernel, kIndex);
@@ -180,11 +227,16 @@ void DilationCS(uint2 DispatchThreadId : SV_DispatchThreadID)
 	
 		if (W > 0.0)
 		{
-			RWColor[PixelPos] = float4(FinalColor / W, 1.0);
-			RWDepth[PixelPos] = FinalDepth / W;
+			RWColorOut[PixelPos] = float4(FinalColor / W, 1.0);
+			RWDepthOut[PixelPos] = FinalDepth / W;
 		}
 
-		RWStencil[PixelPos] = FinalStencil;
+		RWStencilOut[PixelPos] = clamp(FinalStencil, 0, 1);
+
+		if (bCopyToBlurStencil)
+		{
+			RWBlurStencilOut[PixelPos] = clamp(FinalStencil, 0, 1);
+		}
 	}
 }
 
@@ -196,22 +248,28 @@ void BlurCS(uint2 DispatchThreadId : SV_DispatchThreadID)
 	
 	const int kRadius = (FILTER_KERNEL_SIZE - 1) / 2;
 	
-	float3 InitialNormal = NormalFromColorSpace(RWColor[PixelPos].rgb);
-	float InitialDepth = RWDepth[PixelPos];
-	uint Stencil = SceneStencil.Load(int3(PixelPos, 0)) STENCIL_COMPONENT_SWIZZLE;
+	uint Stencil = RWBlurStencilIn[PixelPos];
 	
 	float3 FinalNormal = float3(0.0, 0.0, 0.0);
 	float FinalDepth = 0.0;
 	float W = 0.0;
+
+	// Don't blur any pixels that are beyond the bounds of the azimuthal projection (which are a circle inscribed within the view texture)
+	if (ClipToAzimuthalBounds(PixelPos))
+	{
+		RWColorOut[PixelPos] = RWColorIn[PixelPos];
+		RWDepthOut[PixelPos] = RWDepthIn[PixelPos];
+		return;
+	}
 	
 	UNROLL
 	for (int kIndex = 0; kIndex < FILTER_KERNEL_SIZE; ++kIndex)
 	{
 		int2 PixelOffset = int2(kIndex - kRadius, kIndex - kRadius) * SampleOffsetScale * SampleDirection;
-		uint2 SamplePos = AdjustPixelPos(PixelPos + PixelOffset);
+		uint2 SamplePos = AdjustPixelPos(PixelPos, PixelOffset);
 		
-		float3 Normal = NormalFromColorSpace(RWColor[SamplePos].rgb);
-		float Depth = RWDepth[SamplePos];
+		float3 Normal = NormalFromColorSpace(RWColorIn[SamplePos].rgb);
+		float Depth = RWDepthIn[SamplePos];
 
 		// Get the correct spatial kernel depending on if we are blurring a screen or empty space
 		float Factor = lerp(GET_SCALAR_ARRAY_ELEMENT(SpatialKernel, kIndex), GET_SCALAR_ARRAY_ELEMENT(InteriorSpatialKernel, kIndex), Stencil);
@@ -224,8 +282,8 @@ void BlurCS(uint2 DispatchThreadId : SV_DispatchThreadID)
 	FinalNormal /= W;
 	FinalDepth /= W;
 	
-	RWColor[PixelPos] = float4(NormalToColorSpace(FinalNormal), 1.0);
-	RWDepth[PixelPos] = FinalDepth;
+	RWColorOut[PixelPos] = float4(NormalToColorSpace(FinalNormal), 1.0);
+	RWDepthOut[PixelPos] = FinalDepth;
 }
 
 /** A pixel shader that renders the results of any compute shader passes to an output render target, packing the normal and depth values into the correct RGBA channels */
@@ -234,8 +292,8 @@ void OutputNormalMapPS(noperspective float4 UVAndScreenPos : TEXCOORD0, float4 S
 	const float2 UV = UVAndScreenPos.xy;
 	const uint2 PixelPos = int2(UV * Input_Extent);
 	
-	float3 NormalColor = RWColor[PixelPos].rgb;
-	float Depth = RWDepth[PixelPos];
+	float3 NormalColor = RWColorIn[PixelPos].rgb;
+	float Depth = RWDepthIn[PixelPos];
 	
 	// Normal vector is placed into the RGB channel, while depth is placed into the A channel
 	OutColor = float4(NormalColor, Depth);
diff --git a/Engine/Plugins/Runtime/nDisplay/Source/DisplayCluster/Private/Game/EngineClasses/Scene/DisplayClusterStageGeometryComponent.cpp b/Engine/Plugins/Runtime/nDisplay/Source/DisplayCluster/Private/Game/EngineClasses/Scene/DisplayClusterStageGeometryComponent.cpp
--- a/Engine/Plugins/Runtime/nDisplay/Source/DisplayCluster/Private/Game/EngineClasses/Scene/DisplayClusterStageGeometryComponent.cpp
+++ b/Engine/Plugins/Runtime/nDisplay/Source/DisplayCluster/Private/Game/EngineClasses/Scene/DisplayClusterStageGeometryComponent.cpp
@@ -52,7 +52,8 @@ bool UDisplayClusterStageGeometryComponent::GetStageDistanceAndNormal(const FVec
 {
 	if (bGeometryMapLoaded)
 	{
-		const bool bIsQueryingNorthernHemisphere = InDirection.Z > 0 ? true : false;
+		// Bias ever so slightly to the north to ensure all equator vertices are queried from the same map
+		const bool bIsQueryingNorthernHemisphere = InDirection.Z > -0.001;
 		FDisplayClusterStageGeometryMap& GeometryMap = bIsQueryingNorthernHemisphere ? NorthGeometryMap : SouthGeometryMap;
 
 		auto GetPixel = [&GeometryMap](uint32 InX, uint32 InY)
diff --git a/Engine/Plugins/Runtime/nDisplay/Source/DisplayClusterLightCardEditorShaders/Private/DisplayClusterMeshProjectionRenderer.cpp b/Engine/Plugins/Runtime/nDisplay/Source/DisplayClusterLightCardEditorShaders/Private/DisplayClusterMeshProjectionRenderer.cpp
--- a/Engine/Plugins/Runtime/nDisplay/Source/DisplayClusterLightCardEditorShaders/Private/DisplayClusterMeshProjectionRenderer.cpp
+++ b/Engine/Plugins/Runtime/nDisplay/Source/DisplayClusterLightCardEditorShaders/Private/DisplayClusterMeshProjectionRenderer.cpp
@@ -549,9 +549,10 @@ public:
 		SHADER_PARAMETER_RDG_TEXTURE(Texture2D, SceneColor)
 		SHADER_PARAMETER_RDG_TEXTURE(Texture2D, SceneDepth)
 		SHADER_PARAMETER_RDG_TEXTURE_SRV(Texture2D, SceneStencil)
-		SHADER_PARAMETER_RDG_TEXTURE_UAV(RWTexture2D<float4>, RWColor)
-		SHADER_PARAMETER_RDG_TEXTURE_UAV(RWTexture2D<float>, RWDepth)
-		SHADER_PARAMETER_RDG_TEXTURE_UAV(RWTexture2D<uint>, RWStencil)
+		SHADER_PARAMETER_RDG_TEXTURE_UAV(RWTexture2D<float4>, RWColorOut)
+		SHADER_PARAMETER_RDG_TEXTURE_UAV(RWTexture2D<float>, RWDepthOut)
+		SHADER_PARAMETER_RDG_TEXTURE_UAV(RWTexture2D<uint>, RWStencilOut)
+		SHADER_PARAMETER_RDG_TEXTURE_UAV(RWTexture2D<uint>, RWBlurStencilOut)
 		SHADER_PARAMETER(FMatrix44f, NormalCorrectionMatrix)
 	END_SHADER_PARAMETER_STRUCT()
 
@@ -582,13 +583,19 @@ BEGIN_SHADER_PARAMETER_STRUCT(FMeshProjectionNormalsFilterParameters, )
 	SHADER_PARAMETER_STRUCT_REF(FViewUniformShaderParameters, View)
 	SHADER_PARAMETER_STRUCT(FScreenPassTextureViewportParameters, Input)
 	SHADER_PARAMETER_RDG_TEXTURE_SRV(Texture2D, SceneStencil)
-	SHADER_PARAMETER_RDG_TEXTURE_UAV(RWTexture2D<float4>, RWColor)
-	SHADER_PARAMETER_RDG_TEXTURE_UAV(RWTexture2D<float>, RWDepth)
-	SHADER_PARAMETER_RDG_TEXTURE_UAV(RWTexture2D<uint>, RWStencil)
+	SHADER_PARAMETER_RDG_TEXTURE_UAV(RWTexture2D<float4>, RWColorIn)
+	SHADER_PARAMETER_RDG_TEXTURE_UAV(RWTexture2D<float4>, RWColorOut)
+	SHADER_PARAMETER_RDG_TEXTURE_UAV(RWTexture2D<float>, RWDepthIn)
+	SHADER_PARAMETER_RDG_TEXTURE_UAV(RWTexture2D<float>, RWDepthOut)
+	SHADER_PARAMETER_RDG_TEXTURE_UAV(RWTexture2D<uint>, RWStencilIn)
+	SHADER_PARAMETER_RDG_TEXTURE_UAV(RWTexture2D<uint>, RWStencilOut)
+	SHADER_PARAMETER_RDG_TEXTURE_UAV(RWTexture2D<uint>, RWBlurStencilIn)
+	SHADER_PARAMETER_RDG_TEXTURE_UAV(RWTexture2D<uint>, RWBlurStencilOut)
 	SHADER_PARAMETER_SCALAR_ARRAY(float, SpatialKernel, [FILTER_KERNEL_SIZE])
 	SHADER_PARAMETER_SCALAR_ARRAY(float, InteriorSpatialKernel, [FILTER_KERNEL_SIZE])
 	SHADER_PARAMETER(FVector2f, SampleDirection)
 	SHADER_PARAMETER(FVector2f, SampleOffsetScale)
+	SHADER_PARAMETER(uint32, bCopyToBlurStencil)
 END_SHADER_PARAMETER_STRUCT()
 
 template<EMeshProjectionFilterType FilterType>
@@ -636,8 +643,8 @@ public:
 
 	BEGIN_SHADER_PARAMETER_STRUCT(FParameters, )
 		SHADER_PARAMETER_STRUCT(FScreenPassTextureViewportParameters, Input)
-		SHADER_PARAMETER_RDG_TEXTURE_UAV(RWTexture2D<float4>, RWColor)
-		SHADER_PARAMETER_RDG_TEXTURE_UAV(RWTexture2D<float>, RWDepth)
+		SHADER_PARAMETER_RDG_TEXTURE_UAV(RWTexture2D<float4>, RWColorIn)
+		SHADER_PARAMETER_RDG_TEXTURE_UAV(RWTexture2D<float>, RWDepthIn)
 		RENDER_TARGET_BINDING_SLOTS()
 	END_SHADER_PARAMETER_STRUCT()
 
@@ -1281,14 +1288,41 @@ void GaussianSpatialFilter(float Sigma, TArray<float>& OutKernel)
 	}
 }
 
+/** Structure that holds two textures, one for input and one for output */
+struct FTextureUAVWithIntermediate
+{
+	FRDGTextureRef PrimaryTexture = nullptr;
+	FRDGTextureUAVRef PrimaryUAV = nullptr;
+
+	FRDGTextureRef IntermediateTexture = nullptr;
+	FRDGTextureUAVRef IntermediateUAV = nullptr;
+
+	void CreateTexture(FRDGBuilder& GraphBuilder, const FRDGTextureDesc& TextureDesc, const FString& TextureName)
+	{
+		const FString PrimaryTextureName = *FString::Format(TEXT("{0}.{1}"), { TextureName, TEXT("Primary") });
+		PrimaryTexture = GraphBuilder.CreateTexture(TextureDesc, *PrimaryTextureName);
+		
+		const FString IntermediateTextureName = *FString::Format(TEXT("{0}.{1}"), { TextureName, TEXT("Primary") });
+		IntermediateTexture = GraphBuilder.CreateTexture(TextureDesc, *IntermediateTextureName);
+	}
+	
+	void CreateUAV(FRDGBuilder& GraphBuilder)
+	{
+		PrimaryUAV = GraphBuilder.CreateUAV(PrimaryTexture);
+		IntermediateUAV = GraphBuilder.CreateUAV(IntermediateTexture);
+	}
+};
+
 template<EMeshProjectionFilterType FilterType>
 void AddSeparableFilterPass(FRDGBuilder& GraphBuilder,
 	FGlobalShaderMap* GlobalShaderMap,
 	const FSceneView* View,
 	FRDGTexture* SceneDepth,
-	FRDGTextureUAV* RWColorUAV,
-	FRDGTextureUAV* RWDepthUAV,
-	FRDGTextureUAV* RWStencilUAV)
+	const FTextureUAVWithIntermediate& InColorTextures,
+	const FTextureUAVWithIntermediate& InDepthTextures,
+	const FTextureUAVWithIntermediate& InStencilTextures,
+	const FTextureUAVWithIntermediate& InFinalStencilTextures,
+	bool bCopyToFinalStencil = false)
 {
 	TArray<float> SpatialKernel;
 	TArray<float> InteriorSpatialKernel;
@@ -1296,17 +1330,23 @@ void AddSeparableFilterPass(FRDGBuilder& GraphBuilder,
 	GaussianSpatialFilter(100.0f, SpatialKernel);
 	GaussianSpatialFilter(1.0f, InteriorSpatialKernel);
 
-	const FScreenPassTextureViewport InputViewport(RWColorUAV->Desc.Texture);
+	const FScreenPassTextureViewport InputViewport(InColorTextures.PrimaryUAV->Desc.Texture);
 
 	FMeshProjectionNormalsFilterParameters* HorizontalPassParameters = GraphBuilder.AllocParameters<FMeshProjectionNormalsFilterParameters>();
 	HorizontalPassParameters->View = View->ViewUniformBuffer;
 	HorizontalPassParameters->Input = GetScreenPassTextureViewportParameters(InputViewport);
 	HorizontalPassParameters->SceneStencil = GraphBuilder.CreateSRV(FRDGTextureSRVDesc::CreateWithPixelFormat(SceneDepth, PF_X24_G8));
-	HorizontalPassParameters->RWColor = RWColorUAV;
-	HorizontalPassParameters->RWDepth = RWDepthUAV;
-	HorizontalPassParameters->RWStencil = RWStencilUAV;
+	HorizontalPassParameters->RWColorIn = InColorTextures.PrimaryUAV;
+	HorizontalPassParameters->RWColorOut = InColorTextures.IntermediateUAV;
+	HorizontalPassParameters->RWDepthIn = InDepthTextures.PrimaryUAV;
+	HorizontalPassParameters->RWDepthOut = InDepthTextures.IntermediateUAV;
+	HorizontalPassParameters->RWStencilIn = InStencilTextures.PrimaryUAV;
+	HorizontalPassParameters->RWStencilOut = InStencilTextures.IntermediateUAV;
+	HorizontalPassParameters->RWBlurStencilIn = InFinalStencilTextures.PrimaryUAV;
+	HorizontalPassParameters->RWBlurStencilOut = InFinalStencilTextures.IntermediateUAV;
 	HorizontalPassParameters->SampleDirection = FVector2f(1.0f, 0.0f);
 	HorizontalPassParameters->SampleOffsetScale = FVector2f(1.f);
+	HorizontalPassParameters->bCopyToBlurStencil = bCopyToFinalStencil;
 
 	for (int Index = 0; Index < FILTER_KERNEL_SIZE; ++Index)
 	{
@@ -1332,11 +1372,17 @@ void AddSeparableFilterPass(FRDGBuilder& GraphBuilder,
 	VerticalPassParameters->View = View->ViewUniformBuffer;
 	VerticalPassParameters->Input = GetScreenPassTextureViewportParameters(InputViewport);
 	VerticalPassParameters->SceneStencil = GraphBuilder.CreateSRV(FRDGTextureSRVDesc::CreateWithPixelFormat(SceneDepth, PF_X24_G8));
-	VerticalPassParameters->RWColor = RWColorUAV;
-	VerticalPassParameters->RWDepth = RWDepthUAV;
-	VerticalPassParameters->RWStencil = RWStencilUAV;
+	VerticalPassParameters->RWColorIn = InColorTextures.IntermediateUAV;
+	VerticalPassParameters->RWColorOut = InColorTextures.PrimaryUAV;
+	VerticalPassParameters->RWDepthIn = InDepthTextures.IntermediateUAV;
+	VerticalPassParameters->RWDepthOut = InDepthTextures.PrimaryUAV;
+	VerticalPassParameters->RWStencilIn = InStencilTextures.IntermediateUAV;
+	VerticalPassParameters->RWStencilOut = InStencilTextures.PrimaryUAV;
+	VerticalPassParameters->RWBlurStencilIn = InFinalStencilTextures.IntermediateUAV;
+	VerticalPassParameters->RWBlurStencilOut = InFinalStencilTextures.PrimaryUAV;
 	VerticalPassParameters->SampleDirection = FVector2f(0.0f, 1.0f);
 	VerticalPassParameters->SampleOffsetScale = FVector2f(1.f);
+	VerticalPassParameters->bCopyToBlurStencil = bCopyToFinalStencil;
 
 	for (int Index = 0; Index < FILTER_KERNEL_SIZE; ++Index)
 	{
@@ -1370,16 +1416,23 @@ void FDisplayClusterMeshProjectionRenderer::AddNormalsFilterPass(FRDGBuilder& Gr
 	// * Finally, the RW textures are rendered back to a render target. Here, the normals are placed into the RGB components, while the depth is placed in the A component
 
 	const FRDGTextureDesc RWColorDesc(FRDGTextureDesc::Create2D(SceneColor->Desc.Extent, PF_B8G8R8A8, FClearValueBinding::Black, TexCreate_ShaderResource | TexCreate_UAV));
-	FRDGTextureRef RWColor = GraphBuilder.CreateTexture(RWColorDesc, TEXT("DisplayClusterMeshProjection.NormalsRWTexture"));
-	FRDGTextureUAVRef RWColorUAV = GraphBuilder.CreateUAV(RWColor);
+	FTextureUAVWithIntermediate RWColor;
+	RWColor.CreateTexture(GraphBuilder, RWColorDesc, TEXT("DisplayClusterMeshProjection.NormalsRWTexture"));
+	RWColor.CreateUAV(GraphBuilder);
 
 	const FRDGTextureDesc RWDepthDesc = FRDGTextureDesc::Create2D(SceneDepth->Desc.Extent, PF_G16, FClearValueBinding::White, TexCreate_UAV | TexCreate_ShaderResource);
-	FRDGTextureRef RWDepth = GraphBuilder.CreateTexture(RWDepthDesc, TEXT("DisplayClusterMeshProjection.RWDepth"));
-	FRDGTextureUAVRef RWDepthUAV = GraphBuilder.CreateUAV(RWDepth);
+	FTextureUAVWithIntermediate RWDepth;
+	RWDepth.CreateTexture(GraphBuilder, RWDepthDesc, TEXT("DisplayClusterMeshProjection.RWDepth"));
+	RWDepth.CreateUAV(GraphBuilder);
 
 	const FRDGTextureDesc RWStencilDesc = FRDGTextureDesc::Create2D(SceneDepth->Desc.Extent, PF_R32_UINT, FClearValueBinding::Black, TexCreate_UAV | TexCreate_ShaderResource);
-	FRDGTextureRef RWStencil = GraphBuilder.CreateTexture(RWStencilDesc, TEXT("DisplayClusterMeshProjection.RWStencil"));
-	FRDGTextureUAVRef RWStencilUAV = GraphBuilder.CreateUAV(RWStencil);
+	FTextureUAVWithIntermediate RWStencil;
+	RWStencil.CreateTexture(GraphBuilder, RWStencilDesc, TEXT("DisplayClusterMeshProjection.RWStencil"));
+	RWStencil.CreateUAV(GraphBuilder);
+
+	FTextureUAVWithIntermediate RWBlurStencil;
+	RWBlurStencil.CreateTexture(GraphBuilder, RWStencilDesc, TEXT("DisplayClusterMeshProjection.RWBlurStencil"));
+	RWBlurStencil.CreateUAV(GraphBuilder);
 
 	FGlobalShaderMap* GlobalShaderMap = GetGlobalShaderMap(View->FeatureLevel);
 
@@ -1390,9 +1443,10 @@ void FDisplayClusterMeshProjectionRenderer::AddNormalsFilterPass(FRDGBuilder& Gr
 		PassParameters->SceneColor = SceneColor;
 		PassParameters->SceneDepth = SceneDepth;
 		PassParameters->SceneStencil = GraphBuilder.CreateSRV(FRDGTextureSRVDesc::CreateWithPixelFormat(SceneDepth, PF_X24_G8));
-		PassParameters->RWColor = RWColorUAV;
-		PassParameters->RWDepth = RWDepthUAV;
-		PassParameters->RWStencil = RWStencilUAV;
+		PassParameters->RWColorOut = RWColor.PrimaryUAV;
+		PassParameters->RWDepthOut = RWDepth.PrimaryUAV;
+		PassParameters->RWStencilOut = RWStencil.PrimaryUAV;
+		PassParameters->RWBlurStencilOut = RWBlurStencil.PrimaryUAV;
 		PassParameters->NormalCorrectionMatrix = NormalCorrectionMatrix;
 
 		TShaderMapRef<FMeshProjectionNormalsCreateRWTexturesCS> ComputeShader(GlobalShaderMap);
@@ -1417,28 +1471,31 @@ void FDisplayClusterMeshProjectionRenderer::AddNormalsFilterPass(FRDGBuilder& Gr
 	// Filter the RW buffers, dilating and blurring them appropriately
 	for (int32 Index = 0; Index < DepthDilatePasses; Index++)
 	{
-		AddSeparableFilterPass<EMeshProjectionFilterType::DepthDilate>(GraphBuilder, GlobalShaderMap, View, SceneDepth, RWColorUAV, RWDepthUAV, RWStencilUAV);
+		AddSeparableFilterPass<EMeshProjectionFilterType::DepthDilate>(GraphBuilder, GlobalShaderMap, View, SceneDepth, RWColor, RWDepth, RWStencil, RWBlurStencil);
 	}
 
 	for (int32 Index = 0; Index < DilatePasses; ++Index)
 	{
-		AddSeparableFilterPass<EMeshProjectionFilterType::Dilate>(GraphBuilder, GlobalShaderMap, View, SceneDepth, RWColorUAV, RWDepthUAV, RWStencilUAV);
+		// Copy the second iteration of the dilated stencil to the blur stencil so that the region using the smaller
+		// interior blur kernel extends a little beyond the stage's geometry bounds
+		const bool bCopyToBlurStencil = Index == 1;
+		AddSeparableFilterPass<EMeshProjectionFilterType::Dilate>(GraphBuilder, GlobalShaderMap, View, SceneDepth, RWColor, RWDepth, RWStencil, RWBlurStencil, bCopyToBlurStencil);
 	}
 
 	for (int32 Index = 0; Index < BlurPasses; Index++)
 	{
-		AddSeparableFilterPass<EMeshProjectionFilterType::Blur>(GraphBuilder, GlobalShaderMap, View, SceneDepth, RWColorUAV, RWDepthUAV, RWStencilUAV);
+		AddSeparableFilterPass<EMeshProjectionFilterType::Blur>(GraphBuilder, GlobalShaderMap, View, SceneDepth, RWColor, RWDepth, RWStencil, RWBlurStencil);
 	}
 
 	// Copy the RW buffers to the output render target
 	{
-		const FScreenPassTextureViewport InputViewport(RWColor);
+		const FScreenPassTextureViewport InputViewport(RWColor.PrimaryTexture);
 		const FScreenPassTextureViewport OutputViewport(OutputRenderTargetBinding.GetTexture());
 
 		FMeshProjectionNormalsOutputPS::FParameters* ScreenPassParameters = GraphBuilder.AllocParameters<FMeshProjectionNormalsOutputPS::FParameters>();
 		ScreenPassParameters->Input = GetScreenPassTextureViewportParameters(InputViewport);
-		ScreenPassParameters->RWColor = RWColorUAV;
-		ScreenPassParameters->RWDepth = RWDepthUAV;
+		ScreenPassParameters->RWColorIn = RWColor.PrimaryUAV;
+		ScreenPassParameters->RWDepthIn = RWDepth.PrimaryUAV;
 		ScreenPassParameters->RenderTargets[0] = OutputRenderTargetBinding;
 
 		TShaderMapRef<FScreenPassVS> ScreenPassVS(GlobalShaderMap);
