[nDisplay] Fixed UVLightCard rendering with Substrate

Substrate materials always uses premultiplied alpha for translucent rendering.
Added a custom BlendState for Substrate translucent materials to avoid double premultiplication.

code cleanup.

diff --git a/Engine/Plugins/Runtime/nDisplay/Source/DisplayClusterShaders/Private/Shaders/DisplayClusterShadersPreprocess_UVLightCards.cpp b/Engine/Plugins/Runtime/nDisplay/Source/DisplayClusterShaders/Private/Shaders/DisplayClusterShadersPreprocess_UVLightCards.cpp
--- a/Engine/Plugins/Runtime/nDisplay/Source/DisplayClusterShaders/Private/Shaders/DisplayClusterShadersPreprocess_UVLightCards.cpp
+++ b/Engine/Plugins/Runtime/nDisplay/Source/DisplayClusterShaders/Private/Shaders/DisplayClusterShadersPreprocess_UVLightCards.cpp
@@ -64,14 +64,30 @@ IMPLEMENT_MATERIAL_SHADER_TYPE(, FUVLightCardPS, TEXT("/Plugin/nDisplay/Private/
 
 class FUVLightCardPassProcessor : public FMeshPassProcessor
 {
-public:
-	FUVLightCardPassProcessor(const FSceneView* InView, FMeshPassDrawListContext* InDrawListContext)
-		: FMeshPassProcessor(EMeshPass::Num, nullptr, GMaxRHIFeatureLevel, InView, InDrawListContext)
+	/**
+	* Returns the render scene associated with the given view.
+	* May return nullptr if the view or its family/scene is invalid.
+	*/
+	static const FScene* GetRenderSceneFromView(const FSceneView* InView)
 	{
-		DrawRenderState.SetDepthStencilState(TStaticDepthStencilState<false, CF_Always>::GetRHI());
-		DrawRenderState.SetBlendState(TStaticBlendState<CW_RGBA, BO_Add, BF_SourceAlpha, BF_InverseSourceAlpha, BO_Add, BF_One, BF_InverseSourceAlpha>::GetRHI());
+		if (InView && InView->Family && InView->Family->Scene)
+		{
+			return InView->Family->Scene->GetRenderScene();
+		}
+
+		return nullptr;
 	}
 
+public:
+	FUVLightCardPassProcessor(const FSceneView* InView, FMeshPassDrawListContext* InDrawListContext)
+		: FMeshPassProcessor(
+			EMeshPass::Num, 
+			GetRenderSceneFromView(InView),
+			InView ? InView->GetFeatureLevel() : GMaxRHIFeatureLevel,
+			InView,
+			InDrawListContext)
+	{ }
+
 	virtual void AddMeshBatch(const FMeshBatch& RESTRICT MeshBatch, uint64 BatchElementMask, const FPrimitiveSceneProxy* RESTRICT PrimitiveSceneProxy, int32 StaticMeshId = -1) override final
 	{
 		const FMaterialRenderProxy* FallbackMaterialRenderProxy = nullptr;
@@ -98,9 +114,38 @@ public:
 			return;
 		}
 
-		FMeshDrawingPolicyOverrideSettings OverrideSettings = ComputeMeshOverrideSettings(MeshBatch);
-		ERasterizerFillMode MeshFillMode = ComputeMeshFillMode(Material, OverrideSettings);
-		ERasterizerCullMode MeshCullMode = ComputeMeshCullMode(Material, OverrideSettings);
+		const FMeshDrawingPolicyOverrideSettings OverrideSettings = ComputeMeshOverrideSettings(MeshBatch);
+		const ERasterizerFillMode MeshFillMode = ComputeMeshFillMode(Material, OverrideSettings);
+		const ERasterizerCullMode MeshCullMode = ComputeMeshCullMode(Material, OverrideSettings);
+
+		FMeshPassProcessorRenderState DrawRenderState;
+		{
+			DrawRenderState.SetDepthStencilState(TStaticDepthStencilState<false, CF_Always>::GetRHI());
+
+			const EBlendMode BlendMode = Material.GetBlendMode();
+			const bool bIsTranslucent = IsTranslucentBlendMode(BlendMode);
+
+			// Substrate materials always use premultiplied alpha for translucent rendering.
+			// Apply a custom BlendState so SrcColor is not multiplied by alpha again.
+			if (bIsTranslucent && Material.IsSubstrateMaterial())
+			{
+				// Substrate translucent blending:
+				// Output = SrcColor + DstColor * (1 - SrcAlpha)
+				DrawRenderState.SetBlendState(
+					TStaticBlendState<CW_RGBA,
+					BO_Add, BF_One, BF_InverseSourceAlpha,
+					BO_Add, BF_One, BF_InverseSourceAlpha>::GetRHI());
+			}
+			else
+			{
+				// Standard alpha blending:
+				// Output = SrcColor * SrcAlpha + DstColor * (1 - SrcAlpha)
+				DrawRenderState.SetBlendState(
+					TStaticBlendState<CW_RGBA,
+					BO_Add, BF_SourceAlpha, BF_InverseSourceAlpha,
+					BO_Add, BF_One, BF_InverseSourceAlpha>::GetRHI());
+			}
+		}
 
 		FMeshMaterialShaderElementData ShaderElementData;
 		ShaderElementData.InitializeMeshMaterialData(ViewIfDynamicMeshCommand, PrimitiveSceneProxy, MeshBatch, StaticMeshId, false);
@@ -160,9 +205,6 @@ private:
 		unsigned Mask = -int32(FloatBit >> 31) | 0x80000000;
 		return FloatBit ^ Mask;
 	}
-
-private:
-	FMeshPassProcessorRenderState DrawRenderState;
 };
 
 BEGIN_SHADER_PARAMETER_STRUCT(FUVLightCardPassParameters, )
