[nDisplay]  Logs an error if the backbuffer does not support stereo rendering, but ensures the message is only written once to prevent log spam.

Added new boolFDisplayClusterViewportManagerProxy::bHasLoggedBackbufferStereoUnsupported.

Note: Stereo rendering to the quadbuffer texture expecting a special type of texture ETextureDimension::Texture2DArray that contains 2 layers.
Here is a link to an old JIRA [UE-282114].

Stereo rendering issue should be fixed within the Renderer. Here is  a link https://p4-swarm.epicgames.net/reviews/45890199

diff --git a/Engine/Plugins/Runtime/nDisplay/Source/DisplayCluster/Private/Render/Viewport/DisplayClusterViewportManagerProxy.cpp b/Engine/Plugins/Runtime/nDisplay/Source/DisplayCluster/Private/Render/Viewport/DisplayClusterViewportManagerProxy.cpp
--- a/Engine/Plugins/Runtime/nDisplay/Source/DisplayCluster/Private/Render/Viewport/DisplayClusterViewportManagerProxy.cpp
+++ b/Engine/Plugins/Runtime/nDisplay/Source/DisplayCluster/Private/Render/Viewport/DisplayClusterViewportManagerProxy.cpp
@@ -27,6 +27,8 @@
 #include "Render/Viewport/Postprocess/DisplayClusterViewportPostProcessManager.h"
 #include "Render/Viewport/Containers/DisplayClusterViewportProxyData.h"
 
+#include "Misc/DisplayClusterLog.h"
+
 #include "Render/Viewport/DisplayClusterViewportStrings.h"
 #include "RHIContext.h"
 
@@ -594,6 +596,12 @@ bool FDisplayClusterViewportManagerProxy::ImplResolveFrameTargetToBackBuffer_Ren
 		const bool bStereoTexture = DestTexture->Desc.Dimension == ETextureDimension::Texture2DArray && DestTexture->Desc.ArraySize > 1;
 		if (!bStereoTexture)
 		{
+			if (!bHasLoggedBackbufferStereoUnsupported)
+			{
+				UE_LOG(LogDisplayClusterRender, Error, TEXT("Backbuffer does not support stereo rendering."));
+				bHasLoggedBackbufferStereoUnsupported = 1;
+			}
+
 			return false;
 		}
 	}
diff --git a/Engine/Plugins/Runtime/nDisplay/Source/DisplayCluster/Private/Render/Viewport/DisplayClusterViewportManagerProxy.h b/Engine/Plugins/Runtime/nDisplay/Source/DisplayCluster/Private/Render/Viewport/DisplayClusterViewportManagerProxy.h
--- a/Engine/Plugins/Runtime/nDisplay/Source/DisplayCluster/Private/Render/Viewport/DisplayClusterViewportManagerProxy.h
+++ b/Engine/Plugins/Runtime/nDisplay/Source/DisplayCluster/Private/Render/Viewport/DisplayClusterViewportManagerProxy.h
@@ -169,4 +169,10 @@ private:
 
 	// nDisplay VE object. Used to remove related objects in order.
 	TSharedPtr<class FDisplayClusterViewportManagerViewExtension, ESPMode::ThreadSafe> ViewportManagerViewExtension;
+
+private:
+	// True if the "Backbuffer does not support stereo" error has already been logged.
+	// Used to ensure the error is only written to the log once.
+	mutable uint8 bHasLoggedBackbufferStereoUnsupported : 1 = 0;
+
 };
