[nDisplay] fixed ICVFX related issues

Projection policy 'camera':
- ICVFX CameraComponent settings are now applied only to internal ICVFX viewports.
- Regular viewports continue to use only the system CameraComponent settings.

Reworked media rendering rules in FDisplayClusterViewport::IsRenderEnabledByMedia():
- Invisible In-Camera viewports are rendered only if media outputs are assigned.
- Viewports with media inputs are never rendered.
- Viewports with media outputs are always rendered.
- Other invisible viewports are not rendered.

Fixed logic in FDisplayClusterViewportConfiguration_ICVFXCamera:
- Chromakey (CK) viewports are no longer rendered on headless nodes.

Code cleanup

diff --git a/Engine/Plugins/Runtime/nDisplay/Source/DisplayCluster/Private/Game/EngineClasses/Scene/DisplayClusterICVFXCameraComponent.cpp b/Engine/Plugins/Runtime/nDisplay/Source/DisplayCluster/Private/Game/EngineClasses/Scene/DisplayClusterICVFXCameraComponent.cpp
--- a/Engine/Plugins/Runtime/nDisplay/Source/DisplayCluster/Private/Game/EngineClasses/Scene/DisplayClusterICVFXCameraComponent.cpp
+++ b/Engine/Plugins/Runtime/nDisplay/Source/DisplayCluster/Private/Game/EngineClasses/Scene/DisplayClusterICVFXCameraComponent.cpp
@@ -252,23 +252,18 @@ bool UDisplayClusterICVFXCameraComponent::IsICVFXEnabled() const
 		return false;
 	}
 
-	const ADisplayClusterRootActor* const RootActor = Cast<ADisplayClusterRootActor>(GetOwner());
-
-	if (!RootActor)
+	if (const ADisplayClusterRootActor* const RootActor = Cast<ADisplayClusterRootActor>(GetOwner()))
 	{
-		return false;
-	}
-
-	const UDisplayClusterConfigurationData* const ConfigData = RootActor->GetConfigData();
-
-	if (!ConfigData)
+		if (const IDisplayClusterViewportConfiguration* const Configuration = RootActor->GetViewportConfiguration())
 	{
-		return false;
+			if (const UDisplayClusterConfigurationData* const ConfigurationData = Configuration->GetConfigurationData())
+			{
+				return GetCameraSettingsICVFX().IsICVFXEnabled(*ConfigurationData, Configuration->GetClusterNodeId());
+			}
+		}
 	}
 
-	const FString NodeId = GDisplayCluster->GetPrivateClusterMgr()->GetNodeId();
-
-	return GetCameraSettingsICVFX().IsICVFXEnabled(*ConfigData, NodeId);
+	return false;
 }
 
 const FDisplayClusterConfigurationICVFX_CameraSettings& UDisplayClusterICVFXCameraComponent::GetCameraSettingsICVFX() const
diff --git a/Engine/Plugins/Runtime/nDisplay/Source/DisplayCluster/Private/Render/Viewport/Configuration/DisplayClusterViewportConfiguration.cpp b/Engine/Plugins/Runtime/nDisplay/Source/DisplayCluster/Private/Render/Viewport/Configuration/DisplayClusterViewportConfiguration.cpp
--- a/Engine/Plugins/Runtime/nDisplay/Source/DisplayCluster/Private/Render/Viewport/Configuration/DisplayClusterViewportConfiguration.cpp
+++ b/Engine/Plugins/Runtime/nDisplay/Source/DisplayCluster/Private/Render/Viewport/Configuration/DisplayClusterViewportConfiguration.cpp
@@ -31,6 +31,8 @@
 
 #include "Engine/RendererSettings.h"
 
+#include "Misc/DisplayClusterGlobals.h"
+
 namespace UE::DisplayCluster::Configuration
 {
 	static inline ADisplayClusterRootActor* ImplGetRootActor(const FDisplayClusterActorRef& InConfigurationRootActorRef)
@@ -91,7 +93,7 @@ void FDisplayClusterViewportConfiguration::SetRootActor(const EDisplayClusterRoo
 
 	using namespace UE::DisplayCluster::Configuration;
 
-	// COllect all required RootActor refs changes
+	// Gather all RootActor refs changes
 	EDisplayClusterRootActorType RootActorChanges = (EDisplayClusterRootActorType)0;
 	if (EnumHasAnyFlags(InRootActorType, EDisplayClusterRootActorType::Preview) && ImplIsChangedRootActor(InRootActor, PreviewRootActorRef))
 	{
@@ -106,9 +108,10 @@ void FDisplayClusterViewportConfiguration::SetRootActor(const EDisplayClusterRoo
 		EnumAddFlags(RootActorChanges, EDisplayClusterRootActorType::Configuration);
 	}
 
+	// Handle changes
 	if (RootActorChanges != (EDisplayClusterRootActorType)0)
 	{
-
+		// Handle EndScene for current RootActor
 		if (IsSceneOpened())
 		{
 			if (FDisplayClusterViewportManager* ViewportManager = GetViewportManagerImpl())
@@ -120,7 +123,7 @@ void FDisplayClusterViewportConfiguration::SetRootActor(const EDisplayClusterRoo
 			}
 		}
 
-		// Change RootActor
+		// Apply RootActor changes
 		if (EnumHasAnyFlags(RootActorChanges, EDisplayClusterRootActorType::Preview))
 		{
 			ImplSetRootActor(InRootActor, PreviewRootActorRef);
@@ -185,7 +188,7 @@ void FDisplayClusterViewportConfiguration::OnHandleStartScene()
 void FDisplayClusterViewportConfiguration::OnHandleEndScene()
 {
 	bCurrentSceneActive = false;
-	bCurrentRenderFrameViewportsNeedsToBeUpdated = true;
+	MarkCurrentRenderFrameViewportsOutOfDate();
 
 	// Always reset world ptr at the end of the scene
 	CurrentWorldRef.Reset();
@@ -321,6 +324,34 @@ EDisplayClusterRenderFrameMode FDisplayClusterViewportConfiguration::GetRenderMo
 	return EDisplayClusterRenderFrameMode::PIE_Mono;
 }
 
+bool FDisplayClusterViewportConfiguration::IsMediaAvailable() const
+{
+	// Media is not available in DCRA preview.
+	if (IsPreviewRendering())
+	{
+		return false;
+	}
+
+	// Media is available only when running in cluster mode.
+	const bool bIsClusterMode = GDisplayCluster->GetOperationMode() == EDisplayClusterOperationMode::Cluster;
+
+	return bIsClusterMode;
+}
+
+bool FDisplayClusterViewportConfiguration::IsClusterNodeRenderingOffscreen() const
+{
+	// DCRA Preview: Determined by UDisplayClusterConfigurationClusterNode::bRenderHeadless
+	if (IsPreviewRendering())
+	{
+		return RenderFrameSettings.CurrentNode.bRenderHeadless;
+	}
+
+	// Cluster Runtime: Determined by the "RenderOffscreen" command - line argument
+	static const bool bIsRenderingOffscreen = FParse::Param(FCommandLine::Get(), TEXT("RenderOffscreen"));
+	
+	return bIsRenderingOffscreen;
+}
+
 IDisplayClusterViewportManager* FDisplayClusterViewportConfiguration::GetViewportManager() const
 {
 	return GetViewportManagerImpl();
@@ -341,18 +372,31 @@ bool FDisplayClusterViewportConfiguration::ImplUpdateConfiguration(EDisplayClust
 {
 	check(IsInGameThread());
 
-	FDisplayClusterViewportManager* ViewportManager = GetViewportManagerImpl();
-	if (!ViewportManager || !FDisplayClusterViewportConfigurationHelpers_RenderFrameSettings::UpdateRenderFrameConfiguration(ViewportManager, InRenderMode, *this))
+	if (InViewportNames == nullptr && InClusterNodeId.IsEmpty())
+	{
+		// Requires either a cluster node ID or defined viewport names.
+		return false;
+	}
+
+	if (InRenderMode == EDisplayClusterRenderFrameMode::Unknown)
 	{
+		// Do not initialize for unknown rendering type
 		return false;
 	}
 
-	if (!InWorld)
+	FDisplayClusterViewportManager* ViewportManager = GetViewportManagerImpl();
+	if (!InWorld || !ViewportManager)
 	{
 		// The world is required
 		return false;
 	}
 
+	// Update base settings
+	if(!FDisplayClusterViewportConfigurationHelpers_RenderFrameSettings::UpdateRenderFrameConfiguration(InRenderMode, InClusterNodeId, *this))
+	{
+		return false;
+	}
+
 	// The current world should always be initialized immediately before updating the configuration, since there is a dependency (e.g., visibility lists, etc.)
 	SetCurrentWorldImpl(InWorld);
 
@@ -371,20 +415,10 @@ bool FDisplayClusterViewportConfiguration::ImplUpdateConfiguration(EDisplayClust
 	if (InViewportNames)
 	{
 		ConfigurationViewportManager.UpdateCustomViewports(*InViewportNames);
-
-		// Do not use the cluster node name for this pass type
-		SetClusterNodeId(FString());
 	}
 	else
 	{
-		// Cluster node should be defined
-		check(!InClusterNodeId.IsEmpty());
-
 		ConfigurationViewportManager.UpdateClusterNodeViewports(InClusterNodeId);
-
-		// Use only valid values of cluster node id
-		SetClusterNodeId(InClusterNodeId);
-
 	}
 
 	ConfigurationICVFX.Update();
diff --git a/Engine/Plugins/Runtime/nDisplay/Source/DisplayCluster/Private/Render/Viewport/Configuration/DisplayClusterViewportConfiguration.h b/Engine/Plugins/Runtime/nDisplay/Source/DisplayCluster/Private/Render/Viewport/Configuration/DisplayClusterViewportConfiguration.h
--- a/Engine/Plugins/Runtime/nDisplay/Source/DisplayCluster/Private/Render/Viewport/Configuration/DisplayClusterViewportConfiguration.h
+++ b/Engine/Plugins/Runtime/nDisplay/Source/DisplayCluster/Private/Render/Viewport/Configuration/DisplayClusterViewportConfiguration.h
@@ -77,9 +77,12 @@ public:
 		return bCurrentSceneActive && GetCurrentWorld();
 	}
 
+	virtual bool IsClusterNodeRenderingOffscreen() const override;
+	virtual bool IsMediaAvailable() const override;
+
 	virtual const FString& GetClusterNodeId() const override
 	{
-		return RenderFrameSettings.ClusterNodeId;
+		return RenderFrameSettings.CurrentNode.Id;
 	}
 
 	virtual bool IsPreviewRendering() const override
@@ -121,25 +124,46 @@ public:
 		check(IsInGameThread());
 
 		// Process cluster node name changes
-		SetClusterNodeId(InRenderFrameSettings.ClusterNodeId);
+		SetClusterNodeId(InRenderFrameSettings.CurrentNode.Id);
 
 		RenderFrameSettings = InRenderFrameSettings;
 	}
 
+	// Returns true if the current frame's viewport list is out of date and needs to be refreshed.
+	inline bool IsCurrentRenderFrameViewportsOutOfDate() const
+	{
+		return bCurrentRenderFrameViewportsOutOfDate;
+	}
+
+	// Clears the stale flag, marking the current frame's viewport list as up to date.
+	inline void ClearCurrentRenderFrameViewportsOutOfDate()
+	{
+		bCurrentRenderFrameViewportsOutOfDate = 0;
+	}
+
+	// Raises the stale flag, marking the current frame's viewport list as needing an update.
+	inline void MarkCurrentRenderFrameViewportsOutOfDate()
+	{
+		bCurrentRenderFrameViewportsOutOfDate = 1;
+	}
+
 	/** Set the name of the current cluster node. */
 	void SetClusterNodeId(const FString& InClusterNodeId)
 	{
-		if (RenderFrameSettings.ClusterNodeId != InClusterNodeId)
+		// Tell ViewportManager to update the viewport list:
+		// 1. Cluster node name changed (each node has its own viewport list)
+		// 2. Empty cluster node name (use custom viewport list instead: MRQ, etc)
+		if (RenderFrameSettings.CurrentNode.Id != InClusterNodeId
+		|| InClusterNodeId.IsEmpty())
 		{
-			bCurrentRenderFrameViewportsNeedsToBeUpdated = true;
-			RenderFrameSettings.ClusterNodeId = InClusterNodeId;
+			RenderFrameSettings.CurrentNode.Id = InClusterNodeId;
+			MarkCurrentRenderFrameViewportsOutOfDate();
 		}
 	}
 
 	void OnHandleStartScene();
 	void OnHandleEndScene();
 
-
 	/** Returns the DCRA name or an empty string. */
 	FString GetRootActorName() const;
 
@@ -163,12 +187,12 @@ public:
 	// Reference to configuration proxy object
 	const TSharedRef<FDisplayClusterViewportConfigurationProxy, ESPMode::ThreadSafe> Proxy;
 
+private:
 	// Whether the list of viewports of the current frame needs to be updated
-	bool bCurrentRenderFrameViewportsNeedsToBeUpdated = false;
+	uint8 bCurrentRenderFrameViewportsOutOfDate : 1 = 0;
 
-private:
 	// Is current scene started
-	bool bCurrentSceneActive = false;
+	uint8 bCurrentSceneActive : 1 = 0;
 
 	// Current render frame settings
 	FDisplayClusterRenderFrameSettings RenderFrameSettings;
diff --git a/Engine/Plugins/Runtime/nDisplay/Source/DisplayCluster/Private/Render/Viewport/Configuration/DisplayClusterViewportConfigurationHelpers_ICVFX.cpp b/Engine/Plugins/Runtime/nDisplay/Source/DisplayCluster/Private/Render/Viewport/Configuration/DisplayClusterViewportConfigurationHelpers_ICVFX.cpp
--- a/Engine/Plugins/Runtime/nDisplay/Source/DisplayCluster/Private/Render/Viewport/Configuration/DisplayClusterViewportConfigurationHelpers_ICVFX.cpp
+++ b/Engine/Plugins/Runtime/nDisplay/Source/DisplayCluster/Private/Render/Viewport/Configuration/DisplayClusterViewportConfigurationHelpers_ICVFX.cpp
@@ -277,13 +277,15 @@ TArray<TSharedPtr<FDisplayClusterViewport, ESPMode::ThreadSafe>> FDisplayCluster
 	return OutViewports;
 }
 
-void FDisplayClusterViewportConfigurationHelpers_ICVFX::PreviewReuseInnerFrustumViewportWithinClusterNodes(FDisplayClusterViewport& InCameraViewport, UDisplayClusterICVFXCameraComponent& InCameraComponent, const FDisplayClusterConfigurationICVFX_CameraSettings& InCameraSettings)
+bool FDisplayClusterViewportConfigurationHelpers_ICVFX::PreviewReuseInnerFrustumViewportWithinClusterNodes(FDisplayClusterViewport& InCameraViewport, UDisplayClusterICVFXCameraComponent& InCameraComponent, const FDisplayClusterConfigurationICVFX_CameraSettings& InCameraSettings)
 {
 	if (!InCameraViewport.Configuration->GetRenderFrameSettings().CanReuseViewportWithinClusterNodes())
 	{
-		return;
+		return false;
 	}
 
+	// If the InCamera viewport has already been rendered for another node,
+	// use it as a source to avoid rendering for each node in the cluster.
 	if (FDisplayClusterViewportManager* ViewportManager = InCameraViewport.Configuration->GetViewportManagerImpl())
 	{
 		for (const TSharedPtr<FDisplayClusterViewport, ESPMode::ThreadSafe>& ViewportIt : PreviewGetRenderedInCameraViewports(*InCameraViewport.Configuration, InCameraComponent.GetCameraUniqueId(), false))
@@ -301,19 +303,24 @@ void FDisplayClusterViewportConfigurationHelpers_ICVFX::PreviewReuseInnerFrustum
 
 				// Reuse exist viewport:
 				InOutRenderSettings.SetViewportOverride(ViewportIt->GetId(), ViewportOverrideMode);
-				return;
+
+				return true;
 			}
 		}
 	}
+
+	return false;
 }
 
-void FDisplayClusterViewportConfigurationHelpers_ICVFX::PreviewReuseChromakeyViewportWithinClusterNodes(FDisplayClusterViewport& InChromakeyViewport, const FString& InICVFXCameraId)
+bool FDisplayClusterViewportConfigurationHelpers_ICVFX::PreviewReuseChromakeyViewportWithinClusterNodes(FDisplayClusterViewport& InChromakeyViewport, const FString& InICVFXCameraId)
 {
 	if (!InChromakeyViewport.Configuration->GetRenderFrameSettings().CanReuseViewportWithinClusterNodes())
 	{
-		return;
+		return false;
 	}
 	
+	// If the InCamera chromakey viewport has already been rendered for another node,
+	// use it as a source to avoid rendering for each node in the cluster.
 	if (FDisplayClusterViewportManager* ViewportManager = InChromakeyViewport.Configuration->GetViewportManagerImpl())
 	{
 		for (const TSharedPtr<FDisplayClusterViewport, ESPMode::ThreadSafe>& ViewportIt : PreviewGetRenderedInCameraViewports(*InChromakeyViewport.Configuration, InICVFXCameraId, true))
@@ -331,10 +338,13 @@ void FDisplayClusterViewportConfigurationHelpers_ICVFX::PreviewReuseChromakeyVie
 
 				// Reuse exist viewport from other node
 				InOutRenderSettings.SetViewportOverride(ViewportIt->GetId(), ViewportOverrideMode);
-				return;
+
+				return true;
 			}
 		}
 	}
+
+	return false;
 }
 
 FDisplayClusterViewport* FDisplayClusterViewportConfigurationHelpers_ICVFX::GetOrCreateChromakeyViewport(FDisplayClusterViewportConfiguration& InConfiguration, UDisplayClusterICVFXCameraComponent& InCameraComponent, const FDisplayClusterConfigurationICVFX_CameraSettings& InCameraSettings)
@@ -485,12 +495,9 @@ void FDisplayClusterViewportConfigurationHelpers_ICVFX::UpdateCameraViewportSett
 	// Gain direct access to internal resources of the viewport:
 	FDisplayClusterViewport_RenderSettings& InOutRenderSettings = DstViewport.GetRenderSettingsImpl();
 
-	// incamera textrure used as overlay
+	// The InCamera viewport is used via an overlay (not displayed directly on the backbuffer).
 	InOutRenderSettings.bVisible = false;
 
-	// Update camera viewport projection policy settings
-	UpdateCameraProjectionSettingsICVFX(*DstViewport.Configuration, InCameraComponent, InCameraSettings, DstViewport.GetProjectionPolicy());
-
 	// Update OCIO for Camera Viewport
 	FDisplayClusterViewportConfigurationHelpers_OpenColorIO::UpdateCameraViewportOCIO(DstViewport, InCameraSettings);
 
@@ -811,6 +818,13 @@ void FDisplayClusterViewportConfigurationHelpers_ICVFX::UpdateCameraCustomFrustu
 
 	if (InCameraCustomFrustumConfiguration.bEnable)
 	{
+		const FMarginSet CustomFrustumMargins(
+			InCameraCustomFrustumConfiguration.Left,
+			InCameraCustomFrustumConfiguration.Right,
+			InCameraCustomFrustumConfiguration.Top,
+			InCameraCustomFrustumConfiguration.Bottom
+		);
+
 		OutCustomFrustumSettings.bAdaptResolution = InCameraCustomFrustumConfiguration.bAdaptResolution;
 
 		switch (InCameraCustomFrustumConfiguration.Mode)
@@ -820,20 +834,16 @@ void FDisplayClusterViewportConfigurationHelpers_ICVFX::UpdateCameraCustomFrustu
 			OutCustomFrustumSettings.Unit = EDisplayClusterViewport_FrustumUnit::Percent;
 
 			// Scale 0..100% to 0..1 range
-			OutCustomFrustumSettings.Left = .01f * InCameraCustomFrustumConfiguration.Left;
-			OutCustomFrustumSettings.Right = .01f * InCameraCustomFrustumConfiguration.Right;
-			OutCustomFrustumSettings.Top = .01f * InCameraCustomFrustumConfiguration.Top;
-			OutCustomFrustumSettings.Bottom = .01f * InCameraCustomFrustumConfiguration.Bottom;
+			OutCustomFrustumSettings.AssignMargins(CustomFrustumMargins * .01f);
+
 			break;
 
 		case EDisplayClusterConfigurationViewportCustomFrustumMode::Pixels:
 			OutCustomFrustumSettings.bEnabled = true;
 			OutCustomFrustumSettings.Unit = EDisplayClusterViewport_FrustumUnit::Pixels;
 
-			OutCustomFrustumSettings.Left = InCameraCustomFrustumConfiguration.Left;
-			OutCustomFrustumSettings.Right = InCameraCustomFrustumConfiguration.Right;
-			OutCustomFrustumSettings.Top = InCameraCustomFrustumConfiguration.Top;
-			OutCustomFrustumSettings.Bottom = InCameraCustomFrustumConfiguration.Bottom;
+			OutCustomFrustumSettings.AssignMargins(CustomFrustumMargins);
+
 			break;
 
 		default:
diff --git a/Engine/Plugins/Runtime/nDisplay/Source/DisplayCluster/Private/Render/Viewport/Configuration/DisplayClusterViewportConfigurationHelpers_ICVFX.h b/Engine/Plugins/Runtime/nDisplay/Source/DisplayCluster/Private/Render/Viewport/Configuration/DisplayClusterViewportConfigurationHelpers_ICVFX.h
--- a/Engine/Plugins/Runtime/nDisplay/Source/DisplayCluster/Private/Render/Viewport/Configuration/DisplayClusterViewportConfigurationHelpers_ICVFX.h
+++ b/Engine/Plugins/Runtime/nDisplay/Source/DisplayCluster/Private/Render/Viewport/Configuration/DisplayClusterViewportConfigurationHelpers_ICVFX.h
@@ -48,8 +48,8 @@ public:
 	static TArray<TSharedPtr<FDisplayClusterViewport, ESPMode::ThreadSafe>> GetAllVisibleInnerCameraViewports(FDisplayClusterViewportConfiguration& InConfiguration, bool bGetChromakey = false);
 
 	static TArray<TSharedPtr<FDisplayClusterViewport, ESPMode::ThreadSafe>> PreviewGetRenderedInCameraViewports(FDisplayClusterViewportConfiguration& InConfiguration, const FString& InICVFXCameraId, bool bGetChromakey = false);
-	static void PreviewReuseInnerFrustumViewportWithinClusterNodes(FDisplayClusterViewport& InCameraViewport, UDisplayClusterICVFXCameraComponent& InCameraComponent, const FDisplayClusterConfigurationICVFX_CameraSettings& InCameraSettings);
-	static void PreviewReuseChromakeyViewportWithinClusterNodes(FDisplayClusterViewport& InChromakeyViewport, const FString& InICVFXCameraId);
+	static bool PreviewReuseInnerFrustumViewportWithinClusterNodes(FDisplayClusterViewport& InCameraViewport, UDisplayClusterICVFXCameraComponent& InCameraComponent, const FDisplayClusterConfigurationICVFX_CameraSettings& InCameraSettings);
+	static bool PreviewReuseChromakeyViewportWithinClusterNodes(FDisplayClusterViewport& InChromakeyViewport, const FString& InICVFXCameraId);
 
 	static void UpdateCameraViewportSettings(FDisplayClusterViewport& DstViewport, UDisplayClusterICVFXCameraComponent& InCameraComponent, const FDisplayClusterConfigurationICVFX_CameraSettings& InCameraSettings);
 	static void UpdateChromakeyViewportSettings(FDisplayClusterViewport& DstViewport, FDisplayClusterViewport& InCameraViewport, const FDisplayClusterConfigurationICVFX_CameraSettings& InCameraSettings);
diff --git a/Engine/Plugins/Runtime/nDisplay/Source/DisplayCluster/Private/Render/Viewport/Configuration/DisplayClusterViewportConfigurationHelpers_RenderFrameSettings.cpp b/Engine/Plugins/Runtime/nDisplay/Source/DisplayCluster/Private/Render/Viewport/Configuration/DisplayClusterViewportConfigurationHelpers_RenderFrameSettings.cpp
--- a/Engine/Plugins/Runtime/nDisplay/Source/DisplayCluster/Private/Render/Viewport/Configuration/DisplayClusterViewportConfigurationHelpers_RenderFrameSettings.cpp
+++ b/Engine/Plugins/Runtime/nDisplay/Source/DisplayCluster/Private/Render/Viewport/Configuration/DisplayClusterViewportConfigurationHelpers_RenderFrameSettings.cpp
@@ -5,6 +5,7 @@
 #include "Render/Viewport/RenderFrame/DisplayClusterRenderFrameSettings.h"
 #include "Render/Viewport/DisplayClusterViewportManager.h"
 #include "Render/Viewport/Preview/DisplayClusterViewportManagerPreview.h"
+#include "Render/Viewport/Postprocess/DisplayClusterViewportPostProcessManager.h"
 
 #include "DisplayClusterRootActor.h"
 #include "IDisplayCluster.h"
@@ -112,16 +113,11 @@ EDisplayClusterRenderFrameAlphaChannelCaptureMode FDisplayClusterViewportConfigu
 	return EDisplayClusterRenderFrameAlphaChannelCaptureMode::None;
 }
 
-bool FDisplayClusterViewportConfigurationHelpers_RenderFrameSettings::UpdateRenderFrameConfiguration(FDisplayClusterViewportManager* ViewportManager, EDisplayClusterRenderFrameMode InRenderMode, FDisplayClusterViewportConfiguration& InOutConfiguration)
+bool FDisplayClusterViewportConfigurationHelpers_RenderFrameSettings::UpdateRenderFrameConfiguration(
+	EDisplayClusterRenderFrameMode InRenderMode,
+	const FString& InClusterNodeId,
+	FDisplayClusterViewportConfiguration& InOutConfiguration)
 {
-	if (InRenderMode == EDisplayClusterRenderFrameMode::Unknown)
-	{
-		// Do not initialize for unknown rendering type
-		return false;
-	}
-
-	check(ViewportManager);
-
 	ADisplayClusterRootActor* ConfigurationRootActor = InOutConfiguration.GetRootActor(EDisplayClusterRootActorType::Configuration);
 	if (!(ConfigurationRootActor))
 	{
@@ -129,6 +125,8 @@ bool FDisplayClusterViewportConfigurationHelpers_RenderFrameSettings::UpdateRend
 		return false;
 	}
 
+	// Use settings from the previous frame as base values.
+	// Note: Reset to default values manually if necessary.
 	FDisplayClusterRenderFrameSettings NewRenderFrameSettings = InOutConfiguration.GetRenderFrameSettings();
 
 	// Set current rendering mode
@@ -167,11 +165,6 @@ bool FDisplayClusterViewportConfigurationHelpers_RenderFrameSettings::UpdateRend
 		NewRenderFrameSettings.CrossGPUTransfer.bEnable = GDisplayClusterCrossGPUTransferEnable != 0;
 		NewRenderFrameSettings.CrossGPUTransfer.bLockSteps = GDisplayClusterCrossGPUTransferLockSteps != 0;
 		NewRenderFrameSettings.CrossGPUTransfer.bPullData = GDisplayClusterCrossGPUTransferPullData != 0;
-
-		// Performance: Allow to use parent ViewFamily from parent viewport 
-		// (icvfx has child viewports: lightcard and chromakey with prj_view matrices copied from parent viewport. May sense to use same viewfamily?)
-		// [not implemented yet] Experimental
-		NewRenderFrameSettings.bShouldUseParentViewportRenderFamily = InRenderFrameConfiguration.bShouldUseParentViewportRenderFamily;
 	}
 
 	if (NewRenderFrameSettings.IsPreviewRendering())
@@ -186,7 +179,34 @@ bool FDisplayClusterViewportConfigurationHelpers_RenderFrameSettings::UpdateRend
 		NewRenderFrameSettings.bUseDisplayClusterRenderDevice = (GEngine->StereoRenderingDevice.IsValid() && DisplayClusterAPI.GetOperationMode() == EDisplayClusterOperationMode::Cluster);
 	}
 
-	// Apply new settings:
+	// Settings related to the cluster node.
+	{
+		// Reset to default values
+		NewRenderFrameSettings.CurrentNode = FDisplayClusterRenderFrameSettings::FCurrentClusterNodeSettings();
+
+		// Assign the current node Id.
+		// Note: SetRenderFrameSettings() later also calls SetClusterNodeId() to raise the required flags.
+		NewRenderFrameSettings.CurrentNode.Id = InClusterNodeId;
+
+		// Configure for current node
+		if (!InClusterNodeId.IsEmpty())
+		{
+			const UDisplayClusterConfigurationData* CfgData = ConfigurationRootActor->GetConfigData();
+			if (const UDisplayClusterConfigurationClusterNode* CfgNodeData = CfgData ? CfgData->GetNode(InClusterNodeId) : nullptr)
+			{
+				// Headless node
+				NewRenderFrameSettings.CurrentNode.bRenderHeadless = CfgNodeData->bRenderHeadless;
+
+				// Backbuffer media output
+				NewRenderFrameSettings.CurrentNode.bHasBackbufferMediaOutput =
+					InOutConfiguration.IsMediaAvailable()
+					&& CfgNodeData->MediaSettings.bEnable
+					&& CfgNodeData->MediaSettings.IsMediaOutputAssigned();
+			}
+		}
+	}
+
+	// Applies new render frame settings to the configuration.
 	InOutConfiguration.SetRenderFrameSettings(NewRenderFrameSettings);
 
 	return true;
@@ -194,10 +214,17 @@ bool FDisplayClusterViewportConfigurationHelpers_RenderFrameSettings::UpdateRend
 
 void FDisplayClusterViewportConfigurationHelpers_RenderFrameSettings::PostUpdateRenderFrameConfiguration(FDisplayClusterViewportConfiguration& InOutConfiguration)
 {
-	if (FDisplayClusterViewportManager* ViewportManager = InOutConfiguration.GetViewportManagerImpl())
+	if (const FDisplayClusterViewportManager* ViewportManager = InOutConfiguration.GetViewportManagerImpl())
 	{
 		FDisplayClusterRenderFrameSettings NewRenderFrameSettings = InOutConfiguration.GetRenderFrameSettings();
 
+		// Postprocess flags
+		NewRenderFrameSettings.bUseAdditionalFrameTargetableForPostprocess = ViewportManager->PostProcessManager->ShouldUseAdditionalFrameTargetableResource();
+		NewRenderFrameSettings.bUseFullSizeFrameTargetableForPostprocess = ViewportManager->PostProcessManager->ShouldUseFullSizeFrameTargetableResource();
+
+		// Update immediately, since other ViewportManager functions will use this value
+		InOutConfiguration.SetRenderFrameSettings(NewRenderFrameSettings);
+
 		// Update global flags that control viewport resources
 		NewRenderFrameSettings.bShouldUseOutputFrameTargetableResources = ViewportManager->ShouldUseOutputFrameTargetableResources();
 		NewRenderFrameSettings.bShouldUseAdditionalFrameTargetableResource = ViewportManager->ShouldUseAdditionalFrameTargetableResource();
diff --git a/Engine/Plugins/Runtime/nDisplay/Source/DisplayCluster/Private/Render/Viewport/Configuration/DisplayClusterViewportConfigurationHelpers_RenderFrameSettings.h b/Engine/Plugins/Runtime/nDisplay/Source/DisplayCluster/Private/Render/Viewport/Configuration/DisplayClusterViewportConfigurationHelpers_RenderFrameSettings.h
--- a/Engine/Plugins/Runtime/nDisplay/Source/DisplayCluster/Private/Render/Viewport/Configuration/DisplayClusterViewportConfigurationHelpers_RenderFrameSettings.h
+++ b/Engine/Plugins/Runtime/nDisplay/Source/DisplayCluster/Private/Render/Viewport/Configuration/DisplayClusterViewportConfigurationHelpers_RenderFrameSettings.h
@@ -56,7 +56,7 @@ class FDisplayClusterViewportConfigurationHelpers_RenderFrameSettings
 {
 public:
 	// Update RenderFrameSettings
-	static bool UpdateRenderFrameConfiguration(FDisplayClusterViewportManager* ViewportManager, EDisplayClusterRenderFrameMode InRenderMode, FDisplayClusterViewportConfiguration& InOutConfiguration);
+	static bool UpdateRenderFrameConfiguration(EDisplayClusterRenderFrameMode InRenderMode, const FString& InClusterNodeId, FDisplayClusterViewportConfiguration& InOutConfiguration);
 	static void PostUpdateRenderFrameConfiguration(FDisplayClusterViewportConfiguration& InOutConfiguration);
 
 private:
diff --git a/Engine/Plugins/Runtime/nDisplay/Source/DisplayCluster/Private/Render/Viewport/Configuration/DisplayClusterViewportConfigurationProxy.h b/Engine/Plugins/Runtime/nDisplay/Source/DisplayCluster/Private/Render/Viewport/Configuration/DisplayClusterViewportConfigurationProxy.h
--- a/Engine/Plugins/Runtime/nDisplay/Source/DisplayCluster/Private/Render/Viewport/Configuration/DisplayClusterViewportConfigurationProxy.h
+++ b/Engine/Plugins/Runtime/nDisplay/Source/DisplayCluster/Private/Render/Viewport/Configuration/DisplayClusterViewportConfigurationProxy.h
@@ -40,7 +40,7 @@ public:
 	{
 		check(IsInRenderingThread());
 
-		return RenderFrameSettings.ClusterNodeId;
+		return RenderFrameSettings.CurrentNode.Id;
 	}
 	// ~~END IDisplayClusterViewportConfigurationProxy
 
diff --git a/Engine/Plugins/Runtime/nDisplay/Source/DisplayCluster/Private/Render/Viewport/Configuration/DisplayClusterViewportConfiguration_ICVFX.cpp b/Engine/Plugins/Runtime/nDisplay/Source/DisplayCluster/Private/Render/Viewport/Configuration/DisplayClusterViewportConfiguration_ICVFX.cpp
--- a/Engine/Plugins/Runtime/nDisplay/Source/DisplayCluster/Private/Render/Viewport/Configuration/DisplayClusterViewportConfiguration_ICVFX.cpp
+++ b/Engine/Plugins/Runtime/nDisplay/Source/DisplayCluster/Private/Render/Viewport/Configuration/DisplayClusterViewportConfiguration_ICVFX.cpp
@@ -86,7 +86,7 @@ void FDisplayClusterViewportConfiguration_ICVFX::Update()
 	FDisplayClusterViewportManager* ViewportManager = Configuration.GetViewportManagerImpl();
 
 	// Disable ICVFX if cluster node rendering is not used
-	const FString& ClusterNodeId = RenderFrameSettings.ClusterNodeId;
+	const FString& ClusterNodeId = RenderFrameSettings.CurrentNode.Id;
 	if (!(ViewportManager && StageSettings && !ClusterNodeId.IsEmpty()))
 	{
 		// The nDisplay viewport should now always have the name of the cluster node.
@@ -104,10 +104,7 @@ void FDisplayClusterViewportConfiguration_ICVFX::Update()
 		// We need to be able to render inner views (ICVFX camera view) without any outers. Current pipeline has
 		// a bunch of optimizations that block inner rendering if no outers available. To avoid the limitation
 		// we force inner data initialization if no outers found.
-		static const bool bIsRenderingOffscreen = FParse::Param(FCommandLine::Get(), TEXT("RenderOffscreen"));
-
-		// Cherry-pick ICVFX cameras initialization from IF-block above
-		if (bIsRenderingOffscreen)
+		if(Configuration.IsClusterNodeRenderingOffscreen())
 		{
 			GetAndUpdateStageCameras();
 		}
@@ -334,22 +331,18 @@ void FDisplayClusterViewportConfiguration_ICVFX::GetAndUpdateStageCameras(const
 			}
 		}
 
-		const bool bIsRunningCluster = GDisplayCluster->GetOperationMode() == EDisplayClusterOperationMode::Cluster;
-		const FString ClusterNodeId  = GDisplayCluster->GetPrivateClusterMgr()->GetNodeId();
-
 		// Create camera resources and initialize target ICVFX viewports
 		for (FDisplayClusterViewportConfiguration_ICVFXCamera& CameraIt : StageCameras)
 		{
-			const FDisplayClusterConfigurationMediaICVFX& MediaSettings = CameraIt.GetCameraSettings().RenderSettings.Media;
-
-			const bool bCameraHasMedia = MediaSettings.bEnable && (
-				MediaSettings.HasAnyMediaInputAssigned(ClusterNodeId, MediaSettings.SplitType) ||
-				MediaSettings.HasAnyMediaOutputAssigned(ClusterNodeId, MediaSettings.SplitType));
-
-			if (InTargetViewports == nullptr || !CameraIt.TargetViewports.IsEmpty() || (bIsRunningCluster && bCameraHasMedia))
+			// Ignore cameras that are invisible and have no media assigned.
+			// ("Invisible" = not used on any target (outer) viewport.)
+			if (CameraIt.TargetViewports.IsEmpty() && !CameraIt.HasAnyMediaBindings())
 			{
-				CameraIt.Update();
+				continue;
 			}
+
+			// Update ICVFX camera viewports.
+			CameraIt.Update();
 		}
 	}
 }
diff --git a/Engine/Plugins/Runtime/nDisplay/Source/DisplayCluster/Private/Render/Viewport/Configuration/DisplayClusterViewportConfiguration_ICVFXCamera.cpp b/Engine/Plugins/Runtime/nDisplay/Source/DisplayCluster/Private/Render/Viewport/Configuration/DisplayClusterViewportConfiguration_ICVFXCamera.cpp
--- a/Engine/Plugins/Runtime/nDisplay/Source/DisplayCluster/Private/Render/Viewport/Configuration/DisplayClusterViewportConfiguration_ICVFXCamera.cpp
+++ b/Engine/Plugins/Runtime/nDisplay/Source/DisplayCluster/Private/Render/Viewport/Configuration/DisplayClusterViewportConfiguration_ICVFXCamera.cpp
@@ -31,11 +31,70 @@
 ////////////////////////////////////////////////////////////////////////////////
 // FDisplayClusterViewportConfiguration_ICVFXCamera
 ////////////////////////////////////////////////////////////////////////////////
-bool FDisplayClusterViewportConfiguration_ICVFXCamera::CreateAndSetupInnerCameraViewport()
+bool FDisplayClusterViewportConfiguration_ICVFXCamera::HasAnyMediaBindings() const
 {
-	if (FDisplayClusterViewport* NewCameraViewport = FDisplayClusterViewportConfigurationHelpers_ICVFX::GetOrCreateCameraViewport(Configuration, CameraComponent, GetCameraSettings()))
+	const FDisplayClusterConfigurationMediaICVFX& MediaSettings = GetCameraSettings().RenderSettings.Media;
+	if (Configuration.IsMediaAvailable() && MediaSettings.bEnable)
 	{
-		CameraViewport = NewCameraViewport->AsShared();
+		const bool bMediaInputAssigned = MediaSettings.HasAnyMediaInputAssigned(Configuration.GetClusterNodeId(), MediaSettings.SplitType);
+		const bool MediaOutputAssigned = MediaSettings.HasAnyMediaOutputAssigned(Configuration.GetClusterNodeId(), MediaSettings.SplitType);
+
+		return bMediaInputAssigned || MediaOutputAssigned;
+	}
+
+	return false;
+}
+
+bool FDisplayClusterViewportConfiguration_ICVFXCamera::GetOrCreateAndSetupInnerCameraViewport()
+{
+	if (FDisplayClusterViewport* NewOrExistCameraViewport = FDisplayClusterViewportConfigurationHelpers_ICVFX::GetOrCreateCameraViewport(Configuration, CameraComponent, GetCameraSettings()))
+	{
+		CameraViewport = NewOrExistCameraViewport->AsShared();
+
+		// If this In-Camera viewport is not visible on the outer viewports, set the ICVFX flag.
+		if (TargetViewports.IsEmpty())
+		{
+			// Set flag that indicates that the In-Camera viewport is not visible because no target viewports are assigned.
+			// Note: Flags are reset at the beginning of each new frame by FDisplayClusterViewport_RenderSettingsICVFX::BeginUpdateSettings().
+			EnumAddFlags(NewOrExistCameraViewport->GetRenderSettingsICVFXImpl().Flags, EDisplayClusterViewportICVFXFlags::CameraHasNoTargetViewports);
+		}
+
+		auto DoesAnyTargetViewportSupportInCameraRendering = [&]()
+			{
+				if (HasAnyMediaBindings())
+				{	
+					// Can render the camera if it has any media bindings.
+					// Note: Input media bindings disable rendering in FDisplayClusterViewport::IsRenderEnabledByMedia().
+					return true;
+				}
+
+				if (TargetViewports.IsEmpty())
+				{
+					// Do not render ICVFX cameras that are invisible.
+					return false;
+				}
+
+				for (const FTargetViewport& TargetViewportIt : TargetViewports)
+				{
+					if (TargetViewportIt.ChromakeySource != EDisplayClusterShaderParametersICVFX_ChromakeySource::FrameColor)
+					{
+						// At least one target viewport is configured with chromakey that does not override the In-Camera viewport.
+						return true;
+					}
+				}
+
+				// Performance: If all chromakey sources have a ‘FrameColor’ value for all viewports on the current cluster node,
+				// we can skip rendering the InnerFrustum viewport
+				return false;
+			};
+
+		// Do not render an In-Camera viewport if it is not visible on the outers 
+		// (unless a MediaOutput is assigned).
+		// In this case, the camera viewport is used only as a source of projection data, not textures.
+		const bool bUseCameraForProjectionOnly = !DoesAnyTargetViewportSupportInCameraRendering();
+
+		// Note: Ensure bSkipRendering is set beforehand; otherwise, the InCamera viewport may be split into tiles
+		CameraViewport->GetRenderSettingsImpl().bSkipRendering = bUseCameraForProjectionOnly;
 
 		// overlay rendered only for enabled incamera
 		check(GetCameraSettings().bEnable);
@@ -46,7 +105,10 @@ bool FDisplayClusterViewportConfiguration_ICVFXCamera::CreateAndSetupInnerCamera
 		// Support projection policy update
 		CameraViewport->UpdateConfiguration_ProjectionPolicy();
 
-		// Reuse for EditorPreview
+		// Update camera viewport projection policy settings (Use settings from the UDisplayClusterICVFXCameraComponent)
+		FDisplayClusterViewportConfigurationHelpers_ICVFX::UpdateCameraProjectionSettingsICVFX(*CameraViewport->Configuration, CameraComponent, GetCameraSettings(), CameraViewport->GetProjectionPolicy());
+
+		// Preview performance optimization: the InCamera viewport should only be rendered on one node and then used on all other nodes in the cluster.
 		FDisplayClusterViewportConfigurationHelpers_ICVFX::PreviewReuseInnerFrustumViewportWithinClusterNodes(*CameraViewport, CameraComponent, GetCameraSettings());
 
 		return true;
@@ -72,45 +134,27 @@ bool FDisplayClusterViewportConfiguration_ICVFXCamera::IsCameraProjectionVisible
 
 void FDisplayClusterViewportConfiguration_ICVFXCamera::Update()
 {
-	if (CreateAndSetupInnerCameraViewport())
-	{
-		if (CameraViewport.IsValid())
-		{
-			// Performance: Do not render InnerFrustum if it is not visible.
-			CameraViewport->GetRenderSettingsImpl().bSkipRendering = !EnableInnerFrustumRendering();
-
-			ADisplayClusterRootActor* SceneRootActor = CameraViewport->Configuration->GetRootActor(EDisplayClusterRootActorType::Scene);
-			ADisplayClusterRootActor* ConfigurationRootActor = CameraViewport->Configuration->GetRootActor(EDisplayClusterRootActorType::Configuration);
-			const FDisplayClusterConfigurationICVFX_StageSettings* StageSettings = CameraViewport->Configuration->GetStageSettings();
-
-			if (SceneRootActor && ConfigurationRootActor && StageSettings)
+	// Clears the CK viewports references
+	CameraViewport.Reset();
+	ChromakeyViewport.Reset();
+
+	if (GetOrCreateAndSetupInnerCameraViewport()
+	// CK only applies as composition over the rendered InnerFrustum image.
+	&& !CameraViewport->GetRenderSettings().bSkipRendering
+	// CK composition requires target viewports.
+	&& !TargetViewports.IsEmpty())
 			{
-				const FDisplayClusterConfigurationICVFX_CameraSettings& CameraSettings = GetCameraSettings();
-
-				FDisplayClusterShaderParameters_ICVFX::FCameraSettings ShaderParametersCameraSettings =
-					CameraComponent.GetICVFXCameraShaderParameters(*StageSettings, CameraSettings);
-
-				ShaderParametersCameraSettings.Resource.ViewportId = CameraViewport->GetId();
-
-				// Rendering order for camera overlap
-				const FString InnerFrustumID = CameraComponent.GetCameraUniqueId();
-				const int32 CameraRenderOrder = ConfigurationRootActor->GetInnerFrustumPriority(InnerFrustumID);
-				ShaderParametersCameraSettings.RenderOrder = (CameraRenderOrder < 0) ? CameraSettings.RenderSettings.RenderOrder : CameraRenderOrder;
-
-				// Add this camera data to all visible targets:
-				for (const FTargetViewport& TargetViewportIt : TargetViewports)
+		// Create CK only if this node renders to the backbuffer.
+		if (!Configuration.IsClusterNodeRenderingOffscreen() // visible node
+		|| Configuration.GetRenderFrameSettings().CurrentNode.bHasBackbufferMediaOutput) // or backbuffer shared by media
 				{
-					// Per-viewport CK
-					ShaderParametersCameraSettings.ChromakeySource = TargetViewportIt.ChromakeySource;
-
-					// Gain direct access to internal settings of the viewport:
-					TargetViewportIt.Viewport->GetRenderSettingsICVFXImpl().ICVFX.Cameras.Add(ShaderParametersCameraSettings);
-				}
-			}
+			// Performance: Do not render CK if InnerFrustum is not visible.
+			// Updates the Chromakey viewport and configures it for use in the ICVFX rendering stack.
+			UpdateChromakeyViewport();
 		}
 
-		// Create and assign chromakey for all targets for this camera
-		CreateAndSetupInnerCameraChromakey();
+		// Updates all target viewports.
+		UpdateICVFXSettingsForTargetViewports();
 	}
 }
 
@@ -159,51 +203,42 @@ FString FDisplayClusterViewportConfiguration_ICVFXCamera::GetCameraUniqueId() co
 	return CameraComponent.GetCameraUniqueId();
 }
 
-bool FDisplayClusterViewportConfiguration_ICVFXCamera::EnableChromakeyRendering() const
+bool FDisplayClusterViewportConfiguration_ICVFXCamera::UpdateChromakeyViewport()
 {
-	// Performance: Rnder CK only when it is in use.
-	for (const FTargetViewport& TargetViewportIt : TargetViewports)
-	{
-		if (TargetViewportIt.ChromakeySource == EDisplayClusterShaderParametersICVFX_ChromakeySource::ChromakeyLayers)
+	const FDisplayClusterConfigurationICVFX_StageSettings* StageSettings = Configuration.GetStageSettings();
+	if (!StageSettings)
 		{
-			return true;
-		}
+		return false;
 	}
 
-	return false;
-}
+	// Check the configuration rules to see if chromakey viewport rendering is allowed.
+	const FDisplayClusterConfigurationICVFX_ChromakeyRenderSettings* ChromakeyRenderSettings = GetCameraSettings().Chromakey.GetChromakeyRenderSettings(*StageSettings);
+	const bool bShouldUseChromakeyViewport = ChromakeyRenderSettings && ChromakeyRenderSettings->ShouldUseChromakeyViewport(*StageSettings);
 
-bool FDisplayClusterViewportConfiguration_ICVFXCamera::EnableInnerFrustumRendering() const
-{
-	if (TargetViewports.IsEmpty())
+	// Returns true if any target viewport is configured to use chromakey viewport rendering.
+	auto IsChromakeyUsedByAnyTargetViewport = [&]()
 	{
-		// Headless node
-		return true;
-	}
-
-	// Performance: If all chromakey sources have a ‘FrameColor’ value for all viewports on the current cluster node,
-	// we can skip rendering the InnerFrustum
+			// Performance: Render CK only when it is in use.
 	for (const FTargetViewport& TargetViewportIt : TargetViewports)
 	{
-		if (TargetViewportIt.ChromakeySource != EDisplayClusterShaderParametersICVFX_ChromakeySource::FrameColor)
+				if (TargetViewportIt.ChromakeySource == EDisplayClusterShaderParametersICVFX_ChromakeySource::ChromakeyLayers)
 		{
 			return true;
 		}
 	}
 
 	return false;
-}
-
-bool FDisplayClusterViewportConfiguration_ICVFXCamera::ImplCreateChromakeyViewport()
-{
-	check(CameraViewport.IsValid());
-
-	const FString ICVFXCameraId = CameraComponent.GetCameraUniqueId();
+		};
 
+	// Update CK viewport reference
+	if (CameraViewport.IsValid() // CK requires the InCamera
+		&& bShouldUseChromakeyViewport
+		&& IsChromakeyUsedByAnyTargetViewport())
+	{
 	// Create new chromakey viewport
-	if (FDisplayClusterViewport* NewChromakeyViewport = FDisplayClusterViewportConfigurationHelpers_ICVFX::GetOrCreateChromakeyViewport(Configuration, CameraComponent, GetCameraSettings()))
+		if (FDisplayClusterViewport* NewOrExistChromakeyViewport = FDisplayClusterViewportConfigurationHelpers_ICVFX::GetOrCreateChromakeyViewport(Configuration, CameraComponent, GetCameraSettings()))
 	{
-		ChromakeyViewport = NewChromakeyViewport->AsShared();
+			ChromakeyViewport = NewOrExistChromakeyViewport->AsShared();
 
 		// Update chromakey viewport settings
 		FDisplayClusterViewportConfigurationHelpers_ICVFX::UpdateChromakeyViewportSettings(*ChromakeyViewport, *CameraViewport, GetCameraSettings());
@@ -212,37 +247,55 @@ bool FDisplayClusterViewportConfiguration_ICVFXCamera::ImplCreateChromakeyViewpo
 		ChromakeyViewport->UpdateConfiguration_ProjectionPolicy();
 
 		// reuse for EditorPreview
+			const FString ICVFXCameraId = CameraComponent.GetCameraUniqueId();
 		FDisplayClusterViewportConfigurationHelpers_ICVFX::PreviewReuseChromakeyViewportWithinClusterNodes(*ChromakeyViewport, ICVFXCameraId);
 
 		return true;
 	}
+	}
 
 	return false;
-}
+};
 
-bool FDisplayClusterViewportConfiguration_ICVFXCamera::CreateAndSetupInnerCameraChromakey()
+void FDisplayClusterViewportConfiguration_ICVFXCamera::UpdateICVFXSettingsForTargetViewports()
 {
-	const FDisplayClusterConfigurationICVFX_StageSettings* StageSettings = Configuration.GetStageSettings();
-	if (!StageSettings)
-	{
-		return false;
-	}
+	check(CameraViewport.IsValid());
 
-	// Try create chromakey render on demand
-	if (const FDisplayClusterConfigurationICVFX_ChromakeyRenderSettings* ChromakeyRenderSettings = GetCameraSettings().Chromakey.GetChromakeyRenderSettings(*StageSettings))
-	{
-		if (ChromakeyRenderSettings->ShouldUseChromakeyViewport(*StageSettings) && EnableChromakeyRendering())
+	const FDisplayClusterConfigurationICVFX_StageSettings* StageSettings = Configuration.GetStageSettings();
+	ADisplayClusterRootActor* SceneRootActor = Configuration.GetRootActor(EDisplayClusterRootActorType::Scene);
+	ADisplayClusterRootActor* ConfigurationRootActor = Configuration.GetRootActor(EDisplayClusterRootActorType::Configuration);
+	if (StageSettings == nullptr || SceneRootActor == nullptr || ConfigurationRootActor == nullptr)
 		{
-			ImplCreateChromakeyViewport();
-		}
+		return;
 	}
 
 	// Chromakey viewport name with alpha channel
 	const FString ChromakeyViewportId(ChromakeyViewport.IsValid() ? ChromakeyViewport->GetId() : TEXT(""));
 
-	// Assign this chromakey to all supported targets
+	const FDisplayClusterConfigurationICVFX_CameraSettings& CameraSettings = GetCameraSettings();
+
+	FDisplayClusterShaderParameters_ICVFX::FCameraSettings ShaderParametersCameraSettings =
+		CameraComponent.GetICVFXCameraShaderParameters(*StageSettings, CameraSettings);
+	{
+		// Assign the camera viewport even if it isn’t rendered; math data is still required.
+		ShaderParametersCameraSettings.Resource.ViewportId = CameraViewport->GetId();
+
+		// Rendering order for camera overlap
+		const FString InnerFrustumID = CameraComponent.GetCameraUniqueId();
+		const int32 CameraRenderOrder = ConfigurationRootActor->GetInnerFrustumPriority(InnerFrustumID);
+		ShaderParametersCameraSettings.RenderOrder = (CameraRenderOrder < 0) ? CameraSettings.RenderSettings.RenderOrder : CameraRenderOrder;
+	}
+
+	// Update ICVFX settings for the all target viewports 
 	for (const FTargetViewport& TargetViewportIt : TargetViewports)
 	{
+		// Per-viewport CK
+		ShaderParametersCameraSettings.ChromakeySource = TargetViewportIt.ChromakeySource;
+
+		// Adds shader parameters from this camera to the camera list of the target viewport.
+		TargetViewportIt.Viewport->GetRenderSettingsICVFXImpl().ICVFX.Cameras.Add(ShaderParametersCameraSettings);
+
+		// Assign this chromakey to all supported targets
 		const bool bEnableChromakey = TargetViewportIt.ChromakeySource != EDisplayClusterShaderParametersICVFX_ChromakeySource::Disabled;
 		const bool bEnableChromakeyMarkers = bEnableChromakey && !EnumHasAnyFlags(TargetViewportIt.Viewport->GetRenderSettingsICVFX().Flags, EDisplayClusterViewportICVFXFlags::DisableChromakeyMarkers);
 
@@ -256,6 +309,4 @@ bool FDisplayClusterViewportConfiguration_ICVFXCamera::CreateAndSetupInnerCamera
 		// Setup overlap chromakey with markers
 		FDisplayClusterViewportConfigurationHelpers_ICVFX::UpdateCameraSettings_OverlapChromakey(DstCameraData, *StageSettings, GetCameraSettings(), bEnableChromakeyMarkers);
 	}
-
-	return true;
 }
diff --git a/Engine/Plugins/Runtime/nDisplay/Source/DisplayCluster/Private/Render/Viewport/Configuration/DisplayClusterViewportConfiguration_ICVFXCamera.h b/Engine/Plugins/Runtime/nDisplay/Source/DisplayCluster/Private/Render/Viewport/Configuration/DisplayClusterViewportConfiguration_ICVFXCamera.h
--- a/Engine/Plugins/Runtime/nDisplay/Source/DisplayCluster/Private/Render/Viewport/Configuration/DisplayClusterViewportConfiguration_ICVFXCamera.h
+++ b/Engine/Plugins/Runtime/nDisplay/Source/DisplayCluster/Private/Render/Viewport/Configuration/DisplayClusterViewportConfiguration_ICVFXCamera.h
@@ -34,21 +34,28 @@ public:
 	/** Gets the name of the ICVFX camera. */
 	FString GetCameraUniqueId() const;
 
-	// Performance: Rnder CK only when it is in use.
-	bool EnableChromakeyRendering() const;
-
-	// Performance: If all chromakey sources have a ‘FrameColor’ value for all viewports on the current cluster node, we can skip rendering the InnerFrustum
-	bool EnableInnerFrustumRendering() const;
+	/**
+	* Returns true if this camera has media bindings
+	* (assigned to a media input or output).
+	*
+	* Note: Media is available only in cluster mode.
+	*/
+	bool HasAnyMediaBindings() const;
 
 protected:
-	/** Creating a inner camera viewport for the ICVFX rendering stack. */
-	bool CreateAndSetupInnerCameraViewport();
-
-	/** Creating and setup a inner camera chromakey viewport for the ICVFX rendering stack. */
-	bool CreateAndSetupInnerCameraChromakey();
-
-	/** Creating a camera chromakey viewport for the ICVFX rendering stack. */
-	bool ImplCreateChromakeyViewport();
+	/** Gets or creates an In-Camera viewport and configures it for use in the ICVFX rendering stack. */
+	bool GetOrCreateAndSetupInnerCameraViewport();
+
+	/** Updates the Chromakey viewport and configures it for use in the ICVFX rendering stack. */
+	bool UpdateChromakeyViewport();
+
+	/**
+	 * Updates ICVFX settings for all target viewports.
+	 *
+	 * Applies the necessary ICVFX configuration (e.g., In-Camera, Chromakey, and related
+	 * parameters) so that target viewports are correctly prepared for the ICVFX rendering stack.
+	 */
+	void UpdateICVFXSettingsForTargetViewports();
 
 public:
 	struct FICVFXCameraContext
diff --git a/Engine/Plugins/Runtime/nDisplay/Source/DisplayCluster/Private/Render/Viewport/Configuration/DisplayClusterViewportConfiguration_ProjectionPolicy.cpp b/Engine/Plugins/Runtime/nDisplay/Source/DisplayCluster/Private/Render/Viewport/Configuration/DisplayClusterViewportConfiguration_ProjectionPolicy.cpp
--- a/Engine/Plugins/Runtime/nDisplay/Source/DisplayCluster/Private/Render/Viewport/Configuration/DisplayClusterViewportConfiguration_ProjectionPolicy.cpp
+++ b/Engine/Plugins/Runtime/nDisplay/Source/DisplayCluster/Private/Render/Viewport/Configuration/DisplayClusterViewportConfiguration_ProjectionPolicy.cpp
@@ -74,7 +74,8 @@ void FDisplayClusterViewportConfiguration_ProjectionPolicy::Update()
 					}
 
 					// Projection policies can override postprocess settings
-					// The camera PP override code has been moved from FDisplayClusterViewportConfiguration_ProjectionPolicy::UpdateCameraPolicy_Base() to the projection policy new function:
+					// The camera PP override code has been moved from FDisplayClusterViewportConfiguration_ProjectionPolicy::UpdateCameraPolicyForCameraComponent()
+					// to the projection policy new function
 					ViewportIt->GetProjectionPolicy()->UpdatePostProcessSettings(ViewportIt.Get());
 				}
 			}
@@ -106,22 +107,32 @@ bool FDisplayClusterViewportConfiguration_ProjectionPolicy::UpdateCameraPolicy(F
 		return false;
 	}
 
-	// Get ICVFX camera component
-	if (UpdateCameraPolicy_ICVFX(DstViewport, CameraComponentId))
+
+	// ICVFX InCamera viewport already configured in the FDisplayClusterViewportConfiguration_ICVFXCamera::GetOrCreateAndSetupInnerCameraViewport().
+	if (EnumHasAnyFlags(DstViewport.GetRenderSettingsICVFX().RuntimeFlags, EDisplayClusterViewportRuntimeICVFXFlags::InCamera))
+	{
+		DstViewport.ResetShowLogMsgOnce(EDisplayClusterViewportShowLogMsgOnce::UpdateCameraPolicy);
+
+		return true;
+	}
+
+	// Update projection policy 'camera' for the ICVFXCameraComponent.
+	if (UpdateCameraPolicyForICVFXCameraComponent(DstViewport, CameraComponentId))
 	{
 		DstViewport.ResetShowLogMsgOnce(EDisplayClusterViewportShowLogMsgOnce::UpdateCameraPolicy);
 
 		return true;
 	}
 
-	// Get camera component
-	if (UpdateCameraPolicy_Base(DstViewport, CameraComponentId))
+	// Update projection policy 'camera' for the CameraComponent.
+	if (UpdateCameraPolicyForCameraComponent(DstViewport, CameraComponentId))
 	{
 		DstViewport.ResetShowLogMsgOnce(EDisplayClusterViewportShowLogMsgOnce::UpdateCameraPolicy);
 
 		return true;
 	}
 
+	// Report an error to the log once.
 	if (DstViewport.CanShowLogMsgOnce(EDisplayClusterViewportShowLogMsgOnce::UpdateCameraPolicy_ReferencedCameraNotFound))
 	{
 		UE_LOG(LogDisplayClusterViewport, Error, TEXT("Viewport '%s': referenced camera '%s' not found."), *DstViewport.GetId(), *CameraComponentId);
@@ -130,15 +141,11 @@ bool FDisplayClusterViewportConfiguration_ProjectionPolicy::UpdateCameraPolicy(F
 	return false;
 }
 
-bool FDisplayClusterViewportConfiguration_ProjectionPolicy::UpdateCameraPolicy_Base(FDisplayClusterViewport& DstViewport, const FString& CameraComponentId)
+bool FDisplayClusterViewportConfiguration_ProjectionPolicy::UpdateCameraPolicyForCameraComponent(FDisplayClusterViewport& DstViewport, const FString& CameraComponentId)
 {
-	ADisplayClusterRootActor* ConfigurationRootActor = Configuration.GetRootActor(EDisplayClusterRootActorType::Configuration);
-	ADisplayClusterRootActor*         SceneRootActor = Configuration.GetRootActor(EDisplayClusterRootActorType::Scene);
-	if (ConfigurationRootActor && SceneRootActor)
+	if (ADisplayClusterRootActor* SceneRootActor = Configuration.GetRootActor(EDisplayClusterRootActorType::Scene))
 	{
-		UCameraComponent*         SceneCameraComp = ProjectionPolicyHelpers::FindComponentByName<UCameraComponent>(*SceneRootActor, CameraComponentId);
-		UCameraComponent* ConfigurationCameraComp = ProjectionPolicyHelpers::FindComponentByName<UCameraComponent>(*ConfigurationRootActor, CameraComponentId);
-		if (SceneCameraComp && ConfigurationCameraComp)
+		if (UCameraComponent* SceneCameraComp = ProjectionPolicyHelpers::FindComponentByName<UCameraComponent>(*SceneRootActor, CameraComponentId))
 		{
 			// Now this parameter is always 1 for the standard camera component.
 			const float FOVMultiplier = 1.f;
@@ -157,49 +164,51 @@ bool FDisplayClusterViewportConfiguration_ProjectionPolicy::UpdateCameraPolicy_B
 	return false;
 }
 
-bool FDisplayClusterViewportConfiguration_ProjectionPolicy::UpdateCameraPolicy_ICVFX(FDisplayClusterViewport& DstViewport, const FString& CameraComponentId)
+bool FDisplayClusterViewportConfiguration_ProjectionPolicy::UpdateCameraPolicyForICVFXCameraComponent(FDisplayClusterViewport& DstViewport, const FString& CameraComponentId)
 {
-	FDisplayClusterViewportManager*  ViewportManager = Configuration.GetViewportManagerImpl();
-	ADisplayClusterRootActor* ConfigurationRootActor = Configuration.GetRootActor(EDisplayClusterRootActorType::Configuration);
+	// If a regular viewport window references ICVFXCameraComponent, it should be configured using parameters from UDisplayClusterICVFXCameraComponent.
 	ADisplayClusterRootActor*         SceneRootActor = Configuration.GetRootActor(EDisplayClusterRootActorType::Scene);
-	if (ConfigurationRootActor && SceneRootActor && ViewportManager)
+	UDisplayClusterICVFXCameraComponent* SceneCameraComp = SceneRootActor ?
+		ProjectionPolicyHelpers::FindComponentByName<UDisplayClusterICVFXCameraComponent>(*SceneRootActor, CameraComponentId) : nullptr;
+	if (!SceneCameraComp)
 	{
-		UDisplayClusterICVFXCameraComponent*         SceneCameraComp = ProjectionPolicyHelpers::FindComponentByName<UDisplayClusterICVFXCameraComponent>(*SceneRootActor, CameraComponentId);
-		UDisplayClusterICVFXCameraComponent* ConfigurationCameraComp = ProjectionPolicyHelpers::FindComponentByName<UDisplayClusterICVFXCameraComponent>(*ConfigurationRootActor, CameraComponentId);
-		if (SceneCameraComp && ConfigurationCameraComp)
-		{
-			// If icvfx viewport already rendered, just copy to visible
-			if (FDisplayClusterViewport* ExistCameraViewport = FDisplayClusterViewportConfigurationHelpers_ICVFX::FindCameraViewport(Configuration, *SceneCameraComp))
-			{
-				// Re-use camera RTT, and support OCIO:
-				const EDisplayClusterViewportOverrideMode ViewportOverrideMode =
-					DstViewport.IsOpenColorIOEquals(*ExistCameraViewport)
-					? EDisplayClusterViewportOverrideMode::All
-					: EDisplayClusterViewportOverrideMode::InternalRTT;
-
-				// Gain direct access to internal settings of the viewport:
-				FDisplayClusterViewport_RenderSettings& InOutRenderSettings = DstViewport.GetRenderSettingsImpl();
-
-				InOutRenderSettings.SetViewportOverride(ExistCameraViewport->GetId(), ViewportOverrideMode);
-			}
-			else
-			{
-				// Request render:
-				// Save base settings from camera viewport
-				const FDisplayClusterViewport_RenderSettings           SavedRenderSettings = DstViewport.GetRenderSettings();
-				const FDisplayClusterViewport_RenderSettingsICVFX SavedRenderSettingsICVFX = DstViewport.GetRenderSettingsICVFX();
-
-				// Update render settings for icvfx camera:
-				FDisplayClusterViewportConfigurationHelpers_ICVFX::UpdateCameraViewportSettings(DstViewport, *SceneCameraComp, ConfigurationCameraComp->GetCameraSettingsICVFX());
+		// This is not the projection policy of the ICVFX camera.
+		return false;
+	}
 
-				// Restore some base settings:
-				DstViewport.GetRenderSettingsImpl() = SavedRenderSettings;
-				DstViewport.GetRenderSettingsICVFXImpl() = SavedRenderSettingsICVFX;
-			}
+	ADisplayClusterRootActor* ConfigurationRootActor = Configuration.GetRootActor(EDisplayClusterRootActorType::Configuration);
+	UDisplayClusterICVFXCameraComponent* ConfigurationCameraComp = ConfigurationRootActor
+		? ProjectionPolicyHelpers::FindComponentByName<UDisplayClusterICVFXCameraComponent>(*ConfigurationRootActor, CameraComponentId)
+		: nullptr;
 
-			return true;
-		}
+	if (!ConfigurationCameraComp)
+	{
+		ConfigurationCameraComp = SceneCameraComp;
 	}
 
-	return false;
+	// Note: Use parameters from FDisplayClusterConfigurationICVFX_CameraSettings:
+	// @TODO: Add additional parameters here. (see FDisplayClusterViewportConfigurationHelpers_ICVFX::UpdateCameraViewportSettings)
+
+	// 1. [-] OCIO: use viewport setings	
+	// 2. [-] Viewport Rect: use viewport settings	
+	// 3. Apply postprocess for ICVFX camera
+	SceneCameraComp->ApplyICVFXCameraPostProcessesToViewport(&DstViewport, EDisplayClusterViewportCameraPostProcessFlags::All);
+
+	// 4. [-] Post-render (Override, Blur, Mips): use viewport settings
+	// 5. [-] Custom frustum: None
+	// 6. [-] RenderTargetAdaptRatio: None
+	// 7. [-] BufferRatio (ScreenPercentage): use viewport settings
+	// 8. [-] Upscaler settings: Use viewport settings
+	// 9. [-] Tile rendering: use viewport settings
+
+	// FIN. Update camera viewport projection policy settings (Use settings from the UDisplayClusterICVFXCameraComponent->FDisplayClusterConfigurationICVFX_CameraSettings):
+	// * FOVMultiplier
+	// * FrustumRotation
+	// * FrustumOffset
+	// * OffCenterProjectionOffset
+	return FDisplayClusterViewportConfigurationHelpers_ICVFX::UpdateCameraProjectionSettingsICVFX(
+		*DstViewport.Configuration,
+		*SceneCameraComp,
+		ConfigurationCameraComp->GetCameraSettingsICVFX(),
+		DstViewport.GetProjectionPolicy());
 }
diff --git a/Engine/Plugins/Runtime/nDisplay/Source/DisplayCluster/Private/Render/Viewport/Configuration/DisplayClusterViewportConfiguration_ProjectionPolicy.h b/Engine/Plugins/Runtime/nDisplay/Source/DisplayCluster/Private/Render/Viewport/Configuration/DisplayClusterViewportConfiguration_ProjectionPolicy.h
--- a/Engine/Plugins/Runtime/nDisplay/Source/DisplayCluster/Private/Render/Viewport/Configuration/DisplayClusterViewportConfiguration_ProjectionPolicy.h
+++ b/Engine/Plugins/Runtime/nDisplay/Source/DisplayCluster/Private/Render/Viewport/Configuration/DisplayClusterViewportConfiguration_ProjectionPolicy.h
@@ -27,11 +27,30 @@ private:
 	/** Basic implementation of extended logic for projection policy "camera */
 	bool UpdateCameraPolicy(FDisplayClusterViewport& DstViewport);
 
-	/** Implementation of extended logic for a common camera component */
-	bool UpdateCameraPolicy_Base(FDisplayClusterViewport& DstViewport, const FString& CameraComponentId);
-
-	/** Implementation of extended logic for a ICVFX camera component */
-	bool UpdateCameraPolicy_ICVFX(FDisplayClusterViewport& DstViewport, const FString& CameraComponentId);
+	/**
+	  * Updates the projection policy of a regular (outer) viewport using settings
+	  * from the specified CameraComponent.
+	 *
+	 * @param DstViewport       The outer (regular) viewport to update.
+	 * @param CameraComponentId The identifier of the CameraComponent whose settings will be applied.
+	 *
+	 * @return true if the camera policy was successfully updated, false otherwise.
+	 */
+	bool UpdateCameraPolicyForCameraComponent(FDisplayClusterViewport& DstViewport, const FString& CameraComponentId);
+
+	 /**
+	  * Updates the projection policy of a regular (outer) viewport using settings
+	  * from the specified ICVFX CameraComponent.
+	  *
+	  * This is used when the input viewport has the 'camera' projection policy
+	  * and must be configured to match the ICVFX camera component.
+	  *
+	  * @param DstViewport              The outer (regular) viewport to update.
+	  * @param ICVFXCameraComponentId   The identifier of the ICVFX CameraComponent whose settings will be applied.
+	  *
+	  * @return true if the projection policy was successfully updated, false otherwise.
+	  */
+	bool UpdateCameraPolicyForICVFXCameraComponent(FDisplayClusterViewport& DstViewport, const FString& ICVFXCameraComponentId);
 
 private:
 	FDisplayClusterViewportConfiguration& Configuration;
diff --git a/Engine/Plugins/Runtime/nDisplay/Source/DisplayCluster/Private/Render/Viewport/DisplayClusterViewport.h b/Engine/Plugins/Runtime/nDisplay/Source/DisplayCluster/Private/Render/Viewport/DisplayClusterViewport.h
--- a/Engine/Plugins/Runtime/nDisplay/Source/DisplayCluster/Private/Render/Viewport/DisplayClusterViewport.h
+++ b/Engine/Plugins/Runtime/nDisplay/Source/DisplayCluster/Private/Render/Viewport/DisplayClusterViewport.h
@@ -338,6 +338,13 @@ public:
 	/** Returns true if the rendering of this viewport is allowed by external media objects. */
 	bool IsRenderEnabledByMedia() const;
 
+	/**
+	* Returns true if this viewport should not allocate resources.
+	*
+	* Note: This function also disables rendering, since rendering requires resources.
+	*/
+	bool ShouldDisableViewportResources() const;
+
 	/** Returns true if this viewport used by external media objects. */
 	bool IsUsedByMedia() const;
 
diff --git a/Engine/Plugins/Runtime/nDisplay/Source/DisplayCluster/Private/Render/Viewport/DisplayClusterViewportManager.cpp b/Engine/Plugins/Runtime/nDisplay/Source/DisplayCluster/Private/Render/Viewport/DisplayClusterViewportManager.cpp
--- a/Engine/Plugins/Runtime/nDisplay/Source/DisplayCluster/Private/Render/Viewport/DisplayClusterViewportManager.cpp
+++ b/Engine/Plugins/Runtime/nDisplay/Source/DisplayCluster/Private/Render/Viewport/DisplayClusterViewportManager.cpp
@@ -470,10 +470,9 @@ bool FDisplayClusterViewportManager::ShouldUseFullSizeFrameTargetableResource()
 {
 	check(IsInGameThread());
 
-	if (PostProcessManager->ShouldUseFullSizeFrameTargetableResource())
-	{
-		return true;
-	}
+	// Postprocess requests an full-size FrameTargetable resource via
+	// FDisplayClusterRenderFrameSettings::bUseFullSizeFrameTargetableForPostprocess
+	// Note: FDisplayClusterViewport::ShouldUseFullSizeFrameTargetableResource() relies on this flag.
 
 	for (const TSharedPtr<FDisplayClusterViewport, ESPMode::ThreadSafe>& Viewport : ImplGetCurrentRenderFrameViewports())
 	{
@@ -490,6 +489,17 @@ bool FDisplayClusterViewportManager::ShouldUseOutputFrameTargetableResources() c
 {
 	check(IsInGameThread());
 
+	switch (Configuration->GetRenderFrameSettings().RenderMode)
+	{
+	case EDisplayClusterRenderFrameMode::PreviewInScene:
+	case EDisplayClusterRenderFrameMode::PreviewProxyHitInScene:
+		// Preview uses its own RTTs for each viewport.
+		return false;
+
+	default:
+		break;
+	}
+
 	for (const TSharedPtr<FDisplayClusterViewport, ESPMode::ThreadSafe>& Viewport : ImplGetCurrentRenderFrameViewports())
 	{
 		if (Viewport.IsValid() && Viewport->IsResourceUsed(EDisplayClusterViewportResource::OutputFrameTargetableResources))
@@ -505,11 +515,16 @@ bool FDisplayClusterViewportManager::ShouldUseAdditionalFrameTargetableResource(
 {
 	check(IsInGameThread());
 
-	if (PostProcessManager->ShouldUseAdditionalFrameTargetableResource())
+	// AdditionalFrameTargetable is a subresource of FrameTargetable.
+	if(!ShouldUseOutputFrameTargetableResources())
 	{
-		return true;
+		return false;
 	}
 
+	// Postprocess requests an AdditionalFrameTargetable resource via
+	// FDisplayClusterRenderFrameSettings::bUseAdditionalFrameTargetableForPostprocess
+	// Note: FDisplayClusterViewport::IsResourceUsedImpl() relies on this flag.
+
 	for (const TSharedPtr<FDisplayClusterViewport, ESPMode::ThreadSafe>& Viewport : ImplGetCurrentRenderFrameViewports())
 	{
 		if (Viewport.IsValid() && Viewport->IsResourceUsed(EDisplayClusterViewportResource::AdditionalFrameTargetableResources))
@@ -536,6 +551,9 @@ void FDisplayClusterViewportManager::UpdateCurrentRenderFrameViewports() const
 		{
 			return InViewport1->GetPriority() < InViewport2->GetPriority();
 		});
+
+	// Clears the stale flag, marking the current frame’s viewport list as up to date.
+	Configuration->ClearCurrentRenderFrameViewportsOutOfDate();
 }
 
 void FDisplayClusterViewportManager::RegisterCallbacks()
@@ -991,7 +1009,9 @@ TSharedPtr<FDisplayClusterViewport, ESPMode::ThreadSafe> FDisplayClusterViewport
 
 		// Add viewport on gamethread
 		EntireClusterViewports.Add(NewViewport);
-		Configuration->bCurrentRenderFrameViewportsNeedsToBeUpdated = true;
+		
+		// Raises the stale flag, marking the current frame’s viewport list as needing an update.
+		Configuration->MarkCurrentRenderFrameViewportsOutOfDate();
 
 		if (Configuration->IsSceneOpened())
 		{
@@ -1018,7 +1038,8 @@ void FDisplayClusterViewportManager::ImplDeleteViewport(const TSharedPtr<FDispla
 		EntireClusterViewports.RemoveAt(ViewportIndex);
 	}
 
-	Configuration->bCurrentRenderFrameViewportsNeedsToBeUpdated = true;
+	// Raises the stale flag, marking the current frame’s viewport list as needing an update.
+	Configuration->MarkCurrentRenderFrameViewportsOutOfDate();
 
 	// Reset RTT size after viewport delete
 	ResetSceneRenderTargetSize();
diff --git a/Engine/Plugins/Runtime/nDisplay/Source/DisplayCluster/Private/Render/Viewport/DisplayClusterViewportManager.h b/Engine/Plugins/Runtime/nDisplay/Source/DisplayCluster/Private/Render/Viewport/DisplayClusterViewportManager.h
--- a/Engine/Plugins/Runtime/nDisplay/Source/DisplayCluster/Private/Render/Viewport/DisplayClusterViewportManager.h
+++ b/Engine/Plugins/Runtime/nDisplay/Source/DisplayCluster/Private/Render/Viewport/DisplayClusterViewportManager.h
@@ -153,9 +153,9 @@ public:
 	/** Getting the viewports of the current rendering frame (viewports from the current node or from a special named list). */
 	inline const TArray<TSharedPtr<FDisplayClusterViewport, ESPMode::ThreadSafe>>& ImplGetCurrentRenderFrameViewports() const
 	{
-		if (Configuration->bCurrentRenderFrameViewportsNeedsToBeUpdated)
+		// Updates the current frame’s viewport list if it is out of date and requires refreshing.
+		if (Configuration->IsCurrentRenderFrameViewportsOutOfDate())
 		{
-			Configuration->bCurrentRenderFrameViewportsNeedsToBeUpdated = false;
 			UpdateCurrentRenderFrameViewports();
 		}
 
diff --git a/Engine/Plugins/Runtime/nDisplay/Source/DisplayCluster/Private/Render/Viewport/DisplayClusterViewport_Rules.cpp b/Engine/Plugins/Runtime/nDisplay/Source/DisplayCluster/Private/Render/Viewport/DisplayClusterViewport_Rules.cpp
--- a/Engine/Plugins/Runtime/nDisplay/Source/DisplayCluster/Private/Render/Viewport/DisplayClusterViewport_Rules.cpp
+++ b/Engine/Plugins/Runtime/nDisplay/Source/DisplayCluster/Private/Render/Viewport/DisplayClusterViewport_Rules.cpp
@@ -176,8 +176,51 @@ bool FDisplayClusterViewport::IsExternalRendering() const
 	return false;
 }
 
+bool FDisplayClusterViewport::ShouldDisableViewportResources() const
+{
+	// Viewport used by media.
+	if (RenderSettings.HasAnyMediaStates(EDisplayClusterViewportMediaState::Capture | EDisplayClusterViewportMediaState::Input))
+	{
+		// Always allocate resources for viewports used by media (including offscreen nodes).
+		// Note: this rule takes priority over the ones below.
+		return false;
+	}
+
+	// Any other viewports on the headless node that do not use a backbuffer media output.
+	if (Configuration->IsClusterNodeRenderingOffscreen()
+		&& !Configuration->GetRenderFrameSettings().CurrentNode.bHasBackbufferMediaOutput)
+	{
+		// Resources should not be allocated for these viewports.
+		return true;
+	}
+
+	// An ICVFX camera viewport with no outer viewports assigned is considered invisible.
+	if (EnumHasAnyFlags(RenderSettingsICVFX.Flags, EDisplayClusterViewportICVFXFlags::CameraHasNoTargetViewports)
+		&& EnumHasAnyFlags(RenderSettingsICVFX.RuntimeFlags, EDisplayClusterViewportRuntimeICVFXFlags::InCamera))
+	{
+		// Skip an invisible ICVFX Camera viewport.
+		return true;
+	}
+
+	// Unbound tile viewports
+	if (RenderSettings.TileSettings.GetType() == EDisplayClusterViewportTileType::Tile
+		&& !RenderSettings.TileSettings.HasAnyTileFlags(EDisplayClusterViewportTileFlags::AllowUnboundRender))
+	{
+		// Skip unbound tile viewports.
+		return true;
+	}
+
+	return false;
+}
+
 bool FDisplayClusterViewport::IsRenderEnabled() const
 {
+	// Render requires resources.
+	if (ShouldDisableViewportResources())
+	{
+		return false;
+	}
+
 	if (IsExternalRendering())
 	{
 		// The viewport uses an external rendering that overrides its RTT.
@@ -216,63 +259,32 @@ bool FDisplayClusterViewport::IsRenderEnabled() const
 	return true;
 }
 
-bool FDisplayClusterViewport::IsUsedByMedia() const
-{
-	return RenderSettings.HasAnyMediaStates(EDisplayClusterViewportMediaState::Input | EDisplayClusterViewportMediaState::Capture);
-}
-
-bool FDisplayClusterViewport::ShouldApplyMaxTextureConstraints() const
+bool FDisplayClusterViewport::IsRenderEnabledByMedia() const
 {
-	if (IsUsedByMedia())
+	// Do not render the viewport if it is used as a media input.
+	if (RenderSettings.HasAnyMediaStates(EDisplayClusterViewportMediaState::Input))
 	{
-		// Do not apply this restriction to a viewport that uses media.
+		// This viewport is not rendered but gets the image from the media.
+		// (media input replaces rendering.)
 		return false;
 	}
 
 	return true;
 }
 
-bool FDisplayClusterViewport::IsRenderEnabledByMedia() const
+bool FDisplayClusterViewport::IsUsedByMedia() const
 {
-	// Use this function to disable rendering of this viewport.
-	// All logic for handling media states should reside in this and IsUsedByMedia() functions only.
-
-	static const bool bIsCluster = (GDisplayCluster->GetOperationMode() == EDisplayClusterOperationMode::Cluster);
-	static const bool bIsOffscreen = FParse::Param(FCommandLine::Get(), TEXT("RenderOffscreen"));
-
-	const bool bHasInputAssigned = RenderSettings.HasAnyMediaStates(EDisplayClusterViewportMediaState::Input);
-	const bool bHasOutputAssigned = RenderSettings.HasAnyMediaStates(EDisplayClusterViewportMediaState::Capture);
+	return RenderSettings.HasAnyMediaStates(EDisplayClusterViewportMediaState::Input | EDisplayClusterViewportMediaState::Capture);
+}
 
-	// Do not render the viewport if it is used as a media input.
-	if (bHasInputAssigned)
+bool FDisplayClusterViewport::ShouldApplyMaxTextureConstraints() const
+{
+	if (IsUsedByMedia())
 	{
-		// This viewport is not rendered but gets the image from the media.
-		// (media input replaces rendering.)
+		// Do not apply this restriction to a viewport that uses media.
 		return false;
 	}
 
-	// Applying special media rules for tiles for cluster modes only.
-	if (bIsCluster && !bHasOutputAssigned
-		&& RenderSettings.TileSettings.GetType() == EDisplayClusterViewportTileType::Tile)
-	{
-		// Tiles logic is a little more complicated
-		//
-		// In | Out | Render (offscreen) | Render (onscreen)
-		// ------------------------------------------------------------------
-		//  0 |  0  |        0           | bShouldRenderUnbound
-		//  0 |  1  |        1           | 1
-		//  1 |  0  |        0           | 0
-		//  1 |  1  |        0           | 0
-		// Check if current node was explicitly allowed to render unbound tiles
-		const bool bShouldRenderUnbound = RenderSettings.TileSettings.HasAnyTileFlags(EDisplayClusterViewportTileFlags::AllowUnboundRender);
-		const bool bShouldRender = !bIsOffscreen && bShouldRenderUnbound;
-		if (!bShouldRender)
-		{
-			// Disable tile viewport rendering.
-			return false;
-		}
-	}
-
 	return true;
 }
 
@@ -318,14 +330,6 @@ bool FDisplayClusterViewport::CanSplitIntoTiles() const
 
 bool FDisplayClusterViewport::IsResourceUsedImpl(const EDisplayClusterViewportResource InResourceType) const
 {
-	if (RenderSettings.bSkipRendering)
-	{
-		// When rendering is skipped, the resources isn't used
-		// For example this feature is used when the ICVFX camera uses full-frame chromakey colour,
-		// thereby eliminating the ICVFX camera's viewport rendering resources for optimization.
-		return false;
-	}
-
 	if (RenderSettings.TileSettings.GetType() == EDisplayClusterViewportTileType::Tile)
 	{
 		// The tiled viewport does not use internal resources.
@@ -395,7 +399,7 @@ bool FDisplayClusterViewport::IsResourceUsedImpl(const EDisplayClusterViewportRe
 	switch (InResourceType)
 	{
 	case EDisplayClusterViewportResource::AdditionalTargetableResources:
-		// Supoport projection policy additional resource
+		// Support additional resource demands from the projection policy.
 		if (ProjectionPolicy.IsValid() && ProjectionPolicy->ShouldUseAdditionalTargetableResource(this))
 		{
 			return true;
@@ -409,6 +413,26 @@ bool FDisplayClusterViewport::IsResourceUsedImpl(const EDisplayClusterViewportRe
 		// Only if this viewport is enabled and visible on the final frame.
 		return RenderSettings.bEnable && RenderSettings.bVisible;
 
+	case EDisplayClusterViewportResource::AdditionalFrameTargetableResources:
+		// Only if this viewport is enabled and visible on the final frame.
+		if (RenderSettings.bEnable && RenderSettings.bVisible)
+		{
+			// Viewport Remap requires AdditionalFrameTargetableResources.
+			// [warp] -> AdditionalFrameTargetableResource -> [OutputRemap] -> OutputFrameTargetableResource
+			if (ViewportRemap.IsUsed())
+			{
+				return true;
+			}
+
+			// Postprocess manager may request an AdditionalFrameTargetable resource
+			if (Configuration->GetRenderFrameSettings().bUseAdditionalFrameTargetableForPostprocess)
+			{
+				return true;
+			}
+		}
+		
+		return false;
+
 	default:
 		break;
 	}
@@ -418,30 +442,52 @@ bool FDisplayClusterViewport::IsResourceUsedImpl(const EDisplayClusterViewportRe
 
 bool FDisplayClusterViewport::IsResourceUsed(const EDisplayClusterViewportResource InResourceType) const
 {
-	switch (InResourceType)
+	if (RenderSettings.bSkipRendering)
 	{
-	case EDisplayClusterViewportResource::AdditionalFrameTargetableResources:
-		// OutputFrameTargetableResources must be used for AdditionalFrameTargetableResource
-		if (IsResourceUsedImpl(EDisplayClusterViewportResource::OutputFrameTargetableResources))
+		// When rendering is skipped, the resources isn't used
+		// For example this feature is used when the ICVFX camera uses full-frame chromakey colour,
+		// thereby eliminating the ICVFX camera's viewport rendering resources for optimization.
+		return false;
+	}
+
+	// Resources may be disabled by other rules: media, offscreen nodes, etc.
+	if (ShouldDisableViewportResources())
+	{
+		return false;
+	}
+
+	// Any other viewports on the headless node that do not use a backbuffer media output.
+	if (Configuration->IsClusterNodeRenderingOffscreen()
+	&& !Configuration->GetRenderFrameSettings().CurrentNode.bHasBackbufferMediaOutput)
+	{
+		// Backbuffer not used: skip related resources.
+		switch (InResourceType)
 		{
-			// viewport remap requires AdditionalFrameTargetableResources.
-			// should be reviewed in the future.
-			if (ViewportRemap.IsUsed())
-			{
-				return true;
-			}
+		case EDisplayClusterViewportResource::OutputFrameTargetableResources:
+		case EDisplayClusterViewportResource::AdditionalFrameTargetableResources:
+			// Skip FrameTargetable RTTs (final composite copied to backbuffer).
+			return false;
+
+		default:
+			break;
 		}
-		return false;
+	}
 
-	default:
-		break;
+	if (!IsResourceUsedImpl(InResourceType))
+	{
+		return false;
 	}
 
-	return IsResourceUsedImpl(InResourceType);
+	return true;
 };
 
 bool FDisplayClusterViewport::ShouldUseFullSizeFrameTargetableResource() const
 {
+	if (Configuration->GetRenderFrameSettings().bUseFullSizeFrameTargetableForPostprocess)
+	{
+		return true;
+	}
+
 	if (ViewportRemap.IsUsed())
 	{
 		return true;
diff --git a/Engine/Plugins/Runtime/nDisplay/Source/DisplayCluster/Private/Render/Viewport/RenderFrame/DisplayClusterRenderFrameManager.cpp b/Engine/Plugins/Runtime/nDisplay/Source/DisplayCluster/Private/Render/Viewport/RenderFrame/DisplayClusterRenderFrameManager.cpp
--- a/Engine/Plugins/Runtime/nDisplay/Source/DisplayCluster/Private/Render/Viewport/RenderFrame/DisplayClusterRenderFrameManager.cpp
+++ b/Engine/Plugins/Runtime/nDisplay/Source/DisplayCluster/Private/Render/Viewport/RenderFrame/DisplayClusterRenderFrameManager.cpp
@@ -5,6 +5,7 @@
 #include "Render/Viewport/RenderFrame/DisplayClusterRenderFrameSettings.h"
 #include "Render/Viewport/Configuration/DisplayClusterViewportConfiguration.h"
 #include "Render/Viewport/DisplayClusterViewport.h"
+#include "Render/Viewport/DisplayClusterViewportManager.h"
 #include "Render/Upscaler/DisplayClusterUpscaler.h"
 
 #include "Misc/Parse.h"
@@ -133,12 +134,10 @@ bool FDisplayClusterRenderFrameManager::BuildSimpleFrame(FViewport* InViewport,
 
 bool FDisplayClusterRenderFrameManager::FindFrameTargetRect(FViewport* InViewport, const TArray<TSharedPtr<FDisplayClusterViewport, ESPMode::ThreadSafe>>& InOutViewports, FIntRect& OutFrameTargetRect) const
 {
-	const FDisplayClusterRenderFrameSettings& InRenderFrameSettings = Configuration->GetRenderFrameSettings();
-
 	// Calculate Backbuffer frame
 	bool bIsUsed = false;
 
-	if (InViewport && InRenderFrameSettings.bShouldUseFullSizeFrameTargetableResource)
+	if (InViewport && Configuration->GetRenderFrameSettings().bShouldUseFullSizeFrameTargetableResource)
 	{
 		// Use full-size frame RTT
 		OutFrameTargetRect = FIntRect(FIntPoint(0, 0), InViewport->GetSizeXY());
@@ -171,9 +170,7 @@ bool FDisplayClusterRenderFrameManager::FindFrameTargetRect(FViewport* InViewpor
 		// We need to be able to render inner views (ICVFX camera view) without any outers. Current pipeline has
 		// a bunch of optimizations that block inner rendering if no outers available. To avoid the limitation
 		// we force inner data initialization if no outers found.
-		static const bool bIsRenderingOffscreen = FParse::Param(FCommandLine::Get(), TEXT("RenderOffscreen"));
-
-		if (bIsRenderingOffscreen)
+		if (Configuration->IsClusterNodeRenderingOffscreen())
 		{
 			bIsUsed = true;
 		}
diff --git a/Engine/Plugins/Runtime/nDisplay/Source/DisplayCluster/Private/Render/Viewport/RenderFrame/DisplayClusterRenderFrameSettings.cpp b/Engine/Plugins/Runtime/nDisplay/Source/DisplayCluster/Private/Render/Viewport/RenderFrame/DisplayClusterRenderFrameSettings.cpp
--- a/Engine/Plugins/Runtime/nDisplay/Source/DisplayCluster/Private/Render/Viewport/RenderFrame/DisplayClusterRenderFrameSettings.cpp
+++ b/Engine/Plugins/Runtime/nDisplay/Source/DisplayCluster/Private/Render/Viewport/RenderFrame/DisplayClusterRenderFrameSettings.cpp
@@ -215,22 +215,6 @@ int32 FDisplayClusterRenderFrameSettings::GetViewPerViewportAmount() const
 	return 1;
 }
 
-bool FDisplayClusterRenderFrameSettings::ShouldUseOutputFrameTargetableResources() const
-{
-	switch (RenderMode)
-	{
-	case EDisplayClusterRenderFrameMode::PreviewInScene:
-	case EDisplayClusterRenderFrameMode::PreviewProxyHitInScene:
-		// Preview uses its own RTTs for each viewport.
-		return false;
-
-	default:
-		break;
-	}
-
-	return bShouldUseOutputFrameTargetableResources;
-}
-
 bool FDisplayClusterRenderFrameSettings::ShouldUseStereoRenderingOnMonoscopicDisplay() const
 {
 	switch (RenderMode)
diff --git a/Engine/Plugins/Runtime/nDisplay/Source/DisplayCluster/Private/Render/Viewport/RenderFrame/DisplayClusterRenderFrameSettings.h b/Engine/Plugins/Runtime/nDisplay/Source/DisplayCluster/Private/Render/Viewport/RenderFrame/DisplayClusterRenderFrameSettings.h
--- a/Engine/Plugins/Runtime/nDisplay/Source/DisplayCluster/Private/Render/Viewport/RenderFrame/DisplayClusterRenderFrameSettings.h
+++ b/Engine/Plugins/Runtime/nDisplay/Source/DisplayCluster/Private/Render/Viewport/RenderFrame/DisplayClusterRenderFrameSettings.h
@@ -14,38 +14,67 @@ struct FDisplayClusterRenderFrameSettings
 	// customize mono\stereo render modes
 	EDisplayClusterRenderFrameMode RenderMode = EDisplayClusterRenderFrameMode::Unknown;
 
+	// Alpha channel capture mode for viewports (Lightcard, chromakey)
+	EDisplayClusterRenderFrameAlphaChannelCaptureMode AlphaChannelCaptureMode = EDisplayClusterRenderFrameAlphaChannelCaptureMode::None;
+
 	// nDisplay has its own implementation of cross-GPU transfer.
 	struct FCrossGPUTransfer
 	{
 		// Enable cross-GPU transfers using nDisplay
 		// That replaces the default cross-GPU transfers using UE Core for the nDisplay viewports viewfamilies.
-		bool bEnable = false;
+		uint8 bEnable : 1 = 0;
 
 		// The bLockSteps parameter is simply passed to the FTransferResourceParams structure.
 		// Whether the GPUs must handshake before and after the transfer. Required if the texture rect is being written to in several render passes.
 		// Otherwise, minimal synchronization will be used.
-		bool bLockSteps = false;
+		uint8 bLockSteps : 1 = 0;
 
 		// The bPullData parameter is simply passed to the FTransferResourceParams structure.
 		// Whether the data is read by the dest GPU, or written by the src GPU (not allowed if the texture is a backbuffer)
-		bool bPullData = true;
+		uint8 bPullData : 1 = 1;
 
 	} CrossGPUTransfer;
 
-	// Alpha channel capture mode for viewports (Lightcard, chromakey)
-	EDisplayClusterRenderFrameAlphaChannelCaptureMode AlphaChannelCaptureMode;
+	// This rendering frame uses OutputFrameTargetable resources.
+	uint8 bShouldUseOutputFrameTargetableResources : 1 = 0;
+
+	// This rendering frame uses AdditionalFrameTargetable resources.
+	uint8 bShouldUseAdditionalFrameTargetableResource : 1 = 0;
+
+	// This rendering frame uses full size for the FrameTargetable resources.
+	uint8 bShouldUseFullSizeFrameTargetableResource : 1 = 0;
+
+	// Use an AdditionalFrameTargetable for postprocess in this frame.
+	uint8 bUseAdditionalFrameTargetableForPostprocess : 1 = 0;
 
+	// Use an full-size FrameTargetable resources for postprocess in this frame.
+	uint8 bUseFullSizeFrameTargetableForPostprocess : 1 = 0;
+
+	// Allow warpblend render
+	uint8 bAllowWarpBlend : 1 = 1;
+
+	// Use DC render device for rendering
+	uint8 bUseDisplayClusterRenderDevice : 1 = 1;
+
+	// Performance: Stereoscopic rendering uses a single viewfamily(RTT) for both eyes
+	uint8 bEnableStereoscopicRenderingOptimization : 1 = 0;
 
-	// Some frame postprocess require additional render targetable resources
-	bool bShouldUseAdditionalFrameTargetableResource = false;
+	/** Settings for the current cluster node. */
+	struct FCurrentClusterNodeSettings
+	{
+		// True if this cluster node is configured to render headless.
+		// Mirrors UDisplayClusterConfigurationClusterNode::bRenderHeadless.
+		uint8 bRenderHeadless : 1 = 0;
 
-	// Postprocess can use full size backbuffer. This disables Frame RTT size optimization
-	bool bShouldUseFullSizeFrameTargetableResource = false;
+		// True if this cluster node has a backbuffer media output assigned.
+		uint8 bHasBackbufferMediaOutput : 1 = 0;
 
-	// Create output resources only for visible viewports. This will save GPU memory space.
-	bool bShouldUseOutputFrameTargetableResources = false;
+		// Cluster node name for render
+		FString Id;
 
+	} CurrentNode;
 
+public:
 	// Multiply all viewports RTT size's for whole cluster by this value
 	float ClusterRenderTargetRatioMult = 1.f;
 
@@ -55,7 +84,6 @@ struct FDisplayClusterRenderFrameSettings
 	// Multiply outer viewports RTT size's for whole cluster by this value
 	float ClusterICVFXOuterViewportRenderTargetRatioMult = 1.f;
 
-
 	// Multiply all buffer ratios for whole cluster by this value
 	float ClusterBufferRatioMult = 1.f;
 
@@ -65,9 +93,6 @@ struct FDisplayClusterRenderFrameSettings
 	// Multiply outer viewports buffer ratios for whole cluster by this value
 	float ClusterICVFXOuterViewportBufferRatioMult = 1.f;
 
-	// Allow warpblend render
-	bool bAllowWarpBlend = true;
-
 	// Settings for preview rendering
 	FDisplayClusterViewport_PreviewSettings PreviewSettings;
 
@@ -75,20 +100,6 @@ struct FDisplayClusterRenderFrameSettings
 	// Specifies the mGPU index range for rendering the DCRA preview.
 	TOptional<FIntPoint> PreviewMultiGPURendering;
 
-	// Use DC render device for rendering
-	bool bUseDisplayClusterRenderDevice = true;
-
-	// Performance: Stereoscopic rendering uses a single viewfamily(RTT) for both eyes
-	bool bEnableStereoscopicRenderingOptimization = false;
-
-	// Performance: Allow to use parent ViewFamily from parent viewport 
-	// (icvfx has child viewports: lightcard and chromakey with prj_view matrices copied from parent viewport. May sense to use same viewfamily?)
-	// [not implemented yet] Experimental
-	bool bShouldUseParentViewportRenderFamily = false;
-
-	// Cluster node name for render
-	FString ClusterNodeId;
-
 public:
 	/** Current frame is preview. */
 	bool IsPreviewRendering() const;
@@ -119,9 +130,6 @@ public:
 	/** Should use Holdout. */
 	bool ShouldUseHoldout() const;
 
-	/** true, if output frame RTTs is used. */
-	bool ShouldUseOutputFrameTargetableResources() const;
-
 	/** Is stereo rendering on monoscopic display (sbs, tb) . */
 	bool ShouldUseStereoRenderingOnMonoscopicDisplay() const;
 
diff --git a/Engine/Plugins/Runtime/nDisplay/Source/DisplayCluster/Private/Render/Viewport/RenderTarget/DisplayClusterRenderTargetManager.cpp b/Engine/Plugins/Runtime/nDisplay/Source/DisplayCluster/Private/Render/Viewport/RenderTarget/DisplayClusterRenderTargetManager.cpp
--- a/Engine/Plugins/Runtime/nDisplay/Source/DisplayCluster/Private/Render/Viewport/RenderTarget/DisplayClusterRenderTargetManager.cpp
+++ b/Engine/Plugins/Runtime/nDisplay/Source/DisplayCluster/Private/Render/Viewport/RenderTarget/DisplayClusterRenderTargetManager.cpp
@@ -252,7 +252,7 @@ bool FDisplayClusterRenderTargetManager::AllocateFrameTargets(const FIntPoint& I
 {
 	// ICVFX internal viewports are invisible and do not use output frame resources.
 	// In case of off-screen rendering this texture is not used.
-	if (!Configuration->GetRenderFrameSettings().ShouldUseOutputFrameTargetableResources())
+	if (!Configuration->GetRenderFrameSettings().bShouldUseOutputFrameTargetableResources)
 	{
 		// Skip creating output resources if there are no visible viewports. This will save GPU memory.
 		return true;
diff --git a/Engine/Plugins/Runtime/nDisplay/Source/DisplayCluster/Private/Render/Viewport/Resource/DisplayClusterViewportResourceSettings.cpp b/Engine/Plugins/Runtime/nDisplay/Source/DisplayCluster/Private/Render/Viewport/Resource/DisplayClusterViewportResourceSettings.cpp
--- a/Engine/Plugins/Runtime/nDisplay/Source/DisplayCluster/Private/Render/Viewport/Resource/DisplayClusterViewportResourceSettings.cpp
+++ b/Engine/Plugins/Runtime/nDisplay/Source/DisplayCluster/Private/Render/Viewport/Resource/DisplayClusterViewportResourceSettings.cpp
@@ -12,7 +12,7 @@
 // FDisplayClusterViewportResource
 /////////////////////////////////////////////////////////////////////
 FDisplayClusterViewportResourceSettings::FDisplayClusterViewportResourceSettings(const FDisplayClusterRenderFrameSettings& InRenderFrameSettings, FViewport* InViewport)
-	: ClusterNodeId(InRenderFrameSettings.ClusterNodeId)
+	: ClusterNodeId(InRenderFrameSettings.CurrentNode.Id)
 {
 	if (InViewport && InViewport->GetRenderTargetTexture())
 	{
diff --git a/Engine/Plugins/Runtime/nDisplay/Source/DisplayCluster/Public/Render/Viewport/Containers/DisplayClusterViewport_EnumsICVFX.h b/Engine/Plugins/Runtime/nDisplay/Source/DisplayCluster/Public/Render/Viewport/Containers/DisplayClusterViewport_EnumsICVFX.h
--- a/Engine/Plugins/Runtime/nDisplay/Source/DisplayCluster/Public/Render/Viewport/Containers/DisplayClusterViewport_EnumsICVFX.h
+++ b/Engine/Plugins/Runtime/nDisplay/Source/DisplayCluster/Public/Render/Viewport/Containers/DisplayClusterViewport_EnumsICVFX.h
@@ -40,6 +40,9 @@ enum class EDisplayClusterViewportICVFXFlags : uint16
 
 	/** The order in which the ICVFX cameras are composited over is reversed. Useful for time-multiplexed displays. */
 	ReverseCameraPriority = 1 << 8,
+
+	/** Indicates that the In-Camera viewport is not visible because no target viewports are assigned. */
+	CameraHasNoTargetViewports = 1 << 9,
 };
 ENUM_CLASS_FLAGS(EDisplayClusterViewportICVFXFlags);
 
diff --git a/Engine/Plugins/Runtime/nDisplay/Source/DisplayCluster/Public/Render/Viewport/Containers/DisplayClusterViewport_PreviewSettings.h b/Engine/Plugins/Runtime/nDisplay/Source/DisplayCluster/Public/Render/Viewport/Containers/DisplayClusterViewport_PreviewSettings.h
--- a/Engine/Plugins/Runtime/nDisplay/Source/DisplayCluster/Public/Render/Viewport/Containers/DisplayClusterViewport_PreviewSettings.h
+++ b/Engine/Plugins/Runtime/nDisplay/Source/DisplayCluster/Public/Render/Viewport/Containers/DisplayClusterViewport_PreviewSettings.h
@@ -16,45 +16,47 @@ struct FDisplayClusterViewport_PreviewSettings
 	EDisplayClusterRenderFrameMode EntireClusterPreviewRenderMode = EDisplayClusterRenderFrameMode::PreviewInScene;
 
 	// Render the scene and display it as a preview on the nDisplay root actor in the editor.  This will impact editor performance.
-	bool bPreviewEnable = false;
+	uint8 bPreviewEnable : 1 = 0;
 
 	// Render this DCRA in game for Standalone/Package builds.
-	bool bPreviewInGameEnable = false;
-	bool bPreviewInGameRenderFrustum = false;
+	uint8 bPreviewInGameEnable : 1 = 0;
+	uint8 bPreviewInGameRenderFrustum : 1 = 0;
 
 	/** Enables the use of 'HoldoutComposite' plugin for the preview rendering. */
-	bool bPreviewEnableHoldoutComposite = true;
+	uint8 bPreviewEnableHoldoutComposite : 1 = 1;
 
 	/** Enable TemporalAA/TSR in preview.
 	* Increases memory consumption. This will impact editor performance.
 	*/
-	bool bPreviewEnableTSR = false;
+	uint8 bPreviewEnableTSR : 1 = 0;
 
 	// Preview uses techvis mode for rendering
-	bool bEnablePreviewTechvis = false;
+	uint8 bEnablePreviewTechvis : 1 = 0;
 
 	// Enable/Disable preview rendering. When disabled preview image freeze
-	bool bFreezePreviewRender = false;
+	uint8 bFreezePreviewRender : 1 = 0;
 
 	// Hack preview gamma.
 	// In a scene, PostProcess always renders on top of the preview textures.
 	// But in it, PostProcess is also rendered with the flag turned off.
-	bool bPreviewEnablePostProcess = false;
+	uint8 bPreviewEnablePostProcess : 1 = 0;
 
 	// Enable OCIO on preview
-	bool bPreviewEnableOCIO = false;
+	uint8 bPreviewEnableOCIO : 1 = 0;
 
 	// Show overlay material on the preview mesh when preview rendering is enabled (UMeshComponent::OverlayMaterial)
-	bool bPreviewEnableOverlayMaterial = true;
+	uint8 bPreviewEnableOverlayMaterial : 1 = 1;
 
 	// Allows you to process preview meshes inside DCRA (get or create a mesh from a projection policy, update materials on the preview mesh, etc.).
-	bool bEnablePreviewMesh = false;
+	uint8 bEnablePreviewMesh : 1 = 0;
 
 	// Allows you to process preview editable meshes inside DCRA (get or create a mesh from a projection policy, update materials on the preview editable mesh, etc.).
-	bool bEnablePreviewEditableMesh = false;
+	uint8 bEnablePreviewEditableMesh : 1 = 0;
 
 	// Render ICVFX Frustums
-	bool bPreviewICVFXFrustums = false;
+	uint8 bPreviewICVFXFrustums : 1 = 0;
+
+public:
 	float PreviewICVFXFrustumsFarDistance = 1000.0f;
 
 	// Preview RTT size multiplier
@@ -65,10 +67,10 @@ struct FDisplayClusterViewport_PreviewSettings
 	int32 PreviewMaxTextureDimension = 2048;
 
 	// Tick Per Frame
-	int TickPerFrame = 1;
+	int32 TickPerFrame = 1;
 
 	// Max amount of Viewports Per Frame
-	int ViewportsPerFrame = 1;
+	int32 ViewportsPerFrame = 1;
 
 	// The DisplayDevice component will be obtained from the RootActor with the specified type
 	EDisplayClusterRootActorType DisplayDeviceRootActorType = EDisplayClusterRootActorType::Configuration;
diff --git a/Engine/Plugins/Runtime/nDisplay/Source/DisplayCluster/Public/Render/Viewport/IDisplayClusterViewportConfiguration.h b/Engine/Plugins/Runtime/nDisplay/Source/DisplayCluster/Public/Render/Viewport/IDisplayClusterViewportConfiguration.h
--- a/Engine/Plugins/Runtime/nDisplay/Source/DisplayCluster/Public/Render/Viewport/IDisplayClusterViewportConfiguration.h
+++ b/Engine/Plugins/Runtime/nDisplay/Source/DisplayCluster/Public/Render/Viewport/IDisplayClusterViewportConfiguration.h
@@ -125,6 +125,14 @@ public:
 	*/
 	virtual bool IsSceneOpened() const = 0;
 
+	/**
+	* Returns true if media can be used.
+	*
+	* Note: DCRA media is available only in cluster mode
+	* and should be ignored during preview rendering.
+	*/
+	virtual bool IsMediaAvailable() const = 0;
+
 	/** Returns true if preview rendering mode is used. */
 	virtual bool IsPreviewRendering() const = 0;
 
@@ -133,7 +141,20 @@ public:
 
 	/** Returns true if the DCRA preview feature in Standalone/Package builds is used. */
 	virtual bool IsPreviewInGameEnabled() const = 0;
-	
+
+	/**
+	 * Checks whether the current cluster node is rendering offscreen (headless).
+	 * 
+	 * Note: The DCRA can render in both Preview and Runtime Cluster modes at the same time.
+	 *       This is why this function is DCRA-specific rather than global.
+	 *
+	 * Usage:
+	 *   - Editor/In-Game Preview: Determined by UDisplayClusterConfigurationClusterNode::bRenderHeadless
+	 *   - Cluster Runtime:        Determined by the "RenderOffscreen" command-line argument
+	 *
+	 * @return true if rendering offscreen (headless), false otherwise.
+	 */
+	virtual bool IsClusterNodeRenderingOffscreen() const = 0;
 
 	/** Returns the rendering mode for PIE. */
 	virtual EDisplayClusterRenderFrameMode GetRenderModeForPIE() const = 0;
diff --git a/Engine/Plugins/Runtime/nDisplay/Source/DisplayClusterConfiguration/Private/DisplayClusterConfigurationTypes_ICVFX.cpp b/Engine/Plugins/Runtime/nDisplay/Source/DisplayClusterConfiguration/Private/DisplayClusterConfigurationTypes_ICVFX.cpp
--- a/Engine/Plugins/Runtime/nDisplay/Source/DisplayClusterConfiguration/Private/DisplayClusterConfigurationTypes_ICVFX.cpp
+++ b/Engine/Plugins/Runtime/nDisplay/Source/DisplayClusterConfiguration/Private/DisplayClusterConfigurationTypes_ICVFX.cpp
@@ -201,38 +201,8 @@ UCineCameraComponent* FDisplayClusterConfigurationICVFX_CameraSettings::GetExter
 
 bool FDisplayClusterConfigurationICVFX_CameraSettings::IsICVFXEnabled(const UDisplayClusterConfigurationData& InConfigurationData, const FString& InClusterNodeId) const
 {
-	// When rendering offscreen, we have an extended logic for camera rendering activation
-	static const bool bIsRunningClusterModeOffscreen =
-		(IDisplayCluster::Get().GetOperationMode() == EDisplayClusterOperationMode::Cluster) &&
-		FParse::Param(FCommandLine::Get(), TEXT("RenderOffscreen"));
-
-	if (!bIsRunningClusterModeOffscreen)
-	{
+	// All rules in the FDisplayClusterViewport::IsRenderEnabledByMedia()
 		return bEnable;
-	}
-
-	if (!bEnable)
-	{
-		return false;
-	}
-
-	// If cluster mode + rendering offscreen, discover media output settings
-
-	// First condition to render offscreen: it has media output assigned
-	const bool bUsesMediaOutput = RenderSettings.Media.bEnable && (
-			RenderSettings.Media.HasAnyMediaOutputAssigned(InClusterNodeId, EDisplayClusterConfigurationMediaSplitType::FullFrame) ||
-			RenderSettings.Media.HasAnyMediaOutputAssigned(InClusterNodeId, EDisplayClusterConfigurationMediaSplitType::UniformTiles));
-
-	// Get backbuffer media settings
-	const UDisplayClusterConfigurationClusterNode* const NodeCfg = InConfigurationData.Cluster->GetNode(InClusterNodeId);
-	const FDisplayClusterConfigurationMediaNodeBackbuffer* BackbufferMediaSettings = NodeCfg ? &NodeCfg->MediaSettings : nullptr;
-
-	// Second condition to render offscreen: the backbuffer has media output assigned.
-	// This means the whole frame including ICVFX cameras need to be rendered.
-	const bool bIsBackbufferBeingCaptured = BackbufferMediaSettings ? BackbufferMediaSettings->bEnable && BackbufferMediaSettings->IsMediaOutputAssigned() : false;
-
-	// Finally make a decision if the camera should be rendered
-	return bUsesMediaOutput || bIsBackbufferBeingCaptured;
 }
 
 const FOpenColorIOColorConversionSettings* FDisplayClusterConfigurationICVFX_CameraSettings::FindInnerFrustumOCIOConfiguration(const FString& InClusterNodeId) const
