[nDisplay] Disabled GUI propagation cvar by default

- Disabled propagation by default while working on another approach. UI should appear now as usually.

#rnx

diff --git a/Engine/Plugins/Runtime/nDisplay/Source/DisplayCluster/Private/Render/GUILayer/DisplayClusterGuiLayerController.cpp b/Engine/Plugins/Runtime/nDisplay/Source/DisplayCluster/Private/Render/GUILayer/DisplayClusterGuiLayerController.cpp
--- a/Engine/Plugins/Runtime/nDisplay/Source/DisplayCluster/Private/Render/GUILayer/DisplayClusterGuiLayerController.cpp
+++ b/Engine/Plugins/Runtime/nDisplay/Source/DisplayCluster/Private/Render/GUILayer/DisplayClusterGuiLayerController.cpp
@@ -20,7 +20,7 @@
 
 TAutoConsoleVariable<bool> CVarPropagateGui(
 	TEXT("nDisplay.GUI.Propagate"),
-	true,
+	false,
 	TEXT("Show GUI on viewports\n")
 	TEXT("0 : Disabled\n")
 	TEXT("1 : Enabled\n"),
