[nDisplay] Fixed race condition that rarely crashed the primary node

- Currently, networking data cache is reset on FrameEnd callback which is triggered after FrameEnd barrier. During this step, all internal state variables get reset before starting a new frame. However, it's possible that a secondary node may call GetTimeData from the next frame before p-node has reset its internals. This race condition may crash the p-node. To fix such type of race condition, new FrameEndSafe_AnyThread callback is introduced. It's called right from FrameEnd barrier before any threads leave it. This allows to guarantee that networking data cache is reset safely before handling any requests for a new frame.

#rnx

diff --git a/Engine/Plugins/Runtime/nDisplay/Source/DisplayCluster/Private/Cluster/Failover/DisplayClusterCommDataCache.cpp b/Engine/Plugins/Runtime/nDisplay/Source/DisplayCluster/Private/Cluster/Failover/DisplayClusterCommDataCache.cpp
--- a/Engine/Plugins/Runtime/nDisplay/Source/DisplayCluster/Private/Cluster/Failover/DisplayClusterCommDataCache.cpp
+++ b/Engine/Plugins/Runtime/nDisplay/Source/DisplayCluster/Private/Cluster/Failover/DisplayClusterCommDataCache.cpp
@@ -276,7 +276,7 @@ void FDisplayClusterCommDataCache::OpSaveImpl(const FName& SlotName, Args&... Sl
 void FDisplayClusterCommDataCache::SubscribeToCallbacks()
 {
 	// DCEndFrame is used to invalidate per-frame cache
-	IDisplayCluster::Get().GetCallbacks().OnDisplayClusterEndFrame().AddRaw(this, &FDisplayClusterCommDataCache::ProcessDCEndFrame);
+	IDisplayCluster::Get().GetCallbacks().OnDisplayClusterEndFrameSafe_AnyThread().AddRaw(this, &FDisplayClusterCommDataCache::ProcessDCEndFrame);
 
 	// Set up a post-failure negotiation delegate
 	{
@@ -305,7 +305,7 @@ void FDisplayClusterCommDataCache::SubscribeToCallbacks()
 void FDisplayClusterCommDataCache::UnsubscribeFromCallbacks()
 {
 	// Unsubscribe from DCEndFrame
-	IDisplayCluster::Get().GetCallbacks().OnDisplayClusterEndFrame().RemoveAll(this);
+	IDisplayCluster::Get().GetCallbacks().OnDisplayClusterEndFrameSafe_AnyThread().RemoveAll(this);
 
 	// Clear post-failure negotiation delegate
 	bool bSubscribedToNegotiationSync = false;
diff --git a/Engine/Plugins/Runtime/nDisplay/Source/DisplayCluster/Private/DisplayClusterCallbacks.h b/Engine/Plugins/Runtime/nDisplay/Source/DisplayCluster/Private/DisplayClusterCallbacks.h
--- a/Engine/Plugins/Runtime/nDisplay/Source/DisplayCluster/Private/DisplayClusterCallbacks.h
+++ b/Engine/Plugins/Runtime/nDisplay/Source/DisplayCluster/Private/DisplayClusterCallbacks.h
@@ -32,6 +32,11 @@ public:
 		return DisplayClusterEndFrameEvent;
 	}
 
+	virtual FDisplayClusterEndFrameEvent& OnDisplayClusterEndFrameSafe_AnyThread() override
+	{
+		return DisplayClusterEndFrameEventSafe;
+	}
+
 	virtual FDisplayClusterPreTickEvent& OnDisplayClusterPreTick() override
 	{
 		return DisplayClusterPreTickEvent;
@@ -177,6 +182,7 @@ private:
 	FDisplayClusterEndSessionEvent           DisplayClusterEndSessionEvent;
 	FDisplayClusterStartFrameEvent           DisplayClusterStartFrameEvent;
 	FDisplayClusterEndFrameEvent             DisplayClusterEndFrameEvent;
+	FDisplayClusterEndFrameEvent             DisplayClusterEndFrameEventSafe;
 	FDisplayClusterPreTickEvent              DisplayClusterPreTickEvent;
 	FDisplayClusterTickEvent                 DisplayClusterTickEvent;
 	FDisplayClusterPostTickEvent             DisplayClusterPostTickEvent;
diff --git a/Engine/Plugins/Runtime/nDisplay/Source/DisplayCluster/Private/Game/EngineClasses/Basics/DisplayClusterGameEngine.cpp b/Engine/Plugins/Runtime/nDisplay/Source/DisplayCluster/Private/Game/EngineClasses/Basics/DisplayClusterGameEngine.cpp
--- a/Engine/Plugins/Runtime/nDisplay/Source/DisplayCluster/Private/Game/EngineClasses/Basics/DisplayClusterGameEngine.cpp
+++ b/Engine/Plugins/Runtime/nDisplay/Source/DisplayCluster/Private/Game/EngineClasses/Basics/DisplayClusterGameEngine.cpp
@@ -7,12 +7,13 @@
 #include "Cluster/Controller/IDisplayClusterClusterNodeController.h"
 #include "Cluster/NetAPI/DisplayClusterNetApiFacade.h"
 #include "Config/IPDisplayClusterConfigManager.h"
-#include "DisplayClusterEnums.h"
 #include "Engine/DynamicBlueprintBinding.h"
+#include "Engine/GameInstance.h"
+#include "Engine/NetDriver.h"
 
-#include "DisplayClusterConfigurationTypes.h"
-#include "DisplayClusterConfigurationVersion.h"
-#include "IDisplayClusterConfiguration.h"
+#include "GameFramework/GameModeBase.h"
+#include "Interfaces/IPv4/IPv4Endpoint.h"
+#include "Kismet/GameplayStatics.h"
 
 #include "Misc/App.h"
 #include "Misc/CommandLine.h"
@@ -25,21 +26,23 @@
 #include "Misc/DisplayClusterStrings.h"
 #include "Misc/Parse.h"
 
-#include "SocketSubsystem.h"
-#include "Interfaces/IPv4/IPv4Endpoint.h"
-
-#include "EngineGlobals.h"
-#include "Engine/NetDriver.h"
-#include "Engine/GameInstance.h"
-
-#include "UObject/Linker.h"
-#include "GameMapsSettings.h"
-#include "GameFramework/GameModeBase.h"
+#include "Network/Barrier/IDisplayClusterBarrier.h"
+#include "Network/Configuration/DisplayClusterNetworkConfiguration.h"
+#include "Network/Service/ClusterSync/DisplayClusterClusterSyncService.h"
+#include "Network/Service/DisplayClusterService.h"
 
 #include "Stats/Stats.h"
+#include "UObject/Linker.h"
 
+#include "DisplayClusterConfigurationTypes.h"
+#include "DisplayClusterConfigurationVersion.h"
+#include "DisplayClusterEnums.h"
+#include "EngineGlobals.h"
 #include "GameDelegates.h"
-#include "Kismet/GameplayStatics.h"
+#include "GameMapsSettings.h"
+#include "IDisplayClusterCallbacks.h"
+#include "IDisplayClusterConfiguration.h"
+#include "SocketSubsystem.h"
 
 
 namespace DisplayClusterGameEngineUtils
@@ -179,9 +182,11 @@ bool UDisplayClusterGameEngine::InitializeInternals()
 	// Store diagnostics settings locally
 	Diagnostics = Config->Diagnostics;
 
+	// Listen for JSON sync events
 	FOnClusterEventJsonListener GameSyncTransition = FOnClusterEventJsonListener::CreateUObject(this, &UDisplayClusterGameEngine::GameSyncChange);
 	GDisplayCluster->GetPrivateClusterMgr()->AddClusterEventJsonListener(GameSyncTransition);
 
+	// Enable/disable audio on this node
 	const UDisplayClusterConfigurationClusterNode* CfgLocalNode = GDisplayCluster->GetPrivateConfigMgr()->GetLocalNode();
 	const bool bSoundEnabled = (CfgLocalNode ? CfgLocalNode->bIsSoundEnabled : false);
 	UE_LOG(LogDisplayClusterEngine, Log, TEXT("Configuring sound enabled: %s"), *DisplayClusterTypesConverter::template ToString(bSoundEnabled));
@@ -190,6 +195,25 @@ bool UDisplayClusterGameEngine::InitializeInternals()
 		AudioDeviceManager = nullptr;
 	}
 
+	// Use PreSyncEnd delegate of EndFrame barrier for EndFrameSafe_AnyThread callback
+	const FName& ServiceName = UE::nDisplay::Network::Configuration::ClusterSyncServerName;
+	const TSharedPtr<FDisplayClusterService> ClusterSyncServiceBase = GDisplayCluster->GetPrivateClusterMgr()->GetNodeController()->GetService(ServiceName).Pin();
+	const TSharedPtr<FDisplayClusterClusterSyncService> ClusterSyncService = StaticCastSharedPtr<FDisplayClusterClusterSyncService>(ClusterSyncServiceBase);
+	if (ClusterSyncService)
+	{
+		if (IDisplayClusterBarrier* FrameEndBarrier = ClusterSyncService->GetFrameEndBarrier())
+		{
+			// Trigger EndFrameSafe_AnyThread right from the barrier before any clients leave
+			FrameEndBarrier->GetPreSyncEndDelegate().BindLambda([](FDisplayClusterBarrierPreSyncEndDelegateData&)
+				{
+					if (GDisplayCluster)
+					{
+						GDisplayCluster->GetCallbacks().OnDisplayClusterEndFrameSafe_AnyThread().Broadcast(GFrameCounter);
+					}
+				});
+		}
+	}
+
 	return true;
 }
 
diff --git a/Engine/Plugins/Runtime/nDisplay/Source/DisplayCluster/Private/Network/Service/ClusterSync/DisplayClusterClusterSyncService.h b/Engine/Plugins/Runtime/nDisplay/Source/DisplayCluster/Private/Network/Service/ClusterSync/DisplayClusterClusterSyncService.h
--- a/Engine/Plugins/Runtime/nDisplay/Source/DisplayCluster/Private/Network/Service/ClusterSync/DisplayClusterClusterSyncService.h
+++ b/Engine/Plugins/Runtime/nDisplay/Source/DisplayCluster/Private/Network/Service/ClusterSync/DisplayClusterClusterSyncService.h
@@ -34,6 +34,19 @@ public:
 	virtual void KillSession(const FString& NodeId) override;
 	//~ End IDisplayClusterServer
 
+public:
+
+	/**
+	 * This is a temporary workaround (quick fix) for 5.7.
+	 * The goal is to fix a race condition by providing the ability to handle EndFrame barrier
+	 * sync callback. In 5.8, this code will be re-written in more generic way. The abstract
+	 * FDisplayClusterService will provide a full set of API functions to access any barriers.
+	 */
+	IDisplayClusterBarrier* GetFrameEndBarrier()
+	{
+		return BarrierFrameEnd.Get();
+	}
+
 protected:
 
 	/** Creates session instance for this service */
diff --git a/Engine/Plugins/Runtime/nDisplay/Source/DisplayCluster/Public/IDisplayClusterCallbacks.h b/Engine/Plugins/Runtime/nDisplay/Source/DisplayCluster/Public/IDisplayClusterCallbacks.h
--- a/Engine/Plugins/Runtime/nDisplay/Source/DisplayCluster/Public/IDisplayClusterCallbacks.h
+++ b/Engine/Plugins/Runtime/nDisplay/Source/DisplayCluster/Public/IDisplayClusterCallbacks.h
@@ -55,6 +55,9 @@ public:
 	DECLARE_EVENT_OneParam(IDisplayClusterCallbacks, FDisplayClusterEndFrameEvent, uint64);
 	virtual FDisplayClusterEndFrameEvent& OnDisplayClusterEndFrame() = 0;
 
+	/** Called right after EndFrame barrier sync before any client leaves. Called from a server session thread. **/
+	virtual FDisplayClusterEndFrameEvent& OnDisplayClusterEndFrameSafe_AnyThread() = 0;
+
 	/** Called on DisplayCluster PreTick **/
 	DECLARE_EVENT(IDisplayClusterCallbacks, FDisplayClusterPreTickEvent);
 	virtual FDisplayClusterPreTickEvent& OnDisplayClusterPreTick() = 0;
