nDisplay: Fixed issue where ICVFX window would through failed assert when -rhivalidation was used

diff --git a/Engine/Plugins/Runtime/nDisplay/Source/DisplayClusterLightCardEditorShaders/Private/DisplayClusterMeshProjectionRenderer.cpp b/Engine/Plugins/Runtime/nDisplay/Source/DisplayClusterLightCardEditorShaders/Private/DisplayClusterMeshProjectionRenderer.cpp
--- a/Engine/Plugins/Runtime/nDisplay/Source/DisplayClusterLightCardEditorShaders/Private/DisplayClusterMeshProjectionRenderer.cpp
+++ b/Engine/Plugins/Runtime/nDisplay/Source/DisplayClusterLightCardEditorShaders/Private/DisplayClusterMeshProjectionRenderer.cpp
@@ -1027,7 +1027,7 @@ void FDisplayClusterMeshProjectionRenderer::RenderColorOutput(FRDGBuilder& Graph
 	FRDGTextureRef DepthTexture = GraphBuilder.CreateTexture(DepthDesc, TEXT("DisplayClusterMeshProjection.DepthTexture"));
 
 	FRenderTargetBinding ColorRenderTargetBinding(ColorTexture, ERenderTargetLoadAction::EClear);
-	FDepthStencilBinding DepthStencilBinding(DepthTexture, ERenderTargetLoadAction::EClear, ERenderTargetLoadAction::ENoAction, FExclusiveDepthStencil::DepthWrite_StencilNop);
+	FDepthStencilBinding DepthStencilBinding(DepthTexture, ERenderTargetLoadAction::EClear, ERenderTargetLoadAction::ENoAction, FExclusiveDepthStencil::DepthWrite_StencilWrite);
 
 	AddBaseRenderPass(GraphBuilder, View, RenderSettings, ColorRenderTargetBinding, DepthStencilBinding);
 
@@ -1965,7 +1965,7 @@ void FDisplayClusterMeshProjectionRenderer::RenderColorOutputs(FRDGBuilder& Grap
 	FRDGTextureRef DepthTexture = GraphBuilder.CreateTexture(DepthDesc, TEXT("DisplayClusterMeshProjection.DepthTexture"));
 
 	FRenderTargetBinding ColorRenderTargetBinding(ColorTexture, ERenderTargetLoadAction::EClear);
-	FDepthStencilBinding DepthStencilBinding(DepthTexture, ERenderTargetLoadAction::EClear, ERenderTargetLoadAction::ENoAction, FExclusiveDepthStencil::DepthWrite_StencilNop);
+	FDepthStencilBinding DepthStencilBinding(DepthTexture, ERenderTargetLoadAction::EClear, ERenderTargetLoadAction::ENoAction, FExclusiveDepthStencil::DepthWrite_StencilWrite);
 
 	for (const FSceneView* ViewIt : Views)
 	{
