[nDisplay] Fixed CIS issue after 46534179 (static analyzer)

#rnx

#DisabledValidations Preflight

diff --git a/Engine/Plugins/Runtime/nDisplay/Source/DisplayCluster/Private/Game/EngineClasses/Basics/DisplayClusterGameEngine.cpp b/Engine/Plugins/Runtime/nDisplay/Source/DisplayCluster/Private/Game/EngineClasses/Basics/DisplayClusterGameEngine.cpp
--- a/Engine/Plugins/Runtime/nDisplay/Source/DisplayCluster/Private/Game/EngineClasses/Basics/DisplayClusterGameEngine.cpp
+++ b/Engine/Plugins/Runtime/nDisplay/Source/DisplayCluster/Private/Game/EngineClasses/Basics/DisplayClusterGameEngine.cpp
@@ -175,6 +175,12 @@ EDisplayClusterOperationMode UDisplayClusterGameEngine::DetectOperationMode() co
 
 bool UDisplayClusterGameEngine::InitializeInternals()
 {
+	if (!GDisplayCluster)
+	{
+		UE_LOG(LogDisplayClusterEngine, Warning, TEXT("Could not initialize DC game engine, invalid GDisplayCluster"));
+		return false;
+	}
+
 	// This function is called after a session had been started so it's safe to get config data from the config manager
 	const UDisplayClusterConfigurationData* Config = GDisplayCluster->GetPrivateConfigMgr()->GetConfig();
 	check(Config);
