[nDisplay] Another fix for a CIS issue after 46534179 (static analyzer)

#rnx

diff --git a/Engine/Plugins/Runtime/nDisplay/Source/DisplayCluster/Private/Game/EngineClasses/Basics/DisplayClusterGameEngine.cpp b/Engine/Plugins/Runtime/nDisplay/Source/DisplayCluster/Private/Game/EngineClasses/Basics/DisplayClusterGameEngine.cpp
--- a/Engine/Plugins/Runtime/nDisplay/Source/DisplayCluster/Private/Game/EngineClasses/Basics/DisplayClusterGameEngine.cpp
+++ b/Engine/Plugins/Runtime/nDisplay/Source/DisplayCluster/Private/Game/EngineClasses/Basics/DisplayClusterGameEngine.cpp
@@ -175,14 +175,16 @@ EDisplayClusterOperationMode UDisplayClusterGameEngine::DetectOperationMode() co
 
 bool UDisplayClusterGameEngine::InitializeInternals()
 {
-	if (!GDisplayCluster)
+	// To prevent static analyzer of reporting unchecked pointers usage
+	const IPDisplayCluster* const DCAPI = GDisplayCluster;
+	if (!DCAPI)
 	{
 		UE_LOG(LogDisplayClusterEngine, Warning, TEXT("Could not initialize DC game engine, invalid GDisplayCluster"));
 		return false;
 	}
 
 	// This function is called after a session had been started so it's safe to get config data from the config manager
-	const UDisplayClusterConfigurationData* Config = GDisplayCluster->GetPrivateConfigMgr()->GetConfig();
+	const UDisplayClusterConfigurationData* Config = DCAPI->GetPrivateConfigMgr()->GetConfig();
 	check(Config);
 
 	// Store diagnostics settings locally
@@ -190,10 +192,10 @@ bool UDisplayClusterGameEngine::InitializeInternals()
 
 	// Listen for JSON sync events
 	FOnClusterEventJsonListener GameSyncTransition = FOnClusterEventJsonListener::CreateUObject(this, &UDisplayClusterGameEngine::GameSyncChange);
-	GDisplayCluster->GetPrivateClusterMgr()->AddClusterEventJsonListener(GameSyncTransition);
+	DCAPI->GetPrivateClusterMgr()->AddClusterEventJsonListener(GameSyncTransition);
 
 	// Enable/disable audio on this node
-	const UDisplayClusterConfigurationClusterNode* CfgLocalNode = GDisplayCluster->GetPrivateConfigMgr()->GetLocalNode();
+	const UDisplayClusterConfigurationClusterNode* CfgLocalNode = DCAPI->GetPrivateConfigMgr()->GetLocalNode();
 	const bool bSoundEnabled = (CfgLocalNode ? CfgLocalNode->bIsSoundEnabled : false);
 	UE_LOG(LogDisplayClusterEngine, Log, TEXT("Configuring sound enabled: %s"), *DisplayClusterTypesConverter::template ToString(bSoundEnabled));
 	if (!bSoundEnabled)
@@ -203,7 +205,7 @@ bool UDisplayClusterGameEngine::InitializeInternals()
 
 	// Use PreSyncEnd delegate of EndFrame barrier for EndFrameSafe_AnyThread callback
 	const FName& ServiceName = UE::nDisplay::Network::Configuration::ClusterSyncServerName;
-	const TSharedPtr<FDisplayClusterService> ClusterSyncServiceBase = GDisplayCluster->GetPrivateClusterMgr()->GetNodeController()->GetService(ServiceName).Pin();
+	const TSharedPtr<FDisplayClusterService> ClusterSyncServiceBase = DCAPI->GetPrivateClusterMgr()->GetNodeController()->GetService(ServiceName).Pin();
 	const TSharedPtr<FDisplayClusterClusterSyncService> ClusterSyncService = StaticCastSharedPtr<FDisplayClusterClusterSyncService>(ClusterSyncServiceBase);
 	if (ClusterSyncService)
 	{
