[nDisplay] MRQ rendering with alpha channel.

MRQ must have exclusive rights to DCRA:
- No proxy actors may use this DCRA simultaneously (ICVFX Panel, StageApp, etc.).
- Default scene preview rendering is disabled.
- PIE rendering is disabled.

MRQ rendering alpha output:
- override FDisplayClusterViewport_RenderSettings::bWarpBlendRenderAlphaChannel

Added new parameter FDisplayClusterRenderFrameSettings::bEnableAlphaOutput:
True if alpha channel should be written to all viewport outputs.
Note: Controlled by "DC.EnableAlphaChannelRendering" and "r.PostProcessing.PropagateAlpha".

Removed unused parameter FDisplayClusterRenderFrameSettings::bShouldUseParentViewportRenderFamily

Improved FDisplayClusterMoviePipelineConfiguration
- An empty DCRA reference means: use any available DCRA in the scene.
- Allows configuring specific viewports to render even when the DCRA is undefined.
- UX: User can only edit AllowedViewportNamesList if RootActor is assigned and the bRenderAllViewports checkbox is unchecked.
- UX: Keeps AllowedViewportNamesList values after RootActor is set to nullptr.
Code cleanup

diff --git a/Engine/Plugins/Runtime/nDisplay/Source/DisplayCluster/Private/Game/EngineClasses/Basics/DisplayClusterViewportClient.cpp b/Engine/Plugins/Runtime/nDisplay/Source/DisplayCluster/Private/Game/EngineClasses/Basics/DisplayClusterViewportClient.cpp
--- a/Engine/Plugins/Runtime/nDisplay/Source/DisplayCluster/Private/Game/EngineClasses/Basics/DisplayClusterViewportClient.cpp
+++ b/Engine/Plugins/Runtime/nDisplay/Source/DisplayCluster/Private/Game/EngineClasses/Basics/DisplayClusterViewportClient.cpp
@@ -929,13 +929,13 @@ bool UDisplayClusterViewportClient::Draw_PIE(FViewport* InViewport, FCanvas* Sce
 	// When the PIE is used by this DCRA, we must create a new ViewportManager
 	if(IDisplayClusterViewportManager* ViewportManager = RootActor->GetOrCreateViewportManager())
 	{
-		FDisplayClusterViewport_PreviewSettings NewPreviewSettings = RootActor->GetPreviewSettings();
-		{
-			// Disable frustum preview rendering in PIE
-			NewPreviewSettings.bPreviewICVFXFrustums = false;
+		// Note: Logic that disable frustum preview rendering in PIE has been moved
+		//       to the ADisplayClusterRootActor::PreparePreviewRendererSettings()
 
-			// Note: Normally these settings are not used for previewing in PIE and are ignored.
-			ViewportManager->GetConfiguration().SetPreviewSettings(NewPreviewSettings);
+		// Do nothing if this DCRA is under an exclusive lock.
+		if (ViewportManager->GetConfiguration().IsExclusiveLocked())
+		{
+			return false;
 		}
 
 		const EDisplayClusterRenderFrameMode RenderFrameMode = ViewportManager->GetConfiguration().GetRenderModeForPIE();
diff --git a/Engine/Plugins/Runtime/nDisplay/Source/DisplayCluster/Private/Game/EngineClasses/Scene/DisplayClusterRootActor.cpp b/Engine/Plugins/Runtime/nDisplay/Source/DisplayCluster/Private/Game/EngineClasses/Scene/DisplayClusterRootActor.cpp
--- a/Engine/Plugins/Runtime/nDisplay/Source/DisplayCluster/Private/Game/EngineClasses/Scene/DisplayClusterRootActor.cpp
+++ b/Engine/Plugins/Runtime/nDisplay/Source/DisplayCluster/Private/Game/EngineClasses/Scene/DisplayClusterRootActor.cpp
@@ -1206,20 +1206,20 @@ void ADisplayClusterRootActor::Tick(float DeltaSeconds)
 	Super::Tick(DeltaSeconds);
 }
 
-void ADisplayClusterRootActor::TickPreviewRenderer()
+bool ADisplayClusterRootActor::PreparePreviewRendererSettings(FDisplayClusterViewport_PreviewSettings& OutPreviewSettings) const
 {
 	const bool bIsPrimaryRootActor = IsPrimaryRootActor();
 	const bool bIsRunningDisplayCluster = IsRunningDisplayCluster();
 
 	// Get preview settings from the this actor current source
 	// as they may already be configured externally, so just get the preview settings from the current source
-	FDisplayClusterViewport_PreviewSettings NewPreviewSettings = GetPreviewSettings();
+	OutPreviewSettings = GetPreviewSettings();
 
 	const bool bIsRunningPIE = IsRunningPIE();
 	const bool bIsRunningGame = IsRunningGame();
 
 	// RootActor can have a preview in the scene
-	bool bEnablePreviewInScene = NewPreviewSettings.bPreviewEnable;
+	bool bEnablePreviewInScene = OutPreviewSettings.bPreviewEnable;
 	bool bPreviewInGame = false;
 
 	if (bIsRunningPIE)
@@ -1229,11 +1229,14 @@ void ADisplayClusterRootActor::TickPreviewRenderer()
 		if (IsPrimaryRootActorForPIE())
 		{
 			bEnablePreviewInScene = false;
+
+			// Disable frustum rendering in PIE
+			OutPreviewSettings.bPreviewICVFXFrustums = false;
 		}
 		else
 		{
 			// When running in a game or PIE, a special flag is required to show a preview for PIE/Standalone/Package
-			bEnablePreviewInScene = NewPreviewSettings.bPreviewEnable && NewPreviewSettings.bPreviewInGameEnable;
+			bEnablePreviewInScene = OutPreviewSettings.bPreviewEnable && OutPreviewSettings.bPreviewInGameEnable;
 		}
 	}
 	else if (bIsRunningGame || bIsRunningDisplayCluster)
@@ -1249,64 +1252,62 @@ void ADisplayClusterRootActor::TickPreviewRenderer()
 		else
 		{
 			// When running in a game or PIE, a special flag is required to show a preview for PIE/Standalone/Package
-			bEnablePreviewInScene = NewPreviewSettings.bPreviewEnable && NewPreviewSettings.bPreviewInGameEnable;
+			bEnablePreviewInScene = OutPreviewSettings.bPreviewEnable && OutPreviewSettings.bPreviewInGameEnable;
 		}
 	}
 
-#if WITH_EDITOR
-	if (bMoviePipelineRenderPass)
-	{
-		// Disable preview rendering for RootActor used in MRQ
-		bEnablePreviewInScene = false;
-	}
-#endif
-
 	if (GetGameInstance() && GetGameInstance()->IsDedicatedServerInstance())
 	{
 		// Skip rendering on dedicated server
 		bEnablePreviewInScene = false;
 	}
 
-	if (bPreviewInGame && !NewPreviewSettings.bPreviewInGameRenderFrustum)
+	if (bPreviewInGame && !OutPreviewSettings.bPreviewInGameRenderFrustum)
 	{
 		// Disable frustum preview rendering in game
-		NewPreviewSettings.bPreviewICVFXFrustums = false;
+		OutPreviewSettings.bPreviewICVFXFrustums = false;
 	}
 
 	if (!bEnablePreviewInScene)
 	{
 		// Disable preview rendering
-		NewPreviewSettings.bPreviewEnable = false;
+		OutPreviewSettings.bPreviewEnable = false;
 	}
 
 	// To use the `HoldoutComposite` plugin, a post-process is required
-	if (!NewPreviewSettings.bPreviewEnablePostProcess)
+	if (!OutPreviewSettings.bPreviewEnablePostProcess)
 	{
-		NewPreviewSettings.bPreviewEnableHoldoutComposite = false;
+		OutPreviewSettings.bPreviewEnableHoldoutComposite = false;
 	}
 
 	// The HoldoutComposite plugin requires special project settings.
-	if (NewPreviewSettings.bPreviewEnableHoldoutComposite
-	&& !UCompositeCoreSubsystem::IsProjectSettingsValid())
+	if (OutPreviewSettings.bPreviewEnableHoldoutComposite
+		&& !UCompositeCoreSubsystem::IsProjectSettingsValid())
 	{
-		NewPreviewSettings.bPreviewEnableHoldoutComposite = false;
+		OutPreviewSettings.bPreviewEnableHoldoutComposite = false;
 
 		// In this particular case, also disable the postprocess.
-		NewPreviewSettings.bPreviewEnablePostProcess = false;
+		OutPreviewSettings.bPreviewEnablePostProcess = false;
 	}
 
-#if WITH_EDITOR
-	// The StageIsosphereComponent can only be seen if DCRA preview is enabled.
-	const bool bShowStageIsosphere = NewPreviewSettings.bPreviewEnable && bPreviewStageGeometryMesh;
-	if (StageIsosphereComponent && StageIsosphereComponent->GetVisibleFlag() != bShowStageIsosphere)
-	{
-		StageIsosphereComponent->SetVisibility(bShowStageIsosphere);
-	}
-#endif
+	return bEnablePreviewInScene;
+}
+
+void ADisplayClusterRootActor::TickPreviewRenderer()
+{
+	FDisplayClusterViewport_PreviewSettings NewPreviewSettings;
+	bool bEnablePreviewInScene = PreparePreviewRendererSettings(NewPreviewSettings);
 
-	// Update entire cluster preview rendering
 	if (IDisplayClusterViewportManager* ViewportManager = bEnablePreviewInScene ? GetOrCreateViewportManager() : GetViewportManager())
 	{
+		// Disable preview rendering if this DCRA is under an exclusive lock.
+		if (ViewportManager->GetConfiguration().IsExclusiveLocked())
+		{
+			// Disable preview in scene
+			NewPreviewSettings.bPreviewEnable = false;
+			bEnablePreviewInScene = false;
+		}
+
 		// Update preview settings to new
 		ViewportManager->GetConfiguration().SetPreviewSettings(NewPreviewSettings);
 
@@ -1314,6 +1315,15 @@ void ADisplayClusterRootActor::TickPreviewRenderer()
 		ViewportManager->GetViewportManagerPreview().UpdateEntireClusterPreviewRender(bEnablePreviewInScene);
 	}
 
+#if WITH_EDITOR
+	// The StageIsosphereComponent can only be seen if DCRA preview is enabled.
+	const bool bShowStageIsosphere = NewPreviewSettings.bPreviewEnable && bPreviewStageGeometryMesh;
+	if (StageIsosphereComponent && StageIsosphereComponent->GetVisibleFlag() != bShowStageIsosphere)
+	{
+		StageIsosphereComponent->SetVisibility(bShowStageIsosphere);
+	}
+#endif
+
 	// Update RootActor visibility for game
 	SetActorHiddenInGame(!bEnablePreviewInScene);
 }
diff --git a/Engine/Plugins/Runtime/nDisplay/Source/DisplayCluster/Private/Render/Viewport/Configuration/DisplayClusterViewportConfiguration.cpp b/Engine/Plugins/Runtime/nDisplay/Source/DisplayCluster/Private/Render/Viewport/Configuration/DisplayClusterViewportConfiguration.cpp
--- a/Engine/Plugins/Runtime/nDisplay/Source/DisplayCluster/Private/Render/Viewport/Configuration/DisplayClusterViewportConfiguration.cpp
+++ b/Engine/Plugins/Runtime/nDisplay/Source/DisplayCluster/Private/Render/Viewport/Configuration/DisplayClusterViewportConfiguration.cpp
@@ -324,6 +324,24 @@ EDisplayClusterRenderFrameMode FDisplayClusterViewportConfiguration::GetRenderMo
 	return EDisplayClusterRenderFrameMode::PIE_Mono;
 }
 
+bool FDisplayClusterViewportConfiguration::IsExclusiveLocked() const
+{
+#if WITH_EDITOR
+	// Use the DCRA in scene to ensure that all proxies comply with its rules.
+	// (ICVFX Panel, EpicStageApp, etc.)
+	if (ADisplayClusterRootActor* SceneRootActor = GetRootActor(EDisplayClusterRootActorType::Scene))
+	{
+		// MRQ requires exclusivity: preview and proxies cannot use or modify the DCRA in scene.
+		if (SceneRootActor->bMoviePipelineRenderPass)
+		{
+			return true;
+		}
+	}
+#endif
+
+	return false;
+}
+
 bool FDisplayClusterViewportConfiguration::IsMediaAvailable() const
 {
 	// Media is not available in DCRA preview.
diff --git a/Engine/Plugins/Runtime/nDisplay/Source/DisplayCluster/Private/Render/Viewport/Configuration/DisplayClusterViewportConfiguration.h b/Engine/Plugins/Runtime/nDisplay/Source/DisplayCluster/Private/Render/Viewport/Configuration/DisplayClusterViewportConfiguration.h
--- a/Engine/Plugins/Runtime/nDisplay/Source/DisplayCluster/Private/Render/Viewport/Configuration/DisplayClusterViewportConfiguration.h
+++ b/Engine/Plugins/Runtime/nDisplay/Source/DisplayCluster/Private/Render/Viewport/Configuration/DisplayClusterViewportConfiguration.h
@@ -100,6 +100,8 @@ public:
 		return RenderFrameSettings.IsPreviewInGameEnabled();
 	}
 
+	virtual bool IsExclusiveLocked() const override;
+
 	virtual const float GetWorldToMeters() const override;
 	// ~~END IDisplayClusterViewportConfiguration
 
diff --git a/Engine/Plugins/Runtime/nDisplay/Source/DisplayCluster/Private/Render/Viewport/Configuration/DisplayClusterViewportConfigurationHelpers.cpp b/Engine/Plugins/Runtime/nDisplay/Source/DisplayCluster/Private/Render/Viewport/Configuration/DisplayClusterViewportConfigurationHelpers.cpp
--- a/Engine/Plugins/Runtime/nDisplay/Source/DisplayCluster/Private/Render/Viewport/Configuration/DisplayClusterViewportConfigurationHelpers.cpp
+++ b/Engine/Plugins/Runtime/nDisplay/Source/DisplayCluster/Private/Render/Viewport/Configuration/DisplayClusterViewportConfigurationHelpers.cpp
@@ -108,8 +108,13 @@ void FDisplayClusterViewportConfigurationHelpers::UpdateBaseViewportSetting(FDis
 		InOutRenderSettings.GPUIndex = InConfigurationViewport.GPUIndex;
 		InOutRenderSettings.OverlapOrder = InConfigurationViewport.OverlapOrder;
 
+		if (DstViewport.Configuration->GetRenderFrameSettings().bEnableAlphaOutput)
+		{
+			InOutRenderSettings.bWarpBlendRenderAlphaChannel = true;
+		}
+
 		// update viewport remap data
-		DstViewport.UpdateConfiguration_ViewportRemap(InConfigurationViewport.ViewportRemap);
+		DstViewport.UpdateConfiguration_ViewportRemap(InConfigurationViewport.ViewportRemap); 
 	}
 
 	const FDisplayClusterConfigurationViewport_RenderSettings& InRenderSettings = InConfigurationViewport.RenderSettings;
diff --git a/Engine/Plugins/Runtime/nDisplay/Source/DisplayCluster/Private/Render/Viewport/Configuration/DisplayClusterViewportConfigurationHelpers_RenderFrameSettings.cpp b/Engine/Plugins/Runtime/nDisplay/Source/DisplayCluster/Private/Render/Viewport/Configuration/DisplayClusterViewportConfigurationHelpers_RenderFrameSettings.cpp
--- a/Engine/Plugins/Runtime/nDisplay/Source/DisplayCluster/Private/Render/Viewport/Configuration/DisplayClusterViewportConfigurationHelpers_RenderFrameSettings.cpp
+++ b/Engine/Plugins/Runtime/nDisplay/Source/DisplayCluster/Private/Render/Viewport/Configuration/DisplayClusterViewportConfigurationHelpers_RenderFrameSettings.cpp
@@ -80,21 +80,51 @@ static FAutoConsoleVariableRef CVarDisplayClusterAlphaChannelCaptureMode(
 	ECVF_RenderThreadSafe
 );
 
+int32 GDisplayClusterEnableAlphaChannelRendering = 0;
+static FAutoConsoleVariableRef CVarDisplayClusterEnableAlphaChannelRendering(
+	TEXT("DC.EnableAlphaChannelRendering"),
+	GDisplayClusterEnableAlphaChannelRendering,
+	TEXT("Enable alpha output for all viewports (0 = disabled, 1 = enabled)\n")
+	TEXT("Requires CVar 'r.PostProcessing.PropagateAlpha' to be enabled.\n"),
+	ECVF_RenderThreadSafe
+);
+
+namespace UE::DisplayClusterViewport::Configuration::RenderFrameSettings
+{
+	/**
+	* Returns true if "Alpha Output" is enabled in Project Settings
+	* (backed by the CVar r.PostProcessing.PropagateAlpha).
+	*/
+	static bool IsProjectSettingAlphaOutputEnabled()
+	{
+		static const auto CVarPropagateAlpha = IConsoleManager::Get().FindConsoleVariable(TEXT("r.PostProcessing.PropagateAlpha"));
+		if (CVarPropagateAlpha)
+		{
+			const bool bPropagateAlpha = CVarPropagateAlpha->GetBool();
+
+			return bPropagateAlpha;
+		}
+
+		return false;
+	}
+};
+
 ///////////////////////////////////////////////////////////////////
 // FDisplayClusterViewportConfiguration
 ///////////////////////////////////////////////////////////////////
 EDisplayClusterRenderFrameAlphaChannelCaptureMode FDisplayClusterViewportConfigurationHelpers_RenderFrameSettings::GetAlphaChannelCaptureMode()
 {
-	ECVarDisplayClusterAlphaChannelCaptureMode AlphaChannelCaptureMode = (ECVarDisplayClusterAlphaChannelCaptureMode)FMath::Clamp(GDisplayClusterAlphaChannelCaptureMode, 0, (int32)ECVarDisplayClusterAlphaChannelCaptureMode::COUNT - 1);
-
-	static const auto CVarPropagateAlpha = IConsoleManager::Get().FindConsoleVariable(TEXT("r.PostProcessing.PropagateAlpha"));
-	const bool bPropagateAlpha = CVarPropagateAlpha->GetBool();
+	using namespace UE::DisplayClusterViewport::Configuration::RenderFrameSettings;
 
+	const ECVarDisplayClusterAlphaChannelCaptureMode AlphaChannelCaptureMode = (ECVarDisplayClusterAlphaChannelCaptureMode)FMath::Clamp(GDisplayClusterAlphaChannelCaptureMode, 0, (int32)ECVarDisplayClusterAlphaChannelCaptureMode::COUNT - 1);
 	switch (AlphaChannelCaptureMode)
 	{
 	case ECVarDisplayClusterAlphaChannelCaptureMode::ThroughTonemapper:
-		// Disable alpha capture if PropagateAlpha not valid
-		return bPropagateAlpha ? EDisplayClusterRenderFrameAlphaChannelCaptureMode::ThroughTonemapper : EDisplayClusterRenderFrameAlphaChannelCaptureMode::None;
+		return IsProjectSettingAlphaOutputEnabled()
+			// Capture alpha as processed through the tonemapper.
+			? EDisplayClusterRenderFrameAlphaChannelCaptureMode::ThroughTonemapper
+			// Disable alpha capture when r.PostProcessing.PropagateAlpha is not enabled.
+			: EDisplayClusterRenderFrameAlphaChannelCaptureMode::None;
 
 	case ECVarDisplayClusterAlphaChannelCaptureMode::FXAA:
 		return EDisplayClusterRenderFrameAlphaChannelCaptureMode::FXAA;
@@ -118,6 +148,8 @@ bool FDisplayClusterViewportConfigurationHelpers_RenderFrameSettings::UpdateRend
 	const FString& InClusterNodeId,
 	FDisplayClusterViewportConfiguration& InOutConfiguration)
 {
+	using namespace UE::DisplayClusterViewport::Configuration::RenderFrameSettings;
+
 	ADisplayClusterRootActor* ConfigurationRootActor = InOutConfiguration.GetRootActor(EDisplayClusterRootActorType::Configuration);
 	if (!(ConfigurationRootActor))
 	{
@@ -145,6 +177,12 @@ bool FDisplayClusterViewportConfigurationHelpers_RenderFrameSettings::UpdateRend
 	// Support alpha channel capture
 	NewRenderFrameSettings.AlphaChannelCaptureMode = GetAlphaChannelCaptureMode();
 
+	// True if alpha channel should be written to all viewport outputs.
+	// Note: Controlled by "DC.EnableAlphaChannelRendering" and "r.PostProcessing.PropagateAlpha".
+	NewRenderFrameSettings.bEnableAlphaOutput =
+		GDisplayClusterEnableAlphaChannelRendering
+		&& IsProjectSettingAlphaOutputEnabled();
+
 	// Update RenderFrame configuration
 	const FDisplayClusterConfigurationRenderFrame& InRenderFrameConfiguration = ConfigurationRootActor->GetRenderFrameSettings();
 	{
diff --git a/Engine/Plugins/Runtime/nDisplay/Source/DisplayCluster/Private/Render/Viewport/Configuration/DisplayClusterViewportConfiguration_ICVFX.cpp b/Engine/Plugins/Runtime/nDisplay/Source/DisplayCluster/Private/Render/Viewport/Configuration/DisplayClusterViewportConfiguration_ICVFX.cpp
--- a/Engine/Plugins/Runtime/nDisplay/Source/DisplayCluster/Private/Render/Viewport/Configuration/DisplayClusterViewportConfiguration_ICVFX.cpp
+++ b/Engine/Plugins/Runtime/nDisplay/Source/DisplayCluster/Private/Render/Viewport/Configuration/DisplayClusterViewportConfiguration_ICVFX.cpp
@@ -28,17 +28,6 @@
 #include "Render/Viewport/LightCard/DisplayClusterViewportLightCardManager.h"
 #include "Render/Viewport/RenderFrame/DisplayClusterRenderFrameSettings.h"
 
-
-////////////////////////////////////////////////////////////////////////////////
-int32 GDisplayClusterEnableAlphaChannelRendering = 0;
-static FAutoConsoleVariableRef CVarDisplayClusterEnableAlphaChannelRendering(
-	TEXT("DC.EnableAlphaChannelRendering"),
-	GDisplayClusterEnableAlphaChannelRendering,
-	TEXT("Enable alpha channel rendering to backbuffer (0 == disabled, 1 == enabled)"),
-	ECVF_RenderThreadSafe
-);
-
-
 ///////////////////////////////////////////////////////////////////
 // FDisplayClusterViewportConfiguration_ICVFX
 ///////////////////////////////////////////////////////////////////
@@ -85,9 +74,8 @@ void FDisplayClusterViewportConfiguration_ICVFX::Update()
 	const FDisplayClusterConfigurationICVFX_StageSettings* StageSettings = Configuration.GetStageSettings();
 	FDisplayClusterViewportManager* ViewportManager = Configuration.GetViewportManagerImpl();
 
-	// Disable ICVFX if cluster node rendering is not used
-	const FString& ClusterNodeId = RenderFrameSettings.CurrentNode.Id;
-	if (!(ViewportManager && StageSettings && !ClusterNodeId.IsEmpty()))
+	// ICVFX is node-based; requires a defined cluster node to create viewports
+	if (Configuration.GetClusterNodeId().IsEmpty() || !ViewportManager || !StageSettings)
 	{
 		// The nDisplay viewport should now always have the name of the cluster node.
 		// When rendering MRQ viewports, the list of viewports is used without cluster node names.
@@ -186,26 +174,12 @@ void FDisplayClusterViewportConfiguration_ICVFX::Update()
 
 void FDisplayClusterViewportConfiguration_ICVFX::PostUpdate()
 {
-	if(FDisplayClusterViewportManager* ViewportManager = Configuration.GetViewportManagerImpl())
-	{
 		// Update visibility for icvfx viewports and cameras
 		ImplUpdateVisibility();
 
-		// Support alpha channel rendering for WarpBlend
-		if (GDisplayClusterEnableAlphaChannelRendering != 0)
-		{
-			for (const TSharedPtr<FDisplayClusterViewport, ESPMode::ThreadSafe>& ViewportIt : ViewportManager->ImplGetCurrentRenderFrameViewports())
-			{
-				if (ViewportIt.IsValid())
-				{
-					// Gain direct access to internal settings of the viewport:
-					FDisplayClusterViewport_RenderSettings& InOutRenderSettings = ViewportIt->GetRenderSettingsImpl();
-
-					InOutRenderSettings.bWarpBlendRenderAlphaChannel = true;
-				}
-			}
-		}
-	}
+	// Note: Alpha-output configuration for rendering (unrelated to ICVFX)
+	//       has been moved to FDisplayClusterRenderFrameSettings::bEnableAlphaOutput.
+	// see FDisplayClusterViewportConfigurationHelpers::UpdateBaseViewportSetting().
 }
 
 void FDisplayClusterViewportConfiguration_ICVFX::ImplUpdateVisibility()
diff --git a/Engine/Plugins/Runtime/nDisplay/Source/DisplayCluster/Private/Render/Viewport/RenderFrame/DisplayClusterRenderFrameSettings.h b/Engine/Plugins/Runtime/nDisplay/Source/DisplayCluster/Private/Render/Viewport/RenderFrame/DisplayClusterRenderFrameSettings.h
--- a/Engine/Plugins/Runtime/nDisplay/Source/DisplayCluster/Private/Render/Viewport/RenderFrame/DisplayClusterRenderFrameSettings.h
+++ b/Engine/Plugins/Runtime/nDisplay/Source/DisplayCluster/Private/Render/Viewport/RenderFrame/DisplayClusterRenderFrameSettings.h
@@ -53,6 +53,10 @@ struct FDisplayClusterRenderFrameSettings
 	// Allow warpblend render
 	uint8 bAllowWarpBlend : 1 = 1;
 
+	// True if alpha channel should be written to all viewport outputs.
+	// Note: Controlled by "DC.EnableAlphaChannelRendering" and "r.PostProcessing.PropagateAlpha".
+	uint8 bEnableAlphaOutput : 1 = 0;
+
 	// Use DC render device for rendering
 	uint8 bUseDisplayClusterRenderDevice : 1 = 1;
 
diff --git a/Engine/Plugins/Runtime/nDisplay/Source/DisplayCluster/Public/DisplayClusterRootActor.h b/Engine/Plugins/Runtime/nDisplay/Source/DisplayCluster/Public/DisplayClusterRootActor.h
--- a/Engine/Plugins/Runtime/nDisplay/Source/DisplayCluster/Public/DisplayClusterRootActor.h
+++ b/Engine/Plugins/Runtime/nDisplay/Source/DisplayCluster/Public/DisplayClusterRootActor.h
@@ -458,6 +458,16 @@ public:
 	*/
 	FDisplayClusterViewport_PreviewSettings GetPreviewSettings(bool bIgnorePreviewSetttingsSource = false) const;
 
+private:
+	/**
+	* Prepares new preview settings.
+	* (Access preview settings as read-only data; does not alter DCRA)
+	*
+	* @param OutPreviewSettings  Filled with the prepared settings.
+	* @return True if preview rendering should proceed, false otherwise.
+	*/
+	bool PreparePreviewRendererSettings(FDisplayClusterViewport_PreviewSettings& OutPreviewSettings) const;
+
 //////////////////////////////////////////////////////////////////////////////////////////////
 // EDITOR RELATED SETTINGS
 //////////////////////////////////////////////////////////////////////////////////////////////
diff --git a/Engine/Plugins/Runtime/nDisplay/Source/DisplayCluster/Public/Render/Viewport/IDisplayClusterViewportConfiguration.h b/Engine/Plugins/Runtime/nDisplay/Source/DisplayCluster/Public/Render/Viewport/IDisplayClusterViewportConfiguration.h
--- a/Engine/Plugins/Runtime/nDisplay/Source/DisplayCluster/Public/Render/Viewport/IDisplayClusterViewportConfiguration.h
+++ b/Engine/Plugins/Runtime/nDisplay/Source/DisplayCluster/Public/Render/Viewport/IDisplayClusterViewportConfiguration.h
@@ -142,6 +142,18 @@ public:
 	/** Returns true if the DCRA preview feature in Standalone/Package builds is used. */
 	virtual bool IsPreviewInGameEnabled() const = 0;
 
+	/**
+	* Returns true if this DCRA is under an exclusive lock.
+	* While exclusively locked:
+	*  - No proxy actors may use this DCRA simultaneously.
+	*  - Default scene preview rendering is disabled.
+	*  - PIE rendering is disabled.
+	*
+	* Intended for scenarios where a system (e.g. Movie Render Queue) requires
+	* sole ownership of the DCRA without interference.
+	*/
+	virtual bool IsExclusiveLocked() const = 0;
+
 	/**
 	 * Checks whether the current cluster node is rendering offscreen (headless).
 	 * 
diff --git a/Engine/Plugins/Runtime/nDisplay/Source/DisplayClusterMoviePipeline/Private/DisplayClusterMoviePipelineSettings.cpp b/Engine/Plugins/Runtime/nDisplay/Source/DisplayClusterMoviePipeline/Private/DisplayClusterMoviePipelineSettings.cpp
--- a/Engine/Plugins/Runtime/nDisplay/Source/DisplayClusterMoviePipeline/Private/DisplayClusterMoviePipelineSettings.cpp
+++ b/Engine/Plugins/Runtime/nDisplay/Source/DisplayClusterMoviePipeline/Private/DisplayClusterMoviePipelineSettings.cpp
@@ -29,8 +29,11 @@ ADisplayClusterRootActor* UDisplayClusterMoviePipelineSettings::GetRootActor(con
 
 bool UDisplayClusterMoviePipelineSettings::GetViewports(const UWorld* InWorld, TArray<FString>& OutViewports, TArray<FIntPoint>& OutViewportResolutions) const
 {
-	// When DCRA not selected we get first suitable actor from scene, but always render all viewports
-	const bool bForceRenderAllViewports = !Configuration.DCRootActor.IsValid() || Configuration.bRenderAllViewports;
+	OutViewports.Reset();
+	OutViewportResolutions.Reset();
+
+	// Treat an empty list of allowed viewports as "allow all"
+	const bool bRenderAllViewports = Configuration.bRenderAllViewports || Configuration.AllowedViewportNamesList.IsEmpty();
 
 	if (ADisplayClusterRootActor* RootActorPtr = GetRootActor(InWorld))
 	{
@@ -50,7 +53,7 @@ bool UDisplayClusterMoviePipelineSettings::GetViewports(const UWorld* InWorld, T
 								if (InConfigurationViewport->bAllowRendering)
 								{
 									const FString& InViewportId = InConfigurationViewportIt.Key;
-									if (bForceRenderAllViewports || Configuration.AllowedViewportNamesList.Find(InViewportId) != INDEX_NONE)
+									if (bRenderAllViewports || Configuration.AllowedViewportNamesList.Find(InViewportId) != INDEX_NONE)
 									{
 										OutViewports.Add(InViewportId);
 
diff --git a/Engine/Plugins/Runtime/nDisplay/Source/DisplayClusterMoviePipeline/Private/DisplayClusterMoviePipelineViewportPass.cpp b/Engine/Plugins/Runtime/nDisplay/Source/DisplayClusterMoviePipeline/Private/DisplayClusterMoviePipelineViewportPass.cpp
--- a/Engine/Plugins/Runtime/nDisplay/Source/DisplayClusterMoviePipeline/Private/DisplayClusterMoviePipelineViewportPass.cpp
+++ b/Engine/Plugins/Runtime/nDisplay/Source/DisplayClusterMoviePipeline/Private/DisplayClusterMoviePipelineViewportPass.cpp
@@ -94,16 +94,16 @@ UDisplayClusterMoviePipelineViewportPassBase::UDisplayClusterMoviePipelineViewpo
 	bAllowCameraAspectRatio = false;
 }
 
+ADisplayClusterRootActor* UDisplayClusterMoviePipelineViewportPassBase::GetRootActor() const
+{
+	return RootActorWeakPtr.Get();
+}
+
 void UDisplayClusterMoviePipelineViewportPassBase::SetupImpl(const MoviePipeline::FMoviePipelineRenderPassInitSettings& InPassInitSettings)
 {
 	if (InitializeDisplayCluster())
 	{
-		check(DCRootActor);
-
-#if WITH_EDITORONLY_DATA
-		// Save current DCRA preview settings, and disable preview rendering until MRQ render is done
-		DCRootActor->bMoviePipelineRenderPass = true;
-#endif
+		check(GetRootActor());
 	}
 
 	Super::SetupImpl(InPassInitSettings);
@@ -111,14 +111,6 @@ void UDisplayClusterMoviePipelineViewportPassBase::SetupImpl(const MoviePipeline
 
 void UDisplayClusterMoviePipelineViewportPassBase::TeardownImpl()
 {
-	// Restore original preview value for DCRA actor
-	if (DCRootActor)
-	{
-#if WITH_EDITORONLY_DATA
-		DCRootActor->bMoviePipelineRenderPass = false;
-#endif
-	}
-
 	ReleaseDisplayCluster();
 
 	Super::TeardownImpl();
@@ -296,6 +288,8 @@ FSceneView* UDisplayClusterMoviePipelineViewportPassBase::GetSceneViewForSampleS
 
 	const int32 InViewportIndex = InOutSampleState.OutputState.CameraIndex;
 	FString ViewportId;
+
+	ADisplayClusterRootActor* DCRootActor = GetRootActor();
 	if (DCRootActor && GetViewportId(InViewportIndex, ViewportId))
 	{
 		UMoviePipelineExecutorShot* CurrentShot = GetPipeline()->GetActiveShotList()[GetPipeline()->GetCurrentShotIndex()];
@@ -356,6 +350,8 @@ void UDisplayClusterMoviePipelineViewportPassBase::PostRendererSubmission(const
 	{
 		const int32 InViewportIndex = InSampleState.OutputState.CameraIndex;
 		FString ViewportId;
+
+		ADisplayClusterRootActor* DCRootActor = GetRootActor();
 		if (DCRootActor && GetViewportId(InViewportIndex, ViewportId))
 		{
 			if (FRenderTarget* RenderTarget = InCanvas.GetRenderTarget())
@@ -436,14 +432,20 @@ void UDisplayClusterMoviePipelineViewportPassBase::GetViewportCutOffset(const FM
 
 void UDisplayClusterMoviePipelineViewportPassBase::ReleaseDisplayCluster()
 {
+	ADisplayClusterRootActor* DCRootActor = GetRootActor();
 	if (DCRootActor)
 	{
 		// Delete all used rendering resources
 		DCRootActor->RemoveViewportManager();
+
+#if WITH_EDITORONLY_DATA
+		// Exit from the exclusive mode for MRQ rendering.
+		DCRootActor->bMoviePipelineRenderPass = false;
+#endif
 	}
 
 	// Reset all runtime values
-	DCRootActor = nullptr;
+	RootActorWeakPtr.Reset();
 
 	bFrameWarpBlend = false;
 
@@ -470,7 +472,7 @@ bool UDisplayClusterMoviePipelineViewportPassBase::InitializeDisplayCluster()
 	{
 		if (UDisplayClusterMoviePipelineSettings* DCMoviePipelineSettings = GetPipeline()->FindOrAddSettingForShot<UDisplayClusterMoviePipelineSettings>(CurrentShot))
 		{
-			DCRootActor = DCMoviePipelineSettings->GetRootActor(CurrentWorld);
+			ADisplayClusterRootActor* DCRootActor = DCMoviePipelineSettings->GetRootActor(CurrentWorld);
 			if (DCRootActor == nullptr)
 			{
 				UE_LOG(LogMovieRenderPipeline, Error, TEXT("Failed to find nDisplay RootActor"));
@@ -487,16 +489,18 @@ bool UDisplayClusterMoviePipelineViewportPassBase::InitializeDisplayCluster()
 
 			if (IDisplayClusterViewportManager* ViewportManager = DCRootActor->GetOrCreateViewportManager())
 			{
-				// Get preview settings from DCRootActor properties
-				FDisplayClusterViewport_PreviewSettings NewPreviewSettings = DCRootActor->GetPreviewSettings(true);
-				NewPreviewSettings.bPreviewEnable = false;
-
-				// Don't use preview settings for MRQ rendering (same as for game)
-				ViewportManager->GetConfiguration().SetPreviewSettings(NewPreviewSettings);
-
 				// Update local node viewports (update\create\delete) and build new render frame
 				if (ViewportManager->GetConfiguration().UpdateConfigurationForViewportsList(EDisplayClusterRenderFrameMode::MRQ_Mono, GetWorld(), DisplayClusterViewports))
 				{
+
+#if WITH_EDITORONLY_DATA
+					// Enter to the exclusive mode for MRQ rendering.
+					DCRootActor->bMoviePipelineRenderPass = true;
+#endif
+
+					// Finally set this DCRA for rendering
+					RootActorWeakPtr = DCRootActor;
+
 					return true;
 				}
 			}
@@ -511,6 +515,7 @@ bool UDisplayClusterMoviePipelineViewportPassBase::InitializeDisplayCluster()
 
 bool UDisplayClusterMoviePipelineViewportPassBase::GetViewportId(int32 InViewportIndex, FString& OutViewportId) const
 {
+	ADisplayClusterRootActor* DCRootActor = GetRootActor();
 	if (DCRootActor && InViewportIndex >= 0 && InViewportIndex < DisplayClusterViewports.Num())
 	{
 		OutViewportId = DisplayClusterViewports[InViewportIndex];
@@ -523,6 +528,7 @@ bool UDisplayClusterMoviePipelineViewportPassBase::GetViewportId(int32 InViewpor
 
 IDisplayClusterViewport* UDisplayClusterMoviePipelineViewportPassBase::GetAndCalculateDisplayClusterViewport(const FMoviePipelineRenderPassMetrics& InSampleState, const FString& InViewportId, const FIntPoint& InDestSize, const uint32 InContextNum, FDisplayClusterViewInfo& OutView)
 {
+	ADisplayClusterRootActor* DCRootActor = GetRootActor();
 	check(DCRootActor);
 	
 	UWorld* CurrentWorld = GetWorld();
@@ -558,12 +564,18 @@ IDisplayClusterViewport* UDisplayClusterMoviePipelineViewportPassBase::GetAndCal
 			RenderSettings.bDisableCustomFrustumFeature = true;
 			RenderSettings.bDisableFrustumOverscanFeature = true;
 
+			// Preserve AlphaOutput for MRQ (based on bAccumulatorIncludesAlpha).
+			RenderSettings.bWarpBlendRenderAlphaChannel = bAccumulatorIncludesAlpha;
+
 			DCViewport->SetRenderSettings(RenderSettings);
 		}
 
 		FDisplayClusterRenderFrame RenderFrame;
 		if (ViewportManager->BeginNewFrame(GameViewportClient->Viewport, RenderFrame))
 		{
+			// Initialize frame for render (update proxies data)
+			ViewportManager->InitializeNewFrame();
+
 			// Obtaining the internal viewpoint for a given viewport with stereo eye offset distance.
 			FMinimalViewInfo ViewInfo;
 			if (!DCViewport->SetupViewPoint(InContextNum, ViewInfo))
diff --git a/Engine/Plugins/Runtime/nDisplay/Source/DisplayClusterMoviePipeline/Public/DisplayClusterMoviePipelineSettings.h b/Engine/Plugins/Runtime/nDisplay/Source/DisplayClusterMoviePipeline/Public/DisplayClusterMoviePipelineSettings.h
--- a/Engine/Plugins/Runtime/nDisplay/Source/DisplayClusterMoviePipeline/Public/DisplayClusterMoviePipelineSettings.h
+++ b/Engine/Plugins/Runtime/nDisplay/Source/DisplayClusterMoviePipeline/Public/DisplayClusterMoviePipelineSettings.h
@@ -26,7 +26,7 @@ struct FDisplayClusterMoviePipelineConfiguration
 	bool bRenderAllViewports = true;
 
 	// Render only viewports from this list.
-	UPROPERTY(EditAnywhere, BlueprintReadWrite, Category = "nDisplay", meta = (DisplayName = "Allowed Viewport Names List"))
+	UPROPERTY(EditAnywhere, BlueprintReadWrite, Category = "nDisplay", meta = (DisplayName = "Allowed Viewport Names List", EditCondition = "!bRenderAllViewports && DCRootActor != nullptr"))
 	TArray<FString> AllowedViewportNamesList;
 };
 
diff --git a/Engine/Plugins/Runtime/nDisplay/Source/DisplayClusterMoviePipeline/Public/DisplayClusterMoviePipelineViewportPass.h b/Engine/Plugins/Runtime/nDisplay/Source/DisplayClusterMoviePipeline/Public/DisplayClusterMoviePipelineViewportPass.h
--- a/Engine/Plugins/Runtime/nDisplay/Source/DisplayClusterMoviePipeline/Public/DisplayClusterMoviePipelineViewportPass.h
+++ b/Engine/Plugins/Runtime/nDisplay/Source/DisplayClusterMoviePipeline/Public/DisplayClusterMoviePipelineViewportPass.h
@@ -66,6 +66,9 @@ protected:
 
 	void GetViewportCutOffset(const FMoviePipelineRenderPassMetrics& InSampleState, FIntPoint& OutOffsetMin, FIntPoint& OutOffsetMax) const;
 
+	/** Return current DCRA. */
+	ADisplayClusterRootActor* GetRootActor() const;
+
 public:
 	// Allow warp blend for this pass
 	UPROPERTY(EditAnywhere, BlueprintReadWrite, Category = "nDisplay")
@@ -75,8 +78,8 @@ private:
 	// Render pass name
 	const FString RenderPassName;
 
-	// DCRA for current shot
-	ADisplayClusterRootActor* DCRootActor = nullptr;
+	// Scene DCRA.
+	TWeakObjectPtr<ADisplayClusterRootActor> RootActorWeakPtr;
 
 	// Names of viewports to render in the current shot
 	TArray<FString> DisplayClusterViewports;
@@ -87,7 +90,7 @@ private:
 	TMap<int32, FDisplayClusterViewInfo> DCPrevViews;
 
 	// Runtime flags
-	bool bFrameWarpBlend = false;
+	uint8 bFrameWarpBlend : 1 = 0;
 };
 
 /**
diff --git a/Engine/Plugins/Runtime/nDisplay/Source/DisplayClusterMoviePipelineEditor/Private/Details/DisplayClusterMoviePipelineEditorSettingsCustomization.cpp b/Engine/Plugins/Runtime/nDisplay/Source/DisplayClusterMoviePipelineEditor/Private/Details/DisplayClusterMoviePipelineEditorSettingsCustomization.cpp
--- a/Engine/Plugins/Runtime/nDisplay/Source/DisplayClusterMoviePipelineEditor/Private/Details/DisplayClusterMoviePipelineEditorSettingsCustomization.cpp
+++ b/Engine/Plugins/Runtime/nDisplay/Source/DisplayClusterMoviePipelineEditor/Private/Details/DisplayClusterMoviePipelineEditorSettingsCustomization.cpp
@@ -26,7 +26,9 @@ void FDisplayClusterMoviePipelineEditorSettingsCustomization::Initialize(const T
 	TSharedPtr<IPropertyHandle> DCRootActorHandle = GET_CHILD_HANDLE(FDisplayClusterMoviePipelineConfiguration, DCRootActor);
 	check(DCRootActorHandle.IsValid());
 
-	DCRootActorHandle->SetOnPropertyValueChanged(FSimpleDelegate::CreateSP(this, &FDisplayClusterMoviePipelineEditorSettingsCustomization::OnRootActorReferencedPropertyValueChanged));
+	DCRootActorHandle->SetOnPropertyValueChanged(
+		FSimpleDelegate::CreateSP(
+			this, &FDisplayClusterMoviePipelineEditorSettingsCustomization::OnRootActorReferencedPropertyValueChanged, DCRootActorHandle));
 
 	TSharedPtr<IPropertyHandle> AllowedViewportNamesListHandle = GET_CHILD_HANDLE(FDisplayClusterMoviePipelineConfiguration, AllowedViewportNamesList);
 	check(AllowedViewportNamesListHandle);
@@ -42,8 +44,35 @@ void FDisplayClusterMoviePipelineEditorSettingsCustomization::OnAllowedViewportN
 		NodeSelection->ResetOptionsList();
 	}
 }
-void FDisplayClusterMoviePipelineEditorSettingsCustomization::OnRootActorReferencedPropertyValueChanged()
+void FDisplayClusterMoviePipelineEditorSettingsCustomization::OnRootActorReferencedPropertyValueChanged(TSharedPtr<IPropertyHandle> PropertyHandle)
 {
+	UObject* CurrentValue = nullptr;
+	if (PropertyHandle.IsValid()
+	&& FPropertyAccess::Success == PropertyHandle->GetValue(CurrentValue)
+	&& !CurrentValue)
+	{
+		/**
+		* Do not reset the viewport list when the DCRA reference is cleared.
+		*
+		* Behavior:
+		* - An empty DCRA reference means: use any available DCRA in the scene.
+		* - Allows configuring specific viewports to render even when the DCRA is undefined.
+		*
+		* Example:
+		* - A DCRA class defines 10 viewports.
+		* - Rendering is distributed across 10 computers.
+		* - The scene may change during rendering.
+		* - Multiple DCRA instances (of the same class) can be used interchangeably.
+		*
+		* Expectation:
+		* - The user creates 10 different configurations (one per viewport).
+		* - The DCRA reference is reset to detach from the scene, so each configuration
+		*   can pick up any available DCRA instance dynamically.
+		*/
+
+		return;
+	}
+
 	// When a referenced property is changed, we have to trigger layout refresh
 	if (NodeSelection.IsValid())
 	{
@@ -80,21 +109,26 @@ void FDisplayClusterMoviePipelineEditorSettingsCustomization::SetChildren(const
 
 	if (NodeSelection.IsValid())
 	{
-		const TAttribute<bool> IsEnabledEditCondition = TAttribute<bool>::Create([this, IsRenderAllViewportsHandle]()
+		const TAttribute<bool> IsEnabledEditCondition = TAttribute<bool>::Create([this, IsRenderAllViewportsHandle, DCRootActorHandle]()
 			{
-				bool bCond1 = false;
-				IsRenderAllViewportsHandle->GetValue(bCond1);
-				return bCond1 == false;
+				bool bIsRenderAllViewportsHandleValue = false;
+				const bool bRenderAllViewports = 
+					IsRenderAllViewportsHandle.IsValid()
+					&& FPropertyAccess::Success == IsRenderAllViewportsHandle->GetValue(bIsRenderAllViewportsHandleValue)
+					&& bIsRenderAllViewportsHandleValue;
+				
+				UObject* DCRootActorHandleValue = nullptr;
+				const bool bHasRootActor =
+					DCRootActorHandle.IsValid()
+					&& FPropertyAccess::Success == DCRootActorHandle->GetValue(DCRootActorHandleValue)
+					&& DCRootActorHandleValue;
+
+				return !bRenderAllViewports && bHasRootActor;
 			});
 
-
 		ArrayHandle->SetInstanceMetaData(FDisplayClusterMoviePipelineEditorNodeSelection::NAME_ElementToolTip, ArrayHandle->GetPropertyDisplayName().ToString());
 
-		const TAttribute<bool> IsArrayHandleEnabledEditCondition = TAttribute<bool>::Create([this]()
-			{
-				return true;
-			});
-		NodeSelection->IsEnabled(IsArrayHandleEnabledEditCondition);
+		NodeSelection->IsEnabled(IsEnabledEditCondition);
 		NodeSelection->CreateArrayBuilder(ArrayHandle.ToSharedRef(), InChildBuilder);
 	}
 }
diff --git a/Engine/Plugins/Runtime/nDisplay/Source/DisplayClusterMoviePipelineEditor/Private/Details/DisplayClusterMoviePipelineEditorSettingsCustomization.h b/Engine/Plugins/Runtime/nDisplay/Source/DisplayClusterMoviePipelineEditor/Private/Details/DisplayClusterMoviePipelineEditorSettingsCustomization.h
--- a/Engine/Plugins/Runtime/nDisplay/Source/DisplayClusterMoviePipelineEditor/Private/Details/DisplayClusterMoviePipelineEditorSettingsCustomization.h
+++ b/Engine/Plugins/Runtime/nDisplay/Source/DisplayClusterMoviePipelineEditor/Private/Details/DisplayClusterMoviePipelineEditorSettingsCustomization.h
@@ -27,7 +27,7 @@ protected:
 	virtual void SetChildren(const TSharedRef<IPropertyHandle>& InPropertyHandle, IDetailChildrenBuilder& InChildBuilder, IPropertyTypeCustomizationUtils& CustomizationUtils) override;
 	//~ FDisplayClusterConfiguratorBaseTypeCustomization overrides end
 
-	void OnRootActorReferencedPropertyValueChanged();
+	void OnRootActorReferencedPropertyValueChanged(TSharedPtr<IPropertyHandle> PropertyHandle);
 	void OnAllowedViewportNamesListHandleChanged();
 private:
 	/** The helper object used to generate the custom array widget to display in place of a string list for the AllowedViewportNamesList property */
