[nDisplay]  Rolled back the recent race condition fix 46586267 for UE-337321

- Another solution is required

#rnx

diff --git a/Engine/Plugins/Runtime/nDisplay/Source/DisplayCluster/Private/Cluster/Failover/DisplayClusterCommDataCache.cpp b/Engine/Plugins/Runtime/nDisplay/Source/DisplayCluster/Private/Cluster/Failover/DisplayClusterCommDataCache.cpp
--- a/Engine/Plugins/Runtime/nDisplay/Source/DisplayCluster/Private/Cluster/Failover/DisplayClusterCommDataCache.cpp
+++ b/Engine/Plugins/Runtime/nDisplay/Source/DisplayCluster/Private/Cluster/Failover/DisplayClusterCommDataCache.cpp
@@ -276,7 +276,7 @@ void FDisplayClusterCommDataCache::OpSaveImpl(const FName& SlotName, Args&... Sl
 void FDisplayClusterCommDataCache::SubscribeToCallbacks()
 {
 	// DCEndFrame is used to invalidate per-frame cache
-	IDisplayCluster::Get().GetCallbacks().OnDisplayClusterEndFrameSafe_AnyThread().AddRaw(this, &FDisplayClusterCommDataCache::ProcessDCEndFrame);
+	IDisplayCluster::Get().GetCallbacks().OnDisplayClusterEndFrame().AddRaw(this, &FDisplayClusterCommDataCache::ProcessDCEndFrame);
 
 	// Set up a post-failure negotiation delegate
 	{
@@ -305,7 +305,7 @@ void FDisplayClusterCommDataCache::SubscribeToCallbacks()
 void FDisplayClusterCommDataCache::UnsubscribeFromCallbacks()
 {
 	// Unsubscribe from DCEndFrame
-	IDisplayCluster::Get().GetCallbacks().OnDisplayClusterEndFrameSafe_AnyThread().RemoveAll(this);
+	IDisplayCluster::Get().GetCallbacks().OnDisplayClusterEndFrame().RemoveAll(this);
 
 	// Clear post-failure negotiation delegate
 	bool bSubscribedToNegotiationSync = false;
diff --git a/Engine/Plugins/Runtime/nDisplay/Source/DisplayCluster/Private/DisplayClusterCallbacks.h b/Engine/Plugins/Runtime/nDisplay/Source/DisplayCluster/Private/DisplayClusterCallbacks.h
--- a/Engine/Plugins/Runtime/nDisplay/Source/DisplayCluster/Private/DisplayClusterCallbacks.h
+++ b/Engine/Plugins/Runtime/nDisplay/Source/DisplayCluster/Private/DisplayClusterCallbacks.h
@@ -32,11 +32,6 @@ public:
 		return DisplayClusterEndFrameEvent;
 	}
 
-	virtual FDisplayClusterEndFrameEvent& OnDisplayClusterEndFrameSafe_AnyThread() override
-	{
-		return DisplayClusterEndFrameEventSafe;
-	}
-
 	virtual FDisplayClusterPreTickEvent& OnDisplayClusterPreTick() override
 	{
 		return DisplayClusterPreTickEvent;
@@ -182,7 +177,6 @@ private:
 	FDisplayClusterEndSessionEvent           DisplayClusterEndSessionEvent;
 	FDisplayClusterStartFrameEvent           DisplayClusterStartFrameEvent;
 	FDisplayClusterEndFrameEvent             DisplayClusterEndFrameEvent;
-	FDisplayClusterEndFrameEvent             DisplayClusterEndFrameEventSafe;
 	FDisplayClusterPreTickEvent              DisplayClusterPreTickEvent;
 	FDisplayClusterTickEvent                 DisplayClusterTickEvent;
 	FDisplayClusterPostTickEvent             DisplayClusterPostTickEvent;
diff --git a/Engine/Plugins/Runtime/nDisplay/Source/DisplayCluster/Private/Game/EngineClasses/Basics/DisplayClusterGameEngine.cpp b/Engine/Plugins/Runtime/nDisplay/Source/DisplayCluster/Private/Game/EngineClasses/Basics/DisplayClusterGameEngine.cpp
--- a/Engine/Plugins/Runtime/nDisplay/Source/DisplayCluster/Private/Game/EngineClasses/Basics/DisplayClusterGameEngine.cpp
+++ b/Engine/Plugins/Runtime/nDisplay/Source/DisplayCluster/Private/Game/EngineClasses/Basics/DisplayClusterGameEngine.cpp
@@ -203,25 +203,6 @@ bool UDisplayClusterGameEngine::InitializeInternals()
 		AudioDeviceManager = nullptr;
 	}
 
-	// Use PreSyncEnd delegate of EndFrame barrier for EndFrameSafe_AnyThread callback
-	const FName& ServiceName = UE::nDisplay::Network::Configuration::ClusterSyncServerName;
-	const TSharedPtr<FDisplayClusterService> ClusterSyncServiceBase = DCAPI->GetPrivateClusterMgr()->GetNodeController()->GetService(ServiceName).Pin();
-	const TSharedPtr<FDisplayClusterClusterSyncService> ClusterSyncService = StaticCastSharedPtr<FDisplayClusterClusterSyncService>(ClusterSyncServiceBase);
-	if (ClusterSyncService)
-	{
-		if (IDisplayClusterBarrier* FrameEndBarrier = ClusterSyncService->GetFrameEndBarrier())
-		{
-			// Trigger EndFrameSafe_AnyThread right from the barrier before any clients leave
-			FrameEndBarrier->GetPreSyncEndDelegate().BindLambda([](FDisplayClusterBarrierPreSyncEndDelegateData&)
-				{
-					if (GDisplayCluster)
-					{
-						GDisplayCluster->GetCallbacks().OnDisplayClusterEndFrameSafe_AnyThread().Broadcast(GFrameCounter);
-					}
-				});
-		}
-	}
-
 	return true;
 }
 
diff --git a/Engine/Plugins/Runtime/nDisplay/Source/DisplayCluster/Private/Network/Service/ClusterSync/DisplayClusterClusterSyncService.h b/Engine/Plugins/Runtime/nDisplay/Source/DisplayCluster/Private/Network/Service/ClusterSync/DisplayClusterClusterSyncService.h
--- a/Engine/Plugins/Runtime/nDisplay/Source/DisplayCluster/Private/Network/Service/ClusterSync/DisplayClusterClusterSyncService.h
+++ b/Engine/Plugins/Runtime/nDisplay/Source/DisplayCluster/Private/Network/Service/ClusterSync/DisplayClusterClusterSyncService.h
@@ -34,19 +34,6 @@ public:
 	virtual void KillSession(const FString& NodeId) override;
 	//~ End IDisplayClusterServer
 
-public:
-
-	/**
-	 * This is a temporary workaround (quick fix) for 5.7.
-	 * The goal is to fix a race condition by providing the ability to handle EndFrame barrier
-	 * sync callback. In 5.8, this code will be re-written in more generic way. The abstract
-	 * FDisplayClusterService will provide a full set of API functions to access any barriers.
-	 */
-	IDisplayClusterBarrier* GetFrameEndBarrier()
-	{
-		return BarrierFrameEnd.Get();
-	}
-
 protected:
 
 	/** Creates session instance for this service */
diff --git a/Engine/Plugins/Runtime/nDisplay/Source/DisplayCluster/Public/IDisplayClusterCallbacks.h b/Engine/Plugins/Runtime/nDisplay/Source/DisplayCluster/Public/IDisplayClusterCallbacks.h
--- a/Engine/Plugins/Runtime/nDisplay/Source/DisplayCluster/Public/IDisplayClusterCallbacks.h
+++ b/Engine/Plugins/Runtime/nDisplay/Source/DisplayCluster/Public/IDisplayClusterCallbacks.h
@@ -55,9 +55,6 @@ public:
 	DECLARE_EVENT_OneParam(IDisplayClusterCallbacks, FDisplayClusterEndFrameEvent, uint64);
 	virtual FDisplayClusterEndFrameEvent& OnDisplayClusterEndFrame() = 0;
 
-	/** Called right after EndFrame barrier sync before any client leaves. Called from a server session thread. **/
-	virtual FDisplayClusterEndFrameEvent& OnDisplayClusterEndFrameSafe_AnyThread() = 0;
-
 	/** Called on DisplayCluster PreTick **/
 	DECLARE_EVENT(IDisplayClusterCallbacks, FDisplayClusterPreTickEvent);
 	virtual FDisplayClusterPreTickEvent& OnDisplayClusterPreTick() = 0;
