[nDisplay] Fixed race condition on GetTimeData call

- Currently, networking data cache is reset on FrameEnd callback which is triggered after FrameEnd barrier. During this step, all internal state variables get reset before starting a new frame. However, it's possible that a secondary node may call GetTimeData from the next frame before p-node resets its internals. This race condition may crash the p-node. To fix such type of race condition, TimeData cache is reset right after FrameStart barrier, so we 100% sure it will be empty before any GetTimeData call.

#rnx

diff --git a/Engine/Plugins/Runtime/nDisplay/Source/DisplayCluster/Private/Cluster/DisplayClusterClusterManager.cpp b/Engine/Plugins/Runtime/nDisplay/Source/DisplayCluster/Private/Cluster/DisplayClusterClusterManager.cpp
--- a/Engine/Plugins/Runtime/nDisplay/Source/DisplayCluster/Private/Cluster/DisplayClusterClusterManager.cpp
+++ b/Engine/Plugins/Runtime/nDisplay/Source/DisplayCluster/Private/Cluster/DisplayClusterClusterManager.cpp
@@ -218,6 +218,15 @@ void FDisplayClusterClusterManager::StartFrame(uint64 FrameNum)
 	// synchronization. As a result, we're 100% sure the clients will always get
 	// actual time data.
 	TimeDataCacheReadySignal->Reset();
+
+	// The following code is a fix/workaround for exactly the same problem described
+	// above. With new failover code, there is one more data item (TimeData cache)
+	// that must be safely reset before any other nodes call GetTimeData().
+	// 
+	// Consider this fix temporary. In the future, a more robust solution should be
+	// implemented. The game thread simulation pipeline should use a single barrier
+	// that synchronizes frame start and time data in one call.
+	GetDataCache()->TempWorkaround_ResetTimeDataCache();
 }
 
 void FDisplayClusterClusterManager::EndFrame(uint64 FrameNum)
diff --git a/Engine/Plugins/Runtime/nDisplay/Source/DisplayCluster/Private/Cluster/Failover/DisplayClusterCommDataCache.cpp b/Engine/Plugins/Runtime/nDisplay/Source/DisplayCluster/Private/Cluster/Failover/DisplayClusterCommDataCache.cpp
--- a/Engine/Plugins/Runtime/nDisplay/Source/DisplayCluster/Private/Cluster/Failover/DisplayClusterCommDataCache.cpp
+++ b/Engine/Plugins/Runtime/nDisplay/Source/DisplayCluster/Private/Cluster/Failover/DisplayClusterCommDataCache.cpp
@@ -336,13 +336,23 @@ void FDisplayClusterCommDataCache::OnPostFailureBarrierSync(FDisplayClusterBarri
 	BuildClusterSyncState(SyncData.RequestData, SyncData.ResponseData);
 }
 
+void FDisplayClusterCommDataCache::TempWorkaround_ResetTimeDataCache()
+{
+	// Manually reset TimeData slot only
+	FScopeLock Lock(&LocalDataCache.GameThreadDataCacheCS);
+	if (FDisplayClusterCommDataCache::FCyclicDataCacheBase* const TimeDataCache = LocalDataCache.GetGTSlotData(UE::nDisplay::Failover::Private::Name_SlotTime))
+	{
+		TimeDataCache->Reset();
+	}
+}
+
 void FDisplayClusterCommDataCache::GenerateNodeSyncState(TArray<uint8>& OutNodeSyncState)
 {
 	OutNodeSyncState.Empty(4096);
 
 	FScopeLock Lock(&LocalDataCache.GameThreadDataCacheCS);
 
-	UE_LOG(LogDisplayClusterFailover, Verbose, TEXT("Post-failure recovery. Local sync state:\n%s"), *LocalDataCache.ToLogString())
+	UE_LOG(LogDisplayClusterFailover, Verbose, TEXT("Post-failure recovery. Local sync state:\n%s"), *LocalDataCache.ToLogString());
 
 	// Serialize local state
 	FMemoryWriter MemoryWriter(OutNodeSyncState);
@@ -355,7 +365,7 @@ void FDisplayClusterCommDataCache::UpdateClusterSyncState(const TArray<uint8>& I
 	FMemoryReader MemoryReader(InClusterSyncState);
 	ClusterDataCache.Serialize(MemoryReader);
 
-	UE_LOG(LogDisplayClusterFailover, Verbose, TEXT("Post-failure recovery. Got cluster sync state:\n%s"), *ClusterDataCache.ToLogString())
+	UE_LOG(LogDisplayClusterFailover, Verbose, TEXT("Post-failure recovery. Got cluster sync state:\n%s"), *ClusterDataCache.ToLogString());
 }
 
 void FDisplayClusterCommDataCache::BuildClusterSyncState(const TMap<FString, TArray<uint8>>& RequestData, TMap<FString, TArray<uint8>>& ResponseData)
diff --git a/Engine/Plugins/Runtime/nDisplay/Source/DisplayCluster/Private/Cluster/Failover/DisplayClusterCommDataCache.h b/Engine/Plugins/Runtime/nDisplay/Source/DisplayCluster/Private/Cluster/Failover/DisplayClusterCommDataCache.h
--- a/Engine/Plugins/Runtime/nDisplay/Source/DisplayCluster/Private/Cluster/Failover/DisplayClusterCommDataCache.h
+++ b/Engine/Plugins/Runtime/nDisplay/Source/DisplayCluster/Private/Cluster/Failover/DisplayClusterCommDataCache.h
@@ -87,6 +87,15 @@ public:
 	/** Incrmements barrier sync counter */
 	TFunction<void(const FName& BarrierName, const FName& SyncCallerName)> OpAdvanceBarrierCounter;
 
+public:
+
+	/**
+	 * [Temporary Workaround]
+	 * 
+	 * Reset time data cache manually
+	 */
+	void TempWorkaround_ResetTimeDataCache();
+
 protected:
 
 	/** Generates and exports the state of synchronization of this node */
