ndisplay Late OCIO: Set the correct source texture view rectangle when applying OCIO conversion. This fixes the black bar issue on the inner viewport when Late OCIO is used.

#rnx

diff --git a/Engine/Plugins/Runtime/nDisplay/Source/DisplayClusterMedia/Private/Input/DisplayClusterMediaInputBase.cpp b/Engine/Plugins/Runtime/nDisplay/Source/DisplayClusterMedia/Private/Input/DisplayClusterMediaInputBase.cpp
--- a/Engine/Plugins/Runtime/nDisplay/Source/DisplayClusterMedia/Private/Input/DisplayClusterMediaInputBase.cpp
+++ b/Engine/Plugins/Runtime/nDisplay/Source/DisplayClusterMedia/Private/Input/DisplayClusterMediaInputBase.cpp
@@ -134,7 +134,7 @@ void FDisplayClusterMediaInputBase::OverrideTextureRegions_RenderThread(FIntRect
 
 	// [Workaround]
 	// Based on the discussion, it looks like the problem is the incoming 2110 textures
-	// may have up to 3 ExtraPixelsThreshold extra pixels.
+	// may have up to CVarTempRivermaxExtraPixelsThreshold extra pixels.
 	// If this is the only difference, we just copy the required subregion.
 	if (bRunningRivermaxMedia && CVarTempRivermaxCropWorkaround.GetValueOnRenderThread())
 	{
@@ -255,7 +255,7 @@ void FDisplayClusterMediaInputBase::ImportMediaDataOCIO_RenderThread(FRHICommand
 			GraphBuilder,
 			FScreenPassViewInfo(),
 			GEngine->GetDefaultWorldFeatureLevel(),
-			FScreenPassTexture(InputTexture),
+			FScreenPassTexture(InputTexture, InSrcRect),
 			FScreenPassRenderTarget(DstTexture, InDstRect, ERenderTargetLoadAction::EClear),
 			OCIOResources,
 			1.0f,
