[nDisplay] Fixed crash while closing editor

- Changed the way how lifetime of PIE session data cache is managed. Now every PIE session will have its own instance of data cache. The lifetime is limited by the owning session's lifetime.

#rnx

diff --git a/Engine/Plugins/Runtime/nDisplay/Source/DisplayCluster/Private/Cluster/Failover/DisplayClusterFailoverNodeCtrlEditor.cpp b/Engine/Plugins/Runtime/nDisplay/Source/DisplayCluster/Private/Cluster/Failover/DisplayClusterFailoverNodeCtrlEditor.cpp
--- a/Engine/Plugins/Runtime/nDisplay/Source/DisplayCluster/Private/Cluster/Failover/DisplayClusterFailoverNodeCtrlEditor.cpp
+++ b/Engine/Plugins/Runtime/nDisplay/Source/DisplayCluster/Private/Cluster/Failover/DisplayClusterFailoverNodeCtrlEditor.cpp
@@ -3,7 +3,6 @@
 #include "Cluster/Failover/DisplayClusterFailoverNodeCtrlEditor.h"
 
 #include "Cluster/Controller/IDisplayClusterClusterNodeController.h"
-#include "Cluster/Failover/DisplayClusterCommDataCache.h"
 
 #include "Misc/DisplayClusterAppExit.h"
 #include "Misc/DisplayClusterLog.h"
@@ -22,8 +21,7 @@ bool FDisplayClusterFailoverNodeCtrlEditor::Initialize(const UDisplayClusterConf
 
 TSharedRef<FDisplayClusterCommDataCache> FDisplayClusterFailoverNodeCtrlEditor::GetDataCache()
 {
-	static TSharedRef<FDisplayClusterCommDataCache> DataBroker = MakeShared<FDisplayClusterCommDataCache>();
-	return DataBroker;
+	return DataCache;
 }
 
 bool FDisplayClusterFailoverNodeCtrlEditor::HandleFailure(const FString& FailedNodeId)
diff --git a/Engine/Plugins/Runtime/nDisplay/Source/DisplayCluster/Private/Cluster/Failover/DisplayClusterFailoverNodeCtrlEditor.h b/Engine/Plugins/Runtime/nDisplay/Source/DisplayCluster/Private/Cluster/Failover/DisplayClusterFailoverNodeCtrlEditor.h
--- a/Engine/Plugins/Runtime/nDisplay/Source/DisplayCluster/Private/Cluster/Failover/DisplayClusterFailoverNodeCtrlEditor.h
+++ b/Engine/Plugins/Runtime/nDisplay/Source/DisplayCluster/Private/Cluster/Failover/DisplayClusterFailoverNodeCtrlEditor.h
@@ -3,6 +3,7 @@
 #pragma once
 
 #include "Cluster/Failover/DisplayClusterFailoverNodeCtrlBase.h"
+#include "Cluster/Failover/DisplayClusterCommDataCache.h"
 
 
 /**
@@ -17,6 +18,7 @@ public:
 
 	FDisplayClusterFailoverNodeCtrlEditor(TSharedRef<IDisplayClusterClusterNodeController>& InNodeController)
 		: FDisplayClusterFailoverNodeCtrlBase(InNodeController)
+		, DataCache(MakeShared<FDisplayClusterCommDataCache>())
 	{ }
 
 public:
@@ -75,4 +77,9 @@ public:
 	virtual EDisplayClusterCommResult PostFailureNegotiate(TArray<uint8>& InOutRecoveryData) override;
 	virtual EDisplayClusterCommResult RequestNodeDrop(const FString& NodeId, uint8 DropReason) override;
 	//~ End IDisplayClusterProtocolInternalComm
+
+private:
+
+	/** Stub cache to respect API */
+	TSharedRef<FDisplayClusterCommDataCache> DataCache;
 };
