[nDisplay] No collision on DCRA stage isosphere component

- Updated physics configuration to support line traces, but forbid any rigid body interactions
- Fixed a mesh section issue

#rnx

diff --git a/Engine/Plugins/Runtime/nDisplay/Source/DisplayCluster/Private/Game/EngineClasses/Scene/DisplayClusterStageIsosphereComponent.cpp b/Engine/Plugins/Runtime/nDisplay/Source/DisplayCluster/Private/Game/EngineClasses/Scene/DisplayClusterStageIsosphereComponent.cpp
--- a/Engine/Plugins/Runtime/nDisplay/Source/DisplayCluster/Private/Game/EngineClasses/Scene/DisplayClusterStageIsosphereComponent.cpp
+++ b/Engine/Plugins/Runtime/nDisplay/Source/DisplayCluster/Private/Game/EngineClasses/Scene/DisplayClusterStageIsosphereComponent.cpp
@@ -15,7 +15,10 @@ UDisplayClusterStageIsosphereComponent::UDisplayClusterStageIsosphereComponent(c
 	IsosphereMesh = IsosphereMeshFinder.Object;
 
 	BodyInstance.SetCollisionEnabled(ECollisionEnabled::QueryOnly);
-	
+	BodyInstance.SetObjectType(ECC_WorldDynamic);
+	BodyInstance.SetResponseToAllChannels(ECR_Ignore);
+	BodyInstance.SetResponseToChannel(ECC_Visibility, ECR_Block);
+
 	SetHiddenInGame(true);
 	SetVisibility(false);
 }
@@ -77,7 +80,7 @@ void UDisplayClusterStageIsosphereComponent::ResetIsosphere()
 					Vertex.Normal = Normals[Index];
 				}
 
-				SetProcMeshSection(0, *Section);
+				SetProcMeshSection(SectionIndex, *Section);
 			}
 		}
 
