[nDisplay] Fixed crash with uninitialized projection policy

#rnx

diff --git a/Engine/Plugins/Runtime/nDisplay/Source/DisplayCluster/Private/Render/Viewport/DisplayClusterViewport_ViewPoint.cpp b/Engine/Plugins/Runtime/nDisplay/Source/DisplayCluster/Private/Render/Viewport/DisplayClusterViewport_ViewPoint.cpp
--- a/Engine/Plugins/Runtime/nDisplay/Source/DisplayCluster/Private/Render/Viewport/DisplayClusterViewport_ViewPoint.cpp
+++ b/Engine/Plugins/Runtime/nDisplay/Source/DisplayCluster/Private/Render/Viewport/DisplayClusterViewport_ViewPoint.cpp
@@ -231,6 +231,11 @@ bool FDisplayClusterViewport::PreCalculateViewData(const uint32 InContextNum, co
 		return false;
 	}
 
+	if (!ProjectionPolicy)
+	{
+		return false;
+	}
+
 	// Origin2World transform
 	if (USceneComponent* const OriginComponent = ProjectionPolicy->GetOriginComponent())
 	{
