[nDisplay] Fixed crash while importing an nDisplay config in the nD configurator

#rnx

diff --git a/Engine/Plugins/Runtime/nDisplay/Source/DisplayClusterConfigurator/Private/DisplayClusterConfiguratorBlueprintEditor.cpp b/Engine/Plugins/Runtime/nDisplay/Source/DisplayClusterConfigurator/Private/DisplayClusterConfiguratorBlueprintEditor.cpp
--- a/Engine/Plugins/Runtime/nDisplay/Source/DisplayClusterConfigurator/Private/DisplayClusterConfiguratorBlueprintEditor.cpp
+++ b/Engine/Plugins/Runtime/nDisplay/Source/DisplayClusterConfigurator/Private/DisplayClusterConfiguratorBlueprintEditor.cpp
@@ -49,6 +49,7 @@
 #include "ISCSEditorUICustomization.h"
 #include "Kismet2/BlueprintEditorUtils.h"
 #include "Kismet2/DebuggerCommands.h"
+#include "Kismet2/KismetEditorUtilities.h"
 #include "Misc/DisplayClusterHelpers.h"
 #include "MPCDI/DisplayClusterConfiguratorMPCDIImporter.h"
 #include "SBlueprintEditorToolbar.h"
@@ -658,12 +659,25 @@ bool FDisplayClusterConfiguratorBlueprintEditor::LoadFromFile(const FString& Fil
 	check(LoadedBlueprint.IsValid());
 	if (ADisplayClusterRootActor* NewRootActor = FDisplayClusterConfiguratorUtils::GenerateRootActorFromConfigFile(FilePath))
 	{
-		LoadedBlueprint->SetConfigData(NewRootActor->GetConfigData(), true);
-		UDisplayClusterConfigurationData* ConfigData = LoadedBlueprint->GetConfig();
+		UDisplayClusterBlueprint* const LoadedBlueprintPtr = LoadedBlueprint.Get();
+		if (!IsValid(LoadedBlueprintPtr))
+		{
+			return false;
+		}
+
+		LoadedBlueprintPtr->SetConfigData(NewRootActor->GetConfigData(), true);
+		UDisplayClusterConfigurationData* ConfigData = LoadedBlueprintPtr->GetConfig();
 		check(ConfigData);
 		ConfigData->ImportedPath = FilePath;
-		
-		FDisplayClusterConfiguratorUtils::AddRootActorComponentsToBlueprint(LoadedBlueprint.Get(), NewRootActor);
+
+		FDisplayClusterConfiguratorUtils::AddRootActorComponentsToBlueprint(LoadedBlueprintPtr, NewRootActor, false);
+		ClusterChanged();
+		FKismetEditorUtilities::CompileBlueprint(LoadedBlueprintPtr);
+		UpdatePreviewActor(LoadedBlueprintPtr, true);
+		UpdateSubobjectPreview(true);
+		SubobjectEditor->UpdateTree();
+		LoadedBlueprintPtr->MarkPackageDirty();
+
 		return true;
 	}
 	
