[nDisplay] Improved MRQ settings for using the DCRA class as a reference for multiple scenes.

Added new hidden UPROPERTY DCRootActorClass;

Changes in the UDisplayClusterMoviePipelineSettings::GetRootActor() function:
- First searches for a DCRA instance by reference in the scene using DCRootActor
- Then searches for a DCRA instance by class in the scene using DCRootActorClass
- Third, if DCRootActor and DCRootActorClass are not specified, searches for any DCRA (of any class) in the scene
- Finally, returns nullptr

diff --git a/Engine/Plugins/Runtime/nDisplay/Source/DisplayClusterMoviePipeline/Private/DisplayClusterMoviePipelineSettings.cpp b/Engine/Plugins/Runtime/nDisplay/Source/DisplayClusterMoviePipeline/Private/DisplayClusterMoviePipelineSettings.cpp
--- a/Engine/Plugins/Runtime/nDisplay/Source/DisplayClusterMoviePipeline/Private/DisplayClusterMoviePipelineSettings.cpp
+++ b/Engine/Plugins/Runtime/nDisplay/Source/DisplayClusterMoviePipeline/Private/DisplayClusterMoviePipelineSettings.cpp
@@ -8,22 +8,85 @@
 ///////////////////////////////////////////////////////////////////////////////////////////
 // UDisplayClusterMoviePipelineSettings
 ///////////////////////////////////////////////////////////////////////////////////////////
+#if WITH_EDITOR
+void UDisplayClusterMoviePipelineSettings::PostEditChangeProperty(FPropertyChangedEvent& PropertyChangedEvent)
+{
+	static const FName PropName_Configuration = GET_MEMBER_NAME_CHECKED(UDisplayClusterMoviePipelineSettings, Configuration);
+	static const FName PropMemberName_RootActor = GET_MEMBER_NAME_CHECKED(FDisplayClusterMoviePipelineConfiguration, DCRootActor);
+
+	if (PropertyChangedEvent.Property
+		&& PropertyChangedEvent.Property->GetFName() == PropMemberName_RootActor
+		&& PropertyChangedEvent.MemberProperty
+		&& PropertyChangedEvent.MemberProperty->GetFName() == PropName_Configuration)
+	{
+		if (Configuration.DCRootActor.IsValid())
+		{
+			Configuration.DCRootActorClass = Configuration.DCRootActor->GetClass();      // stores a soft ref to the class
+		}
+		else
+		{
+			Configuration.DCRootActorClass = nullptr;
+		}
+	}
+
+	Super::PostEditChangeProperty(PropertyChangedEvent);
+}
+#endif /** WITH_EDITOR */
+
 ADisplayClusterRootActor* UDisplayClusterMoviePipelineSettings::GetRootActor(const UWorld* InWorld) const
 {
-	if (InWorld)
+	if (!InWorld)
+	{
+		return nullptr;
+	}
+
+	const TSoftObjectPtr<ADisplayClusterRootActor>& RootActorRef = Configuration.DCRootActor;
+	const TSoftClassPtr<ADisplayClusterRootActor>& RootActorClassRef = Configuration.DCRootActorClass;
+
+	UClass* TargetClass = nullptr;
+
+	if (RootActorRef.IsValid())
 	{
-		for (const TWeakObjectPtr<ADisplayClusterRootActor> RootActorRef : TActorRange<ADisplayClusterRootActor>(InWorld))
+		// If DCRootActor points to a live actor in this world, just return it.
+		if (RootActorRef->GetWorld() == InWorld)
 		{
-			if (ADisplayClusterRootActor* RootActorPtr = RootActorRef.Get())
+			return RootActorRef.Get();
+		}
+
+		// Try DCRootActor if still valid (even if it's in another world)
+		TargetClass = RootActorRef->GetClass();
+	}
+	else if (!RootActorClassRef.IsNull())
+	{
+		// Optionally load it synchronously if it's not in memory
+		TargetClass = RootActorClassRef.LoadSynchronous();
+	}
+
+	// If we have a class, return any actor of that class from this world
+	if (TargetClass)
+	{
+		for (const TWeakObjectPtr<ADisplayClusterRootActor> RootActorIt : TActorRange<ADisplayClusterRootActor>(InWorld))
+		{
+			if (RootActorIt.IsValid() && RootActorIt->IsA(TargetClass))
 			{
-				if (!Configuration.DCRootActor.IsValid() || RootActorPtr->GetFName() == Configuration.DCRootActor->GetFName())
-				{
-					return RootActorPtr;
-				}
+				return RootActorIt.Get();
 			}
 		}
+
+		// There is no actor in the world with this class, returning nullptr
+		return nullptr;
+	}
+
+	// DCRootActor==None, just return the first DCRA of any class.
+	for (const TWeakObjectPtr<ADisplayClusterRootActor> RootActorIt : TActorRange<ADisplayClusterRootActor>(InWorld))
+	{
+		if (RootActorIt.IsValid())
+		{
+			return RootActorIt.Get();
+		}
 	}
 
+	// There is no instance of DCRA (of any class) in the world, so nullptr is returned.
 	return nullptr;
 }
 
diff --git a/Engine/Plugins/Runtime/nDisplay/Source/DisplayClusterMoviePipeline/Public/DisplayClusterMoviePipelineSettings.h b/Engine/Plugins/Runtime/nDisplay/Source/DisplayClusterMoviePipeline/Public/DisplayClusterMoviePipelineSettings.h
--- a/Engine/Plugins/Runtime/nDisplay/Source/DisplayClusterMoviePipeline/Public/DisplayClusterMoviePipelineSettings.h
+++ b/Engine/Plugins/Runtime/nDisplay/Source/DisplayClusterMoviePipeline/Public/DisplayClusterMoviePipelineSettings.h
@@ -17,6 +17,11 @@ struct FDisplayClusterMoviePipelineConfiguration
 	UPROPERTY(EditAnywhere, BlueprintReadWrite, Category = "nDisplay", meta = (DisplayName = "DC Root Actor"))
 	TSoftObjectPtr<ADisplayClusterRootActor> DCRootActor;
 
+	// Container with DCRA instance class:
+	// for using multiple instances of DCRA of the same class in multiple scenes.
+	UPROPERTY()
+	TSoftClassPtr<ADisplayClusterRootActor> DCRootActorClass;
+
 	// Render with nDisplay viewport resolutions.
 	UPROPERTY(EditAnywhere, BlueprintReadWrite, Category = "nDisplay")
 	bool bUseViewportResolutions = true;
@@ -53,10 +58,16 @@ public:
 	* 
 	 * return true if is whether it populated 1 or more viewports in the given array.
 	 */
-	bool GetViewports(const UWorld* InWorld, TArray<FString>& OutViewportss, TArray<FIntPoint>& OutViewportResolutions) const;
+	bool GetViewports(const UWorld* InWorld, TArray<FString>& OutViewports, TArray<FIntPoint>& OutViewportResolutions) const;
 
 #if WITH_EDITOR
 	virtual FText GetDisplayText() const override { return NSLOCTEXT("MovieRenderPipeline", "DisplayClusterSettingDisplayName", "nDisplay"); }
+
+	// Begin UObject
+	virtual void PostEditChangeProperty(FPropertyChangedEvent& PropertyChangedEvent) override;
+	// End UObject
+public:
+
 #endif
 
 protected:
