nDisplay: Fixed rare fatal crash caused by texture resource lifetime exceeding texture name pointer lifetime.

diff --git a/Engine/Plugins/Runtime/nDisplay/Source/DisplayClusterLightCardEditorShaders/Private/DisplayClusterMeshProjectionRenderer.cpp b/Engine/Plugins/Runtime/nDisplay/Source/DisplayClusterLightCardEditorShaders/Private/DisplayClusterMeshProjectionRenderer.cpp
--- a/Engine/Plugins/Runtime/nDisplay/Source/DisplayClusterLightCardEditorShaders/Private/DisplayClusterMeshProjectionRenderer.cpp
+++ b/Engine/Plugins/Runtime/nDisplay/Source/DisplayClusterLightCardEditorShaders/Private/DisplayClusterMeshProjectionRenderer.cpp
@@ -1297,13 +1297,10 @@ struct FTextureUAVWithIntermediate
 	FRDGTextureRef IntermediateTexture = nullptr;
 	FRDGTextureUAVRef IntermediateUAV = nullptr;
 
-	void CreateTexture(FRDGBuilder& GraphBuilder, const FRDGTextureDesc& TextureDesc, const FString& TextureName)
+	void CreateTexture(FRDGBuilder& GraphBuilder, const FRDGTextureDesc& TextureDesc, const TCHAR* PrimaryTextureName, const TCHAR* IntermediateTextureName)
 	{
-		const FString PrimaryTextureName = *FString::Format(TEXT("{0}.{1}"), { TextureName, TEXT("Primary") });
-		PrimaryTexture = GraphBuilder.CreateTexture(TextureDesc, *PrimaryTextureName);
-		
-		const FString IntermediateTextureName = *FString::Format(TEXT("{0}.{1}"), { TextureName, TEXT("Primary") });
-		IntermediateTexture = GraphBuilder.CreateTexture(TextureDesc, *IntermediateTextureName);
+		PrimaryTexture = GraphBuilder.CreateTexture(TextureDesc, PrimaryTextureName);
+		IntermediateTexture = GraphBuilder.CreateTexture(TextureDesc, IntermediateTextureName);
 	}
 	
 	void CreateUAV(FRDGBuilder& GraphBuilder)
@@ -1313,6 +1310,9 @@ struct FTextureUAVWithIntermediate
 	}
 };
 
+#define TEXTURE_NAME(TextureName) TEXT(#TextureName)
+#define TEXTURE_NAME_PAIR(TextureName) TEXTURE_NAME(DisplayClusterMeshProjection.TextureName.Primary), TEXTURE_NAME(DisplayClusterMeshProjection.TextureName.Intermediate)
+
 template<EMeshProjectionFilterType FilterType>
 void AddSeparableFilterPass(FRDGBuilder& GraphBuilder,
 	FGlobalShaderMap* GlobalShaderMap,
@@ -1417,21 +1417,21 @@ void FDisplayClusterMeshProjectionRenderer::AddNormalsFilterPass(FRDGBuilder& Gr
 
 	const FRDGTextureDesc RWColorDesc(FRDGTextureDesc::Create2D(SceneColor->Desc.Extent, PF_B8G8R8A8, FClearValueBinding::Black, TexCreate_ShaderResource | TexCreate_UAV));
 	FTextureUAVWithIntermediate RWColor;
-	RWColor.CreateTexture(GraphBuilder, RWColorDesc, TEXT("DisplayClusterMeshProjection.NormalsRWTexture"));
+	RWColor.CreateTexture(GraphBuilder, RWColorDesc, TEXTURE_NAME_PAIR(NormalsRWTexture));
 	RWColor.CreateUAV(GraphBuilder);
 
 	const FRDGTextureDesc RWDepthDesc = FRDGTextureDesc::Create2D(SceneDepth->Desc.Extent, PF_G16, FClearValueBinding::White, TexCreate_UAV | TexCreate_ShaderResource);
 	FTextureUAVWithIntermediate RWDepth;
-	RWDepth.CreateTexture(GraphBuilder, RWDepthDesc, TEXT("DisplayClusterMeshProjection.RWDepth"));
+	RWDepth.CreateTexture(GraphBuilder, RWDepthDesc, TEXTURE_NAME_PAIR(NormalsRWDepth));
 	RWDepth.CreateUAV(GraphBuilder);
 
 	const FRDGTextureDesc RWStencilDesc = FRDGTextureDesc::Create2D(SceneDepth->Desc.Extent, PF_R32_UINT, FClearValueBinding::Black, TexCreate_UAV | TexCreate_ShaderResource);
 	FTextureUAVWithIntermediate RWStencil;
-	RWStencil.CreateTexture(GraphBuilder, RWStencilDesc, TEXT("DisplayClusterMeshProjection.RWStencil"));
+	RWStencil.CreateTexture(GraphBuilder, RWStencilDesc, TEXTURE_NAME_PAIR(NormalsRWStencil));
 	RWStencil.CreateUAV(GraphBuilder);
 
 	FTextureUAVWithIntermediate RWBlurStencil;
-	RWBlurStencil.CreateTexture(GraphBuilder, RWStencilDesc, TEXT("DisplayClusterMeshProjection.RWBlurStencil"));
+	RWBlurStencil.CreateTexture(GraphBuilder, RWStencilDesc, TEXTURE_NAME_PAIR(NormalsRWBlurStencil));
 	RWBlurStencil.CreateUAV(GraphBuilder);
 
 	FGlobalShaderMap* GlobalShaderMap = GetGlobalShaderMap(View->FeatureLevel);
