[nDisplay] Fixed crash while closing nD node after "exit" console command
- It's a hotfix that only fixes the crash with minimum changes and low risks
- A more substantial fix will go into 5.8

#rnx

diff --git a/Engine/Plugins/Runtime/nDisplay/Source/DisplayCluster/Private/Cluster/Failover/DisplayClusterCommDataCache.cpp b/Engine/Plugins/Runtime/nDisplay/Source/DisplayCluster/Private/Cluster/Failover/DisplayClusterCommDataCache.cpp
--- a/Engine/Plugins/Runtime/nDisplay/Source/DisplayCluster/Private/Cluster/Failover/DisplayClusterCommDataCache.cpp
+++ b/Engine/Plugins/Runtime/nDisplay/Source/DisplayCluster/Private/Cluster/Failover/DisplayClusterCommDataCache.cpp
@@ -275,8 +275,16 @@ void FDisplayClusterCommDataCache::OpSaveImpl(const FName& SlotName, Args&... Sl
 
 void FDisplayClusterCommDataCache::SubscribeToCallbacks()
 {
-	// DCEndFrame is used to invalidate per-frame cache
-	IDisplayCluster::Get().GetCallbacks().OnDisplayClusterEndFrame().AddRaw(this, &FDisplayClusterCommDataCache::ProcessDCEndFrame);
+	if (!GDisplayCluster)
+	{
+		return;
+	}
+
+	if (IDisplayCluster::IsAvailable())
+	{
+		// DCEndFrame is used to invalidate per-frame cache
+		IDisplayCluster::Get().GetCallbacks().OnDisplayClusterEndFrame().AddRaw(this, &FDisplayClusterCommDataCache::ProcessDCEndFrame);
+	}
 
 	// Set up a post-failure negotiation delegate
 	{
@@ -304,8 +312,16 @@ void FDisplayClusterCommDataCache::SubscribeToCallbacks()
 
 void FDisplayClusterCommDataCache::UnsubscribeFromCallbacks()
 {
-	// Unsubscribe from DCEndFrame
-	IDisplayCluster::Get().GetCallbacks().OnDisplayClusterEndFrame().RemoveAll(this);
+	if (!GDisplayCluster)
+	{
+		return;
+	}
+
+	if (IDisplayCluster::IsAvailable())
+	{
+		// Unsubscribe from DCEndFrame
+		IDisplayCluster::Get().GetCallbacks().OnDisplayClusterEndFrame().RemoveAll(this);
+	}
 
 	// Clear post-failure negotiation delegate
 	bool bSubscribedToNegotiationSync = false;
