[nDisplay] Removed the unnecessary workaround for ICVFX camera media sub-objects
- ICVFX camera instances no longer own the media subobjects that actually belong to the archetype.
- Fixed crash when TakeRecorder is duplicating an NDC instance

#rnx

diff --git a/Engine/Plugins/Runtime/nDisplay/Source/DisplayCluster/Private/Game/EngineClasses/Scene/DisplayClusterICVFXCameraComponent.cpp b/Engine/Plugins/Runtime/nDisplay/Source/DisplayCluster/Private/Game/EngineClasses/Scene/DisplayClusterICVFXCameraComponent.cpp
--- a/Engine/Plugins/Runtime/nDisplay/Source/DisplayCluster/Private/Game/EngineClasses/Scene/DisplayClusterICVFXCameraComponent.cpp
+++ b/Engine/Plugins/Runtime/nDisplay/Source/DisplayCluster/Private/Game/EngineClasses/Scene/DisplayClusterICVFXCameraComponent.cpp
@@ -138,15 +138,6 @@ void UDisplayClusterICVFXCameraComponent::PostLoad()
 		}
 	}
 	PRAGMA_ENABLE_DEPRECATION_WARNINGS;
-
-	// Propagate Media settings from the Archetype. Works around instanced property limitations.
-	if (!IsTemplate())
-	{
-		if (const UDisplayClusterICVFXCameraComponent* Archetype = Cast<UDisplayClusterICVFXCameraComponent>(GetArchetype()))
-		{
-			CameraSettings.RenderSettings.Media = Archetype->CameraSettings.RenderSettings.Media;
-		}
-	}
 }
 
 void UDisplayClusterICVFXCameraComponent::PostApplyToComponent()
