[nDisplay] Fixed transition states management for nDisplay media
- Fixed transition states management
- Removed unneccessary check

#rnx

diff --git a/Engine/Plugins/Runtime/nDisplay/Source/DisplayCluster/Private/Cluster/NetAPI/DisplayClusterNetApiFacade.cpp b/Engine/Plugins/Runtime/nDisplay/Source/DisplayCluster/Private/Cluster/NetAPI/DisplayClusterNetApiFacade.cpp
--- a/Engine/Plugins/Runtime/nDisplay/Source/DisplayCluster/Private/Cluster/NetAPI/DisplayClusterNetApiFacade.cpp
+++ b/Engine/Plugins/Runtime/nDisplay/Source/DisplayCluster/Private/Cluster/NetAPI/DisplayClusterNetApiFacade.cpp
@@ -126,7 +126,6 @@ public:
 	/** Failover transaction for SynchronizeOnBarrier */
 	virtual EDisplayClusterCommResult SynchronizeOnBarrier() override
 	{
-		check(IsInRHIThread());
 		return static_cast<IDisplayClusterProtocolRenderSync&>(GetFailoverController().Get()).SynchronizeOnBarrier();
 	}
 
diff --git a/Engine/Plugins/Runtime/nDisplay/Source/DisplayCluster/Private/Cluster/NetAPI/DisplayClusterNetApiFacade.h b/Engine/Plugins/Runtime/nDisplay/Source/DisplayCluster/Private/Cluster/NetAPI/DisplayClusterNetApiFacade.h
--- a/Engine/Plugins/Runtime/nDisplay/Source/DisplayCluster/Private/Cluster/NetAPI/DisplayClusterNetApiFacade.h
+++ b/Engine/Plugins/Runtime/nDisplay/Source/DisplayCluster/Private/Cluster/NetAPI/DisplayClusterNetApiFacade.h
@@ -37,7 +37,6 @@ public:
 	/** Access to the render sync API (RHI thread only) */
 	TSharedRef<IDisplayClusterProtocolRenderSync> GetRenderSyncAPI()
 	{
-		check(IsInRHIThread());
 		return RenderSyncAPI;
 	}
 
diff --git a/Engine/Plugins/Runtime/nDisplay/Source/DisplayClusterMedia/Private/Synchronization/LatencyQueue/DisplayClusterFrameQueueItem.cpp b/Engine/Plugins/Runtime/nDisplay/Source/DisplayClusterMedia/Private/Synchronization/LatencyQueue/DisplayClusterFrameQueueItem.cpp
--- a/Engine/Plugins/Runtime/nDisplay/Source/DisplayClusterMedia/Private/Synchronization/LatencyQueue/DisplayClusterFrameQueueItem.cpp
+++ b/Engine/Plugins/Runtime/nDisplay/Source/DisplayClusterMedia/Private/Synchronization/LatencyQueue/DisplayClusterFrameQueueItem.cpp
@@ -13,6 +13,7 @@
 #include "Render/Viewport/IDisplayClusterViewportManagerProxy.h"
 
 #include "RHICommandList.h"
+#include "RHIUtilities.h"
 
 
 FDisplayClusterFrameQueueItem::FDisplayClusterFrameQueueItem(const FDisplayClusterFrameQueueItem& Other)
@@ -24,7 +25,7 @@ FDisplayClusterFrameQueueItem::FDisplayClusterFrameQueueItem(const FDisplayClust
 		{
 			NewItem.Texture = CreateTexture(View.Value.Texture);
 			FRHICommandListImmediate& RHICmdList = GetImmediateCommandList_ForRenderCommand();
-			RHICmdList.CopyTexture(View.Value.Texture.GetReference(), NewItem.Texture, FRHICopyTextureInfo());
+			TransitionAndCopyTexture(RHICmdList, View.Value.Texture.GetReference(), NewItem.Texture, FRHICopyTextureInfo());
 		}
 	}
 }
@@ -52,7 +53,7 @@ void FDisplayClusterFrameQueueItem::SaveView(FRHICommandListImmediate& RHICmdLis
 	// Copy texture data
 	if (View->Texture.IsValid())
 	{
-		RHICmdList.CopyTexture(Texture, View->Texture.GetReference(), FRHICopyTextureInfo());
+		TransitionAndCopyTexture(RHICmdList, Texture, View->Texture.GetReference(), FRHICopyTextureInfo());
 	}
 }
 
@@ -78,7 +79,7 @@ void FDisplayClusterFrameQueueItem::LoadView(FRHICommandListImmediate& RHICmdLis
 		}
 		else
 		{
-			RHICmdList.CopyTexture(View->Texture.GetReference(), Texture, FRHICopyTextureInfo());
+			TransitionAndCopyTexture(RHICmdList, View->Texture.GetReference(), Texture, FRHICopyTextureInfo());
 		}
 	}
 }
@@ -134,20 +135,20 @@ FTextureRHIRef FDisplayClusterFrameQueueItem::CreateTexture(FRHITexture* Referen
 		const int32 SizeY = ReferenceTexture->GetDesc().Extent.Y;
 		const EPixelFormat Format = ReferenceTexture->GetFormat();
 
+		// Preserve sRGB flag from the reference
+		const ETextureCreateFlags RefFlags = ReferenceTexture->GetFlags();
+		const bool bIsSRGB = EnumHasAnyFlags(RefFlags, ETextureCreateFlags::SRGB);
+		const ETextureCreateFlags NewSRGBFlags = (bIsSRGB ? ETextureCreateFlags::SRGB : ETextureCreateFlags::None);
+		const ETextureCreateFlags NewFlags = NewSRGBFlags | ETextureCreateFlags::ShaderResource | ETextureCreateFlags::MultiGPUGraphIgnore;
+
 		// Prepare description
 		FRHITextureCreateDesc Desc =
 			FRHITextureCreateDesc::Create2D(TEXT("DisplayClusterFrameQueueCacheTexture"), SizeX, SizeY, Format)
 			.SetClearValue(FClearValueBinding::Black)
 			.SetNumMips(1)
-			.SetFlags(ETextureCreateFlags::MultiGPUGraphIgnore)
+			.SetFlags(NewFlags)
 			.SetInitialState(ERHIAccess::SRVMask);
 
-		// Leave original flags, but make sure it's ResolveTargetable but not RenderTargetable
-		ETextureCreateFlags Flags = ReferenceTexture->GetFlags();
-		Flags &= ~ETextureCreateFlags::RenderTargetable;
-		Flags |= ETextureCreateFlags::ResolveTargetable;
-		Desc.SetFlags(Flags);
-
 		// Create texture
 		return RHICreateTexture(Desc);
 	}
diff --git a/Engine/Plugins/Runtime/nDisplay/Source/SharedMemoryMedia/Private/SharedMemoryMediaPlayer.cpp b/Engine/Plugins/Runtime/nDisplay/Source/SharedMemoryMedia/Private/SharedMemoryMediaPlayer.cpp
--- a/Engine/Plugins/Runtime/nDisplay/Source/SharedMemoryMedia/Private/SharedMemoryMediaPlayer.cpp
+++ b/Engine/Plugins/Runtime/nDisplay/Source/SharedMemoryMedia/Private/SharedMemoryMediaPlayer.cpp
@@ -7,6 +7,8 @@
 #include "IMediaTextureSampleConverter.h"
 #include "Misc/Timespan.h"
 #include "Tasks/Task.h"
+
+#include "RHIUtilities.h"
 #include "TextureResource.h"
 
 #include "SharedMemoryMediaOutput.h"
@@ -835,7 +837,7 @@ void FSharedMemoryMediaPlayer::JustInTimeSampleRender(FRHICommandListImmediate&
 		RHI_BREADCRUMB_EVENT_STAT(RHICmdList, SharedMemoryMedia_CopyToSampleCommon, "SharedMemoryMedia_CopyToSampleCommon");
 		SCOPED_GPU_STAT(RHICmdList, SharedMemoryMedia_CopyToSampleCommon);
 
-		RHICmdList.CopyTexture(SharedCrossGpuTextures[SharedMemoryIdx], SampleCommonTexture->GetResource()->GetTextureRHI(), CopyInfo);
+		TransitionAndCopyTexture(RHICmdList, SharedCrossGpuTextures[SharedMemoryIdx], SampleCommonTexture->GetResource()->GetTextureRHI(), CopyInfo);
 	}
 	else
 	{
diff --git a/Engine/Source/Runtime/MediaAssets/Private/Misc/MediaTextureResource.cpp b/Engine/Source/Runtime/MediaAssets/Private/Misc/MediaTextureResource.cpp
--- a/Engine/Source/Runtime/MediaAssets/Private/Misc/MediaTextureResource.cpp
+++ b/Engine/Source/Runtime/MediaAssets/Private/Misc/MediaTextureResource.cpp
@@ -23,6 +23,7 @@
 #include "StaticBoundShaderState.h"
 #include "RenderUtils.h"
 #include "RHIStaticStates.h"
+#include "RHIUtilities.h"
 #include "GenerateMips.h"
 #include "ProfilingDebugging/CsvProfiler.h"
 #include "Async/Async.h"
@@ -1448,7 +1449,7 @@ void FMediaTextureResource::CopySample(FRHICommandListImmediate& RHICmdList, con
 			CreateIntermediateRenderTarget(RHICmdList, Sample->GetOutputDim(), MediaTextureResourceHelpers::GetPixelFormat(Sample), MediaTextureResourceHelpers::RequiresSrgbTexture(Sample), ClearColor, NumMips, bNeedsUAVTexture);
 
 			// Copy data into the output texture to able to add mips later on
-			RHICmdList.CopyTexture(SampleTexture, IntermediateTarget, FRHICopyTextureInfo());
+			TransitionAndCopyTexture(RHICmdList, SampleTexture, IntermediateTarget, FRHICopyTextureInfo());
 		}
 	}
 	else
