[nDisplay] Fixed viewport preview flickering
- Removed a piece of code that releases viewport resources if it's not renderd at the current frame. Internal resources should not be released as the owning viewport may not be rendered for optimization purposes. Its output should remain valid.

[RN] Fixed viewport flickering in nDisplay configurator

diff --git a/Engine/Plugins/Runtime/nDisplay/Source/DisplayCluster/Private/Render/Viewport/Preview/DisplayClusterViewportManagerPreview_Render.cpp b/Engine/Plugins/Runtime/nDisplay/Source/DisplayCluster/Private/Render/Viewport/Preview/DisplayClusterViewportManagerPreview_Render.cpp
--- a/Engine/Plugins/Runtime/nDisplay/Source/DisplayCluster/Private/Render/Viewport/Preview/DisplayClusterViewportManagerPreview_Render.cpp
+++ b/Engine/Plugins/Runtime/nDisplay/Source/DisplayCluster/Private/Render/Viewport/Preview/DisplayClusterViewportManagerPreview_Render.cpp
@@ -215,14 +215,6 @@ bool FDisplayClusterViewportManagerPreview::RenderViewportPreview(FViewport* InV
 				// uvLC viewport is not rendered, but late OCIO is used in FDisplayClusterViewportProxy::ApplyOCIO_RenderThread()
 				FSceneView* View = ViewportPtr->ViewportPreview->CalcSceneView(ViewFamily, ViewIt.ContextNum);
 
-				if (View == nullptr)
-				{
-					// This viewport is not rendering. Release all textures.
-					// The preview mesh materials associated with them will revert to default materials.
-					// In this case, all invalid viewports will be displayed gray or black.
-					ViewportPtr->ReleaseTextures();
-				}
-
 				// Remove the view that is not used for rendering.
 				if (View && (!ViewIt.IsViewportContextCanBeRendered() || ViewFamily.RenderTarget == nullptr))
 				{
