nDisplay: Fixed issue where top-level color grading drawer was not getting refreshed on nDisplay stage or ICVFX property changes

diff --git a/Engine/Plugins/Runtime/nDisplay/Source/DisplayClusterColorGrading/Private/Drawer/DisplayClusterColorGradingDrawerSingleton.cpp b/Engine/Plugins/Runtime/nDisplay/Source/DisplayClusterColorGrading/Private/Drawer/DisplayClusterColorGradingDrawerSingleton.cpp
--- a/Engine/Plugins/Runtime/nDisplay/Source/DisplayClusterColorGrading/Private/Drawer/DisplayClusterColorGradingDrawerSingleton.cpp
+++ b/Engine/Plugins/Runtime/nDisplay/Source/DisplayClusterColorGrading/Private/Drawer/DisplayClusterColorGradingDrawerSingleton.cpp
@@ -13,12 +13,15 @@
 #include "Drawer/SDisplayClusterColorGradingDrawer.h"
 
 #include "Editor.h"
+#include "IColorGradingEditor.h"
 #include "Framework/Application/SlateApplication.h"
 #include "Framework/Docking/LayoutExtender.h"
 #include "Framework/Docking/TabManager.h"
 #include "IPropertyRowGenerator.h"
+#include "LevelEditor.h"
 #include "Modules/ModuleManager.h"
 #include "PropertyEditorModule.h"
+#include "SColorGradingPanel.h"
 #include "Widgets/DeclarativeSyntaxSupport.h"
 #include "Widgets/SCompoundWidget.h"
 #include "Widgets/Docking/SDockTab.h"
@@ -89,6 +92,14 @@ void FDisplayClusterColorGradingDrawerSingleton::RefreshColorGradingDrawers()
 			DockedDrawer->Refresh();
 		}
 	}
+
+	// Update any top-level color grading drawers open in the level editor
+	FLevelEditorModule& LevelEditor = FModuleManager::Get().LoadModuleChecked<FLevelEditorModule>("LevelEditor");
+	if (TSharedPtr<SDockTab> ExistingTab = LevelEditor.GetLevelEditorTabManager()->FindExistingLiveTab(IColorGradingEditor::Get().GetColorGradingTabSpawnerId()))
+	{
+		TSharedRef<SColorGradingPanel> ColorGradingPanel = StaticCastSharedRef<SColorGradingPanel>(ExistingTab->GetContent());
+		ColorGradingPanel->Refresh();
+	}
 }
 
 TSharedRef<SWidget> FDisplayClusterColorGradingDrawerSingleton::CreateDrawerContent(bool bIsInDrawer, bool bCopyStateFromActiveDrawer)
