[nDisplay] Fixed the logic for optimizing resources for Offscreen nodes.

Ignore offscreen node optimization for DCRA preview rendering (UE-348442)

Added new CVar to control Offscreen nodes optimization logic:
- CVar "nDisplay.OffscreenNodes.OptimizeResources" (Disable this CVar to render all types of viewports, such as ICVFX)
If true, optimizes resources for offscreen nDisplay nodes.
Unused viewports will not allocate resources and will not be rendered
Default = true

diff --git a/Engine/Plugins/Runtime/nDisplay/Source/DisplayCluster/Private/Render/Viewport/Configuration/DisplayClusterViewportConfiguration.cpp b/Engine/Plugins/Runtime/nDisplay/Source/DisplayCluster/Private/Render/Viewport/Configuration/DisplayClusterViewportConfiguration.cpp
--- a/Engine/Plugins/Runtime/nDisplay/Source/DisplayCluster/Private/Render/Viewport/Configuration/DisplayClusterViewportConfiguration.cpp
+++ b/Engine/Plugins/Runtime/nDisplay/Source/DisplayCluster/Private/Render/Viewport/Configuration/DisplayClusterViewportConfiguration.cpp
@@ -33,6 +33,15 @@
 
 #include "Misc/DisplayClusterGlobals.h"
 
+static bool GDisplayClusterOptimizeResourcesOnOffscreenNodes = true;
+static FAutoConsoleVariableRef CVarDisplayClusterOptimizeResourcesOnOffscreenNodes(
+	TEXT("nDisplay.OffscreenNodes.OptimizeResources"),
+	GDisplayClusterOptimizeResourcesOnOffscreenNodes,
+	TEXT("If true, optimizes resources for offscreen nDisplay nodes.\n")
+	TEXT("Unused viewports will not allocate resources and will not be rendered.\n"),
+	ECVF_RenderThreadSafe
+);
+
 namespace UE::DisplayCluster::Configuration
 {
 	static inline ADisplayClusterRootActor* ImplGetRootActor(const FDisplayClusterActorRef& InConfigurationRootActorRef)
@@ -370,6 +379,13 @@ bool FDisplayClusterViewportConfiguration::IsClusterNodeRenderingOffscreen() con
 	return bIsRenderingOffscreen;
 }
 
+bool FDisplayClusterViewportConfiguration::ShouldOptimizeOutOffscreenNodeResources() const
+{
+	return GDisplayClusterOptimizeResourcesOnOffscreenNodes
+		// Exception: DCRA preview should ignore off-screen node optimizations.
+		&& !IsPreviewRendering();
+}
+
 IDisplayClusterViewportManager* FDisplayClusterViewportConfiguration::GetViewportManager() const
 {
 	return GetViewportManagerImpl();
diff --git a/Engine/Plugins/Runtime/nDisplay/Source/DisplayCluster/Private/Render/Viewport/Configuration/DisplayClusterViewportConfiguration.h b/Engine/Plugins/Runtime/nDisplay/Source/DisplayCluster/Private/Render/Viewport/Configuration/DisplayClusterViewportConfiguration.h
--- a/Engine/Plugins/Runtime/nDisplay/Source/DisplayCluster/Private/Render/Viewport/Configuration/DisplayClusterViewportConfiguration.h
+++ b/Engine/Plugins/Runtime/nDisplay/Source/DisplayCluster/Private/Render/Viewport/Configuration/DisplayClusterViewportConfiguration.h
@@ -120,6 +120,12 @@ public:
 		return RenderFrameSettings;
 	}
 
+	/**
+	* Returns true if unused resources on offscreen nodes should be optimized.
+	*/
+	bool ShouldOptimizeOutOffscreenNodeResources() const;
+
+
 	/** Sets the rendering frame settings. */
 	void SetRenderFrameSettings(const FDisplayClusterRenderFrameSettings& InRenderFrameSettings)
 	{
diff --git a/Engine/Plugins/Runtime/nDisplay/Source/DisplayCluster/Private/Render/Viewport/Configuration/DisplayClusterViewportConfiguration_ICVFXCamera.cpp b/Engine/Plugins/Runtime/nDisplay/Source/DisplayCluster/Private/Render/Viewport/Configuration/DisplayClusterViewportConfiguration_ICVFXCamera.cpp
--- a/Engine/Plugins/Runtime/nDisplay/Source/DisplayCluster/Private/Render/Viewport/Configuration/DisplayClusterViewportConfiguration_ICVFXCamera.cpp
+++ b/Engine/Plugins/Runtime/nDisplay/Source/DisplayCluster/Private/Render/Viewport/Configuration/DisplayClusterViewportConfiguration_ICVFXCamera.cpp
@@ -144,8 +144,9 @@ void FDisplayClusterViewportConfiguration_ICVFXCamera::Update()
 	// CK composition requires target viewports.
 	&& !TargetViewports.IsEmpty())
 			{
-		// Create CK only if this node renders to the backbuffer.
-		if (!Configuration.IsClusterNodeRenderingOffscreen() // visible node
+		// Create CK only when used.
+		if (!(Configuration.IsClusterNodeRenderingOffscreen()
+			&& Configuration.ShouldOptimizeOutOffscreenNodeResources())
 		|| Configuration.GetRenderFrameSettings().CurrentNode.bHasBackbufferMediaOutput) // or backbuffer shared by media
 				{
 			// Performance: Do not render CK if InnerFrustum is not visible.
diff --git a/Engine/Plugins/Runtime/nDisplay/Source/DisplayCluster/Private/Render/Viewport/DisplayClusterViewport.h b/Engine/Plugins/Runtime/nDisplay/Source/DisplayCluster/Private/Render/Viewport/DisplayClusterViewport.h
--- a/Engine/Plugins/Runtime/nDisplay/Source/DisplayCluster/Private/Render/Viewport/DisplayClusterViewport.h
+++ b/Engine/Plugins/Runtime/nDisplay/Source/DisplayCluster/Private/Render/Viewport/DisplayClusterViewport.h
@@ -326,6 +326,10 @@ public:
 	*/
 	bool IsInternalViewport() const;
 
+	/** Return true if this Outer viewport.
+	*/
+	bool IsOuterViewport() const;
+
 	/** Returns true if the RTT of this viewport is changed externally.
 	* This means that this viewport will not be rendered.
 	* This function is used as the basis for many others. Be careful when making changes to it.
diff --git a/Engine/Plugins/Runtime/nDisplay/Source/DisplayCluster/Private/Render/Viewport/DisplayClusterViewport_Rules.cpp b/Engine/Plugins/Runtime/nDisplay/Source/DisplayCluster/Private/Render/Viewport/DisplayClusterViewport_Rules.cpp
--- a/Engine/Plugins/Runtime/nDisplay/Source/DisplayCluster/Private/Render/Viewport/DisplayClusterViewport_Rules.cpp
+++ b/Engine/Plugins/Runtime/nDisplay/Source/DisplayCluster/Private/Render/Viewport/DisplayClusterViewport_Rules.cpp
@@ -152,6 +152,28 @@ bool FDisplayClusterViewport::IsInternalViewport() const
 	return false;
 }
 
+bool FDisplayClusterViewport::IsOuterViewport() const
+{
+	// Ignore ICVFX internal resources.
+	if (EnumHasAnyFlags(GetRenderSettingsICVFX().RuntimeFlags, EDisplayClusterViewportRuntimeICVFXFlags::InternalResource))
+	{
+		return false;
+	}
+
+	// Ignore internal tile viewports.
+	switch (GetRenderSettings().TileSettings.GetType())
+	{
+	case EDisplayClusterViewportTileType::None:
+	case EDisplayClusterViewportTileType::Source:
+		break;
+
+	default:
+		return false;
+	}
+
+	return true;
+}
+
 bool FDisplayClusterViewport::IsExternalRendering() const
 {
 	if (PostRenderSettings.Replace.IsEnabled())
@@ -188,8 +210,17 @@ bool FDisplayClusterViewport::ShouldDisableViewportResources() const
 
 	// Any other viewports on the headless node that do not use a backbuffer media output.
 	if (Configuration->IsClusterNodeRenderingOffscreen()
+		&& Configuration->ShouldOptimizeOutOffscreenNodeResources()
 		&& !Configuration->GetRenderFrameSettings().CurrentNode.bHasBackbufferMediaOutput)
 	{
+		// Exception: Outer viewports on an Offscreen nodes.
+		if (IsOuterViewport())
+		{
+			// Outer viewports are always used on offscreen nDisplay nodes.
+			// They will allocate resources and be rendered even when otherwise unused (e.g. not consumed by media).
+			return false;
+		}
+
 		// Resources should not be allocated for these viewports.
 		return true;
 	}
@@ -458,6 +489,7 @@ bool FDisplayClusterViewport::IsResourceUsed(const EDisplayClusterViewportResour
 
 	// Any other viewports on the headless node that do not use a backbuffer media output.
 	if (Configuration->IsClusterNodeRenderingOffscreen()
+		&& Configuration->ShouldOptimizeOutOffscreenNodeResources()
 	&& !Configuration->GetRenderFrameSettings().CurrentNode.bHasBackbufferMediaOutput)
 	{
 		// Backbuffer not used: skip related resources.
@@ -465,6 +497,15 @@ bool FDisplayClusterViewport::IsResourceUsed(const EDisplayClusterViewportResour
 		{
 		case EDisplayClusterViewportResource::OutputFrameTargetableResources:
 		case EDisplayClusterViewportResource::AdditionalFrameTargetableResources:
+
+			// Exception: Outer viewports on an Offscreen nodes.
+			if (IsOuterViewport())
+			{
+				// Outer viewports are always used on offscreen nDisplay nodes.
+				// They will allocate resources and be rendered even when otherwise unused (e.g. not consumed by media).
+				break;
+			}
+
 			// Skip FrameTargetable RTTs (final composite copied to backbuffer).
 			return false;
 
