[nDisplay] Fixed regression when InCamera is not visible on Outers when the Chromakey type 'Inner Frustum' is enabled.

diff --git a/Engine/Plugins/Runtime/nDisplay/Source/DisplayCluster/Private/Render/Viewport/Configuration/DisplayClusterViewportConfiguration_ICVFXCamera.cpp b/Engine/Plugins/Runtime/nDisplay/Source/DisplayCluster/Private/Render/Viewport/Configuration/DisplayClusterViewportConfiguration_ICVFXCamera.cpp
--- a/Engine/Plugins/Runtime/nDisplay/Source/DisplayCluster/Private/Render/Viewport/Configuration/DisplayClusterViewportConfiguration_ICVFXCamera.cpp
+++ b/Engine/Plugins/Runtime/nDisplay/Source/DisplayCluster/Private/Render/Viewport/Configuration/DisplayClusterViewportConfiguration_ICVFXCamera.cpp
@@ -138,17 +138,25 @@ void FDisplayClusterViewportConfiguration_ICVFXCamera::Update()
 	CameraViewport.Reset();
 	ChromakeyViewport.Reset();
 
-	if (GetOrCreateAndSetupInnerCameraViewport()
-	// CK only applies as composition over the rendered InnerFrustum image.
-	&& !CameraViewport->GetRenderSettings().bSkipRendering
-	// CK composition requires target viewports.
-	&& !TargetViewports.IsEmpty())
-			{
-		// Create CK only when used.
-		if (!(Configuration.IsClusterNodeRenderingOffscreen()
-			&& Configuration.ShouldOptimizeOutOffscreenNodeResources())
-		|| Configuration.GetRenderFrameSettings().CurrentNode.bHasBackbufferMediaOutput) // or backbuffer shared by media
-				{
+	if (GetOrCreateAndSetupInnerCameraViewport())
+	{
+		bool bShouldUseChromakeyViewport =
+			// CK only applies as composition over the rendered InnerFrustum image.
+			!CameraViewport->GetRenderSettings().bSkipRendering
+			// CK composition requires target viewports.
+			&& !TargetViewports.IsEmpty();
+
+		// Optimize CK viewport resources for offscreen nodes
+		if (Configuration.IsClusterNodeRenderingOffscreen()
+			&& Configuration.ShouldOptimizeOutOffscreenNodeResources()
+			// Optimize if backbuffer not used by media
+			&& !Configuration.GetRenderFrameSettings().CurrentNode.bHasBackbufferMediaOutput)
+		{
+			bShouldUseChromakeyViewport = false;
+		}
+
+		if (bShouldUseChromakeyViewport)
+		{
 			// Performance: Do not render CK if InnerFrustum is not visible.
 			// Updates the Chromakey viewport and configures it for use in the ICVFX rendering stack.
 			UpdateChromakeyViewport();
