[nDisplay] Only compile shaders for PC shader platforms (filters out mobile and console).

nDisplay plugin is only supported on Win64/Linux (PlatformAllowList), but enabling other targets
(e.g. -set:WithPS5=true) can still cause shader compilation for those platforms during cook/build.
Guard every ShouldCompilePermutation() with IsPCPlatform(Parameters.Platform) to avoid compiling for consoles/mobile.

diff --git a/Engine/Plugins/Runtime/nDisplay/Source/DisplayClusterShaders/Private/Shaders/DisplayClusterShadersCopyTexture.cpp b/Engine/Plugins/Runtime/nDisplay/Source/DisplayClusterShaders/Private/Shaders/DisplayClusterShadersCopyTexture.cpp
--- a/Engine/Plugins/Runtime/nDisplay/Source/DisplayClusterShaders/Private/Shaders/DisplayClusterShadersCopyTexture.cpp
+++ b/Engine/Plugins/Runtime/nDisplay/Source/DisplayClusterShaders/Private/Shaders/DisplayClusterShadersCopyTexture.cpp
@@ -145,6 +145,12 @@ namespace UE::DisplayClusterShaders::Private
 
 		static bool ShouldCompilePermutation(const FGlobalShaderPermutationParameters& Parameters)
 		{
+			// Only compile for desktop-class "PC" shader platforms (filters out mobile and console).
+			if (!IsPCPlatform(Parameters.Platform))
+			{
+				return false;
+			}
+
 			return FColorEncodingCopyRectPSPermutation::ShouldCompileCommonPSPermutation(FPermutationDomain(Parameters.PermutationId));
 		}
 
diff --git a/Engine/Plugins/Runtime/nDisplay/Source/DisplayClusterShaders/Private/Shaders/DisplayClusterShadersMedia.cpp b/Engine/Plugins/Runtime/nDisplay/Source/DisplayClusterShaders/Private/Shaders/DisplayClusterShadersMedia.cpp
--- a/Engine/Plugins/Runtime/nDisplay/Source/DisplayClusterShaders/Private/Shaders/DisplayClusterShadersMedia.cpp
+++ b/Engine/Plugins/Runtime/nDisplay/Source/DisplayClusterShaders/Private/Shaders/DisplayClusterShadersMedia.cpp
@@ -36,6 +36,12 @@ namespace UE::DisplayClusterShaders::Private
 
 		static bool ShouldCompilePermutation(const FGlobalShaderPermutationParameters& InParameters)
 		{
+			// Only compile for desktop-class "PC" shader platforms (filters out mobile and console).
+			if (!IsPCPlatform(InParameters.Platform))
+			{
+				return false;
+			}
+
 			return true;
 		}
 
@@ -69,6 +75,13 @@ namespace UE::DisplayClusterShaders::Private
 
 		static bool ShouldCompilePermutation(const FGlobalShaderPermutationParameters& InParameters)
 		{
+			// Only compile for desktop-class "PC" shader platforms (filters out mobile and console).
+			if (!IsPCPlatform(InParameters.Platform))
+			{
+				return false;
+			}
+
+
 			return true;
 		}
 
diff --git a/Engine/Plugins/Runtime/nDisplay/Source/DisplayClusterShaders/Private/Shaders/DisplayClusterShadersPostprocess_Blur.cpp b/Engine/Plugins/Runtime/nDisplay/Source/DisplayClusterShaders/Private/Shaders/DisplayClusterShadersPostprocess_Blur.cpp
--- a/Engine/Plugins/Runtime/nDisplay/Source/DisplayClusterShaders/Private/Shaders/DisplayClusterShadersPostprocess_Blur.cpp
+++ b/Engine/Plugins/Runtime/nDisplay/Source/DisplayClusterShaders/Private/Shaders/DisplayClusterShadersPostprocess_Blur.cpp
@@ -38,6 +38,12 @@ class FDirectProjectionVS : public FGlobalShader
 
 	static bool ShouldCompilePermutation(const FGlobalShaderPermutationParameters& Parameters)
 	{
+		// Only compile for desktop-class "PC" shader platforms (filters out mobile and console).
+		if (!IsPCPlatform(Parameters.Platform))
+		{
+			return false;
+		}
+
 		return true;
 	}
 
@@ -71,6 +77,12 @@ public:
 
 	static bool ShouldCompilePermutation(const FGlobalShaderPermutationParameters& Parameters)
 	{
+		// Only compile for desktop-class "PC" shader platforms (filters out mobile and console).
+		if (!IsPCPlatform(Parameters.Platform))
+		{
+			return false;
+		}
+
 		return true;
 	}
 
diff --git a/Engine/Plugins/Runtime/nDisplay/Source/DisplayClusterShaders/Private/Shaders/DisplayClusterShadersPostprocess_OutputRemap.cpp b/Engine/Plugins/Runtime/nDisplay/Source/DisplayClusterShaders/Private/Shaders/DisplayClusterShadersPostprocess_OutputRemap.cpp
--- a/Engine/Plugins/Runtime/nDisplay/Source/DisplayClusterShaders/Private/Shaders/DisplayClusterShadersPostprocess_OutputRemap.cpp
+++ b/Engine/Plugins/Runtime/nDisplay/Source/DisplayClusterShaders/Private/Shaders/DisplayClusterShadersPostprocess_OutputRemap.cpp
@@ -66,7 +66,15 @@ public:
 	using FParameters = FOutputRemapPixelShaderParameters;
 
 	static bool ShouldCompilePermutation(const FGlobalShaderPermutationParameters& Parameters)
-	{ return true; }
+	{
+		// Only compile for desktop-class "PC" shader platforms (filters out mobile and console).
+		if (!IsPCPlatform(Parameters.Platform))
+		{
+			return false;
+		}
+
+		return true;
+	}
 };
 
 IMPLEMENT_SHADER_TYPE(, FOutputRemapPS, OutputRemapShaderFileName, TEXT("OutputRemap_PS"), SF_Pixel);
@@ -81,6 +89,12 @@ public:
 
 	static bool ShouldCompilePermutation(const FGlobalShaderPermutationParameters& Parameters)
 	{
+		// Only compile for desktop-class "PC" shader platforms (filters out mobile and console).
+		if (!IsPCPlatform(Parameters.Platform))
+		{
+			return false;
+		}
+
 		return true;
 	}
 };
diff --git a/Engine/Plugins/Runtime/nDisplay/Source/DisplayClusterShaders/Private/Shaders/DisplayClusterShadersPreprocess_UVLightCards.cpp b/Engine/Plugins/Runtime/nDisplay/Source/DisplayClusterShaders/Private/Shaders/DisplayClusterShadersPreprocess_UVLightCards.cpp
--- a/Engine/Plugins/Runtime/nDisplay/Source/DisplayClusterShaders/Private/Shaders/DisplayClusterShadersPreprocess_UVLightCards.cpp
+++ b/Engine/Plugins/Runtime/nDisplay/Source/DisplayClusterShaders/Private/Shaders/DisplayClusterShadersPreprocess_UVLightCards.cpp
@@ -36,6 +36,12 @@ public:
 
 	static bool ShouldCompilePermutation(const FMeshMaterialShaderPermutationParameters& Parameters)
 	{
+		// Only compile for desktop-class "PC" shader platforms (filters out mobile and console).
+		if (!IsPCPlatform(Parameters.Platform))
+		{
+			return false;
+		}
+
 		return Parameters.VertexFactoryType == FindVertexFactoryType(TEXT("FLocalVertexFactory"));
 	}
 };
@@ -56,6 +62,12 @@ public:
 
 	static bool ShouldCompilePermutation(const FMeshMaterialShaderPermutationParameters& Parameters)
 	{
+		// Only compile for desktop-class "PC" shader platforms (filters out mobile and console).
+		if (!IsPCPlatform(Parameters.Platform))
+		{
+			return false;
+		}
+
 		return Parameters.VertexFactoryType == FindVertexFactoryType(TEXT("FLocalVertexFactory"));
 	}
 };
diff --git a/Engine/Plugins/Runtime/nDisplay/Source/DisplayClusterShaders/Private/Shaders/DisplayClusterShadersWarpblend_ICVFX.cpp b/Engine/Plugins/Runtime/nDisplay/Source/DisplayClusterShaders/Private/Shaders/DisplayClusterShadersWarpblend_ICVFX.cpp
--- a/Engine/Plugins/Runtime/nDisplay/Source/DisplayClusterShaders/Private/Shaders/DisplayClusterShadersWarpblend_ICVFX.cpp
+++ b/Engine/Plugins/Runtime/nDisplay/Source/DisplayClusterShaders/Private/Shaders/DisplayClusterShadersWarpblend_ICVFX.cpp
@@ -317,6 +317,12 @@ public:
 
 	static bool ShouldCompilePermutation(const FGlobalShaderPermutationParameters& Parameters)
 	{
+		// Only compile for desktop-class "PC" shader platforms (filters out mobile and console).
+		if (!IsPCPlatform(Parameters.Platform))
+		{
+			return false;
+		}
+
 		return IcvfxShaderPermutation::ShouldCompileCommonVSPermutation(FPermutationDomain(Parameters.PermutationId));
 	}
 };
@@ -334,6 +340,12 @@ public:
 
 	static bool ShouldCompilePermutation(const FGlobalShaderPermutationParameters& Parameters)
 	{
+		// Only compile for desktop-class "PC" shader platforms (filters out mobile and console).
+		if (!IsPCPlatform(Parameters.Platform))
+		{
+			return false;
+		}
+
 		return IcvfxShaderPermutation::ShouldCompileCommonPSPermutation(FPermutationDomain(Parameters.PermutationId));
 	}
 
@@ -356,6 +368,12 @@ public:
 	using FParameters = FIcvfxPixelShaderParameters;
 	static bool ShouldCompilePermutation(const FGlobalShaderPermutationParameters& Parameters)
 	{
+		// Only compile for desktop-class "PC" shader platforms (filters out mobile and console).
+		if (!IsPCPlatform(Parameters.Platform))
+		{
+			return false;
+		}
+
 		return true;
 	}
 
diff --git a/Engine/Plugins/Runtime/nDisplay/Source/DisplayClusterShaders/Private/Shaders/DisplayClusterShadersWarpblend_MPCDI.cpp b/Engine/Plugins/Runtime/nDisplay/Source/DisplayClusterShaders/Private/Shaders/DisplayClusterShadersWarpblend_MPCDI.cpp
--- a/Engine/Plugins/Runtime/nDisplay/Source/DisplayClusterShaders/Private/Shaders/DisplayClusterShadersWarpblend_MPCDI.cpp
+++ b/Engine/Plugins/Runtime/nDisplay/Source/DisplayClusterShaders/Private/Shaders/DisplayClusterShadersWarpblend_MPCDI.cpp
@@ -124,6 +124,12 @@ public:
 
 	static bool ShouldCompilePermutation(const FGlobalShaderPermutationParameters& Parameters)
 	{
+		// Only compile for desktop-class "PC" shader platforms (filters out mobile and console).
+		if (!IsPCPlatform(Parameters.Platform))
+		{
+			return false;
+		}
+
 		return MpcdiShaderPermutation::ShouldCompileCommonVSPermutation(Parameters, FPermutationDomain(Parameters.PermutationId));
 	}
 };
@@ -141,6 +147,12 @@ public:
 
 	static bool ShouldCompilePermutation(const FGlobalShaderPermutationParameters& Parameters)
 	{
+		// Only compile for desktop-class "PC" shader platforms (filters out mobile and console).
+		if (!IsPCPlatform(Parameters.Platform))
+		{
+			return false;
+		}
+
 		return MpcdiShaderPermutation::ShouldCompileCommonPSPermutation(Parameters, FPermutationDomain(Parameters.PermutationId));
 	}
 };
@@ -156,6 +168,12 @@ public:
 	using FParameters = FMpcdiPixelShaderParameters;
 	static bool ShouldCompilePermutation(const FGlobalShaderPermutationParameters& Parameters)
 	{
+		// Only compile for desktop-class "PC" shader platforms (filters out mobile and console).
+		if (!IsPCPlatform(Parameters.Platform))
+		{
+			return false;
+		}
+
 		return true;
 	}
 };
