[nDisplay] Removed unused member variable from FDisplayClusterClusterManager

#rnx

#FailedValidations Preflight

diff --git a/Engine/Plugins/Runtime/nDisplay/Source/DisplayCluster/Private/Cluster/DisplayClusterClusterManager.cpp b/Engine/Plugins/Runtime/nDisplay/Source/DisplayCluster/Private/Cluster/DisplayClusterClusterManager.cpp
--- a/Engine/Plugins/Runtime/nDisplay/Source/DisplayCluster/Private/Cluster/DisplayClusterClusterManager.cpp
+++ b/Engine/Plugins/Runtime/nDisplay/Source/DisplayCluster/Private/Cluster/DisplayClusterClusterManager.cpp
@@ -179,14 +179,6 @@ void FDisplayClusterClusterManager::EndSession()
 	ClusterNodeId.Empty();
 }
 
-bool FDisplayClusterClusterManager::StartScene(UWorld* InWorld)
-{
-	check(InWorld);
-	CurrentWorld = InWorld;
-
-	return true;
-}
-
 void FDisplayClusterClusterManager::EndScene()
 {
 	{
@@ -203,7 +195,6 @@ void FDisplayClusterClusterManager::EndScene()
 	}
 
 	NativeInputCache.Reset();
-	CurrentWorld = nullptr;
 }
 
 void FDisplayClusterClusterManager::StartFrame(uint64 FrameNum)
diff --git a/Engine/Plugins/Runtime/nDisplay/Source/DisplayCluster/Private/Cluster/DisplayClusterClusterManager.h b/Engine/Plugins/Runtime/nDisplay/Source/DisplayCluster/Private/Cluster/DisplayClusterClusterManager.h
--- a/Engine/Plugins/Runtime/nDisplay/Source/DisplayCluster/Private/Cluster/DisplayClusterClusterManager.h
+++ b/Engine/Plugins/Runtime/nDisplay/Source/DisplayCluster/Private/Cluster/DisplayClusterClusterManager.h
@@ -33,7 +33,6 @@ public:
 	virtual void Release() override;
 	virtual bool StartSession(UDisplayClusterConfigurationData* InConfigData, const FString& NodeID) override;
 	virtual void EndSession() override;
-	virtual bool StartScene(UWorld* World) override;
 	virtual void EndScene() override;
 	virtual void StartFrame(uint64 FrameNum) override;
 	virtual void EndFrame(uint64 FrameNum) override;
@@ -203,9 +202,6 @@ private:
 	mutable FCriticalSection ActiveClusterNodeIdsCS;
 
 
-	// Current world
-	UWorld* CurrentWorld = nullptr;
-
 	FEventRef TimeDataCacheReadySignal{ EEventMode::ManualReset };
 	double DeltaTimeCache = 0.f;
 	double GameTimeCache  = 0.f;
