[TRS Gizmo] Make sure required Editors opt-out of new ITF gizmos and interactions

diff --git a/Engine/Plugins/Runtime/nDisplay/Source/DisplayClusterLightCardEditor/Private/DisplayClusterLightCardEditor.cpp b/Engine/Plugins/Runtime/nDisplay/Source/DisplayClusterLightCardEditor/Private/DisplayClusterLightCardEditor.cpp
--- a/Engine/Plugins/Runtime/nDisplay/Source/DisplayClusterLightCardEditor/Private/DisplayClusterLightCardEditor.cpp
+++ b/Engine/Plugins/Runtime/nDisplay/Source/DisplayClusterLightCardEditor/Private/DisplayClusterLightCardEditor.cpp
@@ -1130,11 +1130,7 @@ void FDisplayClusterLightCardEditor::RegisterLayoutExtension(FLayoutExtender& In
 
 TSharedRef<SDockTab> FDisplayClusterLightCardEditor::SpawnViewportTab(const FSpawnTabArgs& SpawnTabArgs)
 {
-	return SNew(SDockTab)
-	.TabRole(ETabRole::DocumentTab)
-	[
-		CreateViewportWidget()
-	];
+	TSharedRef<SWidget> ViewportWidget = CreateViewportWidget();
 
 	if (ViewportView.IsValid())
 	{
@@ -1147,6 +1143,12 @@ TSharedRef<SDockTab> FDisplayClusterLightCardEditor::SpawnViewportTab(const FSpa
 			}
 		}
 	}
+
+	return SNew(SDockTab)
+	.TabRole(ETabRole::DocumentTab)
+	[
+		ViewportWidget
+	];
 }
 
 TSharedRef<SDockTab> FDisplayClusterLightCardEditor::SpawnOutlinerTab(const FSpawnTabArgs& SpawnTabArgs)
