Temporarily silencing PVS warnings
#rnx

diff --git a/Engine/Plugins/FX/Niagara/Source/Niagara/Private/NiagaraRendererRibbons.cpp b/Engine/Plugins/FX/Niagara/Source/Niagara/Private/NiagaraRendererRibbons.cpp
--- a/Engine/Plugins/FX/Niagara/Source/Niagara/Private/NiagaraRendererRibbons.cpp
+++ b/Engine/Plugins/FX/Niagara/Source/Niagara/Private/NiagaraRendererRibbons.cpp
@@ -1235,6 +1235,7 @@ void FNiagaraRendererRibbons::GetDynamicRayTracingInstances(FRayTracingInstanceC
 }
 #endif
 
+//-V::609 UE-348905
 void FNiagaraRendererRibbons::GenerateShapeStateMultiPlane(FNiagaraRibbonShapeGeometryData& State, int32 MultiPlaneCount, int32 WidthSegmentationCount, bool bEnableAccurateGeometry, bool bUseMaterialBackfaceCulling)
 {
 	using namespace NiagaraRendererRibbonsPrivate;
diff --git a/Engine/Plugins/Runtime/nDisplay/Source/DisplayClusterWarp/Private/WarpBlend/Math/DisplayClusterWarpBlendMath_WarpMap.h b/Engine/Plugins/Runtime/nDisplay/Source/DisplayClusterWarp/Private/WarpBlend/Math/DisplayClusterWarpBlendMath_WarpMap.h
--- a/Engine/Plugins/Runtime/nDisplay/Source/DisplayClusterWarp/Private/WarpBlend/Math/DisplayClusterWarpBlendMath_WarpMap.h
+++ b/Engine/Plugins/Runtime/nDisplay/Source/DisplayClusterWarp/Private/WarpBlend/Math/DisplayClusterWarpBlendMath_WarpMap.h
@@ -44,11 +44,11 @@ public:
 		// Generate valid points for texturebox method:
 		for (int32 low_y = 0; low_y < DivY; low_y++)
 		{
-			int32 y = (Height - 1) * (float(low_y) / (DivY - 1));
+			int32 y = (Height - 1) * (float(low_y) / (DivY - 1)); //-V609 UE-349513 
 
 			for (int32 low_x = 0; low_x < DivX; low_x++)
 			{
-				int32 x = (Width - 1) * (float(low_x) / (DivX - 1));
+				int32 x = (Width - 1) * (float(low_x) / (DivX - 1)); //-V609 UE-349513 
 
 				if (IsValidPoint(x, y))
 				{
