[Engine] [nDisplay] Introduce per-view API for stereo rendering devices to provide and own their FSceneViewStateInterface

Add per-view APIs HasExternalViewState(int32) and GetExternalViewState(int32) to allow stereo rendering devices to provide/manage their own FSceneViewStateInterface.
When HasExternalViewState is true, ULocalPlayer will use GetExternalViewState to obtain ViewState.
Defaults remain false/nullptr, preserving existing behavior for devices that don't manage ViewState externally.

nDisplay now creates and owns the ViewState instead of relying on ULocalPlayer.
FDisplayClusterDeviceBase exposes this device-managed state to the engine by overriding IStereoRendering::HasExternalViewState / GetExternalViewState.
nDisplay also reserves its own stereo view index range by setting FDisplayClusterViewportStereoscopicPass::StereoViewIndexBase = 100,
preventing conflicts with standard view indices.

Added new CVars:
'nDisplay.render.ViewState.LightCard' - Enable ViewState for LightCard viewports (default 1)
'nDisplay.render.ViewState.Chromakey' - Enable ViewState for Chromakey viewports (default 0)

#Synced-CL Unknown

diff --git a/Engine/Plugins/Runtime/nDisplay/Source/DisplayCluster/Private/Render/Device/DisplayClusterDeviceBase.cpp b/Engine/Plugins/Runtime/nDisplay/Source/DisplayCluster/Private/Render/Device/DisplayClusterDeviceBase.cpp
--- a/Engine/Plugins/Runtime/nDisplay/Source/DisplayCluster/Private/Render/Device/DisplayClusterDeviceBase.cpp
+++ b/Engine/Plugins/Runtime/nDisplay/Source/DisplayCluster/Private/Render/Device/DisplayClusterDeviceBase.cpp
@@ -34,9 +34,10 @@
 #include "Render/Synchronization/IDisplayClusterRenderSyncPolicy.h"
 
 
+#include "Render/Viewport/DisplayClusterViewport.h"
 #include "Render/Viewport/DisplayClusterViewportManager.h"
 #include "Render/Viewport/DisplayClusterViewportManagerProxy.h"
-#include "Render/Viewport/IDisplayClusterViewport.h"
+#include "Render/Viewport/DisplayClusterViewportStereoscopicPass.h"
 #include "Render/Viewport/IDisplayClusterViewportProxy.h"
 
 #include "Render/Viewport/Configuration/DisplayClusterViewportConfigurationHelpers_Postprocess.h"
@@ -79,7 +80,7 @@ FDisplayClusterDeviceBase::~FDisplayClusterDeviceBase()
 	//@todo: delete singleton object IDisplayClusterViewportManager
 }
 
-IDisplayClusterViewportManager* FDisplayClusterDeviceBase::GetViewportManager() const
+FDisplayClusterViewportManager* FDisplayClusterDeviceBase::GetViewportManager() const
 {
 	return ViewportManagerWeakPtr.IsValid() ? ViewportManagerWeakPtr.Pin().Get() : nullptr;
 }
@@ -224,28 +225,26 @@ EStereoscopicPass FDisplayClusterDeviceBase::GetViewPassForIndex(bool bStereoReq
 	return EStereoscopicPass::eSSP_FULL;
 }
 
-void FDisplayClusterDeviceBase::AdjustViewRect(int32 ViewIndex, int32& X, int32& Y, uint32& SizeX, uint32& SizeY) const
+void FDisplayClusterDeviceBase::AdjustViewRect(int32 StereoViewIndex, int32& X, int32& Y, uint32& SizeX, uint32& SizeY) const
 {
 	check(IsInGameThread());
 
-	IDisplayClusterViewportManager* ViewportManager = GetViewportManager();
-	if (ViewportManager == nullptr || ViewportManager->GetConfiguration().IsSceneOpened() == false)
+	if (!FDisplayClusterViewportStereoscopicPass::IsValidStereoViewIndex(StereoViewIndex))
 	{
 		return;
 	}
-	// ViewIndex == eSSE_MONOSCOPIC(-1) is a special case called for ISR culling math.
-	// Since nDisplay is not ISR compatible, we ignore this request. This won't be neccessary once
-	// we stop using nDisplay as a stereoscopic rendering device (IStereoRendering).
-	else if (ViewIndex < 0)
+
+	IDisplayClusterViewportManager* ViewportManager = GetViewportManager();
+	if (ViewportManager == nullptr || ViewportManager->GetConfiguration().IsSceneOpened() == false)
 	{
 		return;
 	}
 
 	uint32 ViewportContextNum = 0;
-	IDisplayClusterViewport* ViewportPtr = ViewportManager->FindViewport(ViewIndex, &ViewportContextNum);
+	IDisplayClusterViewport* ViewportPtr = ViewportManager->FindViewport(StereoViewIndex, &ViewportContextNum);
 	if (ViewportPtr == nullptr)
 	{
-		UE_LOG(LogDisplayClusterRender, Warning, TEXT("Viewport StereoViewIndex='%i' not found"), ViewIndex);
+		UE_LOG(LogDisplayClusterRender, Warning, TEXT("Viewport StereoViewIndex='%i' not found"), StereoViewIndex);
 		return;
 	}
 
@@ -260,15 +259,20 @@ void FDisplayClusterDeviceBase::AdjustViewRect(int32 ViewIndex, int32& X, int32&
 	UE_LOG(LogDisplayClusterRender, Verbose, TEXT("Adjusted view rect: Viewport='%s', ViewIndex=%d, [%d,%d - %d,%d]"), *ViewportPtr->GetId(), ViewportContextNum, ViewRect.Min.X, ViewRect.Min.Y, ViewRect.Max.X, ViewRect.Max.Y);
 }
 
-void FDisplayClusterDeviceBase::CalculateStereoViewOffset(const int32 ViewIndex, FRotator& ViewRotation, const float WorldToMeters, FVector& ViewLocation)
+void FDisplayClusterDeviceBase::CalculateStereoViewOffset(const int32 StereoViewIndex, FRotator& ViewRotation, const float WorldToMeters, FVector& ViewLocation)
 {
 	check(IsInGameThread());
 	check(WorldToMeters > 0.f);
 
+	if (!FDisplayClusterViewportStereoscopicPass::IsValidStereoViewIndex(StereoViewIndex))
+	{
+		return;
+	}
+
 	IDisplayClusterViewportManager* ViewportManager = GetViewportManager();
 
 	uint32 ViewportContextNum = 0;
-	IDisplayClusterViewport* ViewportPtr = ViewportManager ? ViewportManager->FindViewport(ViewIndex, &ViewportContextNum) : nullptr;
+	IDisplayClusterViewport* ViewportPtr = ViewportManager ? ViewportManager->FindViewport(StereoViewIndex, &ViewportContextNum) : nullptr;
 	if (ViewportPtr == nullptr)
 	{
 		return;
@@ -292,24 +296,24 @@ void FDisplayClusterDeviceBase::CalculateStereoViewOffset(const int32 ViewIndex,
 	UE_LOG(LogDisplayClusterRender, VeryVerbose, TEXT("ViewLoc: %s, ViewRot: %s"), *ViewLocation.ToString(), *ViewRotation.ToString());
 }
 
-FMatrix FDisplayClusterDeviceBase::GetStereoProjectionMatrix(const int32 ViewIndex) const
+FMatrix FDisplayClusterDeviceBase::GetStereoProjectionMatrix(const int32 StereoViewIndex) const
 {
 	check(IsInGameThread());
 
-	IDisplayClusterViewportManager* ViewportManager = GetViewportManager();
-
 	FMatrix PrjMatrix = FMatrix::Identity;
+	if (!FDisplayClusterViewportStereoscopicPass::IsValidStereoViewIndex(StereoViewIndex))
+	{
+		return PrjMatrix;
+	}
 
-	// ViewIndex == eSSE_MONOSCOPIC(-1) is a special case called for ISR culling math.
-	// Since nDisplay is not ISR compatible, we ignore this request. This won't be neccessary once
-	// we stop using nDisplay as a stereoscopic rendering device (IStereoRendering).
-	if (ViewportManager && ViewportManager->GetConfiguration().IsSceneOpened() && ViewIndex >= 0)
+	IDisplayClusterViewportManager* ViewportManager = GetViewportManager();
+	if (ViewportManager && ViewportManager->GetConfiguration().IsSceneOpened())
 	{
 		uint32 ViewportContextNum = 0;
-		IDisplayClusterViewport* ViewportPtr = ViewportManager->FindViewport(ViewIndex, &ViewportContextNum);
+		IDisplayClusterViewport* ViewportPtr = ViewportManager->FindViewport(StereoViewIndex, &ViewportContextNum);
 		if (ViewportPtr == nullptr)
 		{
-			UE_LOG(LogDisplayClusterRender, Warning, TEXT("Viewport StereoViewIndex='%i' not found"), ViewIndex);
+			UE_LOG(LogDisplayClusterRender, Warning, TEXT("Viewport StereoViewIndex='%i' not found"), StereoViewIndex);
 		}
 		else
 		if (ViewportPtr->GetProjectionMatrix(ViewportContextNum, PrjMatrix) == false)
@@ -321,6 +325,34 @@ FMatrix FDisplayClusterDeviceBase::GetStereoProjectionMatrix(const int32 ViewInd
 	return PrjMatrix;
 }
 
+bool FDisplayClusterDeviceBase::HasExternalViewState(const int32 StereoViewIndex) const
+{
+	check(IsInGameThread());
+
+	// true if StereoViewIndex is in nDisplay range.
+	return FDisplayClusterViewportStereoscopicPass::IsValidStereoViewIndex(StereoViewIndex);
+}
+
+FSceneViewStateInterface* FDisplayClusterDeviceBase::GetExternalViewState(const int32 StereoViewIndex)
+{
+	check(IsInGameThread());
+	check(FDisplayClusterViewportStereoscopicPass::IsValidStereoViewIndex(StereoViewIndex));
+
+	FDisplayClusterViewportManager* ViewportManager = GetViewportManager();
+	if (ViewportManager && ViewportManager->GetConfiguration().IsSceneOpened())
+	{
+		uint32 ViewportContextNum = 0;
+		if (FDisplayClusterViewport* ViewportPtr = ViewportManager->ImplFindViewport(StereoViewIndex, &ViewportContextNum))
+		{
+			return ViewportPtr->GetViewState(ViewportContextNum);
+		}
+
+		UE_LOG(LogDisplayClusterRender, Warning, TEXT("Viewport StereoViewIndex='%i' not found"), StereoViewIndex);
+	}
+
+	return nullptr;
+}
+
 bool FDisplayClusterDeviceBase::BeginNewFrame(FViewport* InViewport, UWorld* InWorld, FDisplayClusterRenderFrame& OutRenderFrame)
 {
 	check(IsInGameThread());
@@ -365,7 +397,7 @@ void FDisplayClusterDeviceBase::InitializeNewFrame()
 		if (IDisplayClusterViewportManager* ViewportManager = RootActor->GetOrCreateViewportManager())
 		{
 			// Begin use viewport manager for current frame
-			ViewportManagerWeakPtr = ViewportManager->ToSharedPtr();
+			ViewportManagerWeakPtr = ViewportManager->ToSharedRef();
 
 			// Initialize frame for render
 			ViewportManager->InitializeNewFrame();
diff --git a/Engine/Plugins/Runtime/nDisplay/Source/DisplayCluster/Private/Render/Device/DisplayClusterDeviceBase.h b/Engine/Plugins/Runtime/nDisplay/Source/DisplayCluster/Private/Render/Device/DisplayClusterDeviceBase.h
--- a/Engine/Plugins/Runtime/nDisplay/Source/DisplayCluster/Private/Render/Device/DisplayClusterDeviceBase.h
+++ b/Engine/Plugins/Runtime/nDisplay/Source/DisplayCluster/Private/Render/Device/DisplayClusterDeviceBase.h
@@ -17,7 +17,7 @@
 
 class IDisplayClusterPostProcess;
 class FDisplayClusterPresentationBase;
-class IDisplayClusterViewportManager;
+class FDisplayClusterViewportManager;
 class FDisplayClusterViewportManagerProxy;
 class FSceneView;
 
@@ -61,6 +61,10 @@ public:
 	virtual void AdjustViewRect(int32 ViewIndex, int32& X, int32& Y, uint32& SizeX, uint32& SizeY) const override;
 	virtual void CalculateStereoViewOffset(const int32 ViewIndex, FRotator& ViewRotation, const float WorldToMeters, FVector& ViewLocation) override;
 	virtual FMatrix GetStereoProjectionMatrix(const int32 ViewIndex) const override;
+	
+	virtual bool HasExternalViewState(const int32 ViewIndex) const override;
+	virtual FSceneViewStateInterface* GetExternalViewState(const int32 ViewIndex) override;
+
 	virtual void RenderTexture_RenderThread(class FRDGBuilder& GraphBuilder, FRDGTextureRef BackBuffer, FRDGTextureRef SrcTexture, FVector2f WindowSize) const override;
 
 	// So far, we don’t need any region bounds
@@ -100,14 +104,14 @@ protected:
 	virtual void EndFinalPostprocessSettings(struct FPostProcessSettings* FinalPostProcessingSettings, const enum EStereoscopicPass StereoPassType, const int32 StereoViewIndex) override;
 
 	/** Get a pointer to the DC ViewportManager if it still exists. */
-	IDisplayClusterViewportManager* GetViewportManager() const;
+	FDisplayClusterViewportManager* GetViewportManager() const;
 
 	/** Get a pointer to the DC ViewportManagerProxy if it still exists. */
 	FDisplayClusterViewportManagerProxy* GetViewportManagerProxy_RenderThread() const;
 
 private:
 	// Pointer to the DC ViewportManager from the active DCRA that is currently being used for rendering.
-	TWeakPtr<IDisplayClusterViewportManager, ESPMode::ThreadSafe> ViewportManagerWeakPtr;
+	TWeakPtr<FDisplayClusterViewportManager, ESPMode::ThreadSafe> ViewportManagerWeakPtr;
 
 	// Pointer to the DC ViewportManagerProxy from the active DCRA that is currently being used for rendering.
 	mutable TWeakPtr<FDisplayClusterViewportManagerProxy, ESPMode::ThreadSafe> ViewportManagerProxyWeakPtr;
diff --git a/Engine/Plugins/Runtime/nDisplay/Source/DisplayCluster/Private/Render/Viewport/DisplayClusterViewportManager.cpp b/Engine/Plugins/Runtime/nDisplay/Source/DisplayCluster/Private/Render/Viewport/DisplayClusterViewportManager.cpp
--- a/Engine/Plugins/Runtime/nDisplay/Source/DisplayCluster/Private/Render/Viewport/DisplayClusterViewportManager.cpp
+++ b/Engine/Plugins/Runtime/nDisplay/Source/DisplayCluster/Private/Render/Viewport/DisplayClusterViewportManager.cpp
@@ -5,6 +5,7 @@
 #include "Render/Viewport/DisplayClusterViewportManagerViewExtension.h"
 #include "Render/Viewport/DisplayClusterViewportManagerViewPointExtension.h"
 #include "Render/Viewport/DisplayClusterViewportFrameStatsViewExtension.h"
+#include "Render/Viewport/DisplayClusterViewportStereoscopicPass.h"
 
 #include "Render/Viewport/Preview/DisplayClusterViewportManagerPreview.h"
 #include "Render/Viewport/Preview/DisplayClusterViewportPreview.h"
@@ -570,7 +571,7 @@ void FDisplayClusterViewportManager::UnregisterCallbacks()
 #endif
 }
 
-int32 FDisplayClusterViewportManager::FindFirstViewportStereoViewIndex(const TSharedPtr<FDisplayClusterViewport, ESPMode::ThreadSafe>& InViewport) const
+int32 FDisplayClusterViewportManager::GetViewportStereoViewIndexBase(const TSharedPtr<FDisplayClusterViewport, ESPMode::ThreadSafe>& InViewport) const
 {
 	const int32 ViewPerViewportAmount = Configuration->GetRenderFrameSettings().GetViewPerViewportAmount();
 	if (InViewport.IsValid() && ViewPerViewportAmount > 0)
@@ -578,10 +579,14 @@ int32 FDisplayClusterViewportManager::FindFirstViewportStereoViewIndex(const TSh
 		const int32 ViewportIndex = ImplGetEntireClusterViewports().Find(InViewport);
 		if (ViewportIndex != INDEX_NONE)
 		{
-			// Begin from 1, because INDEX_NONE use ViewState[0] in LocalPlayer.cpp:786
-			const int32 FirstViewportStereoViewIndex = (ViewportIndex * ViewPerViewportAmount) + 1;
-
-			return FirstViewportStereoViewIndex;
+			// Calculate the unique StereoViewIndex base for the viewport.
+			const int32 ViewportStereoViewIndexBase =
+				// Each StereoViewIndexBase has a distance to others equal to ViewPerViewportAmount.
+				(ViewportIndex * ViewPerViewportAmount)
+				//  Shifts to the nDisplay range; engine indices are below this.
+				+ FDisplayClusterViewportStereoscopicPass::StereoViewIndexBase;
+
+			return ViewportStereoViewIndexBase;
 		}
 	}
 
@@ -628,8 +633,8 @@ bool FDisplayClusterViewportManager::BeginNewFrame(FViewport* InViewport, FDispl
 			continue;
 		}
 
-		const int32 FirstViewportStereoViewIndex = FindFirstViewportStereoViewIndex(ViewportIt);
-		if (FirstViewportStereoViewIndex != INDEX_NONE)
+		const int32 ViewportStereoViewIndexBase = GetViewportStereoViewIndexBase(ViewportIt);
+		if (ViewportStereoViewIndexBase != INDEX_NONE)
 		{
 			// Save orig viewport contexts
 			TArray<FDisplayClusterViewport_Context> PrevContexts;
@@ -640,7 +645,7 @@ bool FDisplayClusterViewportManager::BeginNewFrame(FViewport* InViewport, FDispl
 				ViewportIt->GetProjectionPolicy()->BeginUpdateFrameContexts(ViewportIt.Get());
 			}
 
-			ViewportIt->UpdateFrameContexts(FirstViewportStereoViewIndex);
+			ViewportIt->UpdateFrameContexts(ViewportStereoViewIndexBase);
 
 			if (ViewportIt->GetProjectionPolicy().IsValid())
 			{
@@ -1045,21 +1050,18 @@ void FDisplayClusterViewportManager::ImplDeleteViewport(const TSharedPtr<FDispla
 	ResetSceneRenderTargetSize();
 }
 
-IDisplayClusterViewport* FDisplayClusterViewportManager::FindViewport(const int32 ViewIndex, uint32* OutContextNum) const
+FDisplayClusterViewport* FDisplayClusterViewportManager::ImplFindViewport(const int32 StereoViewIndex, uint32* OutContextNum) const
 {
 	check(IsInGameThread());
 
-	// ViewIndex == eSSE_MONOSCOPIC(-1) is a special case called for ISR culling math.
-	// Since nDisplay is not ISR compatible, we ignore this request. This won't be neccessary once
-	// we stop using nDisplay as a stereoscopic rendering device (IStereoRendering).
-	if (ViewIndex < 0)
+	if (!FDisplayClusterViewportStereoscopicPass::IsValidStereoViewIndex(StereoViewIndex))
 	{
 		return nullptr;
 	}
 
 	for (const TSharedPtr<FDisplayClusterViewport, ESPMode::ThreadSafe>& ViewportIt : ImplGetEntireClusterViewports())
 	{
-		if (ViewportIt.IsValid() && ViewportIt->FindContext(ViewIndex, OutContextNum))
+		if (ViewportIt.IsValid() && ViewportIt->FindContext(StereoViewIndex, OutContextNum))
 		{
 			return ViewportIt.Get();
 		}
diff --git a/Engine/Plugins/Runtime/nDisplay/Source/DisplayCluster/Private/Render/Viewport/DisplayClusterViewportManager.h b/Engine/Plugins/Runtime/nDisplay/Source/DisplayCluster/Private/Render/Viewport/DisplayClusterViewportManager.h
--- a/Engine/Plugins/Runtime/nDisplay/Source/DisplayCluster/Private/Render/Viewport/DisplayClusterViewportManager.h
+++ b/Engine/Plugins/Runtime/nDisplay/Source/DisplayCluster/Private/Render/Viewport/DisplayClusterViewportManager.h
@@ -128,7 +128,10 @@ private:
 
 public:
 	virtual IDisplayClusterViewport* FindViewport(const FString& InViewportId) const override;
-	virtual IDisplayClusterViewport* FindViewport(const int32 ViewIndex, uint32* OutContextNum = nullptr) const override;
+	virtual IDisplayClusterViewport* FindViewport(const int32 ViewIndex, uint32* OutContextNum = nullptr) const override
+	{
+		return ImplFindViewport(ViewIndex, OutContextNum);
+	}
 
 	virtual const TArrayView<TSharedPtr<IDisplayClusterViewport, ESPMode::ThreadSafe>> GetCurrentRenderFrameViewports() const override
 	{
@@ -172,6 +175,8 @@ public:
 	FDisplayClusterViewport* CreateViewport(const FString& ViewportId, const TSharedPtr<IDisplayClusterProjectionPolicy, ESPMode::ThreadSafe>& InProjectionPolicy);
 	bool                     DeleteViewport(const FString& ViewportId);
 
+	FDisplayClusterViewport* ImplFindViewport(const int32 ViewIndex, uint32* OutContextNum = nullptr) const;
+
 	TSharedPtr<FDisplayClusterViewport, ESPMode::ThreadSafe> ImplCreateViewport(const FString& ViewportId, const TSharedPtr<IDisplayClusterProjectionPolicy, ESPMode::ThreadSafe>& InProjectionPolicy);
 	TSharedPtr<FDisplayClusterViewport, ESPMode::ThreadSafe> ImplFindViewport(const FString& InViewportId) const;
 	void ImplDeleteViewport(const TSharedPtr<FDisplayClusterViewport, ESPMode::ThreadSafe>& InExistViewport);
@@ -184,8 +189,16 @@ public:
 		return ViewportManagerProxy;
 	}
 
-	/** Return initial StereoViewIndex for the input viewport. */
-	int32 FindFirstViewportStereoViewIndex(const TSharedPtr<FDisplayClusterViewport, ESPMode::ThreadSafe>& InViewport) const;
+	/** 
+	* Gets the unique StereoViewIndex base for the viewport.
+	* Each StereoViewIndexBase has a distance to others equal to ViewPerViewportAmount.
+	* 
+	* Note: This function has been renamed. Its previous name was FindFirstViewportStereoViewIndex().
+	*
+	* @param InViewport - viewport asked for a StereoViewIndex base
+	* @return StereoViewIndex base or INDEX_NONE, if the viewport cannot be rendered.
+	*/
+	int32 GetViewportStereoViewIndexBase(const TSharedPtr<FDisplayClusterViewport, ESPMode::ThreadSafe>& InViewport) const;
 
 	/** Returns true if the final color is renderable (to support preview in scene). */
 	bool ShouldRenderFinalColor() const;
diff --git a/Engine/Plugins/Runtime/nDisplay/Source/DisplayCluster/Private/Render/Viewport/DisplayClusterViewportManagerProxy.cpp b/Engine/Plugins/Runtime/nDisplay/Source/DisplayCluster/Private/Render/Viewport/DisplayClusterViewportManagerProxy.cpp
--- a/Engine/Plugins/Runtime/nDisplay/Source/DisplayCluster/Private/Render/Viewport/DisplayClusterViewportManagerProxy.cpp
+++ b/Engine/Plugins/Runtime/nDisplay/Source/DisplayCluster/Private/Render/Viewport/DisplayClusterViewportManagerProxy.cpp
@@ -26,6 +26,7 @@
 #include "Render/Viewport/RenderTarget/DisplayClusterRenderTargetResource.h"
 #include "Render/Viewport/Postprocess/DisplayClusterViewportPostProcessManager.h"
 #include "Render/Viewport/Containers/DisplayClusterViewportProxyData.h"
+#include "Render/Viewport/DisplayClusterViewportStereoscopicPass.h"
 
 #include "Misc/DisplayClusterLog.h"
 
@@ -732,10 +733,7 @@ FDisplayClusterViewportProxy* FDisplayClusterViewportManagerProxy::ImplFindViewp
 {
 	check(IsInRenderingThread());
 	
-	// ViewIndex == eSSE_MONOSCOPIC(-1) is a special case called for ISR culling math.
-	// Since nDisplay is not ISR compatible, we ignore this request. This won't be neccessary once
-	// we stop using nDisplay as a stereoscopic rendering device (IStereoRendering).
-	if (StereoViewIndex < 0)
+	if (!FDisplayClusterViewportStereoscopicPass::IsValidStereoViewIndex(StereoViewIndex))
 	{
 		return nullptr;
 	}
diff --git a/Engine/Plugins/Runtime/nDisplay/Source/DisplayCluster/Private/Render/Viewport/DisplayClusterViewportStereoscopicPass.h b/Engine/Plugins/Runtime/nDisplay/Source/DisplayCluster/Private/Render/Viewport/DisplayClusterViewportStereoscopicPass.h
--- a/Engine/Plugins/Runtime/nDisplay/Source/DisplayCluster/Private/Render/Viewport/DisplayClusterViewportStereoscopicPass.h
+++ b/Engine/Plugins/Runtime/nDisplay/Source/DisplayCluster/Private/Render/Viewport/DisplayClusterViewportStereoscopicPass.h
@@ -8,6 +8,17 @@
 
 struct FDisplayClusterViewportStereoscopicPass
 {
+	// Base of the reserved StereoViewIndex range for nDisplay; engine indices are below this.
+	static constexpr int32 StereoViewIndexBase = 100;
+
+	/**
+	* True if this view index falls in the nDisplay-reserved range.
+	*/
+	inline static bool IsValidStereoViewIndex(const int32 ViewIndex)
+	{
+		return ViewIndex >= StereoViewIndexBase;
+	}
+
 	// 'EStereoscopicPass': Stereoscopic rendering passes.
 	// FULL implies stereoscopic rendering isn't enabled for this pass
 	// PRIMARY implies the view needs its own pass, while SECONDARY implies the view can be instanced
diff --git a/Engine/Plugins/Runtime/nDisplay/Source/DisplayCluster/Private/Render/Viewport/DisplayClusterViewport_Rules.cpp b/Engine/Plugins/Runtime/nDisplay/Source/DisplayCluster/Private/Render/Viewport/DisplayClusterViewport_Rules.cpp
--- a/Engine/Plugins/Runtime/nDisplay/Source/DisplayCluster/Private/Render/Viewport/DisplayClusterViewport_Rules.cpp
+++ b/Engine/Plugins/Runtime/nDisplay/Source/DisplayCluster/Private/Render/Viewport/DisplayClusterViewport_Rules.cpp
@@ -95,6 +95,22 @@ static FAutoConsoleVariableRef CVarDisplayClusterPreviewEnableConfiguratorViewSt
 	ECVF_RenderThreadSafe
 );
 
+int32 GDisplayClusterEnableViewStateForLightCard = 1;
+static FAutoConsoleVariableRef CVarDisplayClusterEnableViewStateForLightCard(
+	TEXT("nDisplay.render.ViewState.LightCard"),
+	GDisplayClusterEnableViewStateForLightCard,
+	TEXT("Enable view state for LightCard (0 - disable).\n"),
+	ECVF_RenderThreadSafe
+);
+
+int32 GDisplayClusterEnableViewStateForChromakey = 0;
+static FAutoConsoleVariableRef CVarDisplayClusterEnableViewStateForChromakey(
+	TEXT("nDisplay.render.ViewState.Chromakey"),
+	GDisplayClusterEnableViewStateForChromakey,
+	TEXT("Enable view state for Chromakey (0 - disable).\n"),
+	ECVF_RenderThreadSafe
+);
+
 
 ///////////////////////////////////////////////////////////////////////////////////////
 // FDisplayClusterViewport
@@ -532,28 +548,22 @@ bool FDisplayClusterViewport::ShouldUseViewStates() const
 	// RenderMode rules
 	switch (Configuration->GetRenderFrameSettings().RenderMode)
 	{
-		// Cluster node rendering
-	case EDisplayClusterRenderFrameMode::Mono:
-	case EDisplayClusterRenderFrameMode::Stereo:
-	case EDisplayClusterRenderFrameMode::SideBySide:
-	case EDisplayClusterRenderFrameMode::TopBottom:
-		// These modes use ViewState from ULocalPLayer, outside of nDisplay.
-		return false;
 
-		// MRQ rendering
+	// MRQ has its own ViewState management.
 	case EDisplayClusterRenderFrameMode::MRQ_Mono:
-		// These modes use ViewState from MRQ, outside of nDisplay.
 		return false;
 
 	// Preview-In-Scene
 	case EDisplayClusterRenderFrameMode::PreviewProxyHitInScene:
 	case EDisplayClusterRenderFrameMode::PreviewInScene:
+
+		// CVar 'nDisplay.preview.EnableViewState' disables ViewState for preview in the scene.
 		if (GDisplayClusterPreviewEnableViewState == 0)
 		{
 			return false;
 		}
 
-		// Disable ViewState for the preview inside the Configurator.
+		// CVar 'nDisplay.preview.EnableConfiguratorViewState' disables ViewState for preview inside the Configurator.
 		if (GDisplayClusterPreviewEnableConfiguratorViewState == 0 && Configuration->IsCurrentWorldHasAnyType(EWorldType::EditorPreview))
 		{
 			return false;
@@ -561,26 +571,24 @@ bool FDisplayClusterViewport::ShouldUseViewStates() const
 
 		break;
 
-		// PIE
-	case EDisplayClusterRenderFrameMode::PIE_Mono:
-	case EDisplayClusterRenderFrameMode::PIE_SideBySide:
-	case EDisplayClusterRenderFrameMode::PIE_TopBottom:
-		// We don't have any rules for PIE rendering
+	default:
 		break;
+	}
 
+	// ICVFX Rules
 
-	default:
+	// CVar 'nDisplay.render.ViewState.Chromakey' disables ViewState for Chromakey.
+	if (!GDisplayClusterEnableViewStateForChromakey
+	&& EnumHasAnyFlags(RenderSettingsICVFX.RuntimeFlags, EDisplayClusterViewportRuntimeICVFXFlags::Chromakey))
+	{
 		return false;
 	}
 
-	// Viewport type rules
-	switch (RenderSettings.CaptureMode)
+	// CVar 'nDisplay.render.ViewState.LightCard' disables ViewState for LightCard.
+	if (!GDisplayClusterEnableViewStateForLightCard
+		&& EnumHasAnyFlags(RenderSettingsICVFX.RuntimeFlags, EDisplayClusterViewportRuntimeICVFXFlags::Lightcard | EDisplayClusterViewportRuntimeICVFXFlags::UVLightcard))
 	{
-	case EDisplayClusterViewportCaptureMode::Chromakey:
 		return false;
-
-	default:
-		break;
 	}
 
 	return true;
diff --git a/Engine/Source/Runtime/Engine/Private/LocalPlayer.cpp b/Engine/Source/Runtime/Engine/Private/LocalPlayer.cpp
--- a/Engine/Source/Runtime/Engine/Private/LocalPlayer.cpp
+++ b/Engine/Source/Runtime/Engine/Private/LocalPlayer.cpp
@@ -832,28 +832,42 @@ bool ULocalPlayer::CalcSceneViewInitOptions(
 
 	check(PlayerController && PlayerController->GetWorld());
 
-	const uint32 ViewIndex = StereoViewIndex != INDEX_NONE ? StereoViewIndex : 0;
+	// True when this view should use a device-managed ViewState.
+	const bool bUseExternalViewState = GEngine->IsStereoscopic3D(Viewport)
+		&& GEngine->StereoRenderingDevice->HasExternalViewState(StereoViewIndex);
 
-	// Make sure the ViewStates array has enough elements for the given ViewIndex.
+	if (!bUseExternalViewState)
 	{
-		const int32 RequiredViewStates = (ViewIndex + 1) - ViewStates.Num();
-		
-		if (RequiredViewStates > 0)
+		const uint32 ViewIndex = StereoViewIndex != INDEX_NONE ? StereoViewIndex : 0;
+
+		// Make sure the ViewStates array has enough elements for the given ViewIndex.
 		{
-			ViewStates.AddDefaulted(RequiredViewStates);		
+			const int32 RequiredViewStates = (ViewIndex + 1) - ViewStates.Num();
+
+			if (RequiredViewStates > 0)
+			{
+				ViewStates.AddDefaulted(RequiredViewStates);
+			}
 		}
-	}
 
-	// Allocate the current ViewState if necessary
-	if (ViewStates[ViewIndex].GetReference() == nullptr)
-	{
-		const UWorld* CurrentWorld = GetWorld();
-		const ERHIFeatureLevel::Type FeatureLevel = CurrentWorld ? CurrentWorld->GetFeatureLevel() : GMaxRHIFeatureLevel;
+		// Allocate the current ViewState if necessary
+		if (ViewStates[ViewIndex].GetReference() == nullptr)
+		{
+			const UWorld* CurrentWorld = GetWorld();
+			const ERHIFeatureLevel::Type FeatureLevel = CurrentWorld ? CurrentWorld->GetFeatureLevel() : GMaxRHIFeatureLevel;
+
+			ViewStates[ViewIndex].Allocate(FeatureLevel);
+		}
 
-		ViewStates[ViewIndex].Allocate(FeatureLevel);
+		// Use ViewState from ULocalPlayer
+		ViewInitOptions.SceneViewStateInterface = ViewStates[ViewIndex].GetReference();
+	}
+	else
+	{
+		// Use ViewState from IStereoRendering
+		ViewInitOptions.SceneViewStateInterface = GEngine->StereoRenderingDevice->GetExternalViewState(StereoViewIndex);
 	}
 
-	ViewInitOptions.SceneViewStateInterface = ViewStates[ViewIndex].GetReference();
 	ViewInitOptions.ViewActor = PlayerController->GetViewTarget();
 
 	// TODO: Switch to GetLocalPlayerIndex during GetControllerId deprecation, this is only used by MotionControllerComponent
diff --git a/Engine/Source/Runtime/Engine/Public/StereoRendering.h b/Engine/Source/Runtime/Engine/Public/StereoRendering.h
--- a/Engine/Source/Runtime/Engine/Public/StereoRendering.h
+++ b/Engine/Source/Runtime/Engine/Public/StereoRendering.h
@@ -12,6 +12,7 @@
 class FSceneView;
 class IStereoLayers;
 class IStereoRenderTargetManager;
+class FSceneViewStateInterface;
 
 /**
 * Stereoscopic rendering passes.  FULL implies stereoscopic rendering isn't enabled for this pass
@@ -160,6 +161,25 @@ public:
 	 */
 	virtual FMatrix GetStereoProjectionMatrix(const int32 ViewIndex) const = 0;
 
+
+	/**
+	* True if the device manages the ViewState for the specified view index.
+	*/
+	virtual bool HasExternalViewState(const int32 ViewIndex) const
+	{
+		return false;
+	}
+
+	/**
+	* Gets a ViewState for the device, given the specified view index.
+	* Note: The device may return nullptr, but it is expected to configure
+	*       rendering properly (disable TemporalAA, Postprocessing, etc.).
+	*/
+	virtual FSceneViewStateInterface* GetExternalViewState(const int32 ViewIndex)
+	{
+		return nullptr;
+	}
+
 	/**
 	 * Sets view-specific params (such as view projection matrix) for the canvas.
 	 */
