nDisplay: Fixed editor crash when instanced stereo is enabled caused by bad pixel shader parameters when rendering nDisplay stage geometry

#Synced-CL 47990237

diff --git a/Engine/Plugins/Runtime/nDisplay/Shaders/Private/MeshProjectionNormalSmoothing.usf b/Engine/Plugins/Runtime/nDisplay/Shaders/Private/MeshProjectionNormalSmoothing.usf
--- a/Engine/Plugins/Runtime/nDisplay/Shaders/Private/MeshProjectionNormalSmoothing.usf
+++ b/Engine/Plugins/Runtime/nDisplay/Shaders/Private/MeshProjectionNormalSmoothing.usf
@@ -287,7 +287,7 @@ void BlurCS(uint2 DispatchThreadId : SV_DispatchThreadID)
 }
 
 /** A pixel shader that renders the results of any compute shader passes to an output render target, packing the normal and depth values into the correct RGBA channels */
-void OutputNormalMapPS(noperspective float4 UVAndScreenPos : TEXCOORD0, float4 SvPosition : SV_POSITION, out float4 OutColor : SV_Target0)
+void OutputNormalMapPS(noperspective float4 UVAndScreenPos : TEXCOORD0, out float4 OutColor : SV_Target0)
 {
 	const float2 UV = UVAndScreenPos.xy;
 	const uint2 PixelPos = int2(UV * Input_Extent);
