[nDisplay] Fixed deadlock when presenting frames from the submission thread

- Some nDisplay sync policies wait until frame completion within Present() call. It's Ok to do from the RHI thread, but it deadlocks when called from the submission thread. The fix makes D3D12 to call Present() from RHI thread only.

#rnx

#Synced-CL Unknown

diff --git a/Engine/Plugins/Runtime/nDisplay/Source/DisplayCluster/Private/Game/EngineClasses/Basics/DisplayClusterGameEngine.cpp b/Engine/Plugins/Runtime/nDisplay/Source/DisplayCluster/Private/Game/EngineClasses/Basics/DisplayClusterGameEngine.cpp
--- a/Engine/Plugins/Runtime/nDisplay/Source/DisplayCluster/Private/Game/EngineClasses/Basics/DisplayClusterGameEngine.cpp
+++ b/Engine/Plugins/Runtime/nDisplay/Source/DisplayCluster/Private/Game/EngineClasses/Basics/DisplayClusterGameEngine.cpp
@@ -12,6 +12,7 @@
 #include "Engine/NetDriver.h"
 
 #include "GameFramework/GameModeBase.h"
+#include "HAL/IConsoleManager.h"
 #include "Interfaces/IPv4/IPv4Endpoint.h"
 #include "Kismet/GameplayStatics.h"
 
@@ -67,6 +68,9 @@ void UDisplayClusterGameEngine::Init(class IEngineLoop* InEngineLoop)
 	// Detect requested operation mode
 	OperationMode = DetectOperationMode();
 
+	// Modify any required CVars before keep initialization
+	ModifyCVars();
+
 	// Initialize Display Cluster
 	if (!GDisplayCluster->Init(OperationMode))
 	{
@@ -665,3 +669,15 @@ void UDisplayClusterGameEngine::PendingLevelUpdate(FWorldContext& Context, float
 		TransitionType = ETransitionType::None;
 	}
 }
+
+void UDisplayClusterGameEngine::ModifyCVars()
+{
+	// Always call Present() from the RHI thread to prevent the submission thread from deadlocking.
+	// Some synchronization policies wait for frame completion. In other words, they use GPU fences
+	// for synchronization during Present() call. This can't be done on the submission thread as
+	// it would deadlock itself.
+	if (IConsoleVariable* CVar = IConsoleManager::Get().FindConsoleVariable(TEXT("r.D3D12.PresentOnSubmissionThread")))
+	{
+		CVar->Set(false);
+	}
+}
diff --git a/Engine/Plugins/Runtime/nDisplay/Source/DisplayCluster/Public/DisplayClusterGameEngine.h b/Engine/Plugins/Runtime/nDisplay/Source/DisplayCluster/Public/DisplayClusterGameEngine.h
--- a/Engine/Plugins/Runtime/nDisplay/Source/DisplayCluster/Public/DisplayClusterGameEngine.h
+++ b/Engine/Plugins/Runtime/nDisplay/Source/DisplayCluster/Public/DisplayClusterGameEngine.h
@@ -87,6 +87,9 @@ private:
 	 */
 	void PendingLevelUpdate(FWorldContext& Context, float DeltaSeconds);
 
+	/** Sets up CVars */
+	void ModifyCVars();
+
 private:
 
 	FDisplayClusterConfigurationDiagnostics Diagnostics;
