nDisplay: Added a way to control icvfx camera composition order per viewport.

The original behavior is that each camera has a render order value. Suppose they are as follows:

CamD - 4
CamC - 3
CamB - 2
CamA - 1

This will put CamD on top because it is the last to be composed.

Ghost Frame requires this order to be overridden per viewport. If a given viewport enforces the following order:

Index 0 - CamB
Index 1 - CamD

Then we expect the layered appearance to be:

CamB (as in the viewport enforced order)
CamD (as in the viewport enforced order)
CamC (above CamA as it was before the viewport setting)
CamA (below CamC as it was before the viewport setting)

Note that the viewport enforced order appears in the details panel in the same order as the inner camera layers, with Index 0 composed on top.

This supports ghost frame workflows with more than 2 cameras, which the reverse order flag was limited to.

Note: ReverseOrder option is still there and reverses the final order when used, but may be deprecated in the future since this option now supercedes it.

[RN] Inner camera composition order can now be overridden per viewport, to support led processors that support time slicing.

#Synced-CL 48544109

diff --git a/Engine/Plugins/Runtime/nDisplay/Source/DisplayCluster/Private/Render/Viewport/Configuration/DisplayClusterViewportConfigurationHelpers_ICVFX.cpp b/Engine/Plugins/Runtime/nDisplay/Source/DisplayCluster/Private/Render/Viewport/Configuration/DisplayClusterViewportConfigurationHelpers_ICVFX.cpp
--- a/Engine/Plugins/Runtime/nDisplay/Source/DisplayCluster/Private/Render/Viewport/Configuration/DisplayClusterViewportConfigurationHelpers_ICVFX.cpp
+++ b/Engine/Plugins/Runtime/nDisplay/Source/DisplayCluster/Private/Render/Viewport/Configuration/DisplayClusterViewportConfigurationHelpers_ICVFX.cpp
@@ -86,6 +86,16 @@ bool FDisplayClusterViewportConfigurationHelpers_ICVFX::CreateProjectionPolicyIC
 	return true;
 }
 
+FString FDisplayClusterViewportConfigurationHelpers_ICVFX::BuildCameraViewportId(const FString& InClusterNodeId, const FString& InCameraName)
+{
+	return ImplGetNameICVFX(
+		InClusterNodeId,
+		InCameraName,
+		DisplayClusterViewportStrings::icvfx::camera,
+		DisplayClusterViewportStrings::icvfx::suffix::none
+	);
+}
+
 bool FDisplayClusterViewportConfigurationHelpers_ICVFX::UpdateCameraProjectionSettingsICVFX(FDisplayClusterViewportConfiguration& InConfiguration, UDisplayClusterICVFXCameraComponent& InCameraComponent, const FDisplayClusterConfigurationICVFX_CameraSettings& InCameraSettings, const TSharedPtr<IDisplayClusterProjectionPolicy, ESPMode::ThreadSafe>& InProjectionPolicy)
 {
 	const FDisplayClusterConfigurationICVFX_StageSettings* StageSettings = InConfiguration.GetStageSettings();
diff --git a/Engine/Plugins/Runtime/nDisplay/Source/DisplayCluster/Private/Render/Viewport/Configuration/DisplayClusterViewportConfigurationHelpers_ICVFX.h b/Engine/Plugins/Runtime/nDisplay/Source/DisplayCluster/Private/Render/Viewport/Configuration/DisplayClusterViewportConfigurationHelpers_ICVFX.h
--- a/Engine/Plugins/Runtime/nDisplay/Source/DisplayCluster/Private/Render/Viewport/Configuration/DisplayClusterViewportConfigurationHelpers_ICVFX.h
+++ b/Engine/Plugins/Runtime/nDisplay/Source/DisplayCluster/Private/Render/Viewport/Configuration/DisplayClusterViewportConfigurationHelpers_ICVFX.h
@@ -75,4 +75,13 @@ public:
 
 	static FDisplayClusterViewport* FindViewportICVFX(FDisplayClusterViewportConfiguration& InConfiguration, const FString& InViewportId, const FString& InResourceId, const FString& InResourceSuffixId);
 	static FDisplayClusterViewport* CreateViewportICVFX(FDisplayClusterViewportConfiguration& InConfiguration, const FString& InViewportId, const FString& InResourceId, const FString& InResourceSuffixId, TSharedPtr<IDisplayClusterProjectionPolicy, ESPMode::ThreadSafe>& InProjectionPolicy);
+
+	/** Constructs the full ICVFX camera viewport ID from a camera name.
+	 * 
+	 * @param InClusterNodeId - The cluster node ID
+	 * @param InCameraName - The camera component name
+	 * 
+	 * @return The constructed viewport ID string.
+	 */
+	static FString BuildCameraViewportId(const FString& InClusterNodeId, const FString& InCameraName);
 };
diff --git a/Engine/Plugins/Runtime/nDisplay/Source/DisplayCluster/Private/Render/Viewport/Configuration/DisplayClusterViewportConfiguration_ICVFX.cpp b/Engine/Plugins/Runtime/nDisplay/Source/DisplayCluster/Private/Render/Viewport/Configuration/DisplayClusterViewportConfiguration_ICVFX.cpp
--- a/Engine/Plugins/Runtime/nDisplay/Source/DisplayCluster/Private/Render/Viewport/Configuration/DisplayClusterViewportConfiguration_ICVFX.cpp
+++ b/Engine/Plugins/Runtime/nDisplay/Source/DisplayCluster/Private/Render/Viewport/Configuration/DisplayClusterViewportConfiguration_ICVFX.cpp
@@ -161,6 +161,29 @@ void FDisplayClusterViewportConfiguration_ICVFX::Update()
 			// Gain direct access to internal settings of the viewport:
 			FDisplayClusterViewport_RenderSettingsICVFX& InOutRenderSettingsICVFX = TargetIt->GetRenderSettingsICVFXImpl();
 
+			// Clear enforced camera order each frame to avoid stale data
+			InOutRenderSettingsICVFX.ICVFX.EnforcedCameraOrder.Empty();
+
+			// Build enforced camera order by resolving camera names to viewport IDs
+			if (const UDisplayClusterConfigurationViewport* ViewportCfg = TargetIt->GetViewportConfigurationData())
+			{
+
+				if (!ViewportCfg->ICVFX.EnforcedCameraOrder.IsEmpty())
+				{
+					const FString& ClusterNodeId = Configuration.GetClusterNodeId();
+
+					// Build full viewport IDs from camera names using helper
+					for (const FString& CameraName : ViewportCfg->ICVFX.EnforcedCameraOrder)
+					{
+						if (!CameraName.IsEmpty())
+						{
+							const FString ViewportId = FDisplayClusterViewportConfigurationHelpers_ICVFX::BuildCameraViewportId(ClusterNodeId, CameraName);
+							InOutRenderSettingsICVFX.ICVFX.EnforcedCameraOrder.Add(ViewportId);
+						}
+					}
+				}
+			}
+
 			// Sort cameras by render order for each target
 			InOutRenderSettingsICVFX.ICVFX.SortCamerasRenderOrder(EnumHasAnyFlags(InOutRenderSettingsICVFX.Flags, EDisplayClusterViewportICVFXFlags::ReverseCameraPriority));
 
@@ -323,9 +346,13 @@ void FDisplayClusterViewportConfiguration_ICVFX::GetAndUpdateStageCameras(const
 
 bool FDisplayClusterViewportConfiguration_ICVFX::ImplGetStageCameras()
 {
+	// Clear stage cameras unconditionally to prevent stale data
+	StageCameras.Empty();
+
 	ADisplayClusterRootActor* SceneRootActor = Configuration.GetRootActor(EDisplayClusterRootActorType::Scene);
 	ADisplayClusterRootActor* ConfigurationRootActor = Configuration.GetRootActor(EDisplayClusterRootActorType::Configuration);
 	const UDisplayClusterConfigurationData* ConfigurationData = Configuration.GetConfigurationData();
+
 	if (!(SceneRootActor && ConfigurationRootActor && ConfigurationData))
 	{
 		return false;
diff --git a/Engine/Plugins/Runtime/nDisplay/Source/DisplayClusterConfiguration/Public/DisplayClusterConfigurationTypes_Viewport.h b/Engine/Plugins/Runtime/nDisplay/Source/DisplayClusterConfiguration/Public/DisplayClusterConfigurationTypes_Viewport.h
--- a/Engine/Plugins/Runtime/nDisplay/Source/DisplayClusterConfiguration/Public/DisplayClusterConfigurationTypes_Viewport.h
+++ b/Engine/Plugins/Runtime/nDisplay/Source/DisplayClusterConfiguration/Public/DisplayClusterConfigurationTypes_Viewport.h
@@ -80,6 +80,14 @@ public:
 	UPROPERTY(EditAnywhere, BlueprintReadWrite, Category = "In Camera VFX")
 	TMap<FString, EDisplayClusterConfigurationICVFX_OverrideChromakeyType> PerCameraOverrideChromakeyType;
 
+	/** Enforced rendering order for ICVFX cameras on this viewport.
+	 * Cameras listed here will be rendered in this order (index 0 is composed on top).
+	 * Cameras not in this list will be composed below, in their default order based on RenderOrder property.
+	 * If empty, all cameras use default render order. 
+	 */
+	UPROPERTY(EditAnywhere, BlueprintReadWrite, Category = "In Camera VFX")
+	TArray<FString> EnforcedCameraOrder;
+
 	/** The order in which the ICVFX cameras are composited over is reversed. Useful for time-multiplexed displays. */
 	UPROPERTY(EditAnywhere, BlueprintReadWrite, Category = "In Camera VFX")
 	bool bReverseCameraPriority = false;
diff --git a/Engine/Plugins/Runtime/nDisplay/Source/DisplayClusterConfigurator/Private/DisplayClusterConfiguratorModule.cpp b/Engine/Plugins/Runtime/nDisplay/Source/DisplayClusterConfigurator/Private/DisplayClusterConfiguratorModule.cpp
--- a/Engine/Plugins/Runtime/nDisplay/Source/DisplayClusterConfigurator/Private/DisplayClusterConfiguratorModule.cpp
+++ b/Engine/Plugins/Runtime/nDisplay/Source/DisplayClusterConfigurator/Private/DisplayClusterConfiguratorModule.cpp
@@ -1,59 +1,61 @@
 // Copyright Epic Games, Inc. All Rights Reserved.
 
 #include "DisplayClusterConfiguratorModule.h"
-#include "DisplayClusterConfiguratorStyle.h"
-#include "DisplayClusterConfiguratorCommands.h"
+
 #include "DisplayClusterConfiguratorAssetTypeActions.h"
+#include "DisplayClusterConfiguratorCommands.h"
+#include "DisplayClusterConfiguratorLog.h"
+#include "DisplayClusterConfiguratorStyle.h"
 #include "DisplayClusterConfiguratorVersionUtils.h"
 #include "Framework/Application/SlateApplication.h"
 #include "PropertyEditorModule.h"
 #include "Settings/DisplayClusterConfiguratorSettings.h"
 #include "Views/Details/DisplayClusterRootActorDetailsCustomization.h"
 #include "Views/OutputMapping/DisplayClusterConfiguratorOutputMappingCommands.h"
-#include "DisplayClusterConfiguratorLog.h"
 
-#include "Views/Details/DisplayClusterConfiguratorBaseTypeCustomization.h"
-#include "Views/Details/DisplayClusterEditorPropertyReferenceTypeCustomization.h"
-#include "Views/Details/DisplayClusterConfiguratorBaseDetailCustomization.h"
-#include "Views/Details/Cluster/DisplayClusterConfiguratorDataDetailsCustomization.h"
+#include "Views/Details/Cluster/DisplayClusterConfiguratorCameraOrderCustomization.h"
 #include "Views/Details/Cluster/DisplayClusterConfiguratorClusterDetailsCustomization.h"
+#include "Views/Details/Cluster/DisplayClusterConfiguratorClusterReferenceListCustomization.h"
+#include "Views/Details/Cluster/DisplayClusterConfiguratorDataDetailsCustomization.h"
 #include "Views/Details/Cluster/DisplayClusterConfiguratorExternalImageTypeCustomization.h"
 #include "Views/Details/Cluster/DisplayClusterConfiguratorGenerateMipsCustomization.h"
 #include "Views/Details/Cluster/DisplayClusterConfiguratorNodeSelectionCustomization.h"
-#include "Views/Details/Cluster/DisplayClusterConfiguratorClusterReferenceListCustomization.h"
+#include "Views/Details/Cluster/DisplayClusterConfiguratorRectangleCustomization.h"
 #include "Views/Details/Cluster/DisplayClusterConfiguratorViewportDetailsCustomization.h"
 #include "Views/Details/Cluster/DisplayClusterConfiguratorViewportRemapCustomization.h"
-#include "Views/Details/Cluster/DisplayClusterConfiguratorRectangleCustomization.h"
+#include "Views/Details/Components/DisplayClusterCameraComponentDetailsCustomization.h"
 #include "Views/Details/Components/DisplayClusterConfiguratorScreenComponentDetailsCustomization.h"
 #include "Views/Details/Components/DisplayClusterICVFXCameraComponentDetailsCustomization.h"
+#include "Views/Details/DisplayClusterConfiguratorBaseDetailCustomization.h"
+#include "Views/Details/DisplayClusterConfiguratorBaseTypeCustomization.h"
+#include "Views/Details/DisplayClusterEditorPropertyReferenceTypeCustomization.h"
 #include "Views/Details/Media/DCConfiguratorClusterNodeMediaCustomization.h"
 #include "Views/Details/Media/DCConfiguratorICVFXMediaCustomization.h"
 #include "Views/Details/Media/DCConfiguratorViewportMediaCustomization.h"
 #include "Views/Details/Media/DisplayClusterConfiguratorMediaFullFrameCustomization.h"
 #include "Views/Details/Media/DisplayClusterConfiguratorMediaTileCustomization.h"
-#include "Views/Details/Components/DisplayClusterCameraComponentDetailsCustomization.h"
 #include "Views/Details/Policies/DisplayClusterConfiguratorPolicyDetailCustomization.h"
 #include "Views/Details/Upscaler/DisplayClusterConfigurationUpscalerSettingsDetailCustomization.h"
 
 #include "Blueprints/DisplayClusterBlueprint.h"
-#include "Components/DisplayClusterScreenComponent.h"
 #include "Components/DisplayClusterCameraComponent.h"
-#include "Components/DisplayClusterInFrustumFitCameraComponent.h"
 #include "Components/DisplayClusterICVFXCameraComponent.h"
-#include "Misc/DisplayClusterObjectRef.h"
+#include "Components/DisplayClusterInFrustumFitCameraComponent.h"
+#include "Components/DisplayClusterScreenComponent.h"
 #include "DisplayClusterRootActor.h"
+#include "Misc/DisplayClusterObjectRef.h"
 
+#include "AssetRegistry/AssetRegistryModule.h"
 #include "AssetToolsModule.h"
+#include "AssetTypeCategories.h"
 #include "IAssetTools.h"
 #include "ISettingsModule.h"
-#include "AssetRegistry/AssetRegistryModule.h"
-#include "AssetTypeCategories.h"
 
+#include "ActorFactories/ActorFactoryBlueprint.h"
 #include "DisplayClusterConfiguratorBlueprintEditor.h"
 #include "DisplayClusterConfiguratorToolbarExtensions.h"
 #include "HAL/IConsoleManager.h"
 #include "Modules/ModuleManager.h"
-#include "ActorFactories/ActorFactoryBlueprint.h"
 #include "ScopedTransaction.h"
 #include "SEditorViewportToolBarButton.h"
 #include "Subsystems/PanelExtensionSubsystem.h"
@@ -164,7 +166,7 @@ void FDisplayClusterConfiguratorModule::ShutdownModule()
 	IKismetCompilerInterface& KismetCompilerModule = FModuleManager::GetModuleChecked<IKismetCompilerInterface>("KismetCompiler");
 	KismetCompilerModule.GetCompilers().Remove(&BlueprintCompiler);
 
-	if (FilesLoadedHandle.IsValid() && FModuleManager::Get().IsModuleLoaded("AssetTools"))
+	if (FilesLoadedHandle.IsValid() && FModuleManager::Get().IsModuleLoaded("AssetRegistry"))
 	{
 		FAssetRegistryModule& AssetRegistryModule = FModuleManager::LoadModuleChecked<FAssetRegistryModule>(TEXT("AssetRegistry"));
 		AssetRegistryModule.Get().OnFilesLoaded().Remove(FilesLoadedHandle);
@@ -208,7 +210,7 @@ void FDisplayClusterConfiguratorModule::UnregisterSettings()
 {
 	if (ISettingsModule* SettingsModule = FModuleManager::GetModulePtr<ISettingsModule>("Settings"))
 	{
-		SettingsModule->UnregisterSettings("Editor", "ContentEditors", "nDisplayEditor");
+		SettingsModule->UnregisterSettings("Editor", "Plugins", "nDisplayEditor");
 	}
 }
 
diff --git a/Engine/Plugins/Runtime/nDisplay/Source/DisplayClusterConfigurator/Private/Views/Details/Cluster/DisplayClusterConfiguratorCameraOrderCustomization.cpp b/Engine/Plugins/Runtime/nDisplay/Source/DisplayClusterConfigurator/Private/Views/Details/Cluster/DisplayClusterConfiguratorCameraOrderCustomization.cpp
new file mode 100644
--- /dev/null
+++ b/Engine/Plugins/Runtime/nDisplay/Source/DisplayClusterConfigurator/Private/Views/Details/Cluster/DisplayClusterConfiguratorCameraOrderCustomization.cpp
@@ -0,0 +1,49 @@
+// Copyright Epic Games, Inc. All Rights Reserved.
+
+#include "DisplayClusterConfiguratorCameraOrderCustomization.h"
+
+#include "DisplayClusterConfigurationTypes_Viewport.h"
+#include "DisplayClusterConfiguratorBlueprintEditor.h"
+#include "DisplayClusterConfiguratorNodeSelectionCustomization.h"
+#include "DisplayClusterRootActor.h"
+
+#include "DetailLayoutBuilder.h"
+#include "IDetailChildrenBuilder.h"
+#include "PropertyHandle.h"
+
+void FDisplayClusterConfiguratorCameraOrderCustomization::Initialize(const TSharedRef<IPropertyHandle>& InPropertyHandle, IPropertyTypeCustomizationUtils& CustomizationUtils)
+{
+	FDisplayClusterConfiguratorBaseTypeCustomization::Initialize(InPropertyHandle, CustomizationUtils);
+
+	// Get the EnforcedCameraOrder array property handle
+	CameraOrderArrayHandle = InPropertyHandle;
+
+	FDisplayClusterConfiguratorBlueprintEditor* BlueprintEditor = FindBlueprintEditor();
+	ADisplayClusterRootActor* RootActor = FindRootActor();
+
+	if (RootActor != nullptr || BlueprintEditor != nullptr)
+	{
+		// Create a node selection builder configured for ICVFX cameras
+		NodeSelectionBuilder = MakeShared<FDisplayClusterConfiguratorNodeSelection>(
+			FDisplayClusterConfiguratorNodeSelection::EOperationMode::ICVFXCameras,
+			RootActor,
+			BlueprintEditor);
+		bCanCustomizeDisplay = true;
+	}
+}
+
+void FDisplayClusterConfiguratorCameraOrderCustomization::SetChildren(const TSharedRef<IPropertyHandle>& InPropertyHandle, IDetailChildrenBuilder& InChildBuilder, IPropertyTypeCustomizationUtils& CustomizationUtils)
+{
+	if (bCanCustomizeDisplay && CameraOrderArrayHandle && CameraOrderArrayHandle->IsValidHandle())
+	{
+		CameraOrderArrayHandle->SetPropertyDisplayName(InPropertyHandle->GetPropertyDisplayName());
+		CameraOrderArrayHandle->SetToolTipText(InPropertyHandle->GetToolTipText());
+
+		NodeSelectionBuilder->IsEnabled(true);
+		NodeSelectionBuilder->CreateArrayBuilder(CameraOrderArrayHandle.ToSharedRef(), InChildBuilder);
+	}
+	else
+	{
+		FDisplayClusterConfiguratorBaseTypeCustomization::SetChildren(InPropertyHandle, InChildBuilder, CustomizationUtils);
+	}
+}
diff --git a/Engine/Plugins/Runtime/nDisplay/Source/DisplayClusterConfigurator/Private/Views/Details/Cluster/DisplayClusterConfiguratorCameraOrderCustomization.h b/Engine/Plugins/Runtime/nDisplay/Source/DisplayClusterConfigurator/Private/Views/Details/Cluster/DisplayClusterConfiguratorCameraOrderCustomization.h
new file mode 100644
--- /dev/null
+++ b/Engine/Plugins/Runtime/nDisplay/Source/DisplayClusterConfigurator/Private/Views/Details/Cluster/DisplayClusterConfiguratorCameraOrderCustomization.h
@@ -0,0 +1,45 @@
+// Copyright Epic Games, Inc. All Rights Reserved.
+
+#pragma once
+
+#include "CoreMinimal.h"
+#include "Views/Details/DisplayClusterConfiguratorBaseTypeCustomization.h"
+
+class ADisplayClusterRootActor;
+class FDisplayClusterConfiguratorBlueprintEditor;
+class FDisplayClusterConfiguratorNodeSelection;
+
+/**
+ * Type customization for the EnforcedCameraOrder TArray<FString> property in FDisplayClusterConfigurationViewport_ICVFX.
+ * Displays a dropdown for each array element listing all available ICVFX cameras.
+ */
+class FDisplayClusterConfiguratorCameraOrderCustomization final : public FDisplayClusterConfiguratorBaseTypeCustomization
+{
+public:
+	static TSharedRef<IPropertyTypeCustomization> MakeInstance()
+	{
+		return MakeShared<FDisplayClusterConfiguratorCameraOrderCustomization>();
+	}
+
+	~FDisplayClusterConfiguratorCameraOrderCustomization()
+	{
+		NodeSelectionBuilder.Reset();
+	}
+
+protected:
+	//~ FDisplayClusterConfiguratorBaseTypeCustomization overrides begin
+	virtual void Initialize(const TSharedRef<IPropertyHandle>& InPropertyHandle, IPropertyTypeCustomizationUtils& CustomizationUtils) override;
+	virtual void SetChildren(const TSharedRef<IPropertyHandle>& InPropertyHandle, IDetailChildrenBuilder& InChildBuilder, IPropertyTypeCustomizationUtils& CustomizationUtils) override;
+	virtual bool ShouldShowHeader(const TSharedRef<IPropertyHandle>& InPropertyHandle) const override { return false; }
+	//~ FDisplayClusterConfiguratorBaseTypeCustomization overrides end
+
+private:
+	/** The helper object used to generate the custom array widget to display in place of a string list for the EnforcedCameraOrder property */
+	TSharedPtr<FDisplayClusterConfiguratorNodeSelection> NodeSelectionBuilder;
+
+	/** The property handle for the EnforcedCameraOrder array */
+	TSharedPtr<IPropertyHandle> CameraOrderArrayHandle;
+
+	/** Indicates whether this customizer can customize the array property's display or not */
+	bool bCanCustomizeDisplay = false;
+};
diff --git a/Engine/Plugins/Runtime/nDisplay/Source/DisplayClusterConfigurator/Private/Views/Details/Cluster/DisplayClusterConfiguratorNodeSelectionCustomization.cpp b/Engine/Plugins/Runtime/nDisplay/Source/DisplayClusterConfigurator/Private/Views/Details/Cluster/DisplayClusterConfiguratorNodeSelectionCustomization.cpp
--- a/Engine/Plugins/Runtime/nDisplay/Source/DisplayClusterConfigurator/Private/Views/Details/Cluster/DisplayClusterConfiguratorNodeSelectionCustomization.cpp
+++ b/Engine/Plugins/Runtime/nDisplay/Source/DisplayClusterConfigurator/Private/Views/Details/Cluster/DisplayClusterConfiguratorNodeSelectionCustomization.cpp
@@ -2,10 +2,11 @@
 
 #include "DisplayClusterConfiguratorNodeSelectionCustomization.h"
 
-#include "DisplayClusterRootActor.h"
+#include "Components/DisplayClusterICVFXCameraComponent.h"
 #include "DisplayClusterConfigurationTypes.h"
 #include "DisplayClusterConfiguratorBlueprintEditor.h"
 #include "DisplayClusterConfiguratorUtils.h"
+#include "DisplayClusterRootActor.h"
 #include "Views/Details/Widgets/SDisplayClusterConfigurationSearchableComboBox.h"
 
 #include "DetailLayoutBuilder.h"
@@ -51,7 +52,8 @@ ADisplayClusterRootActor* FDisplayClusterConfiguratorNodeSelection::GetRootActor
 		RootActor = RootActorPtr.Get();
 	}
 
-	check(RootActor);
+	ensure(RootActor);
+
 	return RootActor;
 }
 
@@ -155,7 +157,25 @@ void FDisplayClusterConfiguratorNodeSelection::GenerateSelectionWidget(
 void FDisplayClusterConfiguratorNodeSelection::ResetOptions()
 {
 	Options.Reset();
-	if (UDisplayClusterConfigurationData* ConfigData = GetConfigData())
+
+	if (OperationMode == ICVFXCameras)
+	{
+		// Get ICVFX camera components from the root actor
+		if (ADisplayClusterRootActor* RootActor = GetRootActor())
+		{
+			TArray<UActorComponent*> CameraComponents;
+			RootActor->GetComponents(UDisplayClusterICVFXCameraComponent::StaticClass(), CameraComponents);
+
+			for (UActorComponent* Component : CameraComponents)
+			{
+				if (Component->IsA<UDisplayClusterICVFXCameraComponent>())
+				{
+					Options.Add(MakeShared<FString>(Component->GetName()));
+				}
+			}
+		}
+	}
+	else if (UDisplayClusterConfigurationData* ConfigData = GetConfigData())
 	{
 		for (const TTuple<FString, TObjectPtr<UDisplayClusterConfigurationClusterNode>>& Node : ConfigData->Cluster->Nodes)
 		{
diff --git a/Engine/Plugins/Runtime/nDisplay/Source/DisplayClusterConfigurator/Private/Views/Details/Cluster/DisplayClusterConfiguratorNodeSelectionCustomization.h b/Engine/Plugins/Runtime/nDisplay/Source/DisplayClusterConfigurator/Private/Views/Details/Cluster/DisplayClusterConfiguratorNodeSelectionCustomization.h
--- a/Engine/Plugins/Runtime/nDisplay/Source/DisplayClusterConfigurator/Private/Views/Details/Cluster/DisplayClusterConfiguratorNodeSelectionCustomization.h
+++ b/Engine/Plugins/Runtime/nDisplay/Source/DisplayClusterConfigurator/Private/Views/Details/Cluster/DisplayClusterConfiguratorNodeSelectionCustomization.h
@@ -14,7 +14,7 @@ class SWidget;
 
 /**
  * Helper class used to construct and manage a custom array property widget that displays for each element a dropdown containing
- * a list of the display cluster's cluster nodes or viewports to select from.
+ * a list of the display cluster's cluster nodes, viewports, or ICVFX cameras to select from.
  */
 class FDisplayClusterConfiguratorNodeSelection : public TSharedFromThis<FDisplayClusterConfiguratorNodeSelection>
 {
@@ -26,7 +26,8 @@ public:
 	enum EOperationMode
 	{
 		Viewports,
-		ClusterNodes
+		ClusterNodes,
+		ICVFXCameras
 	};
 
 	FDisplayClusterConfiguratorNodeSelection(EOperationMode InMode, ADisplayClusterRootActor* InRootActor, FDisplayClusterConfiguratorBlueprintEditor* InToolkitPtr);
@@ -53,13 +54,14 @@ public:
 		return *this;
 	}
 
-	/** Inspects the specified property's metadata to determine if the node selection array should show cluster nodes or viewports in its dropdown lists */
+	/** Inspects the specified property's metadata to determine if the node selection array should show cluster nodes or viewports or icvfx cameras in its dropdown lists */
 	static EOperationMode GetOperationModeFromProperty(FProperty* Property);
 
-protected:
 	/** Generates a widget to display the currently selected cluster item for the specified element in the array widget and adds it to the details panel */
 	void GenerateSelectionWidget(TSharedRef<IPropertyHandle> PropertyHandle, int32 ArrayIndex, IDetailChildrenBuilder& ChildrenBuilder);
 
+protected:
+
 	/** Resets the list of cluster items to display in the dropdown lists for each element of the array widget */
 	void ResetOptions();
 
@@ -88,7 +90,7 @@ private:
 	/** A weak reference to the current display cluster root actor that is being edited by the details panel */
 	TWeakObjectPtr<ADisplayClusterRootActor> RootActorPtr;
 
-	/** The operation mode which indicates whether to display cluster nodes or viewports in each element's combo box */
+	/** The operation mode which indicates whether to display cluster nodes, viewports, or ICVFX cameras in each element's combo box */
 	EOperationMode OperationMode = ClusterNodes;
 };
 
diff --git a/Engine/Plugins/Runtime/nDisplay/Source/DisplayClusterConfigurator/Private/Views/Details/Cluster/DisplayClusterConfiguratorViewportDetailsCustomization.cpp b/Engine/Plugins/Runtime/nDisplay/Source/DisplayClusterConfigurator/Private/Views/Details/Cluster/DisplayClusterConfiguratorViewportDetailsCustomization.cpp
--- a/Engine/Plugins/Runtime/nDisplay/Source/DisplayClusterConfigurator/Private/Views/Details/Cluster/DisplayClusterConfiguratorViewportDetailsCustomization.cpp
+++ b/Engine/Plugins/Runtime/nDisplay/Source/DisplayClusterConfigurator/Private/Views/Details/Cluster/DisplayClusterConfiguratorViewportDetailsCustomization.cpp
@@ -2,19 +2,22 @@
 
 #include "DisplayClusterConfiguratorViewportDetailsCustomization.h"
 
-#include "DisplayClusterRootActor.h"
+#include "Components/DisplayClusterCameraComponent.h"
 #include "DisplayClusterConfigurationTypes.h"
 #include "DisplayClusterConfigurationTypes_Viewport.h"
-#include "DisplayClusterProjectionStrings.h"
 #include "DisplayClusterConfiguratorLog.h"
-#include "Components/DisplayClusterCameraComponent.h"
+#include "DisplayClusterProjectionStrings.h"
+#include "DisplayClusterRootActor.h"
 #include "Render/DisplayDevice/Components/DisplayClusterDisplayDeviceBaseComponent.h"
-#include "Views/Details/Widgets/SDisplayClusterConfiguratorComponentPicker.h"
+#include "Views/Details/Cluster/DisplayClusterConfiguratorNodeSelectionCustomization.h"
 #include "Views/Details/Widgets/SDisplayClusterConfigurationSearchableComboBox.h"
+#include "Views/Details/Widgets/SDisplayClusterConfiguratorComponentPicker.h"
 
-#include "DetailLayoutBuilder.h"
 #include "DetailCategoryBuilder.h"
+#include "DetailLayoutBuilder.h"
 #include "DetailWidgetRow.h"
+#include "IDetailChildrenBuilder.h"
+#include "PropertyCustomizationHelpers.h"
 #include "PropertyHandle.h"
 
 void FDisplayClusterConfiguratorViewportDetailsCustomization::CustomizeDetails(IDetailLayoutBuilder& InLayoutBuilder)
@@ -40,13 +43,23 @@ void FDisplayClusterConfiguratorViewportDetailsCustomization::CustomizeDetails(I
 	{
 		ConfigurationViewportPtr = Cast<UDisplayClusterConfigurationViewport>(SelectedObjects[0]);
 	}
-	check(ConfigurationViewportPtr != nullptr);
+
+	if (ConfigurationViewportPtr == nullptr)
+	{
+		UE_LOG(DisplayClusterConfiguratorLog, Warning, TEXT("No valid viewport selected in details panel."));
+		return;
+	}
 
 	if (!ConfigurationViewportPtr->ProjectionPolicy.ShouldDisplayICVFXCategory())
 	{
 		IDetailCategoryBuilder& Category = InLayoutBuilder.EditCategory(DisplayClusterConfigurationStrings::categories::InCameraVFXCategory);
 		Category.SetCategoryVisibility(false);
 	}
+	else
+	{
+		// Customize the EnforcedCameraOrder array property with ICVFX camera dropdowns
+		CustomizeEnforcedCameraOrder(InLayoutBuilder);
+	}
 
 	{
 		CameraHandle = InLayoutBuilder.GetProperty(GET_MEMBER_NAME_CHECKED(UDisplayClusterConfigurationViewport, Camera));
@@ -133,8 +146,14 @@ void FDisplayClusterConfiguratorViewportDetailsCustomization::ResetCameraOptions
 
 	AActor* RootActor = GetRootActor();
 
+	if (!RootActor)
+	{
+		return;
+	}
+
 	TArray<UActorComponent*> ActorComponents;
 	RootActor->GetComponents(UDisplayClusterCameraComponent::StaticClass(), ActorComponents);
+
 	for (UActorComponent* ActorComponent : ActorComponents)
 	{
 		const FString ComponentName = ActorComponent->GetName();
@@ -209,3 +228,60 @@ FText FDisplayClusterConfiguratorViewportDetailsCustomization::GetSelectedCamera
 
 	return FText::FromString(SelectedOption);
 }
+
+void FDisplayClusterConfiguratorViewportDetailsCustomization::CustomizeEnforcedCameraOrder(IDetailLayoutBuilder& InLayoutBuilder)
+{
+	// Get the ICVFX property group
+	TSharedPtr<IPropertyHandle> ICVFXHandle = InLayoutBuilder.GetProperty(
+		GET_MEMBER_NAME_CHECKED(UDisplayClusterConfigurationViewport, ICVFX));
+
+	if (!ICVFXHandle.IsValid() || !ICVFXHandle->IsValidHandle())
+	{
+		return;
+	}
+
+	// Get the EnforcedCameraOrder array property
+	TSharedPtr<IPropertyHandle> EnforcedCameraOrderHandle = ICVFXHandle->GetChildHandle(
+		GET_MEMBER_NAME_CHECKED(FDisplayClusterConfigurationViewport_ICVFX, EnforcedCameraOrder));
+
+	if (!EnforcedCameraOrderHandle.IsValid() || !EnforcedCameraOrderHandle->IsValidHandle())
+	{
+		return;
+	}
+
+	ADisplayClusterRootActor* RootActor = GetRootActor();
+
+	if (RootActor == nullptr)
+	{
+		return;
+	}
+
+	// Create a node selection builder configured for ICVFX cameras
+	CameraOrderNodeSelection = MakeShared<FDisplayClusterConfiguratorNodeSelection>(
+		FDisplayClusterConfiguratorNodeSelection::EOperationMode::ICVFXCameras,
+		RootActor,
+		nullptr /* FDisplayClusterConfiguratorBlueprintEditor */);
+
+	CameraOrderNodeSelection->IsEnabled(true);
+
+	// Get the property row and hide it since we're replacing with custom builder
+	IDetailPropertyRow* PropertyRow = InLayoutBuilder.EditDefaultProperty(EnforcedCameraOrderHandle);
+	if (PropertyRow)
+	{
+		PropertyRow->Visibility(EVisibility::Collapsed);
+	}
+
+	// Add the custom array builder to the ICVFX category
+	IDetailCategoryBuilder& ICVFXCategory = InLayoutBuilder.EditCategory(
+		DisplayClusterConfigurationStrings::categories::InCameraVFXCategory);
+
+	// Create a custom array builder manually since we need to pass IDetailCategoryBuilder
+	TSharedRef<FDetailArrayBuilder> ArrayBuilder = MakeShared<FDetailArrayBuilder>(
+		EnforcedCameraOrderHandle.ToSharedRef());
+
+	ArrayBuilder->OnGenerateArrayElementWidget(FOnGenerateArrayElementWidget::CreateSP(
+		CameraOrderNodeSelection.Get(),
+		&FDisplayClusterConfiguratorNodeSelection::GenerateSelectionWidget));
+
+	ICVFXCategory.AddCustomBuilder(ArrayBuilder);
+}
diff --git a/Engine/Plugins/Runtime/nDisplay/Source/DisplayClusterConfigurator/Private/Views/Details/Cluster/DisplayClusterConfiguratorViewportDetailsCustomization.h b/Engine/Plugins/Runtime/nDisplay/Source/DisplayClusterConfigurator/Private/Views/Details/Cluster/DisplayClusterConfiguratorViewportDetailsCustomization.h
--- a/Engine/Plugins/Runtime/nDisplay/Source/DisplayClusterConfigurator/Private/Views/Details/Cluster/DisplayClusterConfiguratorViewportDetailsCustomization.h
+++ b/Engine/Plugins/Runtime/nDisplay/Source/DisplayClusterConfigurator/Private/Views/Details/Cluster/DisplayClusterConfiguratorViewportDetailsCustomization.h
@@ -6,10 +6,11 @@
 #include "Views/Details/DisplayClusterConfiguratorBaseDetailCustomization.h"
 #include "Types/SlateEnums.h"
 
+class FDisplayClusterConfiguratorNodeSelection;
+class SDisplayClusterConfigurationSearchableComboBox;
+class SWidget;
 class UDisplayClusterConfigurationData;
 class UDisplayClusterConfigurationViewport;
-class SWidget;
-class SDisplayClusterConfigurationSearchableComboBox;
 
 /** Details panel customization for the UDisplayClusterConfigurationViewport object */
 class FDisplayClusterConfiguratorViewportDetailsCustomization final : public FDisplayClusterConfiguratorBaseDetailCustomization
@@ -49,6 +50,9 @@ private:
 	/** Gets the text to display for the currently selected camera */
 	FText GetSelectedCameraText() const;
 
+	/** Customizes the EnforcedCameraOrder array property with camera name dropdowns */
+	void CustomizeEnforcedCameraOrder(IDetailLayoutBuilder& InLayoutBuilder);
+
 private:
 	/** The list of camera items to display in the dropdown menu */
 	TArray<TSharedPtr<FString>>	CameraOptions;
@@ -67,4 +71,7 @@ private:
 
 	/** The combo box that is being displayed in the details panel for the Camera property */
 	TSharedPtr<SDisplayClusterConfigurationSearchableComboBox> CameraComboBox;
+
+	/** The helper object used to generate the custom array widget for the EnforcedCameraOrder property */
+	TSharedPtr<FDisplayClusterConfiguratorNodeSelection> CameraOrderNodeSelection;
 };
diff --git a/Engine/Plugins/Runtime/nDisplay/Source/DisplayClusterShaders/Public/ShaderParameters/DisplayClusterShaderParameters_ICVFX.h b/Engine/Plugins/Runtime/nDisplay/Source/DisplayClusterShaders/Public/ShaderParameters/DisplayClusterShaderParameters_ICVFX.h
--- a/Engine/Plugins/Runtime/nDisplay/Source/DisplayClusterShaders/Public/ShaderParameters/DisplayClusterShaderParameters_ICVFX.h
+++ b/Engine/Plugins/Runtime/nDisplay/Source/DisplayClusterShaders/Public/ShaderParameters/DisplayClusterShaderParameters_ICVFX.h
@@ -112,6 +112,7 @@ public:
 		LightCardOver.Reset();
 		UVLightCardUnder.Reset();
 		UVLightCardOver.Reset();
+		EnforcedCameraOrder.Reset();
 
 		OverrideUVLightCardRenderMode = EDisplayClusterShaderParametersICVFX_LightCardRenderMode::None;
 		LightCardGamma = 2.2;
@@ -134,30 +135,71 @@ public:
 		OverrideUVLightCardRenderMode = InParameters.OverrideUVLightCardRenderMode;
 
 		LightCardGamma = InParameters.LightCardGamma;
+
+		EnforcedCameraOrder = InParameters.EnforcedCameraOrder;
 	}
 
 	inline void SortCamerasRenderOrder(const bool bReverseOrder = false)
 	{
-		Cameras.Sort([bReverseOrder](const FCameraSettings& It1, const FCameraSettings& It2)
+		// Default sorting: by RenderOrder, then alphabetically by ViewportId
+		auto DefaultCameraSort = [](const FCameraSettings& It1, const FCameraSettings& It2)
 		{
-			bool bComparisonResult;
-
 			if (It1.RenderOrder == It2.RenderOrder)
 			{
-				bComparisonResult = It1.Resource.ViewportId.Compare(It2.Resource.ViewportId, ESearchCase::IgnoreCase) < 0;
+				return It1.Resource.ViewportId.Compare(It2.Resource.ViewportId, ESearchCase::IgnoreCase) < 0;
 			}
-			else
+			return It1.RenderOrder < It2.RenderOrder;
+		};
+
+		// If we have an enforced camera order, apply it first
+		if (!EnforcedCameraOrder.IsEmpty())
+		{
+			TArray<FCameraSettings> EnforcedCameras;
+			TArray<FCameraSettings> RemainingCameras;
+
+			// Separate cameras into enforced and remaining groups
+			for (const FCameraSettings& CameraIt : Cameras)
 			{
-				bComparisonResult = It1.RenderOrder < It2.RenderOrder;
+				const int32 EnforcedIndex = EnforcedCameraOrder.IndexOfByKey(CameraIt.Resource.ViewportId);
+				if (EnforcedIndex != INDEX_NONE)
+				{
+					EnforcedCameras.Add(CameraIt);
+				}
+				else
+				{
+					RemainingCameras.Add(CameraIt);
+				}
 			}
 
-			if (bReverseOrder)
+			// Sort remaining cameras using default logic
+			RemainingCameras.Sort(DefaultCameraSort);
+
+			// Sort enforced cameras by their position in the EnforcedCameraOrder array
+			// Index 0 in EnforcedCameraOrder composes on top, and conveniently appears at top in the details panel.
+			// Note: The final order will be reversed if the bReverseOrder option is used.
+			EnforcedCameras.Sort([this](const FCameraSettings& It1, const FCameraSettings& It2)
 			{
-				bComparisonResult = !bComparisonResult;
-			}
+				const int32 Index1 = EnforcedCameraOrder.IndexOfByKey(It1.Resource.ViewportId);
+				const int32 Index2 = EnforcedCameraOrder.IndexOfByKey(It2.Resource.ViewportId);
+				return Index1 > Index2; // Reverse: lower index renders later (on top)
+			});
+
+			// Combine EnforcedCameras with RemainingCameras. EnforcedCameras fill the top of the stack,
+			// while the remaining ones are appended and preserve their original relative order.
+			Cameras = RemainingCameras;
+			Cameras.Append(EnforcedCameras);
+		}
+		else
+		{
+			// No enforced order, use default sorting
+			Cameras.Sort(DefaultCameraSort);
+		}
 
-			return bComparisonResult;
-		});
+		// Apply reverse order as the final step
+		if (bReverseOrder)
+		{
+			Algo::Reverse(Cameras);
+		}
 	}
 
 public:
@@ -291,6 +333,11 @@ public:
 		/** The texture to use to render the chromakey markers for any regions of overlaping ICVFX frustums */
 		FTextureRHIRef OverlapChromakeyMarkerTextureRHI;
 
+		/**
+		 * Camera render order priority. Lower values render first (appear underneath), higher values render last (appear on top).
+		 * When multiple cameras have the same RenderOrder, they are sorted alphabetically by viewport ID.
+		 * Can be overridden per-viewport using EnforcedCameraOrder array.
+		 */
 		int32 RenderOrder = -1;
 	};
 
@@ -367,9 +414,13 @@ public:
 	// UV LightCard resources
 	FDisplayClusterShaderParametersICVFX_ViewportResource    UVLightCardUnder;
 	FDisplayClusterShaderParametersICVFX_ViewportResource    UVLightCardOver;
+
 	// Force UVLC to render merged (per-viewport LC)
 	EDisplayClusterShaderParametersICVFX_LightCardRenderMode OverrideUVLightCardRenderMode = EDisplayClusterShaderParametersICVFX_LightCardRenderMode::None;
 
 	// The gamma that the light card renders have been encoded with, used to linearize during final composite
 	float LightCardGamma = 2.2;
+
+	// Enforced camera order for custom sorting. Contains the camera viewport ids, where index 0 should be composed on top, index 1 below it, and so on. 
+	TArray<FString> EnforcedCameraOrder;
 };
