nDisplay: Reference projection policy now customized with a dropdown.

This avoids the user having to type the referenced viewport name.

Only shows viewports in the same node and excludes itself. When viewports of different nodes are being edited at the same time, the dropdown is disabled to avoid referencing a out of node viewport.

#rnx

#Synced-CL 48544109

diff --git a/Engine/Plugins/Runtime/nDisplay/Source/DisplayClusterConfigurator/Private/Views/Details/Policies/DisplayClusterConfiguratorPolicyDetailCustomization.cpp b/Engine/Plugins/Runtime/nDisplay/Source/DisplayClusterConfigurator/Private/Views/Details/Policies/DisplayClusterConfiguratorPolicyDetailCustomization.cpp
--- a/Engine/Plugins/Runtime/nDisplay/Source/DisplayClusterConfigurator/Private/Views/Details/Policies/DisplayClusterConfiguratorPolicyDetailCustomization.cpp
+++ b/Engine/Plugins/Runtime/nDisplay/Source/DisplayClusterConfigurator/Private/Views/Details/Policies/DisplayClusterConfiguratorPolicyDetailCustomization.cpp
@@ -634,11 +634,17 @@ void FDisplayClusterConfiguratorProjectionCustomization::CreateEasyBlendPolicy(U
 
 void FDisplayClusterConfiguratorProjectionCustomization::CreateReferencePolicy(UDisplayClusterBlueprint* Blueprint)
 {
-	CustomPolicyParameters.Add(MakeShared<FPolicyParameterInfoText>(
+	TSharedPtr<FPolicyParameterInfoViewportCombo> ViewportCombo = MakeShared<FPolicyParameterInfoViewportCombo>(
 		"Source Viewport",
 		DisplayClusterProjectionStrings::cfg::reference::ViewportId,
 		Blueprint,
-		ConfigurationViewports));
+		ConfigurationViewports);
+
+	ViewportCombo->SetParameterTooltip(LOCTEXT("ReferencePolicySourceViewportTooltip",
+		"The viewport to reference for this projection. Only viewports within the same cluster node can be referenced. "
+		"This option is disabled when editing viewports from different cluster nodes."));
+
+	CustomPolicyParameters.Add(ViewportCombo);
 }
 
 void FDisplayClusterConfiguratorProjectionCustomization::CreateManualPolicy(UDisplayClusterBlueprint* Blueprint)
diff --git a/Engine/Plugins/Runtime/nDisplay/Source/DisplayClusterConfigurator/Private/Views/Details/Policies/DisplayClusterConfiguratorPolicyParameterCustomization.cpp b/Engine/Plugins/Runtime/nDisplay/Source/DisplayClusterConfigurator/Private/Views/Details/Policies/DisplayClusterConfiguratorPolicyParameterCustomization.cpp
--- a/Engine/Plugins/Runtime/nDisplay/Source/DisplayClusterConfigurator/Private/Views/Details/Policies/DisplayClusterConfiguratorPolicyParameterCustomization.cpp
+++ b/Engine/Plugins/Runtime/nDisplay/Source/DisplayClusterConfigurator/Private/Views/Details/Policies/DisplayClusterConfiguratorPolicyParameterCustomization.cpp
@@ -273,8 +273,9 @@ TSharedRef<SWidget> FPolicyParameterInfoCombo::MakeCustomParameterValueComboWidg
 
 void FPolicyParameterInfoCombo::OnCustomParameterValueSelected(TSharedPtr<FString> InValue, ESelectInfo::Type SelectInfo)
 {
-	UpdateCustomParameterValueText(InValue.IsValid() ? *InValue.Get() : "");
-	OnItemSelected.ExecuteIfBound(*InValue);
+	const FString ValueString = InValue.IsValid() ? *InValue.Get() : "";
+	UpdateCustomParameterValueText(ValueString);
+	OnItemSelected.ExecuteIfBound(ValueString);
 }
 
 
@@ -322,6 +323,154 @@ void FPolicyParameterInfoComponentCombo::CreateParameterValues(ADisplayClusterRo
 	}
 }
 
+FPolicyParameterInfoViewportCombo::FPolicyParameterInfoViewportCombo(
+	const FString& InDisplayName,
+	const FString& InKey,
+	UDisplayClusterBlueprint* InBlueprint,
+	const TArray<TWeakObjectPtr<UDisplayClusterConfigurationViewport>>& InConfigurationViewports)
+	: FPolicyParameterInfoCombo(InDisplayName, InKey, InBlueprint, InConfigurationViewports, {}, nullptr)
+{
+	CreateParameterValues();
+}
+
+void FPolicyParameterInfoViewportCombo::CreateParameterValues()
+{
+	// Reference policy can only reference viewports within the same cluster node.
+	// When multiple viewports are selected for editing, we use the first valid one to determine
+	// which cluster node's viewports to populate. If the selected viewports belong to different
+	// cluster nodes, AreAllViewportsInSameClusterNode() will return false and disable the control.
+	UDisplayClusterConfigurationClusterNode* ClusterNode = nullptr;
+	TSet<FString> EditedViewportIds;
+
+	// Find the first valid cluster node and build the set of edited viewport IDs
+	for (const TWeakObjectPtr<UDisplayClusterConfigurationViewport>& Viewport : ConfigurationViewports)
+	{
+		if (Viewport.IsValid())
+		{
+			UDisplayClusterConfigurationClusterNode* CurrentClusterNode = Viewport->GetTypedOuter<UDisplayClusterConfigurationClusterNode>();
+			if (CurrentClusterNode)
+			{
+				// Use the first valid cluster node we find
+				if (!ClusterNode)
+				{
+					ClusterNode = CurrentClusterNode;
+				}
+
+				// Build set of currently edited viewport IDs from the first cluster node
+				// (AreAllViewportsInSameClusterNode will validate they're all the same)
+				if (ClusterNode == CurrentClusterNode)
+				{
+					for (const auto& ViewportPair : ClusterNode->Viewports)
+					{
+						if (ViewportPair.Value == Viewport.Get())
+						{
+							EditedViewportIds.Add(ViewportPair.Key);
+							break;
+						}
+					}
+				}
+			}
+		}
+	}
+
+	// If we couldn't find a cluster node, we can't populate the list
+	if (!ClusterNode)
+	{
+		return;
+	}
+
+	// Collect all viewports from the cluster node, excluding the currently edited ones
+	for (const auto& ViewportPair : ClusterNode->Viewports)
+	{
+		const FString& ViewportId = ViewportPair.Key;
+
+		// Skip if this is one of the viewports being edited
+		if (EditedViewportIds.Contains(ViewportId))
+		{
+			continue;
+		}
+
+		CustomParameterOptions.Add(MakeShared<FString>(ViewportId));
+	}
+
+	// Sort the options alphabetically for better UX
+	CustomParameterOptions.Sort([](const TSharedPtr<FString>& A, const TSharedPtr<FString>& B)
+	{
+		return *A < *B;
+	});
+}
+
+void FPolicyParameterInfoViewportCombo::CreateCustomRowWidget(IDetailChildrenBuilder& InDetailWidgetRow)
+{
+	InDetailWidgetRow.AddCustomRow(GetParameterDisplayName())
+	.NameContent()
+	[
+		SNew(STextBlock)
+		.Font(IDetailLayoutBuilder::GetDetailFont())
+		.Text(this, &FPolicyParameterInfo::GetParameterDisplayName)
+		.ToolTipText(this, &FPolicyParameterInfo::GetParameterTooltip)
+	]
+	.ValueContent()
+	[
+		SAssignNew(GetCustomParameterValueComboBox(), SDisplayClusterConfigurationSearchableComboBox)
+		.OptionsSource(&GetCustomParameterOptions())
+		.InitiallySelectedItem(InitialValue)
+		.OnGenerateWidget_Lambda([this](TSharedPtr<FString> InItem)
+		{
+			return MakeCustomParameterValueComboWidget(InItem);
+		})
+		.OnSelectionChanged_Lambda([this](TSharedPtr<FString> InValue, ESelectInfo::Type SelectInfo)
+		{
+			OnCustomParameterValueSelected(InValue, SelectInfo);
+		})
+		.Visibility_Lambda([this]()
+		{
+			return IsParameterVisible();
+		})
+		.IsEnabled(this, &FPolicyParameterInfoViewportCombo::AreAllViewportsInSameClusterNode)
+		.ContentPadding(2)
+		.MaxListHeight(200.0f)
+		.Content()
+		[
+			SNew(STextBlock)
+			.Text(this, &FPolicyParameterInfoCombo::GetOrAddCustomParameterValueText)
+			.Font(IDetailLayoutBuilder::GetDetailFont())
+		]
+	];
+}
+
+bool FPolicyParameterInfoViewportCombo::AreAllViewportsInSameClusterNode() const
+{
+	UDisplayClusterConfigurationClusterNode* FirstClusterNode = nullptr;
+
+	for (const TWeakObjectPtr<UDisplayClusterConfigurationViewport>& Viewport : ConfigurationViewports)
+	{
+		if (Viewport.IsValid())
+		{
+			UDisplayClusterConfigurationClusterNode* ClusterNode = Viewport->GetTypedOuter<UDisplayClusterConfigurationClusterNode>();
+
+			if (!ClusterNode)
+			{
+				// If we can't find a cluster node, disable the control
+				return false;
+			}
+
+			if (!FirstClusterNode)
+			{
+				FirstClusterNode = ClusterNode;
+			}
+			else if (FirstClusterNode != ClusterNode)
+			{
+				// Different cluster nodes found, disable the control
+				return false;
+			}
+		}
+	}
+
+	// If we didn't find any valid cluster node, disable the control
+	return FirstClusterNode != nullptr;
+}
+
 
 FPolicyParameterInfoText::FPolicyParameterInfoText(const FString& InDisplayName, const FString& InKey,
 	UDisplayClusterBlueprint* InBlueprint,
diff --git a/Engine/Plugins/Runtime/nDisplay/Source/DisplayClusterConfigurator/Private/Views/Details/Policies/DisplayClusterConfiguratorPolicyParameterCustomization.h b/Engine/Plugins/Runtime/nDisplay/Source/DisplayClusterConfigurator/Private/Views/Details/Policies/DisplayClusterConfiguratorPolicyParameterCustomization.h
--- a/Engine/Plugins/Runtime/nDisplay/Source/DisplayClusterConfigurator/Private/Views/Details/Policies/DisplayClusterConfiguratorPolicyParameterCustomization.h
+++ b/Engine/Plugins/Runtime/nDisplay/Source/DisplayClusterConfigurator/Private/Views/Details/Policies/DisplayClusterConfiguratorPolicyParameterCustomization.h
@@ -172,6 +172,30 @@ private:
 	TArray<TSubclassOf<UActorComponent>> ComponentTypes;
 };
 
+/**
+ * Policy info for combo box picker for viewports (for Reference projection policy).
+ */
+class FPolicyParameterInfoViewportCombo final : public FPolicyParameterInfoCombo
+{
+public:
+	FPolicyParameterInfoViewportCombo(
+		const FString& InDisplayName,
+		const FString& InKey,
+		UDisplayClusterBlueprint* InBlueprint,
+		const TArray<TWeakObjectPtr<UDisplayClusterConfigurationViewport>>& InConfigurationViewports);
+
+	// FPolicyParameterInfo
+	virtual void CreateCustomRowWidget(IDetailChildrenBuilder& InDetailWidgetRow) override;
+	// ~FPolicyParameterInfo
+
+	/** Populate CustomParameterOptions with all viewports except the currently edited ones. */
+	void CreateParameterValues();
+
+private:
+	/** Check if all edited viewports belong to the same cluster node. */
+	bool AreAllViewportsInSameClusterNode() const;
+};
+
 /**
  * Policy info for number value.
  */
