[nDisplay] Made NVIDIA Present Barrier the default NVIDIA sync policy

- Switched to the new sync approach in order to start testing within 5.8 cycle

#rnx
#Synced-CL 0

diff --git a/Engine/Plugins/Runtime/nDisplay/Source/DisplayCluster/Private/Render/Synchronization/DisplayClusterRenderSyncPolicyFactoryInternal.cpp b/Engine/Plugins/Runtime/nDisplay/Source/DisplayCluster/Private/Render/Synchronization/DisplayClusterRenderSyncPolicyFactoryInternal.cpp
--- a/Engine/Plugins/Runtime/nDisplay/Source/DisplayCluster/Private/Render/Synchronization/DisplayClusterRenderSyncPolicyFactoryInternal.cpp
+++ b/Engine/Plugins/Runtime/nDisplay/Source/DisplayCluster/Private/Render/Synchronization/DisplayClusterRenderSyncPolicyFactoryInternal.cpp
@@ -64,11 +64,12 @@ TSharedPtr<IDisplayClusterRenderSyncPolicy> FDisplayClusterRenderSyncPolicyFacto
 #if PLATFORM_WINDOWS
 		if (RHIType == ERHIInterfaceType::D3D12 && CVarUseNvidiaPresentBarrierPolicy.GetValueOnGameThread() && WITH_NVAPI)
 		{
-			UE_LOG(LogDisplayClusterRenderSync, Log, TEXT("Experimental NVIDIA synchronization approach will be used!"));
+			UE_LOG(LogDisplayClusterRenderSync, Log, TEXT("NVIDIA Present Barrier (new) synchronization approach will be used!"));
 			return MakeShared<FDisplayClusterRenderSyncPolicyNvidiaPresentBarrier>(Parameters);
 		}
 		else if (RHIType == ERHIInterfaceType::D3D11 || RHIType == ERHIInterfaceType::D3D12 && WITH_NVAPI)
 		{
+			UE_LOG(LogDisplayClusterRenderSync, Log, TEXT("NVIDIA Swap Barrier (old) synchronization approach will be used!"));
 			return MakeShared<FDisplayClusterRenderSyncPolicyNvidiaSwapBarrier>(Parameters);
 		}
 		else if (RHIType == ERHIInterfaceType::Vulkan)
