[nDisplay] Made NVIDIA Present Barrier the default NVIDIA sync policy

- Switched to the new sync approach in order to start testing within 5.8 cycle
- Changed the CVar (missed in the first CL)
- Fixed wrong logic statement

#rnx

#Synced-CL 0

diff --git a/Engine/Plugins/Runtime/nDisplay/Source/DisplayCluster/Private/Render/Synchronization/DisplayClusterRenderSyncPolicyFactoryInternal.cpp b/Engine/Plugins/Runtime/nDisplay/Source/DisplayCluster/Private/Render/Synchronization/DisplayClusterRenderSyncPolicyFactoryInternal.cpp
--- a/Engine/Plugins/Runtime/nDisplay/Source/DisplayCluster/Private/Render/Synchronization/DisplayClusterRenderSyncPolicyFactoryInternal.cpp
+++ b/Engine/Plugins/Runtime/nDisplay/Source/DisplayCluster/Private/Render/Synchronization/DisplayClusterRenderSyncPolicyFactoryInternal.cpp
@@ -20,7 +20,7 @@
 // the GUI configurator as well as other policies are.
 static TAutoConsoleVariable<bool> CVarUseNvidiaPresentBarrierPolicy(
 	TEXT("nDisplay.sync.nvidia.UsePresentBarrierPolicy"),
-	false,
+	true,
 	TEXT("NVIDIA Present Barrier synchronization policy (experimental)"),
 	ECVF_ReadOnly
 );
@@ -67,7 +67,7 @@ TSharedPtr<IDisplayClusterRenderSyncPolicy> FDisplayClusterRenderSyncPolicyFacto
 			UE_LOG(LogDisplayClusterRenderSync, Log, TEXT("NVIDIA Present Barrier (new) synchronization approach will be used!"));
 			return MakeShared<FDisplayClusterRenderSyncPolicyNvidiaPresentBarrier>(Parameters);
 		}
-		else if (RHIType == ERHIInterfaceType::D3D11 || RHIType == ERHIInterfaceType::D3D12 && WITH_NVAPI)
+		else if ((RHIType == ERHIInterfaceType::D3D11 || RHIType == ERHIInterfaceType::D3D12) && WITH_NVAPI)
 		{
 			UE_LOG(LogDisplayClusterRenderSync, Log, TEXT("NVIDIA Swap Barrier (old) synchronization approach will be used!"));
 			return MakeShared<FDisplayClusterRenderSyncPolicyNvidiaSwapBarrier>(Parameters);
