nDisplay: Fixed rare assert failure where ICVFX window would attempt to get DPI on render thread which would result in trying to access FSlateApplication on an unsupported thread

#Synced-CL 48991784

diff --git a/Engine/Plugins/Runtime/nDisplay/Source/DisplayClusterLightCardEditor/Private/Viewport/DisplayClusterLightCardEditorViewportClient.cpp b/Engine/Plugins/Runtime/nDisplay/Source/DisplayClusterLightCardEditor/Private/Viewport/DisplayClusterLightCardEditorViewportClient.cpp
--- a/Engine/Plugins/Runtime/nDisplay/Source/DisplayClusterLightCardEditor/Private/Viewport/DisplayClusterLightCardEditorViewportClient.cpp
+++ b/Engine/Plugins/Runtime/nDisplay/Source/DisplayClusterLightCardEditor/Private/Viewport/DisplayClusterLightCardEditorViewportClient.cpp
@@ -215,6 +215,11 @@ void FDisplayClusterLightCardEditorViewportClient::Tick(float DeltaSeconds)
 	{
 		UpdatePreviewActor(LightCardEditorPtr.Pin()->GetActiveRootActor().Get(), true);
 	}
+
+	if (EditorWidget.IsValid())
+	{
+		EditorWidget->Tick(this, DeltaSeconds);
+	}
 }
 
 void FDisplayClusterLightCardEditorViewportClient::Draw(FViewport* InViewport, FCanvas* Canvas)
diff --git a/Engine/Plugins/Runtime/nDisplay/Source/DisplayClusterLightCardEditor/Private/Viewport/DisplayClusterLightCardEditorWidget.cpp b/Engine/Plugins/Runtime/nDisplay/Source/DisplayClusterLightCardEditor/Private/Viewport/DisplayClusterLightCardEditorWidget.cpp
--- a/Engine/Plugins/Runtime/nDisplay/Source/DisplayClusterLightCardEditor/Private/Viewport/DisplayClusterLightCardEditorWidget.cpp
+++ b/Engine/Plugins/Runtime/nDisplay/Source/DisplayClusterLightCardEditor/Private/Viewport/DisplayClusterLightCardEditorWidget.cpp
@@ -15,6 +15,12 @@ FDisplayClusterLightCardEditorWidget::FDisplayClusterLightCardEditorWidget()
 {
 }
 
+void FDisplayClusterLightCardEditorWidget::Tick(const FDisplayClusterLightCardEditorViewportClient* ViewportClient, float InDeltaSeconds)
+{
+	// Cache the DPI scale from the viewport. Must be done here to avoid calling on render thread and potentially failing an assert. See UE-353611
+	DPIScale = ViewportClient->GetDPIScale();
+}
+
 void FDisplayClusterLightCardEditorWidget::Draw(const FSceneView* View, const FDisplayClusterLightCardEditorViewportClient* ViewportClient, FPrimitiveDrawInterface* PDI, bool bDrawZAxis)
 {
 	const bool bIsOrthographic = ViewportClient->GetRenderViewportType() != LVT_Perspective;
@@ -168,8 +174,6 @@ float FDisplayClusterLightCardEditorWidget::GetLengthScreenScalar(const FSceneVi
 	// height is resized, the view is clamped
 
 	const bool bIsOrthographic = ViewportClient->GetRenderViewportType() != LVT_Perspective;
-
-	const float DPIScale = ViewportClient->GetDPIScale();
 	const float ResolutionScale = View->UnconstrainedViewRect.Size().X / (DPIScale * ReferenceResolution);
 
 	const float DistanceFromView = FMath::Max(FVector::Dist(Origin, View->ViewMatrices.GetViewOrigin()), 1.f);
@@ -194,8 +198,6 @@ float FDisplayClusterLightCardEditorWidget::GetSizeScreenScalar(const FSceneView
 	// height is resized, the view is clamped
 
 	const bool bIsOrthographic = ViewportClient->GetRenderViewportType() != LVT_Perspective;
-
-	const float DPIScale = ViewportClient->GetDPIScale();
 	const float ResolutionScale = View->UnconstrainedViewRect.Size().X / (DPIScale * ReferenceResolution);
 
 	const FMatrix& ProjectionMatrix = View->ViewMatrices.GetProjectionMatrix();
diff --git a/Engine/Plugins/Runtime/nDisplay/Source/DisplayClusterLightCardEditor/Private/Viewport/DisplayClusterLightCardEditorWidget.h b/Engine/Plugins/Runtime/nDisplay/Source/DisplayClusterLightCardEditor/Private/Viewport/DisplayClusterLightCardEditorWidget.h
--- a/Engine/Plugins/Runtime/nDisplay/Source/DisplayClusterLightCardEditor/Private/Viewport/DisplayClusterLightCardEditorWidget.h
+++ b/Engine/Plugins/Runtime/nDisplay/Source/DisplayClusterLightCardEditor/Private/Viewport/DisplayClusterLightCardEditorWidget.h
@@ -52,6 +52,9 @@ public:
 
 	FDisplayClusterLightCardEditorWidget();
 
+	/** Called from the game thread, allows the widget to perform any game thread specific actions */
+	void Tick(const FDisplayClusterLightCardEditorViewportClient* ViewportClient, float InDeltaSeconds);
+	
 	void Draw(const FSceneView* View, const FDisplayClusterLightCardEditorViewportClient* ViewportClient, FPrimitiveDrawInterface* PDI, bool bDrawZAxis);
 
 	EWidgetMode GetWidgetMode() const { return WidgetMode; }
@@ -104,4 +107,7 @@ private:
 
 	/** A scale factor to scale the widget's rendered size */
 	float WidgetScale = 1.f;
+
+	/** Cached DPI scale used to scale the widget's rendered size */
+	float DPIScale = 1.f;
 };
\ No newline at end of file
