nDisplay: Remove UDisplayClusterGameEngine::WorldContextObject reference that was leaking GC when UnloadWorld.

#rnx

#cherrypick of CL 46125790

#Synced-CL 49034323

diff --git a/Engine/Plugins/Runtime/nDisplay/Source/DisplayCluster/Public/DisplayClusterGameEngine.h b/Engine/Plugins/Runtime/nDisplay/Source/DisplayCluster/Public/DisplayClusterGameEngine.h
--- a/Engine/Plugins/Runtime/nDisplay/Source/DisplayCluster/Public/DisplayClusterGameEngine.h
+++ b/Engine/Plugins/Runtime/nDisplay/Source/DisplayCluster/Public/DisplayClusterGameEngine.h
@@ -27,7 +27,6 @@ class DISPLAYCLUSTER_API UDisplayClusterGameEngine
 {
 	GENERATED_BODY()
 
-	UPROPERTY()
 	TObjectPtr<UObject> WorldContextObject = nullptr;
 
 public:
