[nDisplay] Extended fix for a static data cache issue
- Like it's done in FDisplayClusterNodeCtrlEditor, avoid using of static variable in FDisplayClusterNodeCtrlDisabled

#rnx

#Synced-CL 0

diff --git a/Engine/Plugins/Runtime/nDisplay/Source/DisplayCluster/Private/Cluster/Failover/DisplayClusterFailoverNodeCtrlDisabled.cpp b/Engine/Plugins/Runtime/nDisplay/Source/DisplayCluster/Private/Cluster/Failover/DisplayClusterFailoverNodeCtrlDisabled.cpp
--- a/Engine/Plugins/Runtime/nDisplay/Source/DisplayCluster/Private/Cluster/Failover/DisplayClusterFailoverNodeCtrlDisabled.cpp
+++ b/Engine/Plugins/Runtime/nDisplay/Source/DisplayCluster/Private/Cluster/Failover/DisplayClusterFailoverNodeCtrlDisabled.cpp
@@ -1,7 +1,6 @@
 // Copyright Epic Games, Inc. All Rights Reserved.
 
 #include "Cluster/Failover/DisplayClusterFailoverNodeCtrlDisabled.h"
-#include "Cluster/Failover/DisplayClusterCommDataCache.h"
 
 
 ////////////////////////////////////////////////////////////////////////////////////////////////////
@@ -15,8 +14,7 @@ bool FDisplayClusterFailoverNodeCtrlDisabled::Initialize(const UDisplayClusterCo
 
 TSharedRef<FDisplayClusterCommDataCache> FDisplayClusterFailoverNodeCtrlDisabled::GetDataCache()
 {
-	static TSharedRef<FDisplayClusterCommDataCache> DataBroker = MakeShared<FDisplayClusterCommDataCache>();
-	return DataBroker;
+	return DataCache;
 }
 
 bool FDisplayClusterFailoverNodeCtrlDisabled::HandleFailure(const FString& FailedNodeId)
diff --git a/Engine/Plugins/Runtime/nDisplay/Source/DisplayCluster/Private/Cluster/Failover/DisplayClusterFailoverNodeCtrlDisabled.h b/Engine/Plugins/Runtime/nDisplay/Source/DisplayCluster/Private/Cluster/Failover/DisplayClusterFailoverNodeCtrlDisabled.h
--- a/Engine/Plugins/Runtime/nDisplay/Source/DisplayCluster/Private/Cluster/Failover/DisplayClusterFailoverNodeCtrlDisabled.h
+++ b/Engine/Plugins/Runtime/nDisplay/Source/DisplayCluster/Private/Cluster/Failover/DisplayClusterFailoverNodeCtrlDisabled.h
@@ -3,6 +3,7 @@
 #pragma once
 
 #include "Cluster/Failover/DisplayClusterFailoverNodeCtrlBase.h"
+#include "Cluster/Failover/DisplayClusterCommDataCache.h"
 
 
 /**
@@ -17,6 +18,7 @@ public:
 
 	FDisplayClusterFailoverNodeCtrlDisabled(TSharedRef<IDisplayClusterClusterNodeController>& InNodeController)
 		: FDisplayClusterFailoverNodeCtrlBase(InNodeController)
+		, DataCache(MakeShared<FDisplayClusterCommDataCache>())
 	{ }
 
 public:
@@ -75,4 +77,9 @@ public:
 	virtual EDisplayClusterCommResult PostFailureNegotiate(TArray<uint8>& InOutRecoveryData) override;
 	virtual EDisplayClusterCommResult RequestNodeDrop(const FString& NodeId, uint8 DropReason) override;
 	//~ End IDisplayClusterProtocolInternalComm
+
+private:
+
+	/** Stub cache to respect API */
+	TSharedRef<FDisplayClusterCommDataCache> DataCache;
 };
