nDisplay: Fixed issue where editor would crash when importing an nDisplay config into an existing nDisplay asset while hovering mouse over output mapping tab

#Synced-CL 49143073

diff --git a/Engine/Plugins/Runtime/nDisplay/Source/DisplayClusterConfigurator/Private/Views/OutputMapping/GraphNodes/SDisplayClusterConfiguratorBaseNode.cpp b/Engine/Plugins/Runtime/nDisplay/Source/DisplayClusterConfigurator/Private/Views/OutputMapping/GraphNodes/SDisplayClusterConfiguratorBaseNode.cpp
--- a/Engine/Plugins/Runtime/nDisplay/Source/DisplayClusterConfigurator/Private/Views/OutputMapping/GraphNodes/SDisplayClusterConfiguratorBaseNode.cpp
+++ b/Engine/Plugins/Runtime/nDisplay/Source/DisplayClusterConfigurator/Private/Views/OutputMapping/GraphNodes/SDisplayClusterConfiguratorBaseNode.cpp
@@ -159,6 +159,11 @@ FVector2D SDisplayClusterConfiguratorBaseNode::ComputeDesiredSize(float) const
 
 void SDisplayClusterConfiguratorBaseNode::MoveTo(const FVector2f& NewPosition, FNodeSet& NodeFilter, bool bMarkDirty)
 {
+	if (!IsGraphNodeValid())
+	{
+		return;
+	}
+
 	if (!IsNodeUnlocked())
 	{
 		NodeFilter.Add(SharedThis(this));
@@ -363,6 +368,11 @@ void SDisplayClusterConfiguratorBaseNode::BeginUserInteraction() const
 		return;
 	}
 
+	if (!IsGraphNodeValid())
+	{
+		return;
+	}
+
 	UDisplayClusterConfiguratorBaseNode* EdNode = GetGraphNodeChecked<UDisplayClusterConfiguratorBaseNode>();
 
 	bool bIsInteractingDirectly = false;
@@ -396,6 +406,11 @@ void SDisplayClusterConfiguratorBaseNode::BeginUserInteraction() const
 
 void SDisplayClusterConfiguratorBaseNode::SetNodeSize(const FVector2D InLocalSize, bool bFixedAspectRatio)
 {
+	if (!IsGraphNodeValid())
+	{
+		return;
+	}
+
 	TSharedPtr<FDisplayClusterConfiguratorBlueprintEditor> Toolkit = ToolkitPtr.Pin();
 	check(Toolkit.IsValid());
 
@@ -457,18 +472,33 @@ void SDisplayClusterConfiguratorBaseNode::SetNodeSize(const FVector2D InLocalSiz
 
 bool SDisplayClusterConfiguratorBaseNode::IsNodeVisible() const
 {
+	if (!IsGraphNodeValid())
+	{
+		return false;
+	}
+
 	UDisplayClusterConfiguratorBaseNode* EdNode = GetGraphNodeChecked<UDisplayClusterConfiguratorBaseNode>();
 	return EdNode->IsNodeVisible();
 }
 
 bool SDisplayClusterConfiguratorBaseNode::IsNodeEnabled() const
 {
+	if (!IsGraphNodeValid())
+	{
+		return false;
+	}
+
 	UDisplayClusterConfiguratorBaseNode* EdNode = GetGraphNodeChecked<UDisplayClusterConfiguratorBaseNode>();
 	return EdNode->IsNodeEnabled();
 }
 
 bool SDisplayClusterConfiguratorBaseNode::IsNodeUnlocked() const
 {
+	if (!IsGraphNodeValid())
+	{
+		return false;
+	}
+
 	UDisplayClusterConfiguratorBaseNode* EdNode = GetGraphNodeChecked<UDisplayClusterConfiguratorBaseNode>();
 	return EdNode->IsNodeUnlocked();
 }
@@ -510,12 +540,22 @@ TOptional<EMouseCursor::Type> SDisplayClusterConfiguratorBaseNode::GetCursor() c
 
 int32 SDisplayClusterConfiguratorBaseNode::GetNodeLogicalLayer() const
 {
+	if (!IsGraphNodeValid())
+	{
+		return 0;
+	}
+
 	UDisplayClusterConfiguratorBaseNode* EdNode = GetGraphNodeChecked<UDisplayClusterConfiguratorBaseNode>();
 	return EdNode->GetNodeLayer(ObjectPtrDecay(GetOwnerPanel()->SelectionManager.SelectedNodes));
 }
 
 int32 SDisplayClusterConfiguratorBaseNode::GetNodeVisualLayer() const
 {
+	if (!IsGraphNodeValid())
+	{
+		return 0;
+	}
+
 	UDisplayClusterConfiguratorBaseNode* EdNode = GetGraphNodeChecked<UDisplayClusterConfiguratorBaseNode>();
 	return EdNode->GetNodeLayer(ObjectPtrDecay(GetOwnerPanel()->SelectionManager.SelectedNodes));
 }
@@ -626,3 +666,8 @@ void SDisplayClusterConfiguratorBaseNode::AddAlignmentRulerToOverlay(TArray<FOve
 
 	OverlayWidgets.Add(AlignmentRulerInfo);
 }
+
+bool SDisplayClusterConfiguratorBaseNode::IsGraphNodeValid() const
+{
+	return IsValid(GraphNode) && !GraphNode->HasAnyFlags(RF_BeginDestroyed | RF_FinishDestroyed);
+}
diff --git a/Engine/Plugins/Runtime/nDisplay/Source/DisplayClusterConfigurator/Private/Views/OutputMapping/GraphNodes/SDisplayClusterConfiguratorBaseNode.h b/Engine/Plugins/Runtime/nDisplay/Source/DisplayClusterConfigurator/Private/Views/OutputMapping/GraphNodes/SDisplayClusterConfiguratorBaseNode.h
--- a/Engine/Plugins/Runtime/nDisplay/Source/DisplayClusterConfigurator/Private/Views/OutputMapping/GraphNodes/SDisplayClusterConfiguratorBaseNode.h
+++ b/Engine/Plugins/Runtime/nDisplay/Source/DisplayClusterConfigurator/Private/Views/OutputMapping/GraphNodes/SDisplayClusterConfiguratorBaseNode.h
@@ -154,6 +154,9 @@ protected:
 		return CastedNode;
 	}
 
+	/** Checks that the graph node this widget displays is valid and has not been destroyed or garbage collected */
+	bool IsGraphNodeValid() const;
+	
 protected:
 	TWeakPtr<FDisplayClusterConfiguratorBlueprintEditor> ToolkitPtr;
 
diff --git a/Engine/Plugins/Runtime/nDisplay/Source/DisplayClusterConfigurator/Private/Views/OutputMapping/GraphNodes/SDisplayClusterConfiguratorCanvasNode.cpp b/Engine/Plugins/Runtime/nDisplay/Source/DisplayClusterConfigurator/Private/Views/OutputMapping/GraphNodes/SDisplayClusterConfiguratorCanvasNode.cpp
--- a/Engine/Plugins/Runtime/nDisplay/Source/DisplayClusterConfigurator/Private/Views/OutputMapping/GraphNodes/SDisplayClusterConfiguratorCanvasNode.cpp
+++ b/Engine/Plugins/Runtime/nDisplay/Source/DisplayClusterConfigurator/Private/Views/OutputMapping/GraphNodes/SDisplayClusterConfiguratorCanvasNode.cpp
@@ -132,6 +132,11 @@ FMargin SDisplayClusterConfiguratorCanvasNode::GetBackgroundPosition() const
 
 FText SDisplayClusterConfiguratorCanvasNode::GetCanvasSizeText() const
 {
+	if (!IsGraphNodeValid())
+	{
+		return FText::GetEmpty();
+	}
+
 	UDisplayClusterConfiguratorCanvasNode* CanvasNode = GetGraphNodeChecked<UDisplayClusterConfiguratorCanvasNode>();
 	const FVector2D& Resolution = CanvasNode->GetResolution();
 	return FText::Format(LOCTEXT("ClusterResolution", "Cluster Resolution {0} x {1}"), FText::AsNumber(FMath::RoundToInt(Resolution.X)), FText::AsNumber(FMath::RoundToInt(Resolution.Y)));
diff --git a/Engine/Plugins/Runtime/nDisplay/Source/DisplayClusterConfigurator/Private/Views/OutputMapping/GraphNodes/SDisplayClusterConfiguratorHostNode.cpp b/Engine/Plugins/Runtime/nDisplay/Source/DisplayClusterConfigurator/Private/Views/OutputMapping/GraphNodes/SDisplayClusterConfiguratorHostNode.cpp
--- a/Engine/Plugins/Runtime/nDisplay/Source/DisplayClusterConfigurator/Private/Views/OutputMapping/GraphNodes/SDisplayClusterConfiguratorHostNode.cpp
+++ b/Engine/Plugins/Runtime/nDisplay/Source/DisplayClusterConfigurator/Private/Views/OutputMapping/GraphNodes/SDisplayClusterConfiguratorHostNode.cpp
@@ -161,6 +161,11 @@ public:
 			TSharedPtr<SGraphPanel> GraphPanel = ParentNode->GetOwnerPanel();
 			check(GraphPanel.IsValid());
 
+			if (!ParentNode->IsGraphNodeValid())
+			{
+				return FReply::Unhandled();
+			}
+
 			const float DPIScale = GetDPIScale();
 
 			FVector2D NewOrigin = MouseEvent.GetScreenSpacePosition() - ParentNode->GetTickSpaceGeometry().GetAbsolutePosition();
@@ -219,6 +224,11 @@ private:
 		TSharedPtr<SDisplayClusterConfiguratorHostNode> ParentNode = ParentNodePtr.Pin();
 		check(ParentNode.IsValid());
 
+		if (!ParentNode->IsGraphNodeValid())
+		{
+			return FText::GetEmpty();
+		}
+
 		UDisplayClusterConfiguratorHostNode* HostEdNode = ParentNode->GetGraphNodeChecked<UDisplayClusterConfiguratorHostNode>();
 		const FVector2D& Origin = HostEdNode->GetHostOrigin();
 		return FText::Format(LOCTEXT("OriginLabel", "Host Origin @ {0}, {1}"), FMath::FloorToFloat(Origin.X), FMath::FloorToFloat(Origin.Y));
@@ -409,6 +419,11 @@ void SDisplayClusterConfiguratorHostNode::UpdateGraphNode()
 
 void SDisplayClusterConfiguratorHostNode::MoveTo(const FVector2f& NewPosition, FNodeSet& NodeFilter, bool bMarkDirty)
 {
+	if (!IsGraphNodeValid())
+	{
+		return;
+	}
+
 	UDisplayClusterConfiguratorHostNode* HostEdNode = GetGraphNodeChecked<UDisplayClusterConfiguratorHostNode>();
 	if (!HostEdNode->CanUserMoveNode())
 	{
@@ -437,6 +452,11 @@ void SDisplayClusterConfiguratorHostNode::SetNodeSize(const FVector2D InLocalSiz
 
 bool SDisplayClusterConfiguratorHostNode::CanNodeBeResized() const
 {
+	if (!IsGraphNodeValid())
+	{
+		return false;
+	}
+
 	UDisplayClusterConfiguratorHostNode* HostEdNode = GetGraphNodeChecked<UDisplayClusterConfiguratorHostNode>();
 	return HostEdNode->CanUserResizeNode();
 }
@@ -449,6 +469,11 @@ FMargin SDisplayClusterConfiguratorHostNode::GetBackgroundPosition() const
 
 FMargin SDisplayClusterConfiguratorHostNode::GetNodeOriginPosition() const
 {
+	if (!IsGraphNodeValid())
+	{
+		return FMargin();
+	}
+
 	UDisplayClusterConfiguratorHostNode* HostEdNode = GetGraphNodeChecked<UDisplayClusterConfiguratorHostNode>();
 	const FVector2D GlobalOrigin = HostEdNode->GetHostOrigin(true);
 	
@@ -457,6 +482,11 @@ FMargin SDisplayClusterConfiguratorHostNode::GetNodeOriginPosition() const
 
 int32 SDisplayClusterConfiguratorHostNode::GetNodeOriginLayerOffset() const
 {
+	if (!IsGraphNodeValid())
+	{
+		return 0;
+	}
+
 	UDisplayClusterConfiguratorHostNode* HostEdNode = GetGraphNodeChecked<UDisplayClusterConfiguratorHostNode>();
 	int32 NodeLayerIndex = HostEdNode->GetNodeLayer(ObjectPtrDecay(GetOwnerPanel()->SelectionManager.SelectedNodes));
 	int32 OrnamentLayerIndex = DisplayClusterConfiguratorGraphLayers::OrnamentLayerIndex;
@@ -466,6 +496,11 @@ int32 SDisplayClusterConfiguratorHostNode::GetNodeOriginLayerOffset() const
 
 EVisibility SDisplayClusterConfiguratorHostNode::GetHostNameVisibility() const
 {
+	if (!IsGraphNodeValid())
+	{
+		return EVisibility::Collapsed;
+	}
+
 	UDisplayClusterConfiguratorHostNode* PCEdNode = GetGraphNodeChecked<UDisplayClusterConfiguratorHostNode>();
 	const FText HostName = PCEdNode->GetHostName();
 
@@ -474,12 +509,22 @@ EVisibility SDisplayClusterConfiguratorHostNode::GetHostNameVisibility() const
 
 FText SDisplayClusterConfiguratorHostNode::GetHostNameText() const
 {
+	if (!IsGraphNodeValid())
+	{
+		return FText::GetEmpty();
+	}
+
 	UDisplayClusterConfiguratorHostNode* PCEdNode = GetGraphNodeChecked<UDisplayClusterConfiguratorHostNode>();
 	return PCEdNode->GetHostName();
 }
 
 FText SDisplayClusterConfiguratorHostNode::GetHostIPText() const
 {
+	if (!IsGraphNodeValid())
+	{
+		return FText::GetEmpty();
+	}
+
 	UDisplayClusterConfiguratorHostNode* PCEdNode = GetGraphNodeChecked<UDisplayClusterConfiguratorHostNode>();
 	FString IPAddress = PCEdNode->GetNodeName();
 	return FText::FromString(*IPAddress);
@@ -487,6 +532,11 @@ FText SDisplayClusterConfiguratorHostNode::GetHostIPText() const
 
 FText SDisplayClusterConfiguratorHostNode::GetHostResolutionText() const
 {
+	if (!IsGraphNodeValid())
+	{
+		return FText::GetEmpty();
+	}
+
 	UDisplayClusterConfiguratorHostNode* HostEdNode = GetGraphNodeChecked<UDisplayClusterConfiguratorHostNode>();
 	const FVector2D NodeSize = HostEdNode->GetNodeLocalSize();
 	return FText::Format(LOCTEXT("HostResolution_Formatted", "[{0} x {1}]"), FText::AsNumber(FMath::RoundToInt(NodeSize.X)), FText::AsNumber(FMath::RoundToInt(NodeSize.Y)));
@@ -494,6 +544,11 @@ FText SDisplayClusterConfiguratorHostNode::GetHostResolutionText() const
 
 FSlateColor SDisplayClusterConfiguratorHostNode::GetBorderColor() const
 {
+	if (!IsGraphNodeValid())
+	{
+		return FSlateColor();
+	}
+
 	if (GetOwnerPanel()->SelectionManager.SelectedNodes.Contains(GraphNode))
 	{
 		// Selected Case
@@ -508,6 +563,11 @@ FSlateColor SDisplayClusterConfiguratorHostNode::GetBorderColor() const
 
 FSlateColor SDisplayClusterConfiguratorHostNode::GetTextColor() const
 {
+	if (!IsGraphNodeValid())
+	{
+		return FSlateColor();
+	}
+
 	if (GetOwnerPanel()->SelectionManager.SelectedNodes.Contains(GraphNode))
 	{
 		// Selected Case
diff --git a/Engine/Plugins/Runtime/nDisplay/Source/DisplayClusterConfigurator/Private/Views/OutputMapping/GraphNodes/SDisplayClusterConfiguratorViewportNode.cpp b/Engine/Plugins/Runtime/nDisplay/Source/DisplayClusterConfigurator/Private/Views/OutputMapping/GraphNodes/SDisplayClusterConfiguratorViewportNode.cpp
--- a/Engine/Plugins/Runtime/nDisplay/Source/DisplayClusterConfigurator/Private/Views/OutputMapping/GraphNodes/SDisplayClusterConfiguratorViewportNode.cpp
+++ b/Engine/Plugins/Runtime/nDisplay/Source/DisplayClusterConfigurator/Private/Views/OutputMapping/GraphNodes/SDisplayClusterConfiguratorViewportNode.cpp
@@ -37,6 +37,11 @@ void SDisplayClusterConfiguratorViewportNode::UpdateGraphNode()
 {
 	SDisplayClusterConfiguratorBaseNode::UpdateGraphNode();
 
+	if (!IsGraphNodeValid())
+	{
+		return;
+	}
+
 	BackgroundImage = SNew(SImage)
 		.ColorAndOpacity(this, &SDisplayClusterConfiguratorViewportNode::GetBackgroundColor)
 		.Image(this, &SDisplayClusterConfiguratorViewportNode::GetBackgroundBrush)
@@ -170,6 +175,11 @@ void SDisplayClusterConfiguratorViewportNode::Tick(const FGeometry& AllottedGeom
 
 bool SDisplayClusterConfiguratorViewportNode::IsNodeVisible() const
 {
+	if (!IsGraphNodeValid())
+	{
+		return false;
+	}
+
 	UDisplayClusterConfiguratorViewportNode* ViewportEdNode = GetGraphNodeChecked<UDisplayClusterConfiguratorViewportNode>();
 	TSharedPtr<FDisplayClusterConfiguratorBlueprintEditor> Toolkit = ToolkitPtr.Pin();
 	check(Toolkit.IsValid());
@@ -193,6 +203,11 @@ float SDisplayClusterConfiguratorViewportNode::GetNodeMaximumSize() const
 
 bool SDisplayClusterConfiguratorViewportNode::IsAspectRatioFixed() const
 {
+	if (!IsGraphNodeValid())
+	{
+		return false;
+	}
+
 	UDisplayClusterConfiguratorViewportNode* ViewportEdNode = GetGraphNodeChecked<UDisplayClusterConfiguratorViewportNode>();
 	return ViewportEdNode->IsFixedAspectRatio();
 }
@@ -201,6 +216,11 @@ void SDisplayClusterConfiguratorViewportNode::BeginUserInteraction() const
 {
 	SDisplayClusterConfiguratorBaseNode::BeginUserInteraction();
 
+	if (!IsGraphNodeValid())
+	{
+		return;
+	}
+
 	const UDisplayClusterConfiguratorViewportNode* ViewportEdNode = GetGraphNodeChecked<UDisplayClusterConfiguratorViewportNode>();
 	UDisplayClusterConfigurationViewport* Viewport = ViewportEdNode->GetObjectChecked<UDisplayClusterConfigurationViewport>();
 	bIsManagingPreviewTexture = Viewport->DisablePreviewTexture();
@@ -210,6 +230,11 @@ void SDisplayClusterConfiguratorViewportNode::EndUserInteraction() const
 {
 	SDisplayClusterConfiguratorBaseNode::EndUserInteraction();
 	
+	if (!IsGraphNodeValid())
+	{
+		return;
+	}
+
 	if (bIsManagingPreviewTexture)
 	{
 		const UDisplayClusterConfiguratorViewportNode* ViewportEdNode = GetGraphNodeChecked<UDisplayClusterConfiguratorViewportNode>();
@@ -225,6 +250,11 @@ FSlateColor SDisplayClusterConfiguratorViewportNode::GetBackgroundColor() const
 	TSharedPtr<FDisplayClusterConfiguratorBlueprintEditor> Toolkit = ToolkitPtr.Pin();
 	check(Toolkit.IsValid());
 
+	if (!IsGraphNodeValid())
+	{
+		return FSlateColor();
+	}
+
 	TSharedRef<IDisplayClusterConfiguratorViewOutputMapping> OutputMapping = Toolkit->GetViewOutputMapping();
 
 	const bool bIsSelected = GetOwnerPanel()->SelectionManager.SelectedNodes.Contains(GraphNode);
@@ -294,6 +324,11 @@ const FSlateBrush* SDisplayClusterConfiguratorViewportNode::GetBackgroundBrush()
 
 TOptional<FSlateRenderTransform> SDisplayClusterConfiguratorViewportNode::GetBackgroundRenderTransform() const
 {
+	if (!IsGraphNodeValid())
+	{
+		return TOptional<FSlateRenderTransform>();
+	}
+
 	UDisplayClusterConfiguratorViewportNode* ViewportEdNode = GetGraphNodeChecked<UDisplayClusterConfiguratorViewportNode>();
 	const FDisplayClusterConfigurationRectangle& Region = ViewportEdNode->GetCfgViewportRegion();
 	const FDisplayClusterConfigurationViewport_RemapData& RemapData = ViewportEdNode->GetCfgViewportRemap();
@@ -333,6 +368,11 @@ const FSlateBrush* SDisplayClusterConfiguratorViewportNode::GetNodeShadowBrush()
 
 const FSlateBrush* SDisplayClusterConfiguratorViewportNode::GetBorderBrush() const
 {
+	if (!IsGraphNodeValid())
+	{
+		return nullptr;
+	}
+
 	if (GetOwnerPanel()->SelectionManager.SelectedNodes.Contains(GraphNode))
 	{
 		return FDisplayClusterConfiguratorStyle::Get().GetBrush("DisplayClusterConfigurator.Node.Viewport.Border.Brush.Selected");
@@ -373,6 +413,11 @@ FSlateColor SDisplayClusterConfiguratorViewportNode::GetTextColor() const
 
 FText SDisplayClusterConfiguratorViewportNode::GetPositionAndSizeText() const
 {
+	if (!IsGraphNodeValid())
+	{
+		return FText::GetEmpty();
+	}
+
 	UDisplayClusterConfiguratorViewportNode* ViewportEdNode = GetGraphNodeChecked<UDisplayClusterConfiguratorViewportNode>();
 	const FDisplayClusterConfigurationRectangle& CfgViewportRegion = ViewportEdNode->GetCfgViewportRegion();
 
@@ -381,6 +426,11 @@ FText SDisplayClusterConfiguratorViewportNode::GetPositionAndSizeText() const
 
 FText SDisplayClusterConfiguratorViewportNode::GetTransformText() const
 {
+	if (!IsGraphNodeValid())
+	{
+		return FText::GetEmpty();
+	}
+
 	UDisplayClusterConfiguratorViewportNode* ViewportEdNode = GetGraphNodeChecked<UDisplayClusterConfiguratorViewportNode>();
 	const FDisplayClusterConfigurationViewport_RemapData& RemapData = ViewportEdNode->GetCfgViewportRemap();
 
@@ -421,6 +471,11 @@ FText SDisplayClusterConfiguratorViewportNode::GetTransformText() const
 
 EVisibility SDisplayClusterConfiguratorViewportNode::GetTransformTextVisibility() const
 {
+	if (!IsGraphNodeValid())
+	{
+		return EVisibility::Collapsed;
+	}
+
 	UDisplayClusterConfiguratorViewportNode* ViewportEdNode = GetGraphNodeChecked<UDisplayClusterConfiguratorViewportNode>();
 	const FDisplayClusterConfigurationViewport_RemapData& RemapData = ViewportEdNode->GetCfgViewportRemap();
 	return RemapData.IsValid() ? EVisibility::Visible : EVisibility::Collapsed;
@@ -454,6 +509,11 @@ EVisibility SDisplayClusterConfiguratorViewportNode::GetLockIconVisibility() con
 
 void SDisplayClusterConfiguratorViewportNode::UpdatePreviewTexture()
 {
+	if (!IsGraphNodeValid())
+	{
+		return;
+	}
+
 	UDisplayClusterConfiguratorViewportNode* ViewportEdNode = GetGraphNodeChecked<UDisplayClusterConfiguratorViewportNode>();
 	UTexture* CurrentTexture = ViewportEdNode->GetPreviewTexture();
 
diff --git a/Engine/Plugins/Runtime/nDisplay/Source/DisplayClusterConfigurator/Private/Views/OutputMapping/GraphNodes/SDisplayClusterConfiguratorWindowNode.cpp b/Engine/Plugins/Runtime/nDisplay/Source/DisplayClusterConfigurator/Private/Views/OutputMapping/GraphNodes/SDisplayClusterConfiguratorWindowNode.cpp
--- a/Engine/Plugins/Runtime/nDisplay/Source/DisplayClusterConfigurator/Private/Views/OutputMapping/GraphNodes/SDisplayClusterConfiguratorWindowNode.cpp
+++ b/Engine/Plugins/Runtime/nDisplay/Source/DisplayClusterConfigurator/Private/Views/OutputMapping/GraphNodes/SDisplayClusterConfiguratorWindowNode.cpp
@@ -219,6 +219,11 @@ public:
 
 	FText GetNodeName() const
 	{
+		if (!ParentNode->IsGraphNodeValid())
+		{
+			return FText::GetEmpty();
+		}
+
 		UDisplayClusterConfiguratorWindowNode* WindowEdNode = ParentNode->GetGraphNodeChecked<UDisplayClusterConfiguratorWindowNode>();
 		const FText NodeName = FText::FromString(WindowEdNode->GetNodeName());
 		if (WindowEdNode->IsPrimary())
@@ -233,6 +238,11 @@ public:
 
 	FText GetPositionAndSizeText() const
 	{
+		if (!ParentNode->IsGraphNodeValid())
+		{
+			return FText::GetEmpty();
+		}
+
 		UDisplayClusterConfiguratorWindowNode* WindowEdNode = ParentNode->GetGraphNodeChecked<UDisplayClusterConfiguratorWindowNode>();
 		FDisplayClusterConfigurationRectangle WindowRect = WindowEdNode->GetCfgWindowRect();
 		return FText::Format(LOCTEXT("ResAndOffset", "[{0} x {1}] @ {2}, {3}"), WindowRect.W, WindowRect.H, WindowRect.X, WindowRect.Y);
@@ -240,6 +250,11 @@ public:
 
 	FOptionalSize GetTitleWidth() const
 	{
+		if (!ParentNode->IsGraphNodeValid())
+		{
+			return FOptionalSize();
+		}
+
 		UDisplayClusterConfiguratorWindowNode* WindowEdNode = ParentNode->GetGraphNodeChecked<UDisplayClusterConfiguratorWindowNode>();
 		return WindowEdNode->NodeWidth;
 	}
@@ -261,6 +276,11 @@ private:
 
 SDisplayClusterConfiguratorWindowNode::~SDisplayClusterConfiguratorWindowNode()
 {
+	if (!IsGraphNodeValid())
+	{
+		return;
+	}
+
 	UDisplayClusterConfiguratorWindowNode* WindowEdNode = GetGraphNodeChecked<UDisplayClusterConfiguratorWindowNode>();
 	WindowEdNode->UnregisterOnPreviewImageChanged(ImageChangedHandle);
 }
@@ -365,12 +385,22 @@ FVector2f SDisplayClusterConfiguratorWindowNode::GetPosition2f() const
 
 bool SDisplayClusterConfiguratorWindowNode::IsAspectRatioFixed() const
 {
+	if (!IsGraphNodeValid())
+	{
+		return false;
+	}
+
 	UDisplayClusterConfiguratorWindowNode* WindowEdNode = GetGraphNodeChecked<UDisplayClusterConfiguratorWindowNode>();
 	return WindowEdNode->IsFixedAspectRatio();
 }
 
 int32 SDisplayClusterConfiguratorWindowNode::GetNodeLogicalLayer() const
 {
+	if (!IsGraphNodeValid())
+	{
+		return 0;
+	}
+
 	UDisplayClusterConfiguratorWindowNode* WindowEdNode = GetGraphNodeChecked<UDisplayClusterConfiguratorWindowNode>();
 	if (bLayerAboveViewports)
 	{
@@ -395,6 +425,11 @@ int32 SDisplayClusterConfiguratorWindowNode::GetNodeLogicalLayer() const
 
 TSharedRef<SWidget> SDisplayClusterConfiguratorWindowNode::CreateBackground(const TAttribute<FSlateColor>& InColorAndOpacity)
 {
+	if (!IsGraphNodeValid())
+	{
+		return SNullWidget::NullWidget;
+	}
+
 	UDisplayClusterConfiguratorWindowNode* WindowEdNode = GetGraphNodeChecked<UDisplayClusterConfiguratorWindowNode>();
 
 	return SNew(SOverlay)
@@ -441,6 +476,11 @@ const FSlateBrush* SDisplayClusterConfiguratorWindowNode::GetBorderBrush() const
 
 int32 SDisplayClusterConfiguratorWindowNode::GetBorderLayerOffset() const
 {
+	if (!IsGraphNodeValid())
+	{
+		return 0;
+	}
+
 	// If the window node is selected, we want to render the border at the same layer as the viewport nodes to ensure it is visible
 	// in the case that the child viewport nodes completely fill the window, since the border is a key indicator that the window node is selected.
 	if (GetOwnerPanel()->SelectionManager.SelectedNodes.Contains(GraphNode))
@@ -496,6 +536,11 @@ FVector2D SDisplayClusterConfiguratorWindowNode::GetPreviewImageSize() const
 
 EVisibility SDisplayClusterConfiguratorWindowNode::GetPreviewImageVisibility() const
 {
+	if (!IsGraphNodeValid())
+	{
+		return EVisibility::Collapsed;
+	}
+
 	UDisplayClusterConfiguratorWindowNode* WindowEdNode = GetGraphNodeChecked<UDisplayClusterConfiguratorWindowNode>();
 	const FString& PreviewImagePath = WindowEdNode->GetPreviewImagePath();
 	return !PreviewImagePath.IsEmpty() ? EVisibility::Visible : EVisibility::Hidden;
@@ -504,6 +549,11 @@ EVisibility SDisplayClusterConfiguratorWindowNode::GetPreviewImageVisibility() c
 
 int32 SDisplayClusterConfiguratorWindowNode::GetNodeTitleLayerOffset() const
 {
+	if (!IsGraphNodeValid())
+	{
+		return 0;
+	}
+
 	UDisplayClusterConfiguratorWindowNode* WindowEdNode = GetGraphNodeChecked<UDisplayClusterConfiguratorWindowNode>();
 	return WindowEdNode->GetAuxiliaryLayer(ObjectPtrDecay(GetOwnerPanel()->SelectionManager.SelectedNodes));
 }
@@ -551,6 +601,11 @@ bool SDisplayClusterConfiguratorWindowNode::IsClusterNodeLocked() const
 
 void SDisplayClusterConfiguratorWindowNode::OnPreviewImageChanged()
 {
+	if (!IsGraphNodeValid())
+	{
+		return;
+	}
+
 	if (PreviewImageWidget.IsValid())
 	{
 		UDisplayClusterConfiguratorWindowNode* WindowEdNode = GetGraphNodeChecked<UDisplayClusterConfiguratorWindowNode>();
