[nDisplay] Restrain precise keyword to when REFACTORING_ALLOWED is defined, per specification.

Changes based on the UE-347560.

#Synced-CL 0

diff --git a/Engine/Plugins/Runtime/nDisplay/Shaders/Private/DisplayClusterColorCorrection.ush b/Engine/Plugins/Runtime/nDisplay/Shaders/Private/DisplayClusterColorCorrection.ush
--- a/Engine/Plugins/Runtime/nDisplay/Shaders/Private/DisplayClusterColorCorrection.ush
+++ b/Engine/Plugins/Runtime/nDisplay/Shaders/Private/DisplayClusterColorCorrection.ush
@@ -3,6 +3,12 @@
 #include "/Engine/Public/Platform.ush"
 #include "/Engine/Private/GammaCorrectionCommon.ush"
 
+#ifndef REFACTORING_ALLOWED
+// Empty definition when REFACTORING_ALLOWED flag is not defined,
+// per specification: https://learn.microsoft.com/en-us/windows/win32/direct3dhlsl/precise
+#define precise
+#endif
+
 // This function copied from the `\\Engine\Plugins\Compositing\OpenColorIO\Shaders\Private\OpenColorIOShader.usf`
 // Note: For our (zero-tolerance) tests to succeed against a no-op, precise is required with the pow function.
 float3 DisplayGammaCorrection(precise float3 InColor, precise float InGamma)
