[nDisplay] Fixed RHI issue.

EDisplayClusterViewportResource::InputShaderResources: In some cases, it can be used as RTT.

#Synced-CL 0

diff --git a/Engine/Plugins/Runtime/nDisplay/Source/DisplayCluster/Private/Render/Viewport/RenderTarget/DisplayClusterRenderTargetManager.cpp b/Engine/Plugins/Runtime/nDisplay/Source/DisplayCluster/Private/Render/Viewport/RenderTarget/DisplayClusterRenderTargetManager.cpp
--- a/Engine/Plugins/Runtime/nDisplay/Source/DisplayCluster/Private/Render/Viewport/RenderTarget/DisplayClusterRenderTargetManager.cpp
+++ b/Engine/Plugins/Runtime/nDisplay/Source/DisplayCluster/Private/Render/Viewport/RenderTarget/DisplayClusterRenderTargetManager.cpp
@@ -204,7 +204,9 @@ bool FDisplayClusterRenderTargetManager::AllocateRenderFrameResources(FViewport*
 						TSharedPtr<FDisplayClusterViewportResource, ESPMode::ThreadSafe>& ReallocResource = ViewportIt->GetViewportResourcesImpl(EDisplayClusterViewportResource::InputShaderResources)[ContextNum];
 						if (!ReallocResource.IsValid() || !EnumHasAnyFlags(ReallocResource->GetResourceState(), EDisplayClusterViewportResourceState::DisableReallocate))
 						{
-							ReallocResource = ResourcesPool->AllocateResource(EmptyString, ContextSize, CustomFormat, EDisplayClusterViewportResourceSettingsFlags::ResolveTargetableTexture | CustomFlags);
+							ReallocResource = ResourcesPool->AllocateResource(EmptyString, ContextSize, CustomFormat,
+								EDisplayClusterViewportResourceSettingsFlags::RenderTargetableTexture // 'InputShader' resource can be used as RTT.
+								| EDisplayClusterViewportResourceSettingsFlags::ResolveTargetableTexture | CustomFlags);
 							ImplViewportTextureResourceLogging(ViewportIt, ContextNum, TEXT("InputShader"), ViewportIt->GetViewportResources(EDisplayClusterViewportResource::InputShaderResources)[ContextNum]);
 						}
 					}
diff --git a/Engine/Plugins/Runtime/nDisplay/Source/DisplayClusterShaders/Private/Shaders/DisplayClusterShadersCopyTexture.cpp b/Engine/Plugins/Runtime/nDisplay/Source/DisplayClusterShaders/Private/Shaders/DisplayClusterShadersCopyTexture.cpp
--- a/Engine/Plugins/Runtime/nDisplay/Source/DisplayClusterShaders/Private/Shaders/DisplayClusterShadersCopyTexture.cpp
+++ b/Engine/Plugins/Runtime/nDisplay/Source/DisplayClusterShaders/Private/Shaders/DisplayClusterShadersCopyTexture.cpp
@@ -381,6 +381,11 @@ namespace UE::DisplayClusterShaders::Private
 				return false;
 			}
 
+#if UE_BUILD_DEBUG
+			// Dest should be RTT (debug purpose)
+			check(EnumHasAnyFlags(DestTexture->GetDesc().Flags, ETextureCreateFlags::RenderTargetable));
+#endif
+
 			const FIntPoint& SrcTextureSize = SrcTexture->GetDesc().Extent;
 			const FIntPoint& DestTextureSize = DestTexture->GetDesc().Extent;
 
