[nDisplay] Propagate per-viewport hide list to LC and CK

Note: only for 3D-LC

Code cleanup.

#Synced-CL 0

diff --git a/Engine/Plugins/Runtime/nDisplay/Source/DisplayCluster/Private/Render/Viewport/Configuration/DisplayClusterViewportConfiguration.cpp b/Engine/Plugins/Runtime/nDisplay/Source/DisplayCluster/Private/Render/Viewport/Configuration/DisplayClusterViewportConfiguration.cpp
--- a/Engine/Plugins/Runtime/nDisplay/Source/DisplayCluster/Private/Render/Viewport/Configuration/DisplayClusterViewportConfiguration.cpp
+++ b/Engine/Plugins/Runtime/nDisplay/Source/DisplayCluster/Private/Render/Viewport/Configuration/DisplayClusterViewportConfiguration.cpp
@@ -476,21 +476,13 @@ bool FDisplayClusterViewportConfiguration::ImplUpdateConfiguration(EDisplayClust
 	return true;
 }
 
-void FDisplayClusterViewportConfiguration::ImplUpdateConfigurationVisibility() const
+void FDisplayClusterViewportConfiguration::ImplUpdateConfigurationVisibility()
 {
-	ADisplayClusterRootActor* SceneRootActor = GetRootActor(EDisplayClusterRootActorType::Scene);
-	FDisplayClusterViewportManager* ViewportManager = GetViewportManagerImpl();
+	RootActorHidePrimitivesList.Reset();
 
 	// Hide root actor components for all viewports
-	TSet<FPrimitiveComponentId> RootActorHidePrimitivesList;
-	if (ViewportManager && SceneRootActor && SceneRootActor->GetHiddenInGamePrimitives(RootActorHidePrimitivesList))
+	if (ADisplayClusterRootActor* SceneRootActor = GetRootActor(EDisplayClusterRootActorType::Scene))
 	{
-		for (const TSharedPtr<FDisplayClusterViewport, ESPMode::ThreadSafe>& ViewportIt : ViewportManager->ImplGetCurrentRenderFrameViewports())
-		{
-			if (ViewportIt.IsValid())
-			{
-				ViewportIt->GetVisibilitySettingsImpl().SetRootActorHideList(RootActorHidePrimitivesList);
-			}
-		}
+		SceneRootActor->GetHiddenInGamePrimitives(RootActorHidePrimitivesList);
 	}
 }
diff --git a/Engine/Plugins/Runtime/nDisplay/Source/DisplayCluster/Private/Render/Viewport/Configuration/DisplayClusterViewportConfiguration.h b/Engine/Plugins/Runtime/nDisplay/Source/DisplayCluster/Private/Render/Viewport/Configuration/DisplayClusterViewportConfiguration.h
--- a/Engine/Plugins/Runtime/nDisplay/Source/DisplayCluster/Private/Render/Viewport/Configuration/DisplayClusterViewportConfiguration.h
+++ b/Engine/Plugins/Runtime/nDisplay/Source/DisplayCluster/Private/Render/Viewport/Configuration/DisplayClusterViewportConfiguration.h
@@ -175,6 +175,12 @@ public:
 	/** Returns the DCRA name or an empty string. */
 	FString GetRootActorName() const;
 
+	/** Returns a list of hidden RootActor primitives. */
+	const TSet<FPrimitiveComponentId>& GetRootActorHidePrimitivesList() const
+	{
+		return RootActorHidePrimitivesList;
+	}
+
 private:
 	/** Update configuration implementation.
 	* 
@@ -186,7 +192,7 @@ private:
 	bool ImplUpdateConfiguration(EDisplayClusterRenderFrameMode InRenderMode, const UWorld* InWorld, const FString& InClusterNodeId, const TArray<FString>* InViewportNames);
 
 	/** Hide DCRA components for nDisplay rendering.*/
-	void ImplUpdateConfigurationVisibility() const;
+	void ImplUpdateConfigurationVisibility();
 
 	/** Set current world.*/
 	void SetCurrentWorldImpl(const UWorld* InWorld);
@@ -220,4 +226,7 @@ private:
 
 	// Pointer to the current world to be rendered in.
 	TWeakObjectPtr<UWorld> CurrentWorldRef;
+
+	// Additional hide primitives list from root actor
+	TSet<FPrimitiveComponentId> RootActorHidePrimitivesList;
 };
diff --git a/Engine/Plugins/Runtime/nDisplay/Source/DisplayCluster/Private/Render/Viewport/Configuration/DisplayClusterViewportConfigurationHelpers_Visibility.cpp b/Engine/Plugins/Runtime/nDisplay/Source/DisplayCluster/Private/Render/Viewport/Configuration/DisplayClusterViewportConfigurationHelpers_Visibility.cpp
--- a/Engine/Plugins/Runtime/nDisplay/Source/DisplayCluster/Private/Render/Viewport/Configuration/DisplayClusterViewportConfigurationHelpers_Visibility.cpp
+++ b/Engine/Plugins/Runtime/nDisplay/Source/DisplayCluster/Private/Render/Viewport/Configuration/DisplayClusterViewportConfigurationHelpers_Visibility.cpp
@@ -143,7 +143,38 @@ void FDisplayClusterViewportConfigurationHelpers_Visibility::UpdateShowOnlyList_
 	TSet<FPrimitiveComponentId> ComponentsList;
 	UE::DisplayCluster::Configuration::VisibilityHelpers::ImplCollectComponentsFromVisibilityList(*DstViewport.Configuration, InVisibilityList, ComponentsList, &DstViewport);
 
+	// Create show-only list for LC/CK viewports
 	DstViewport.GetVisibilitySettingsImpl().SetVisibilityModeAndComponentsList(EDisplayClusterViewport_VisibilityMode::ShowOnly, ComponentsList);
+
+	// Update exclude list for LC/CK viewports
+	const FString& ParentViewportId = DstViewport.GetRenderSettings().GetParentViewportId();
+	if (!ParentViewportId.IsEmpty())
+	{
+		const FDisplayClusterConfigurationICVFX_StageSettings* StageSettings = DstViewport.Configuration->GetStageSettings();
+		ADisplayClusterRootActor* ConfigurationRootActor = DstViewport.Configuration->GetRootActor(EDisplayClusterRootActorType::Configuration);
+		if (!StageSettings || !ConfigurationRootActor)
+		{
+			return;
+		}
+
+		// Exclude Outer viewport HideList
+		TSet<FPrimitiveComponentId> OuterViewportHiddenComponentsList;
+		UE::DisplayCluster::Configuration::VisibilityHelpers::ImplCollectComponentsFromVisibilityList(
+			DstViewport.Configuration.Get(), StageSettings->OuterViewportHideList, OuterViewportHiddenComponentsList);
+		DstViewport.GetVisibilitySettingsImpl().AddExcludedComponents(OuterViewportHiddenComponentsList);
+
+		// Exclude Per-viewport HideList
+		const UDisplayClusterConfigurationData* ConfigData = ConfigurationRootActor->GetConfigData();
+		if (const UDisplayClusterConfigurationViewport* ViewportCfg = ConfigData
+			? ConfigData->GetViewport(DstViewport.ClusterNodeId, ParentViewportId) : nullptr)
+		{
+			TSet<FPrimitiveComponentId> PerViewportHiddenComponentsList;
+			UE::DisplayCluster::Configuration::VisibilityHelpers::ImplCollectComponentsFromVisibilityList(
+				DstViewport.Configuration.Get(), ViewportCfg->RenderSettings.HiddenContent, PerViewportHiddenComponentsList);
+
+			DstViewport.GetVisibilitySettingsImpl().AddExcludedComponents(PerViewportHiddenComponentsList);
+		}
+	}
 }
 
 void FDisplayClusterViewportConfigurationHelpers_Visibility::AppendHideList_ICVFX(FDisplayClusterViewport& DstViewport, const FDisplayClusterConfigurationICVFX_VisibilityList& InHideList)
diff --git a/Engine/Plugins/Runtime/nDisplay/Source/DisplayCluster/Private/Render/Viewport/DisplayClusterViewport.cpp b/Engine/Plugins/Runtime/nDisplay/Source/DisplayCluster/Private/Render/Viewport/DisplayClusterViewport.cpp
--- a/Engine/Plugins/Runtime/nDisplay/Source/DisplayCluster/Private/Render/Viewport/DisplayClusterViewport.cpp
+++ b/Engine/Plugins/Runtime/nDisplay/Source/DisplayCluster/Private/Render/Viewport/DisplayClusterViewport.cpp
@@ -348,7 +348,7 @@ void FDisplayClusterViewport::SetupSceneView(uint32 ContextNum, class UWorld* Wo
 	if(RenderSettings.CaptureMode == EDisplayClusterViewportCaptureMode::MoviePipeline)
 	{
 		// Apply visibility settigns to view
-		VisibilitySettings.SetupSceneView(InOutView);
+		VisibilitySettings.SetupSceneView(*this, InOutView);
 
 		return;
 	}
@@ -380,7 +380,7 @@ void FDisplayClusterViewport::SetupSceneView(uint32 ContextNum, class UWorld* Wo
 	}
 
 	// Apply visibility settigns to view
-	VisibilitySettings.SetupSceneView(InOutView);
+	VisibilitySettings.SetupSceneView(*this, InOutView);
 
 	// Handle Motion blur parameters
 	CameraMotionBlur.SetupSceneView(Contexts[ContextNum], InOutView);
diff --git a/Engine/Plugins/Runtime/nDisplay/Source/DisplayCluster/Private/Render/Viewport/DisplayClusterViewportManager.h b/Engine/Plugins/Runtime/nDisplay/Source/DisplayCluster/Private/Render/Viewport/DisplayClusterViewportManager.h
--- a/Engine/Plugins/Runtime/nDisplay/Source/DisplayCluster/Private/Render/Viewport/DisplayClusterViewportManager.h
+++ b/Engine/Plugins/Runtime/nDisplay/Source/DisplayCluster/Private/Render/Viewport/DisplayClusterViewportManager.h
@@ -171,6 +171,18 @@ public:
 	}
 
 	// internal use only
+
+	/**
+	* Find viewport object by name [Game thread func]
+	*
+	* @param ViewportId - Viewport name
+	* @return - viewport object ref, or nullptr
+	*/
+	inline FDisplayClusterViewport* FindViewportByName(const FString& InViewportId) const
+	{
+		return ImplFindViewport(InViewportId).Get();
+	}
+
 	FDisplayClusterViewport* CreateViewport(const FString& ViewportId, const class UDisplayClusterConfigurationViewport& ConfigurationViewport);
 	FDisplayClusterViewport* CreateViewport(const FString& ViewportId, const TSharedPtr<IDisplayClusterProjectionPolicy, ESPMode::ThreadSafe>& InProjectionPolicy);
 	bool                     DeleteViewport(const FString& ViewportId);
diff --git a/Engine/Plugins/Runtime/nDisplay/Source/DisplayCluster/Private/Render/Viewport/DisplayClusterViewport_VisibilitySettings.cpp b/Engine/Plugins/Runtime/nDisplay/Source/DisplayCluster/Private/Render/Viewport/DisplayClusterViewport_VisibilitySettings.cpp
--- a/Engine/Plugins/Runtime/nDisplay/Source/DisplayCluster/Private/Render/Viewport/DisplayClusterViewport_VisibilitySettings.cpp
+++ b/Engine/Plugins/Runtime/nDisplay/Source/DisplayCluster/Private/Render/Viewport/DisplayClusterViewport_VisibilitySettings.cpp
@@ -8,8 +8,10 @@
 #include "EngineUtils.h"
 #include "SceneView.h"
 
-void FDisplayClusterViewport_VisibilitySettings::SetupSceneView(FSceneView& InOutView) const
+void FDisplayClusterViewport_VisibilitySettings::SetupSceneView(const FDisplayClusterViewport& Owner, FSceneView& InOutView) const
 {
+	const TSet<FPrimitiveComponentId>& RootActorHidePrimitivesList = Owner.Configuration->GetRootActorHidePrimitivesList();
+
 	switch (VisibilityMode)
 	{
 	case EDisplayClusterViewport_VisibilityMode::ShowOnly:
@@ -23,6 +25,12 @@ void FDisplayClusterViewport_VisibilitySettings::SetupSceneView(FSceneView& InOu
 				continue;
 			}
 
+			// Except components from the exclude list
+			if (ExcludedComponents.Contains(ComponentId))
+			{
+				continue;
+			}
+
 			InOutView.ShowOnlyPrimitives->Add(ComponentId);
 		}
 
@@ -30,6 +38,23 @@ void FDisplayClusterViewport_VisibilitySettings::SetupSceneView(FSceneView& InOu
 	}
 	case EDisplayClusterViewport_VisibilityMode::Hide:
 	{
+		// If ExcludedComponents is used, filter ComponentsList.
+		if (!ExcludedComponents.IsEmpty())
+		{
+			for (const FPrimitiveComponentId& ComponentId : ComponentsList)
+			{
+				// Except components from the exclude list
+				if (ExcludedComponents.Contains(ComponentId))
+				{
+					continue;
+				}
+
+				InOutView.HiddenPrimitives.Add(ComponentId);
+			}
+
+			break;
+		}
+
 		InOutView.HiddenPrimitives.Append(ComponentsList);
 		break;
 	}
diff --git a/Engine/Plugins/Runtime/nDisplay/Source/DisplayCluster/Private/Render/Viewport/DisplayClusterViewport_VisibilitySettings.h b/Engine/Plugins/Runtime/nDisplay/Source/DisplayCluster/Private/Render/Viewport/DisplayClusterViewport_VisibilitySettings.h
--- a/Engine/Plugins/Runtime/nDisplay/Source/DisplayCluster/Private/Render/Viewport/DisplayClusterViewport_VisibilitySettings.h
+++ b/Engine/Plugins/Runtime/nDisplay/Source/DisplayCluster/Private/Render/Viewport/DisplayClusterViewport_VisibilitySettings.h
@@ -30,7 +30,13 @@ public:
 	{
 		VisibilityMode = EDisplayClusterViewport_VisibilityMode::None;
 		ComponentsList.Empty();
-		RootActorHidePrimitivesList.Empty();
+		ExcludedComponents.Empty();
+	}
+
+	/** Return components list. */
+	const TSet<FPrimitiveComponentId>& GetComponentsList() const
+	{
+		return ComponentsList;
 	}
 
 	/** Sets the visibility mode and component list. */
@@ -54,9 +60,10 @@ public:
 		return false;
 	}
 
-	void SetRootActorHideList(TSet<FPrimitiveComponentId>& InHidePrimitivesList)
+	/** Exclude these components from the output list. */
+	void AddExcludedComponents(const TSet<FPrimitiveComponentId>& InComponentsList)
 	{
-		RootActorHidePrimitivesList = InHidePrimitivesList;
+		ExcludedComponents.Append(InComponentsList);
 	}
 
 	/** Returns true if this viewport contains any geometry and can be rendered. */
@@ -70,12 +77,13 @@ public:
 		return true;
 	}
 
-	void SetupSceneView(FSceneView& InOutView) const;
+	/** Configure InOutView using the viewport settings. */
+	void SetupSceneView(const FDisplayClusterViewport& Owner, FSceneView& InOutView) const;
 
 private:
 	EDisplayClusterViewport_VisibilityMode VisibilityMode = EDisplayClusterViewport_VisibilityMode::None;
 	TSet<FPrimitiveComponentId> ComponentsList;
 
-	// Additional hide primitives list from root actor
-	TSet<FPrimitiveComponentId> RootActorHidePrimitivesList;
+	// Exclude this components from ComponentsList
+	TSet<FPrimitiveComponentId> ExcludedComponents;
 };
