[nDisplay] Fixed Divide by zero.

Code cleanup.

#Synced-CL 0

diff --git a/Engine/Plugins/Runtime/nDisplay/Source/DisplayClusterWarp/Private/WarpBlend/DisplayClusterWarpBlend_GeometryProxy.cpp b/Engine/Plugins/Runtime/nDisplay/Source/DisplayClusterWarp/Private/WarpBlend/DisplayClusterWarpBlend_GeometryProxy.cpp
--- a/Engine/Plugins/Runtime/nDisplay/Source/DisplayClusterWarp/Private/WarpBlend/DisplayClusterWarpBlend_GeometryProxy.cpp
+++ b/Engine/Plugins/Runtime/nDisplay/Source/DisplayClusterWarp/Private/WarpBlend/DisplayClusterWarpBlend_GeometryProxy.cpp
@@ -266,15 +266,17 @@ bool FDisplayClusterWarpBlend_GeometryProxy::ImplUpdateFrustumGeometryCache_Warp
 		}
 
 		// Update cache
-		FDisplayClusterWarpBlendMath_WarpMap DataHelper(*(WarpMapTexture.Get()));
-
-		GeometryCache.AABBox = DataHelper.GetAABBox();
-		GeometryCache.SurfaceViewNormal = DataHelper.GetSurfaceViewNormal();
-		GeometryCache.SurfaceViewPlane  = DataHelper.GetSurfaceViewPlane();
+		FDisplayClusterWarpBlendMath_WarpMap DataHelper(WarpMapTexture.ToSharedRef());
+		if (DataHelper.CanBeUsed())
+		{
+			GeometryCache.AABBox = DataHelper.GetAABBox();
+			GeometryCache.SurfaceViewNormal = DataHelper.GetSurfaceViewNormal();
+			GeometryCache.SurfaceViewPlane = DataHelper.GetSurfaceViewPlane();
 
-		bIsGeometryCacheValid = true;
+			bIsGeometryCacheValid = true;
 
-		return true;
+			return true;
+		}
 	}
 
 	bIsGeometryCacheValid = false;
@@ -334,10 +336,7 @@ bool FDisplayClusterWarpBlend_GeometryProxy::UpdateFrustumGeometryLOD(const FInt
 		case EDisplayClusterWarpFrustumGeometryType::WarpMap:
 		{
 			// Generate valid points for texturebox method:
-			FDisplayClusterWarpBlendMath_WarpMap DataHelper(*(WarpMapTexture.Get()));
-			DataHelper.BuildIndexLOD(InSizeLOD.X, InSizeLOD.Y, GeometryCache.IndexLOD);
-
-			return true;
+			return FDisplayClusterWarpBlendMath_WarpMap(WarpMapTexture.ToSharedRef()).BuildIndexLOD(InSizeLOD, GeometryCache.IndexLOD);
 		}
 		default:
 			break;
diff --git a/Engine/Plugins/Runtime/nDisplay/Source/DisplayClusterWarp/Private/WarpBlend/Math/DisplayClusterWarpBlendMath_FrustumProjection.cpp b/Engine/Plugins/Runtime/nDisplay/Source/DisplayClusterWarp/Private/WarpBlend/Math/DisplayClusterWarpBlendMath_FrustumProjection.cpp
--- a/Engine/Plugins/Runtime/nDisplay/Source/DisplayClusterWarp/Private/WarpBlend/Math/DisplayClusterWarpBlendMath_FrustumProjection.cpp
+++ b/Engine/Plugins/Runtime/nDisplay/Source/DisplayClusterWarp/Private/WarpBlend/Math/DisplayClusterWarpBlendMath_FrustumProjection.cpp
@@ -177,9 +177,16 @@ bool FDisplayClusterWarpBlendMath_Frustum::ImplCalcFrustumProjection_WarpMap(con
 	const TSharedPtr<IDisplayClusterRender_Texture, ESPMode::ThreadSafe>& WarpMapTexture = GeometryContext.GeometryProxy.WarpMapTexture;
 	if (const FVector4f* SourcePts = (FVector4f*)(WarpMapTexture.IsValid() ? WarpMapTexture->GetData() : nullptr))
 	{
+		const FIntPoint WarpMapTextureSize(WarpMapTexture->GetWidth(), WarpMapTexture->GetHeight());
+		if (WarpMapTextureSize.GetMin() < 2)
+		{
+			// Minimal texture size should be at least 2x2
+			return false;
+		}
+
 		bool bAllPointsInFrustum = true;
 
-		const int32 PointsAmount = WarpMapTexture->GetWidth() * WarpMapTexture->GetHeight();
+		const int32 PointsAmount = WarpMapTextureSize.X * WarpMapTextureSize.Y;
 
 		switch (WarpData.FrustumType)
 		{
diff --git a/Engine/Plugins/Runtime/nDisplay/Source/DisplayClusterWarp/Private/WarpBlend/Math/DisplayClusterWarpBlendMath_WarpMap.h b/Engine/Plugins/Runtime/nDisplay/Source/DisplayClusterWarp/Private/WarpBlend/Math/DisplayClusterWarpBlendMath_WarpMap.h
--- a/Engine/Plugins/Runtime/nDisplay/Source/DisplayClusterWarp/Private/WarpBlend/Math/DisplayClusterWarpBlendMath_WarpMap.h
+++ b/Engine/Plugins/Runtime/nDisplay/Source/DisplayClusterWarp/Private/WarpBlend/Math/DisplayClusterWarpBlendMath_WarpMap.h
@@ -12,20 +12,32 @@
 class FDisplayClusterWarpBlendMath_WarpMap
 {
 public:
-	FDisplayClusterWarpBlendMath_WarpMap(const IDisplayClusterRender_Texture& InWarpMap)
-		: Width(InWarpMap.GetWidth())
-		, Height(InWarpMap.GetHeight())
-		, Data((FVector4f*)InWarpMap.GetData())
+	FDisplayClusterWarpBlendMath_WarpMap(const TSharedRef<IDisplayClusterRender_Texture, ESPMode::ThreadSafe>& InWarpMapTexture)
+		: WarpMapTexture(InWarpMapTexture)
 	{ }
 
+	/** Check if WarpMap data can be used. */
+	inline bool CanBeUsed() const
+	{
+		return GetWidth() >= 2
+			&& GetHeight() >= 2
+			&& GetData();
+	}
+
 public:
 	/** Build a bbox from valid points. */
 	FBox GetAABBox() const
 	{
+		// The input texture must be valid and have a size of at least 2x2.
+		if (!CanBeUsed())
+		{
+			return FDisplayClusterWarpAABB();
+		}
+
 		FDisplayClusterWarpAABB WarpAABB;
-		for (int32 Y = 0; Y < Height; ++Y)
+		for (int32 Y = 0; Y < GetHeight(); ++Y)
 		{
-			for (int32 X = 0; X < Width; ++X)
+			for (int32 X = 0; X < GetWidth(); ++X)
 			{
 				WarpAABB.UpdateAABB(GetPoint(X, Y));
 			}
@@ -34,21 +46,31 @@ public:
 		return WarpAABB;
 	}
 
-	void BuildIndexLOD(int32 DivX, int32 DivY, TArray<int32>& OutIndexLOD)
+	/** Build an LOD mesh with the size specified in InSizeLOD. */
+	bool BuildIndexLOD(const FIntPoint& InSizeLOD, TArray<int32>& OutIndexLOD)
 	{
-		check(DivX > 0);
-		check(DivY > 0);
+		OutIndexLOD.Reset();
+
+		// The input texture must be valid and have a size of at least 2x2.
+		if (!CanBeUsed())
+		{
+			return false;
+		}
+
+		// Let's limit the LOD size to the size of the input texture.
+		const int32 DivX = FMath::Clamp(InSizeLOD.X, 2, GetWidth());
+		const int32 DivY = FMath::Clamp(InSizeLOD.Y, 2, GetHeight());
 
 		OutIndexLOD.Reserve(DivX * DivY);
 
 		// Generate valid points for texturebox method:
 		for (int32 low_y = 0; low_y < DivY; low_y++)
 		{
-			int32 y = (Height - 1) * (float(low_y) / (DivY - 1)); //-V609 UE-349513 
+			int32 y = (GetHeight() - 1) * (float(low_y) / (DivY - 1)); //-V609 UE-349513 
 
 			for (int32 low_x = 0; low_x < DivX; low_x++)
 			{
-				int32 x = (Width - 1) * (float(low_x) / (DivX - 1)); //-V609 UE-349513 
+				int32 x = (GetWidth() - 1) * (float(low_x) / (DivX - 1)); //-V609 UE-349513 
 
 				if (IsValidPoint(x, y))
 				{
@@ -65,9 +87,11 @@ public:
 				}
 			}
 		}
+
+		return true;
 	}
 
-public:
+private:
 	int32 GetSavedPointIndex()
 	{
 		return GetPointIndex(ValidPointX, ValidPointY);
@@ -75,17 +99,17 @@ public:
 
 	FORCEINLINE int32 GetPointIndex(int32 InX, int32 InY) const
 	{
-		return InX + InY * Width;
+		return InX + InY * GetWidth();
 	}
 
 	FORCEINLINE const FVector4f& GetPoint(int32 InX, int32 InY) const
 	{
-		return Data[GetPointIndex(InX, InY)];
+		return GetData()[GetPointIndex(InX, InY)];
 	}
 
 	FORCEINLINE const FVector4f& GetPoint(int32 PointIndex) const
 	{
-		return Data[PointIndex];
+		return GetData()[PointIndex];
 	}
 
 	bool IsValidPoint(int32 InX, int32 InY)
@@ -98,7 +122,7 @@ public:
 		X0 = InX;
 		Y0 = InY;
 
-		for (int32 Range = 1; Range < Width; Range++)
+		for (int32 Range = 1; Range < GetWidth(); Range++)
 		{
 			if (FindValidPointInRange(Range))
 			{
@@ -109,14 +133,21 @@ public:
 		return false;
 	}
 
+public:
 	FVector GetSurfaceViewNormal()
 	{
+		// The input texture must be valid and have a size of at least 2x2.
+		if (!CanBeUsed())
+		{
+			return FVector::ZeroVector;
+		}
+
 		int32 Ncount = 0;
 		double Nxyz[3] = { 0,0,0 };
 
-		for (int32 ItY = 0; ItY < (Height - 2); ++ItY)
+		for (int32 ItY = 0; ItY < (GetHeight() - 2); ++ItY)
 		{
-			for (int32 ItX = 0; ItX < (Width - 2); ++ItX)
+			for (int32 ItX = 0; ItX < (GetWidth() - 2); ++ItX)
 			{
 				const FVector4f& Pts0 = GetPoint(ItX, ItY);
 				const FVector4f& Pts1 = GetPoint(ItX + 1, ItY);
@@ -138,6 +169,11 @@ public:
 			}
 		}
 
+		if (!Ncount)
+		{
+			return FVector::ZeroVector;
+		}
+
 		double Scale = double(1) / Ncount;
 		for (int32 AxisIndex = 0; AxisIndex < 3; AxisIndex++)
 		{
@@ -149,9 +185,15 @@ public:
 
 	FVector GetSurfaceViewPlane()
 	{
+		// The input texture must be valid and have a size of at least 2x2.
+		if (!CanBeUsed())
+		{
+			return FVector::ZeroVector;
+		}
+
 		const FVector4f& Pts0 = GetValidPoint(0, 0);
-		const FVector4f& Pts1 = GetValidPoint(Width - 1, 0);
-		const FVector4f& Pts2 = GetValidPoint(0, Height - 1);
+		const FVector4f& Pts1 = GetValidPoint(GetWidth() - 1, 0);
+		const FVector4f& Pts2 = GetValidPoint(0, GetHeight() - 1);
 
 		const FVector N1 = FVector4(Pts1 - Pts0);
 		const FVector N2 = FVector4(Pts2 - Pts0);
@@ -194,13 +236,13 @@ private:
 
 	bool IsValid(int32 newX, int32 newY)
 	{
-		if (newX < 0 || newY < 0 || newX >= Width || newY >= Height)
+		if (newX < 0 || newY < 0 || newX >= GetWidth() || newY >= GetHeight())
 		{
 			// Out of texture
 			return false;
 		}
 
-		if (Data[GetPointIndex(newX, newY)].W > 0)
+		if (GetData()[GetPointIndex(newX, newY)].W > 0)
 		{
 			// Store valid result
 			ValidPointX = newX;
@@ -213,9 +255,28 @@ private:
 	}
 
 private:
-	const int32 Width;
-	const int32 Height;
-	const FVector4f* Data;
+
+	/** Return WarpMap texture width. */
+	inline const int32 GetWidth() const
+	{
+		return WarpMapTexture->GetWidth();
+	}
+
+	/** Return WarpMap texture height. */
+	inline const int32 GetHeight() const
+	{
+		return WarpMapTexture->GetHeight();
+	}
+
+	/** Return WarpMap texture points data. */
+	inline const FVector4f* GetData() const
+	{
+		return (FVector4f*)WarpMapTexture->GetData();
+	}
+
+private:
+	// WarpMap texture reference
+	const TSharedRef<IDisplayClusterRender_Texture, ESPMode::ThreadSafe> WarpMapTexture;
 
 	// Internal logic
 	int32 ValidPointX = 0;
