[nDisplay] HideList rules for 3D-LC & CK

"StageSettings" -> "Content Hidden from Entire Cluster" for 3D-LC & CK
"In Camera VFX" -> "Content Hidden from Inner Frustum" for CK

#Synced-CL 0

diff --git a/Engine/Plugins/Runtime/nDisplay/Source/DisplayCluster/Private/Render/Viewport/Configuration/DisplayClusterViewportConfigurationHelpers_ICVFX.cpp b/Engine/Plugins/Runtime/nDisplay/Source/DisplayCluster/Private/Render/Viewport/Configuration/DisplayClusterViewportConfigurationHelpers_ICVFX.cpp
--- a/Engine/Plugins/Runtime/nDisplay/Source/DisplayCluster/Private/Render/Viewport/Configuration/DisplayClusterViewportConfigurationHelpers_ICVFX.cpp
+++ b/Engine/Plugins/Runtime/nDisplay/Source/DisplayCluster/Private/Render/Viewport/Configuration/DisplayClusterViewportConfigurationHelpers_ICVFX.cpp
@@ -583,7 +583,7 @@ void FDisplayClusterViewportConfigurationHelpers_ICVFX::UpdateChromakeyViewportS
 		{
 			check(InChromakeyRenderSettings->ShowOnlyList.IsVisibilityListValid());
 
-			FDisplayClusterViewportConfigurationHelpers_Visibility::UpdateShowOnlyList_ICVFX(DstViewport, InChromakeyRenderSettings->ShowOnlyList);
+			FDisplayClusterViewportConfigurationHelpers_Visibility::UpdateShowOnlyList_ICVFX(DstViewport, InChromakeyRenderSettings->ShowOnlyList, &InCameraSettings);
 		}
 
 		DstViewport.UpdateConfiguration_OverlayRenderSettings(InChromakeyRenderSettings->AdvancedRenderSettings);
diff --git a/Engine/Plugins/Runtime/nDisplay/Source/DisplayCluster/Private/Render/Viewport/Configuration/DisplayClusterViewportConfigurationHelpers_Visibility.cpp b/Engine/Plugins/Runtime/nDisplay/Source/DisplayCluster/Private/Render/Viewport/Configuration/DisplayClusterViewportConfigurationHelpers_Visibility.cpp
--- a/Engine/Plugins/Runtime/nDisplay/Source/DisplayCluster/Private/Render/Viewport/Configuration/DisplayClusterViewportConfigurationHelpers_Visibility.cpp
+++ b/Engine/Plugins/Runtime/nDisplay/Source/DisplayCluster/Private/Render/Viewport/Configuration/DisplayClusterViewportConfigurationHelpers_Visibility.cpp
@@ -138,25 +138,35 @@ namespace UE::DisplayCluster::Configuration::VisibilityHelpers
 	}
 };
 
-void FDisplayClusterViewportConfigurationHelpers_Visibility::UpdateShowOnlyList_ICVFX(FDisplayClusterViewport& DstViewport, const FDisplayClusterConfigurationICVFX_VisibilityList& InVisibilityList)
+void FDisplayClusterViewportConfigurationHelpers_Visibility::UpdateShowOnlyList_ICVFX(
+	FDisplayClusterViewport& DstViewport,
+	const FDisplayClusterConfigurationICVFX_VisibilityList& InVisibilityList,
+	const FDisplayClusterConfigurationICVFX_CameraSettings* InCameraSettings)
 {
 	TSet<FPrimitiveComponentId> ComponentsList;
 	UE::DisplayCluster::Configuration::VisibilityHelpers::ImplCollectComponentsFromVisibilityList(*DstViewport.Configuration, InVisibilityList, ComponentsList, &DstViewport);
 
-	// Create show-only list for LC/CK viewports
+	// Create show-only list for 3D-LC/CK viewports
 	DstViewport.GetVisibilitySettingsImpl().SetVisibilityModeAndComponentsList(EDisplayClusterViewport_VisibilityMode::ShowOnly, ComponentsList);
 
-	// Update exclude list for LC/CK viewports
-	const FString& ParentViewportId = DstViewport.GetRenderSettings().GetParentViewportId();
-	if (!ParentViewportId.IsEmpty())
+	const FDisplayClusterConfigurationICVFX_StageSettings* StageSettings = DstViewport.Configuration->GetStageSettings();
+	ADisplayClusterRootActor* ConfigurationRootActor = DstViewport.Configuration->GetRootActor(EDisplayClusterRootActorType::Configuration);
+	if (!StageSettings || !ConfigurationRootActor)
 	{
-		const FDisplayClusterConfigurationICVFX_StageSettings* StageSettings = DstViewport.Configuration->GetStageSettings();
-		ADisplayClusterRootActor* ConfigurationRootActor = DstViewport.Configuration->GetRootActor(EDisplayClusterRootActorType::Configuration);
-		if (!StageSettings || !ConfigurationRootActor)
-		{
-			return;
-		}
+		return;
+	}
+
+	// Exclude entire cluster HideList
+	TSet<FPrimitiveComponentId> EntireClusterHiddenComponentsList;
+	UE::DisplayCluster::Configuration::VisibilityHelpers::ImplCollectComponentsFromVisibilityList(
+		DstViewport.Configuration.Get(), StageSettings->HideList, EntireClusterHiddenComponentsList);
+	DstViewport.GetVisibilitySettingsImpl().AddExcludedComponents(EntireClusterHiddenComponentsList);
 
+	// Update exclude list for 3D-LC viewports
+	const FString& ParentViewportId = DstViewport.GetRenderSettings().GetParentViewportId();
+	const bool bLightcardViewport = EnumHasAnyFlags(DstViewport.GetRenderSettingsICVFX().RuntimeFlags, EDisplayClusterViewportRuntimeICVFXFlags::Lightcard);
+	if (bLightcardViewport && !ParentViewportId.IsEmpty())
+	{
 		// Exclude Outer viewport HideList
 		TSet<FPrimitiveComponentId> OuterViewportHiddenComponentsList;
 		UE::DisplayCluster::Configuration::VisibilityHelpers::ImplCollectComponentsFromVisibilityList(
@@ -175,6 +185,17 @@ void FDisplayClusterViewportConfigurationHelpers_Visibility::UpdateShowOnlyList_
 			DstViewport.GetVisibilitySettingsImpl().AddExcludedComponents(PerViewportHiddenComponentsList);
 		}
 	}
+
+	// Update exclude list for CK viewports
+	const bool bChromakeyViewport = EnumHasAnyFlags(DstViewport.GetRenderSettingsICVFX().RuntimeFlags, EDisplayClusterViewportRuntimeICVFXFlags::Chromakey);
+	if (bChromakeyViewport && InCameraSettings)
+	{
+		// Exclude camera HideList
+		TSet<FPrimitiveComponentId> CameraHiddenComponentsList;
+		UE::DisplayCluster::Configuration::VisibilityHelpers::ImplCollectComponentsFromVisibilityList(
+			DstViewport.Configuration.Get(), InCameraSettings->CameraHideList, CameraHiddenComponentsList);
+		DstViewport.GetVisibilitySettingsImpl().AddExcludedComponents(CameraHiddenComponentsList);
+	}
 }
 
 void FDisplayClusterViewportConfigurationHelpers_Visibility::AppendHideList_ICVFX(FDisplayClusterViewport& DstViewport, const FDisplayClusterConfigurationICVFX_VisibilityList& InHideList)
@@ -194,6 +215,7 @@ void FDisplayClusterViewportConfigurationHelpers_Visibility::UpdateHideList_ICVF
 	{
 		TSet<FPrimitiveComponentId> ComponentsList;
 
+		// Entire cluster HideList
 		UE::DisplayCluster::Configuration::VisibilityHelpers::ImplCollectComponentsFromVisibilityList(InConfiguration, StageSettings->HideList, ComponentsList);
 
 		// Hide lightcard
diff --git a/Engine/Plugins/Runtime/nDisplay/Source/DisplayCluster/Private/Render/Viewport/Configuration/DisplayClusterViewportConfigurationHelpers_Visibility.h b/Engine/Plugins/Runtime/nDisplay/Source/DisplayCluster/Private/Render/Viewport/Configuration/DisplayClusterViewportConfigurationHelpers_Visibility.h
--- a/Engine/Plugins/Runtime/nDisplay/Source/DisplayCluster/Private/Render/Viewport/Configuration/DisplayClusterViewportConfigurationHelpers_Visibility.h
+++ b/Engine/Plugins/Runtime/nDisplay/Source/DisplayCluster/Private/Render/Viewport/Configuration/DisplayClusterViewportConfigurationHelpers_Visibility.h
@@ -9,6 +9,7 @@ class FDisplayClusterViewportConfiguration;
 class FDisplayClusterViewport;
 class AActor;
 struct FDisplayClusterConfigurationICVFX_VisibilityList;
+struct FDisplayClusterConfigurationICVFX_CameraSettings;
 
 /**
 * Visibility configuration helper class.
@@ -16,8 +17,8 @@ struct FDisplayClusterConfigurationICVFX_VisibilityList;
 class FDisplayClusterViewportConfigurationHelpers_Visibility
 {
 public:
-	// Update ShowOnly list for DstViewport
-	static void UpdateShowOnlyList_ICVFX(FDisplayClusterViewport& DstViewport, const FDisplayClusterConfigurationICVFX_VisibilityList& InVisibilityList);
+	// Update ShowOnly list for 3D-LC and CK viewports
+	static void UpdateShowOnlyList_ICVFX(FDisplayClusterViewport& DstViewport, const FDisplayClusterConfigurationICVFX_VisibilityList& InVisibilityList, const FDisplayClusterConfigurationICVFX_CameraSettings* InCameraSettings = nullptr);
 
 	// Update hide lists for DstViewports (lightcards, chromakey, stage settings hide list, outer viewports hide list, etc)
 	static void UpdateHideList_ICVFX(FDisplayClusterViewportConfiguration& InConfiguration, TArray<TSharedPtr<FDisplayClusterViewport, ESPMode::ThreadSafe>>& DstViewports);
