[nDisplay] Implemented custom states (shared, distributed)

[RN] Implemented custom states that make it easy to share arbitrary data across the cluster

#Synced-CL 0

diff --git a/Engine/Plugins/Runtime/nDisplay/Source/DisplayCluster/Private/Cluster/Controller/DisplayClusterClusterNodeCtrlDisabled.cpp b/Engine/Plugins/Runtime/nDisplay/Source/DisplayCluster/Private/Cluster/Controller/DisplayClusterClusterNodeCtrlDisabled.cpp
--- a/Engine/Plugins/Runtime/nDisplay/Source/DisplayCluster/Private/Cluster/Controller/DisplayClusterClusterNodeCtrlDisabled.cpp
+++ b/Engine/Plugins/Runtime/nDisplay/Source/DisplayCluster/Private/Cluster/Controller/DisplayClusterClusterNodeCtrlDisabled.cpp
@@ -83,6 +83,11 @@ EDisplayClusterCommResult FDisplayClusterClusterNodeCtrlDisabled::GetNativeInput
 	return EDisplayClusterCommResult::NotAllowed;
 }
 
+EDisplayClusterCommResult FDisplayClusterClusterNodeCtrlDisabled::PropagateStatesData(const TArray<uint8>& InLocalStatesData, TArray<uint8>& OutClusterStatesData)
+{
+	return EDisplayClusterCommResult::NotAllowed;
+}
+
 
 //////////////////////////////////////////////////////////////////////////////////////////////
 // IDisplayClusterProtocolRenderSync
diff --git a/Engine/Plugins/Runtime/nDisplay/Source/DisplayCluster/Private/Cluster/Controller/DisplayClusterClusterNodeCtrlDisabled.h b/Engine/Plugins/Runtime/nDisplay/Source/DisplayCluster/Private/Cluster/Controller/DisplayClusterClusterNodeCtrlDisabled.h
--- a/Engine/Plugins/Runtime/nDisplay/Source/DisplayCluster/Private/Cluster/Controller/DisplayClusterClusterNodeCtrlDisabled.h
+++ b/Engine/Plugins/Runtime/nDisplay/Source/DisplayCluster/Private/Cluster/Controller/DisplayClusterClusterNodeCtrlDisabled.h
@@ -39,6 +39,7 @@ public:
 	virtual EDisplayClusterCommResult GetObjectsData(const EDisplayClusterSyncGroup InSyncGroup, TMap<FString, FString>& OutObjectsData) override;
 	virtual EDisplayClusterCommResult GetEventsData(TArray<TSharedPtr<FDisplayClusterClusterEventJson>>& OutJsonEvents, TArray<TSharedPtr<FDisplayClusterClusterEventBinary>>& OutBinaryEvents) override;
 	virtual EDisplayClusterCommResult GetNativeInputData(TMap<FString, FString>& OutNativeInputData) override;
+	virtual EDisplayClusterCommResult PropagateStatesData(const TArray<uint8>& InLocalStatesData, TArray<uint8>& OutClusterStatesData) override;
 	//~ End IDisplayClusterProtocolClusterSync
 
 public:
diff --git a/Engine/Plugins/Runtime/nDisplay/Source/DisplayCluster/Private/Cluster/Controller/DisplayClusterClusterNodeCtrlEditor.cpp b/Engine/Plugins/Runtime/nDisplay/Source/DisplayCluster/Private/Cluster/Controller/DisplayClusterClusterNodeCtrlEditor.cpp
--- a/Engine/Plugins/Runtime/nDisplay/Source/DisplayCluster/Private/Cluster/Controller/DisplayClusterClusterNodeCtrlEditor.cpp
+++ b/Engine/Plugins/Runtime/nDisplay/Source/DisplayCluster/Private/Cluster/Controller/DisplayClusterClusterNodeCtrlEditor.cpp
@@ -120,6 +120,12 @@ EDisplayClusterCommResult FDisplayClusterClusterNodeCtrlEditor::GetNativeInputDa
 	return EDisplayClusterCommResult::Ok;
 }
 
+EDisplayClusterCommResult FDisplayClusterClusterNodeCtrlEditor::PropagateStatesData(const TArray<uint8>& InLocalStatesData, TArray<uint8>& OutClusterStatesData)
+{
+	// PIE stub
+	return EDisplayClusterCommResult::Ok;
+}
+
 
 //////////////////////////////////////////////////////////////////////////////////////////////
 // IDisplayClusterProtocolRenderSync
diff --git a/Engine/Plugins/Runtime/nDisplay/Source/DisplayCluster/Private/Cluster/Controller/DisplayClusterClusterNodeCtrlEditor.h b/Engine/Plugins/Runtime/nDisplay/Source/DisplayCluster/Private/Cluster/Controller/DisplayClusterClusterNodeCtrlEditor.h
--- a/Engine/Plugins/Runtime/nDisplay/Source/DisplayCluster/Private/Cluster/Controller/DisplayClusterClusterNodeCtrlEditor.h
+++ b/Engine/Plugins/Runtime/nDisplay/Source/DisplayCluster/Private/Cluster/Controller/DisplayClusterClusterNodeCtrlEditor.h
@@ -41,6 +41,7 @@ public:
 	virtual EDisplayClusterCommResult GetObjectsData(const EDisplayClusterSyncGroup InSyncGroup, TMap<FString, FString>& OutObjectsData) override;
 	virtual EDisplayClusterCommResult GetEventsData(TArray<TSharedPtr<FDisplayClusterClusterEventJson>>& OutJsonEvents, TArray<TSharedPtr<FDisplayClusterClusterEventBinary>>& OutBinaryEvents) override;
 	virtual EDisplayClusterCommResult GetNativeInputData(TMap<FString, FString>& OutNativeInputData) override;
+	virtual EDisplayClusterCommResult PropagateStatesData(const TArray<uint8>& InLocalStatesData, TArray<uint8>& OutClusterStatesData) override;
 	//~ End IDisplayClusterProtocolClusterSync
 
 public:
diff --git a/Engine/Plugins/Runtime/nDisplay/Source/DisplayCluster/Private/Cluster/Controller/DisplayClusterClusterNodeCtrlMain.cpp b/Engine/Plugins/Runtime/nDisplay/Source/DisplayCluster/Private/Cluster/Controller/DisplayClusterClusterNodeCtrlMain.cpp
--- a/Engine/Plugins/Runtime/nDisplay/Source/DisplayCluster/Private/Cluster/Controller/DisplayClusterClusterNodeCtrlMain.cpp
+++ b/Engine/Plugins/Runtime/nDisplay/Source/DisplayCluster/Private/Cluster/Controller/DisplayClusterClusterNodeCtrlMain.cpp
@@ -350,6 +350,16 @@ EDisplayClusterCommResult FDisplayClusterClusterNodeCtrlMain::GetNativeInputData
 	return EDisplayClusterCommResult::InternalError;
 }
 
+EDisplayClusterCommResult FDisplayClusterClusterNodeCtrlMain::PropagateStatesData(const TArray<uint8>& InLocalStatesData, TArray<uint8>& OutClusterStatesData)
+{
+	if (const TSharedPtr<FNodeClientSet>& ClientSet = GetActiveClientSet())
+	{
+		return ClientSet->ClusterSyncClient->PropagateStatesData(InLocalStatesData, OutClusterStatesData);
+	}
+
+	return EDisplayClusterCommResult::InternalError;
+}
+
 
 //////////////////////////////////////////////////////////////////////////////////////////////
 // IDisplayClusterProtocolRenderSync
diff --git a/Engine/Plugins/Runtime/nDisplay/Source/DisplayCluster/Private/Cluster/Controller/DisplayClusterClusterNodeCtrlMain.h b/Engine/Plugins/Runtime/nDisplay/Source/DisplayCluster/Private/Cluster/Controller/DisplayClusterClusterNodeCtrlMain.h
--- a/Engine/Plugins/Runtime/nDisplay/Source/DisplayCluster/Private/Cluster/Controller/DisplayClusterClusterNodeCtrlMain.h
+++ b/Engine/Plugins/Runtime/nDisplay/Source/DisplayCluster/Private/Cluster/Controller/DisplayClusterClusterNodeCtrlMain.h
@@ -51,6 +51,7 @@ public:
 	virtual EDisplayClusterCommResult GetObjectsData(const EDisplayClusterSyncGroup InSyncGroup, TMap<FString, FString>& OutObjectsData) override;
 	virtual EDisplayClusterCommResult GetEventsData(TArray<TSharedPtr<FDisplayClusterClusterEventJson>>& OutJsonEvents, TArray<TSharedPtr<FDisplayClusterClusterEventBinary>>& OutBinaryEvents) override;
 	virtual EDisplayClusterCommResult GetNativeInputData(TMap<FString, FString>& OutNativeInputData) override;
+	virtual EDisplayClusterCommResult PropagateStatesData(const TArray<uint8>& InLocalStatesData, TArray<uint8>& OutClusterStatesData) override;
 	//~ End IDisplayClusterProtocolClusterSync
 
 public:
diff --git a/Engine/Plugins/Runtime/nDisplay/Source/DisplayCluster/Private/Cluster/CustomStates/DisplayClusterCustomStateBase.cpp b/Engine/Plugins/Runtime/nDisplay/Source/DisplayCluster/Private/Cluster/CustomStates/DisplayClusterCustomStateBase.cpp
new file mode 100644
--- /dev/null
+++ b/Engine/Plugins/Runtime/nDisplay/Source/DisplayCluster/Private/Cluster/CustomStates/DisplayClusterCustomStateBase.cpp
@@ -0,0 +1,47 @@
+// Copyright Epic Games, Inc. All Rights Reserved.
+
+#include "Cluster/CustomStates/DisplayClusterCustomStateBase.h"
+#include "RenderingThread.h"
+
+
+namespace UE::nDisplay::Private
+{
+	FCustomStateBase::FCustomStateBase(const FName& InUniqueName, const FName& InClusterNodeId)
+		: UniqueName(InUniqueName)
+		, ClusterNodeId(InClusterNodeId)
+		, StateCS(MakeShared<FCriticalSection>())
+	{
+		if (IDisplayCluster::IsAvailable())
+		{
+			ClusterMgr = IDisplayCluster::Get().GetClusterMgr();
+		}
+	}
+
+	FName FCustomStateBase::GetName() const
+	{
+		return UniqueName;
+	}
+
+	void FCustomStateBase::Lock() const
+	{
+		StateCS->Lock();
+	}
+
+	void FCustomStateBase::Unlock() const
+	{
+		StateCS->Unlock();
+	}
+
+	void FCustomStateBase::ExecuteOnRenderThread(TUniqueFunction<void()> InFunc)
+	{
+		if (InFunc.IsSet())
+		{
+			// Just schedule update on the render thread
+			ENQUEUE_RENDER_COMMAND(CustomState_ExecuteOnRenderThread)(
+				[Func = MoveTemp(InFunc)](FRHICommandListImmediate& RHICmdList)
+				{
+					Func();
+				});
+		}
+	}
+}
diff --git a/Engine/Plugins/Runtime/nDisplay/Source/DisplayCluster/Private/Cluster/CustomStates/DisplayClusterCustomStateHelpers.cpp b/Engine/Plugins/Runtime/nDisplay/Source/DisplayCluster/Private/Cluster/CustomStates/DisplayClusterCustomStateHelpers.cpp
new file mode 100644
--- /dev/null
+++ b/Engine/Plugins/Runtime/nDisplay/Source/DisplayCluster/Private/Cluster/CustomStates/DisplayClusterCustomStateHelpers.cpp
@@ -0,0 +1,91 @@
+// Copyright Epic Games, Inc. All Rights Reserved.
+
+#include "Cluster/CustomStates/DisplayClusterCustomStateHelpers.h"
+#include "Cluster/CustomStates/IDisplayClusterCustomState.h"
+
+#include "Serialization/Archive.h"
+#include "Serialization/MemoryWriter.h"
+
+
+namespace UE::nDisplay::Cluster::Private
+{
+	FCustomStateRequestBlob::FCustomStateRequestBlob(const TSharedPtr<IDisplayClusterCustomState>& InState)
+		: StateName(InState->GetName())
+		, bCustomUpstream(InState->HasCustomUpstreamConfiguration())
+		, UpstreamNodes(bCustomUpstream ? InState->GetUpstreams() : TSet<FName>{})
+		, bExposesData(InState->ShouldPropagate())
+		, StateType(InState->GetType())
+	{
+		if (bExposesData)
+		{
+			FMemoryWriter DataWriter(StateData);
+			InState->Serialize(DataWriter);
+		}
+	}
+
+	void FCustomStateRequestBlob::Serialize(FArchive& Ar)
+	{
+		Ar << StateName;
+		Ar << StateType;
+
+		// Conditional upstream configuration
+		Ar << bCustomUpstream;
+		if (bCustomUpstream)
+		{
+			Ar << UpstreamNodes;
+		}
+		
+		// Conditional state exposure
+		Ar << bExposesData;
+		if (bExposesData)
+		{
+			Ar << StateData;
+		}
+	}
+
+	void FCustomStateResponseBlob::Serialize(FArchive& Ar)
+	{
+		Ar << StateName;
+
+		if (Ar.IsSaving())
+		{
+			// Nodes num
+			int32 NodesNum = NodeStates.Num();
+			Ar << NodesNum;
+
+			for (TPair<FName, TPair<FName, TArray<uint8>>>& NodeState : NodeStates)
+			{
+				// Node ID
+				Ar << NodeState.Key;
+				// Type ID
+				Ar << NodeState.Value.Key;
+				// State data
+				Ar << NodeState.Value.Value;
+			}
+		}
+		else
+		{
+			// Nodes num
+			int32 NodesNum = 0;
+			Ar << NodesNum;
+
+			for (int32 NodeIdx = 0; NodeIdx < NodesNum; ++NodeIdx)
+			{
+				// Node ID
+				FName NodeId;
+				Ar << NodeId;
+
+				// Type ID
+				FName TypeId;
+				Ar << TypeId;
+
+				// State data
+				TArray<uint8> StateData;
+				Ar << StateData;
+
+				// Add to the map
+				NodeStates.Emplace(NodeId, MakeTuple(TypeId, MoveTemp(StateData)));
+			}
+		}
+	}
+}
diff --git a/Engine/Plugins/Runtime/nDisplay/Source/DisplayCluster/Private/Cluster/CustomStates/DisplayClusterCustomStateHelpers.h b/Engine/Plugins/Runtime/nDisplay/Source/DisplayCluster/Private/Cluster/CustomStates/DisplayClusterCustomStateHelpers.h
new file mode 100644
--- /dev/null
+++ b/Engine/Plugins/Runtime/nDisplay/Source/DisplayCluster/Private/Cluster/CustomStates/DisplayClusterCustomStateHelpers.h
@@ -0,0 +1,84 @@
+// Copyright Epic Games, Inc. All Rights Reserved.
+
+#pragma once
+
+#include "CoreMinimal.h"
+
+class FArchive;
+class IDisplayClusterCustomState;
+
+
+namespace UE::nDisplay::Cluster::Private
+{
+	/**
+	 * Auxiliary class to pass custom state synchronization requests in binary
+	 */
+	class FCustomStateRequestBlob
+	{
+	public:
+
+		FCustomStateRequestBlob() = default;
+		FCustomStateRequestBlob(const TSharedPtr<IDisplayClusterCustomState>& InState);
+		~FCustomStateRequestBlob() = default;
+
+	public:
+
+		/** Serialization implementation */
+		void Serialize(FArchive& Ar);
+
+	public:
+
+		/** State name */
+		FName StateName;
+
+		/** Whether a state has custom upstream nodes configuration */
+		bool bCustomUpstream = false;
+
+		/**
+		 * When custom upstream configuration is used, this set contains
+		 * node IDs that the state instance wants to get updates from.
+		 */
+		TSet<FName> UpstreamNodes;
+
+		/** Whether this state exposes any data outside during synchronization */
+		bool bExposesData = false;
+
+		/** Type ID of this state */
+		FName StateType;
+
+		/** Serialized state data */
+		TArray<uint8> StateData;
+	};
+
+
+	/**
+	 * Auxiliary class to pass custom state synchronization responses in binary
+	 */
+	class FCustomStateResponseBlob
+	{
+	public:
+
+		FCustomStateResponseBlob() = default;
+
+		FCustomStateResponseBlob(const FName& InStateName)
+			: StateName(InStateName)
+		{ }
+
+	public:
+
+		/** Serialization implementation */
+		void Serialize(FArchive& Ar);
+
+	public:
+
+		/** State name */
+		FName StateName;
+
+		/**
+		 * Contains serialized state data in a form of <NodeId-to-{ StateType, StateData }> map.
+		 * So any state may have multiple cluster nodes that it comes from, each one would
+		 * provide its own serialized data.
+		 */
+		TMap<FName, TPair<FName, TArray<uint8>>> NodeStates;
+	};
+}
diff --git a/Engine/Plugins/Runtime/nDisplay/Source/DisplayCluster/Private/Cluster/DisplayClusterClusterManager.cpp b/Engine/Plugins/Runtime/nDisplay/Source/DisplayCluster/Private/Cluster/DisplayClusterClusterManager.cpp
--- a/Engine/Plugins/Runtime/nDisplay/Source/DisplayCluster/Private/Cluster/DisplayClusterClusterManager.cpp
+++ b/Engine/Plugins/Runtime/nDisplay/Source/DisplayCluster/Private/Cluster/DisplayClusterClusterManager.cpp
@@ -12,6 +12,8 @@
 #include "Cluster/Controller/DisplayClusterClusterNodeCtrlEditor.h"
 #include "Cluster/Controller/DisplayClusterClusterNodeCtrlMain.h"
 
+#include "Cluster/CustomStates/IDisplayClusterCustomState.h"
+
 #include "Cluster/Failover/DisplayClusterFailoverNodeCtrlDisabled.h"
 #include "Cluster/Failover/DisplayClusterFailoverNodeCtrlEditor.h"
 #include "Cluster/Failover/DisplayClusterFailoverNodeCtrlMain.h"
@@ -31,12 +33,27 @@
 #include "Misc/DisplayClusterTypesConverter.h"
 #include "Misc/ScopeLock.h"
 
+#include "Serialization/Archive.h"
 #include "UObject/Interface.h"
 
 #include "DisplayClusterConfigurationTypes.h"
 #include "IDisplayClusterCallbacks.h"
 
 
+namespace UE::nDisplay::DisplayCluster::Private
+{
+	static bool GIsCustomStatesPropagationEnabled = true;
+	static FAutoConsoleVariableRef CVarDisplayClusterOptimizeResourcesOnOffscreenNodes(
+		TEXT("nDisplay.Cluster.CustomStatesPropagation.Enabled"),
+		GIsCustomStatesPropagationEnabled,
+		TEXT("Whether states synchronization is required (consumes a bit of traffic)\n")
+		TEXT("0 : Disabled\n")
+		TEXT("1 : Enabled\n"),
+		ECVF_Default
+	);
+}
+
+
 FDisplayClusterClusterManager::FDisplayClusterClusterManager()
 	: NodeCtrl(MakeShared<FDisplayClusterClusterNodeCtrlDisabled>())
 	, FailoverCtrl(MakeShared<FDisplayClusterFailoverNodeCtrlDisabled>(NodeCtrl))
@@ -239,6 +256,9 @@ void FDisplayClusterClusterManager::EndFrame(uint64 FrameNum)
 
 void FDisplayClusterClusterManager::PreTick(float DeltaSeconds)
 {
+	// Sync custom states
+	SyncStates();
+
 	// Sync cluster objects (PreTick)
 	SyncObjects(EDisplayClusterSyncGroup::PreTick);
 
@@ -354,6 +374,49 @@ void FDisplayClusterClusterManager::UnregisterSyncObject(IDisplayClusterClusterS
 	}
 }
 
+bool FDisplayClusterClusterManager::RegisterCustomState(const FName& UniqueStateName, TSharedPtr<IDisplayClusterCustomState>& StateVariable)
+{
+	if (!StateVariable.IsValid())
+	{
+		return false;
+	}
+
+	// Name duplication is not allowed
+	if (IsCustomStateRegistered(UniqueStateName))
+	{
+		return false;
+	}
+
+	// Store for future replication
+	{
+		FScopeLock Lock(&CustomStatesCS);
+		CustomStates.Emplace(UniqueStateName, StateVariable);
+	}
+
+	return true;
+}
+
+bool FDisplayClusterClusterManager::UnregisterCustomState(const FName& UniqueStateName)
+{
+	FScopeLock Lock(&CustomStatesCS);
+	const int32 NumRemovedPairs = CustomStates.Remove(UniqueStateName);
+	return NumRemovedPairs > 0;
+}
+
+TSet<FName> FDisplayClusterClusterManager::GetRegisteredCustomStateNames() const
+{
+	FScopeLock Lock(&CustomStatesCS);
+	TSet<FName> Output;
+	CustomStates.GetKeys(Output);
+	return Output;
+}
+
+bool FDisplayClusterClusterManager::IsCustomStateRegistered(const FName& UniqueStateName) const
+{
+	FScopeLock Lock(&CustomStatesCS);
+	return CustomStates.Contains(UniqueStateName);
+}
+
 TSharedRef<IDisplayClusterGenericBarriersClient> FDisplayClusterClusterManager::CreateGenericBarriersClient()
 {
 	return MakeShared<FDisplayClusterGenericBarrierAPI>();
@@ -618,6 +681,33 @@ void FDisplayClusterClusterManager::ImportTimeData(const double& InDeltaTime, co
 	}
 }
 
+void FDisplayClusterClusterManager::SyncStates()
+{
+	// Skip states syncrhonization if disabled by a cvar
+	if (!UE::nDisplay::DisplayCluster::Private::GIsCustomStatesPropagationEnabled)
+	{
+		return;
+	}
+
+	FScopeLock Lock(&CustomStatesCS);
+
+	// Update frame bound data before propagating local updates
+	AdvanceFrameOnCustomStates();
+
+	// Prepare request data. All local states will be serialized into a request buffer.
+	TArray<uint8> StatesRequestData;
+	GenerateLocalStatesPropagationData(StatesRequestData);
+
+	// Propagate our local states, and receive the states of entire cluster
+	TArray<uint8> StatesResponse;
+	UE_LOG(LogDisplayClusterCluster, Verbose, TEXT("Downloading cluster states data (custom states)..."));
+	GetNetApi().GetClusterSyncAPI()->PropagateStatesData(StatesRequestData, StatesResponse);
+	UE_LOG(LogDisplayClusterCluster, Verbose, TEXT("Downloading finished. %d bytes received (custom states)."), StatesResponse.Num());
+
+	// Deserialize the response, and propagate updates among local state instances
+	ProcessClusterStatesData(StatesResponse);
+}
+
 void FDisplayClusterClusterManager::SyncObjects(EDisplayClusterSyncGroup InSyncGroup)
 {
 	TMap<FString, FString> ObjectsData;
@@ -1091,3 +1181,105 @@ void FDisplayClusterClusterManager::OnClusterNodeFailed(const FString& FailedNod
 		ActiveClusterNodeIds.Remove(FailedNodeId);
 	}
 }
+
+void FDisplayClusterClusterManager::AdvanceFrameOnCustomStates()
+{
+	FScopeLock Lock(&CustomStatesCS);
+
+	for (const TPair<FName, TSharedPtr<IDisplayClusterCustomState>>& CustomState : CustomStates)
+	{
+		CustomState.Value->AdvanceFrame();
+	}
+}
+
+void FDisplayClusterClusterManager::GenerateLocalStatesPropagationData(TArray<uint8>& OutLocalStatesData)
+{
+	FScopeLock Lock(&CustomStatesCS);
+
+	OutLocalStatesData.Reset();
+	FMemoryWriter DataWriter(OutLocalStatesData);
+
+	// States num
+	{
+		int32 StatesNum = CustomStates.Num();
+		DataWriter << StatesNum;
+	}
+
+	// States data
+	{
+		for (TPair<FName, TSharedPtr<IDisplayClusterCustomState>>& CustomState : CustomStates)
+		{
+			// Serialize the sate instance into a blob
+			CustomState.Value->Lock();
+			UE::nDisplay::Cluster::Private::FCustomStateRequestBlob StateBlob(CustomState.Value);
+			CustomState.Value->Unlock();
+
+			// Write this blob to the buffer
+			StateBlob.Serialize(DataWriter);
+		}
+	}
+}
+
+void FDisplayClusterClusterManager::ProcessClusterStatesData(const TArray<uint8>& ClusterStatesData)
+{
+	// Just skip if there is no input data
+	if (ClusterStatesData.Num() <= 0)
+	{
+		return;
+	}
+
+	FMemoryReader DataReader(ClusterStatesData);
+
+	// Read states num
+	int32 StatesNum = 0;
+	DataReader << StatesNum;
+	if (StatesNum <= 0)
+	{
+		checkSlow(StatesNum < 0);
+		return;
+	}
+
+	// Clear data & reserve memory
+	TArray<FClusterStateBlob> ClusterStatesDataMap;
+	ClusterStatesDataMap.Reserve(StatesNum);
+
+	// Read states
+	for (int32 StateIdx = 0; StateIdx < StatesNum; ++StateIdx)
+	{
+		FClusterStateBlob StateBlob;
+		StateBlob.Serialize(DataReader);
+		ClusterStatesDataMap.Add(MoveTemp(StateBlob));
+	}
+
+	// Propagate the updates locally
+	{
+		FScopeLock Lock(&CustomStatesCS);
+
+		// Iterate through all cluster states
+		for (FClusterStateBlob& ClusterState : ClusterStatesDataMap)
+		{
+			// If this state exists locally, let it update itself
+			if (TSharedPtr<IDisplayClusterCustomState>* FoundState = CustomStates.Find(ClusterState.StateName))
+			{
+				// Final data to pass to the local state instance
+				TMap<FName, TArray<uint8>> CompatibleStates;
+
+				// The type of local state data
+				const FName LocalStateType = (*FoundState)->GetType();
+
+				// Iterate through all the states we have received
+				for (TPair<FName, TPair<FName, TArray<uint8>>>& NodeState : ClusterState.NodeStates)
+				{
+					// If types match
+					if (LocalStateType == NodeState.Value.Key)
+					{
+						CompatibleStates.Emplace(NodeState.Key, MoveTemp(NodeState.Value.Value));
+					}
+				}
+
+				// Finally, let the state instance update itself
+				(*FoundState)->Update(CompatibleStates);
+			}
+		}
+	}
+}
diff --git a/Engine/Plugins/Runtime/nDisplay/Source/DisplayCluster/Private/Cluster/DisplayClusterClusterManager.h b/Engine/Plugins/Runtime/nDisplay/Source/DisplayCluster/Private/Cluster/DisplayClusterClusterManager.h
--- a/Engine/Plugins/Runtime/nDisplay/Source/DisplayCluster/Private/Cluster/DisplayClusterClusterManager.h
+++ b/Engine/Plugins/Runtime/nDisplay/Source/DisplayCluster/Private/Cluster/DisplayClusterClusterManager.h
@@ -4,6 +4,7 @@
 
 #include "Cluster/IPDisplayClusterClusterManager.h"
 #include "Cluster/DisplayClusterClusterEvent.h"
+#include "Cluster/CustomStates/DisplayClusterCustomStateHelpers.h"
 #include "HAL/Event.h"
 #include "Misc/Optional.h"
 #include "Misc/QualifiedFrameTime.h"
@@ -62,6 +63,11 @@ public:
 	virtual void RegisterSyncObject(IDisplayClusterClusterSyncObject* SyncObj, EDisplayClusterSyncGroup SyncGroup) override;
 	virtual void UnregisterSyncObject(IDisplayClusterClusterSyncObject* SyncObj) override;
 
+	virtual bool RegisterCustomState(const FName& UniqueStateName, TSharedPtr<IDisplayClusterCustomState>& StateVariable) override;
+	virtual bool UnregisterCustomState(const FName& UniqueStateName) override;
+	virtual TSet<FName> GetRegisteredCustomStateNames() const override;
+	virtual bool IsCustomStateRegistered(const FName& UniqueStateName) const override;
+
 	virtual TSharedRef<IDisplayClusterGenericBarriersClient> CreateGenericBarriersClient() override;
 
 	virtual void AddClusterEventListener(TScriptInterface<IDisplayClusterClusterEventListener>) override;
@@ -97,6 +103,9 @@ public:
 	virtual void ExportTimeData(      double& OutDeltaTime,      double& OutGameTime,      TOptional<FQualifiedFrameTime>& OutFrameTime) override;
 	virtual void ImportTimeData(const double& InDeltaTime, const double& InGameTime, const TOptional<FQualifiedFrameTime>& InFrameTime) override;
 
+	// Sync states
+	virtual void SyncStates() override;
+
 	// Sync objects
 	virtual void SyncObjects(EDisplayClusterSyncGroup SyncGroup) override;
 	virtual void CacheObjects(EDisplayClusterSyncGroup SyncGroup) override;
@@ -158,6 +167,27 @@ private:
 	/** Handles node failure events */
 	void OnClusterNodeFailed(const FString& FailedNodeId);
 
+private:
+
+	using FClusterStateBlob = UE::nDisplay::Cluster::Private::FCustomStateResponseBlob;
+
+	/** Updates frame bound data on the state instances */
+	void AdvanceFrameOnCustomStates();
+
+	/**
+	 * Serializes local cluster states before propagation
+	 * 
+	 * @param OutLocalStatesData - [out] Serialized state of all local state instanfces
+	 */
+	void GenerateLocalStatesPropagationData(TArray<uint8>& OutLocalStatesData);
+
+	/**
+	 * Deserializes response data, and propagate the updates among local state instances
+	 * 
+	 * @param ClusterStatesData - Serialized state of all states from every node
+	 */
+	void ProcessClusterStatesData(const TArray<uint8>& ClusterStatesData);
+
 private:
 
 	/** Networking controller */
@@ -238,4 +268,8 @@ private:
 	// Cluster event listeners
 	mutable FCriticalSection ClusterEventListenersCS;
 	TArray<TScriptInterface<IDisplayClusterClusterEventListener>> ClusterEventListeners;
+
+	// Custom states
+	mutable FCriticalSection CustomStatesCS;
+	TMap<FName, TSharedPtr<IDisplayClusterCustomState>> CustomStates;
 };
diff --git a/Engine/Plugins/Runtime/nDisplay/Source/DisplayCluster/Private/Cluster/Failover/DisplayClusterFailoverNodeCtrlDisabled.cpp b/Engine/Plugins/Runtime/nDisplay/Source/DisplayCluster/Private/Cluster/Failover/DisplayClusterFailoverNodeCtrlDisabled.cpp
--- a/Engine/Plugins/Runtime/nDisplay/Source/DisplayCluster/Private/Cluster/Failover/DisplayClusterFailoverNodeCtrlDisabled.cpp
+++ b/Engine/Plugins/Runtime/nDisplay/Source/DisplayCluster/Private/Cluster/Failover/DisplayClusterFailoverNodeCtrlDisabled.cpp
@@ -63,6 +63,11 @@ EDisplayClusterCommResult FDisplayClusterFailoverNodeCtrlDisabled::GetNativeInpu
 	return EDisplayClusterCommResult::NotAllowed;
 }
 
+EDisplayClusterCommResult FDisplayClusterFailoverNodeCtrlDisabled::PropagateStatesData(const TArray<uint8>& InLocalStatesData, TArray<uint8>& OutClusterStatesData)
+{
+	return EDisplayClusterCommResult::NotAllowed;
+}
+
 
 ////////////////////////////////////////////////////////////////////////////////////////////////////
 // IDisplayClusterProtocolRenderSync
diff --git a/Engine/Plugins/Runtime/nDisplay/Source/DisplayCluster/Private/Cluster/Failover/DisplayClusterFailoverNodeCtrlDisabled.h b/Engine/Plugins/Runtime/nDisplay/Source/DisplayCluster/Private/Cluster/Failover/DisplayClusterFailoverNodeCtrlDisabled.h
--- a/Engine/Plugins/Runtime/nDisplay/Source/DisplayCluster/Private/Cluster/Failover/DisplayClusterFailoverNodeCtrlDisabled.h
+++ b/Engine/Plugins/Runtime/nDisplay/Source/DisplayCluster/Private/Cluster/Failover/DisplayClusterFailoverNodeCtrlDisabled.h
@@ -37,6 +37,7 @@ public:
 	virtual EDisplayClusterCommResult GetObjectsData(const EDisplayClusterSyncGroup InSyncGroup, TMap<FString, FString>& OutObjectsData) override;
 	virtual EDisplayClusterCommResult GetEventsData(TArray<TSharedPtr<FDisplayClusterClusterEventJson>>& OutJsonEvents, TArray<TSharedPtr<FDisplayClusterClusterEventBinary>>& OutBinaryEvents) override;
 	virtual EDisplayClusterCommResult GetNativeInputData(TMap<FString, FString>& OutNativeInputData) override;
+	virtual EDisplayClusterCommResult PropagateStatesData(const TArray<uint8>& InLocalStatesData, TArray<uint8>& OutClusterStatesData) override;
 	//~ End IDisplayClusterProtocolClusterSync
 
 public:
diff --git a/Engine/Plugins/Runtime/nDisplay/Source/DisplayCluster/Private/Cluster/Failover/DisplayClusterFailoverNodeCtrlEditor.cpp b/Engine/Plugins/Runtime/nDisplay/Source/DisplayCluster/Private/Cluster/Failover/DisplayClusterFailoverNodeCtrlEditor.cpp
--- a/Engine/Plugins/Runtime/nDisplay/Source/DisplayCluster/Private/Cluster/Failover/DisplayClusterFailoverNodeCtrlEditor.cpp
+++ b/Engine/Plugins/Runtime/nDisplay/Source/DisplayCluster/Private/Cluster/Failover/DisplayClusterFailoverNodeCtrlEditor.cpp
@@ -75,6 +75,12 @@ EDisplayClusterCommResult FDisplayClusterFailoverNodeCtrlEditor::GetNativeInputD
 	return GetNodeController()->GetNativeInputData(OutNativeInputData);
 }
 
+EDisplayClusterCommResult FDisplayClusterFailoverNodeCtrlEditor::PropagateStatesData(const TArray<uint8>& InLocalStatesData, TArray<uint8>& OutClusterStatesData)
+{
+	// `Editor` failover controller doesn't perform any validations, just forward the call
+	return GetNodeController()->PropagateStatesData(InLocalStatesData, OutClusterStatesData);
+}
+
 
 ////////////////////////////////////////////////////////////////////////////////////////////////////
 // IDisplayClusterProtocolRenderSync
diff --git a/Engine/Plugins/Runtime/nDisplay/Source/DisplayCluster/Private/Cluster/Failover/DisplayClusterFailoverNodeCtrlEditor.h b/Engine/Plugins/Runtime/nDisplay/Source/DisplayCluster/Private/Cluster/Failover/DisplayClusterFailoverNodeCtrlEditor.h
--- a/Engine/Plugins/Runtime/nDisplay/Source/DisplayCluster/Private/Cluster/Failover/DisplayClusterFailoverNodeCtrlEditor.h
+++ b/Engine/Plugins/Runtime/nDisplay/Source/DisplayCluster/Private/Cluster/Failover/DisplayClusterFailoverNodeCtrlEditor.h
@@ -39,6 +39,7 @@ public:
 	virtual EDisplayClusterCommResult GetObjectsData(const EDisplayClusterSyncGroup InSyncGroup, TMap<FString, FString>& OutObjectsData) override;
 	virtual EDisplayClusterCommResult GetEventsData(TArray<TSharedPtr<FDisplayClusterClusterEventJson>>& OutJsonEvents, TArray<TSharedPtr<FDisplayClusterClusterEventBinary>>& OutBinaryEvents) override;
 	virtual EDisplayClusterCommResult GetNativeInputData(TMap<FString, FString>& OutNativeInputData) override;
+	virtual EDisplayClusterCommResult PropagateStatesData(const TArray<uint8>& InLocalStatesData, TArray<uint8>& OutClusterStatesData) override;
 	//~ End IDisplayClusterProtocolClusterSync
 
 public:
diff --git a/Engine/Plugins/Runtime/nDisplay/Source/DisplayCluster/Private/Cluster/Failover/DisplayClusterFailoverNodeCtrlMain.cpp b/Engine/Plugins/Runtime/nDisplay/Source/DisplayCluster/Private/Cluster/Failover/DisplayClusterFailoverNodeCtrlMain.cpp
--- a/Engine/Plugins/Runtime/nDisplay/Source/DisplayCluster/Private/Cluster/Failover/DisplayClusterFailoverNodeCtrlMain.cpp
+++ b/Engine/Plugins/Runtime/nDisplay/Source/DisplayCluster/Private/Cluster/Failover/DisplayClusterFailoverNodeCtrlMain.cpp
@@ -371,6 +371,23 @@ EDisplayClusterCommResult FDisplayClusterFailoverNodeCtrlMain::GetNativeInputDat
 	return CommResult;
 }
 
+EDisplayClusterCommResult FDisplayClusterFailoverNodeCtrlMain::PropagateStatesData(const TArray<uint8>& InLocalStatesData, TArray<uint8>& OutClusterStatesData)
+{
+	static const FName TxnName     = TEXT("PropagateStatesData");
+	static const FName BarrierName = TEXT("SyncStatesBarrier");
+
+	const EDisplayClusterCommResult CommResult = ProcessTransactionSYNC(
+		TxnName,
+		BarrierName,
+		ThisNodeIdName,
+		[&]() -> EDisplayClusterCommResult // OpSendReq
+		{
+			return GetNodeController()->PropagateStatesData(InLocalStatesData, OutClusterStatesData);
+		});
+
+	return CommResult;
+}
+
 
 ////////////////////////////////////////////////////////////////////////////////////////////////////
 // IDisplayClusterProtocolRenderSync
diff --git a/Engine/Plugins/Runtime/nDisplay/Source/DisplayCluster/Private/Cluster/Failover/DisplayClusterFailoverNodeCtrlMain.h b/Engine/Plugins/Runtime/nDisplay/Source/DisplayCluster/Private/Cluster/Failover/DisplayClusterFailoverNodeCtrlMain.h
--- a/Engine/Plugins/Runtime/nDisplay/Source/DisplayCluster/Private/Cluster/Failover/DisplayClusterFailoverNodeCtrlMain.h
+++ b/Engine/Plugins/Runtime/nDisplay/Source/DisplayCluster/Private/Cluster/Failover/DisplayClusterFailoverNodeCtrlMain.h
@@ -50,6 +50,7 @@ public:
 	virtual EDisplayClusterCommResult GetObjectsData(const EDisplayClusterSyncGroup InSyncGroup, TMap<FString, FString>& OutObjectsData) override;
 	virtual EDisplayClusterCommResult GetEventsData(TArray<TSharedPtr<FDisplayClusterClusterEventJson>>& OutJsonEvents, TArray<TSharedPtr<FDisplayClusterClusterEventBinary>>& OutBinaryEvents) override;
 	virtual EDisplayClusterCommResult GetNativeInputData(TMap<FString, FString>& OutNativeInputData) override;
+	virtual EDisplayClusterCommResult PropagateStatesData(const TArray<uint8>& InLocalStatesData, TArray<uint8>& OutClusterStatesData) override;
 	//~ End IDisplayClusterProtocolClusterSync
 
 public:
diff --git a/Engine/Plugins/Runtime/nDisplay/Source/DisplayCluster/Private/Cluster/IPDisplayClusterClusterManager.h b/Engine/Plugins/Runtime/nDisplay/Source/DisplayCluster/Private/Cluster/IPDisplayClusterClusterManager.h
--- a/Engine/Plugins/Runtime/nDisplay/Source/DisplayCluster/Private/Cluster/IPDisplayClusterClusterManager.h
+++ b/Engine/Plugins/Runtime/nDisplay/Source/DisplayCluster/Private/Cluster/IPDisplayClusterClusterManager.h
@@ -65,6 +65,11 @@ public: // Time data sync
 	/** Import time data from external source */
 	virtual void ImportTimeData(const double& InDeltaTime, const double& InGameTime, const TOptional<FQualifiedFrameTime>& InFrameTime) = 0;
 
+public: // Custom states
+
+	/** Custom states synchronization (procedure entry point) */
+	virtual void SyncStates() = 0;
+
 public: // Objects sync
 
 	/** Custom objects synchronization (procedure entry point) */
diff --git a/Engine/Plugins/Runtime/nDisplay/Source/DisplayCluster/Private/Cluster/NetAPI/DisplayClusterNetApiFacade.cpp b/Engine/Plugins/Runtime/nDisplay/Source/DisplayCluster/Private/Cluster/NetAPI/DisplayClusterNetApiFacade.cpp
--- a/Engine/Plugins/Runtime/nDisplay/Source/DisplayCluster/Private/Cluster/NetAPI/DisplayClusterNetApiFacade.cpp
+++ b/Engine/Plugins/Runtime/nDisplay/Source/DisplayCluster/Private/Cluster/NetAPI/DisplayClusterNetApiFacade.cpp
@@ -53,35 +53,35 @@ public:
 	/** Failover transaction for WaitForGameStart */
 	virtual EDisplayClusterCommResult WaitForGameStart() override
 	{
-		check(IsInGameThread());
+		checkSlow(IsInGameThread());
 		return GetFailoverController()->WaitForGameStart();
 	}
 
 	/** Failover transaction for WaitForFrameStart */
 	virtual EDisplayClusterCommResult WaitForFrameStart() override
 	{
-		check(IsInGameThread());
+		checkSlow(IsInGameThread());
 		return GetFailoverController()->WaitForFrameStart();
 	}
 
 	/** Failover transaction for WaitForFrameEnd */
 	virtual EDisplayClusterCommResult WaitForFrameEnd() override
 	{
-		check(IsInGameThread());
+		checkSlow(IsInGameThread());
 		return GetFailoverController()->WaitForFrameEnd();
 	}
 
 	/** Failover transaction for GetTimeData */
 	virtual EDisplayClusterCommResult GetTimeData(double& OutDeltaTime, double& OutGameTime, TOptional<FQualifiedFrameTime>& OutFrameTime) override
 	{
-		check(IsInGameThread());
+		checkSlow(IsInGameThread());
 		return GetFailoverController()->GetTimeData(OutDeltaTime, OutGameTime, OutFrameTime);
 	}
 
 	/** Failover transaction for GetObjectsData */
 	virtual EDisplayClusterCommResult GetObjectsData(const EDisplayClusterSyncGroup InSyncGroup, TMap<FString, FString>& OutObjectsData) override
 	{
-		check(IsInGameThread());
+		checkSlow(IsInGameThread());
 		return GetFailoverController()->GetObjectsData(InSyncGroup, OutObjectsData);
 	}
 
@@ -90,17 +90,24 @@ public:
 		TArray<TSharedPtr<FDisplayClusterClusterEventJson>>& OutJsonEvents,
 		TArray<TSharedPtr<FDisplayClusterClusterEventBinary>>& OutBinaryEvents) override
 	{
-		check(IsInGameThread());
+		checkSlow(IsInGameThread());
 		return GetFailoverController()->GetEventsData(OutJsonEvents, OutBinaryEvents);
 	}
 
 	/** Failover transaction for GetNativeInputData */
 	virtual EDisplayClusterCommResult GetNativeInputData(TMap<FString, FString>& OutNativeInputData) override
 	{
-		check(IsInGameThread());
+		checkSlow(IsInGameThread());
 		return GetFailoverController()->GetNativeInputData(OutNativeInputData);
 	}
 
+	/** Failover transaction for PropagateStatesData */
+	virtual EDisplayClusterCommResult PropagateStatesData(const TArray<uint8>& InLocalStatesData, TArray<uint8>& OutClusterStatesData) override
+	{
+		checkSlow(IsInGameThread());
+		return GetFailoverController()->PropagateStatesData(InLocalStatesData, OutClusterStatesData);
+	}
+
 	//~ End IDisplayClusterProtocolClusterSync
 };
 
diff --git a/Engine/Plugins/Runtime/nDisplay/Source/DisplayCluster/Private/Cluster/RequestHandler/DisplayClusterCommRequestHandlerLocal.cpp b/Engine/Plugins/Runtime/nDisplay/Source/DisplayCluster/Private/Cluster/RequestHandler/DisplayClusterCommRequestHandlerLocal.cpp
--- a/Engine/Plugins/Runtime/nDisplay/Source/DisplayCluster/Private/Cluster/RequestHandler/DisplayClusterCommRequestHandlerLocal.cpp
+++ b/Engine/Plugins/Runtime/nDisplay/Source/DisplayCluster/Private/Cluster/RequestHandler/DisplayClusterCommRequestHandlerLocal.cpp
@@ -101,3 +101,9 @@ EDisplayClusterCommResult FDisplayClusterCommRequestHandlerLocal::GetNativeInput
 
 	return EDisplayClusterCommResult::Ok;
 }
+
+EDisplayClusterCommResult FDisplayClusterCommRequestHandlerLocal::PropagateStatesData(const TArray<uint8>& InLocalStatesData, TArray<uint8>& OutClusterStatesData)
+{
+	// Not used currently
+	return EDisplayClusterCommResult::NotImplemented;
+}
diff --git a/Engine/Plugins/Runtime/nDisplay/Source/DisplayCluster/Private/Cluster/RequestHandler/DisplayClusterCommRequestHandlerLocal.h b/Engine/Plugins/Runtime/nDisplay/Source/DisplayCluster/Private/Cluster/RequestHandler/DisplayClusterCommRequestHandlerLocal.h
--- a/Engine/Plugins/Runtime/nDisplay/Source/DisplayCluster/Private/Cluster/RequestHandler/DisplayClusterCommRequestHandlerLocal.h
+++ b/Engine/Plugins/Runtime/nDisplay/Source/DisplayCluster/Private/Cluster/RequestHandler/DisplayClusterCommRequestHandlerLocal.h
@@ -40,5 +40,6 @@ public:
 	virtual EDisplayClusterCommResult GetObjectsData(const EDisplayClusterSyncGroup InSyncGroup, TMap<FString, FString>& OutObjectsData) override;
 	virtual EDisplayClusterCommResult GetEventsData(TArray<TSharedPtr<FDisplayClusterClusterEventJson>>& OutJsonEvents, TArray<TSharedPtr<FDisplayClusterClusterEventBinary>>& OutBinaryEvents) override;
 	virtual EDisplayClusterCommResult GetNativeInputData(TMap<FString, FString>& OutNativeInputData) override;
+	virtual EDisplayClusterCommResult PropagateStatesData(const TArray<uint8>& InLocalStatesData, TArray<uint8>& OutClusterStatesData) override;
 	//~ End IDisplayClusterProtocolClusterSync
 };
diff --git a/Engine/Plugins/Runtime/nDisplay/Source/DisplayCluster/Private/Cluster/RequestHandler/DisplayClusterCommRequestHandlerRemote.cpp b/Engine/Plugins/Runtime/nDisplay/Source/DisplayCluster/Private/Cluster/RequestHandler/DisplayClusterCommRequestHandlerRemote.cpp
--- a/Engine/Plugins/Runtime/nDisplay/Source/DisplayCluster/Private/Cluster/RequestHandler/DisplayClusterCommRequestHandlerRemote.cpp
+++ b/Engine/Plugins/Runtime/nDisplay/Source/DisplayCluster/Private/Cluster/RequestHandler/DisplayClusterCommRequestHandlerRemote.cpp
@@ -98,3 +98,9 @@ EDisplayClusterCommResult FDisplayClusterCommRequestHandlerRemote::GetNativeInpu
 
 	return EDisplayClusterCommResult::Ok;
 }
+
+EDisplayClusterCommResult FDisplayClusterCommRequestHandlerRemote::PropagateStatesData(const TArray<uint8>& InLocalStatesData, TArray<uint8>& OutClusterStatesData)
+{
+	// Not used currently
+	return EDisplayClusterCommResult::NotImplemented;
+}
diff --git a/Engine/Plugins/Runtime/nDisplay/Source/DisplayCluster/Private/Cluster/RequestHandler/DisplayClusterCommRequestHandlerRemote.h b/Engine/Plugins/Runtime/nDisplay/Source/DisplayCluster/Private/Cluster/RequestHandler/DisplayClusterCommRequestHandlerRemote.h
--- a/Engine/Plugins/Runtime/nDisplay/Source/DisplayCluster/Private/Cluster/RequestHandler/DisplayClusterCommRequestHandlerRemote.h
+++ b/Engine/Plugins/Runtime/nDisplay/Source/DisplayCluster/Private/Cluster/RequestHandler/DisplayClusterCommRequestHandlerRemote.h
@@ -40,5 +40,6 @@ public:
 	virtual EDisplayClusterCommResult GetObjectsData(const EDisplayClusterSyncGroup InSyncGroup, TMap<FString, FString>& OutObjectsData) override;
 	virtual EDisplayClusterCommResult GetEventsData(TArray<TSharedPtr<FDisplayClusterClusterEventJson>>& OutJsonEvents, TArray<TSharedPtr<FDisplayClusterClusterEventBinary>>& OutBinaryEvents) override;
 	virtual EDisplayClusterCommResult GetNativeInputData(TMap<FString, FString>& OutNativeInputData) override;
+	virtual EDisplayClusterCommResult PropagateStatesData(const TArray<uint8>& InLocalStatesData, TArray<uint8>& OutClusterStatesData) override;
 	//~ End IDisplayClusterProtocolClusterSync
 };
diff --git a/Engine/Plugins/Runtime/nDisplay/Source/DisplayCluster/Private/Network/DisplayClusterNetworkTypes.h b/Engine/Plugins/Runtime/nDisplay/Source/DisplayCluster/Private/Network/DisplayClusterNetworkTypes.h
--- a/Engine/Plugins/Runtime/nDisplay/Source/DisplayCluster/Private/Network/DisplayClusterNetworkTypes.h
+++ b/Engine/Plugins/Runtime/nDisplay/Source/DisplayCluster/Private/Network/DisplayClusterNetworkTypes.h
@@ -57,6 +57,6 @@ protected:
 public:
 	FString ToString() const
 	{
-		return FString::Printf(TEXT("%llu:%s:%s:%s:%s"), SessionId, *Protocol, *Endpoint.ToString(), *SessionName, *NodeId.Get(TEXT("")));
+		return FString::Printf(TEXT("%llu:%s:%s:%s:%s"), SessionId, *Protocol, *Endpoint.ToString(), *SessionName, *NodeId.Get(FString()));
 	}
 };
diff --git a/Engine/Plugins/Runtime/nDisplay/Source/DisplayCluster/Private/Network/Protocol/IDisplayClusterProtocolClusterSync.h b/Engine/Plugins/Runtime/nDisplay/Source/DisplayCluster/Private/Network/Protocol/IDisplayClusterProtocolClusterSync.h
--- a/Engine/Plugins/Runtime/nDisplay/Source/DisplayCluster/Private/Network/Protocol/IDisplayClusterProtocolClusterSync.h
+++ b/Engine/Plugins/Runtime/nDisplay/Source/DisplayCluster/Private/Network/Protocol/IDisplayClusterProtocolClusterSync.h
@@ -44,4 +44,7 @@ public:
 
 	/** Native UE input syncrhonization */
 	virtual EDisplayClusterCommResult GetNativeInputData(TMap<FString, FString>& OutNativeInputData) = 0;
+
+	/** Custom states synchronization */
+	virtual EDisplayClusterCommResult PropagateStatesData(const TArray<uint8>& InLocalStatesData, TArray<uint8>& OutClusterStatesData) = 0;
 };
diff --git a/Engine/Plugins/Runtime/nDisplay/Source/DisplayCluster/Private/Network/Service/ClusterEventsBinary/DisplayClusterClusterEventsBinaryService.cpp b/Engine/Plugins/Runtime/nDisplay/Source/DisplayCluster/Private/Network/Service/ClusterEventsBinary/DisplayClusterClusterEventsBinaryService.cpp
--- a/Engine/Plugins/Runtime/nDisplay/Source/DisplayCluster/Private/Network/Service/ClusterEventsBinary/DisplayClusterClusterEventsBinaryService.cpp
+++ b/Engine/Plugins/Runtime/nDisplay/Source/DisplayCluster/Private/Network/Service/ClusterEventsBinary/DisplayClusterClusterEventsBinaryService.cpp
@@ -31,8 +31,7 @@ FDisplayClusterClusterEventsBinaryService::~FDisplayClusterClusterEventsBinarySe
 
 FString FDisplayClusterClusterEventsBinaryService::GetProtocolName() const
 {
-	static const FString ProtocolName(DisplayClusterClusterEventsBinaryStrings::ProtocolName);
-	return ProtocolName;
+	return DisplayClusterClusterEventsBinaryStrings::ProtocolName;
 }
 
 TSharedPtr<IDisplayClusterSession> FDisplayClusterClusterEventsBinaryService::CreateSession(FDisplayClusterSessionInfo& SessionInfo)
diff --git a/Engine/Plugins/Runtime/nDisplay/Source/DisplayCluster/Private/Network/Service/ClusterEventsBinary/DisplayClusterClusterEventsBinaryStrings.cpp b/Engine/Plugins/Runtime/nDisplay/Source/DisplayCluster/Private/Network/Service/ClusterEventsBinary/DisplayClusterClusterEventsBinaryStrings.cpp
new file mode 100644
--- /dev/null
+++ b/Engine/Plugins/Runtime/nDisplay/Source/DisplayCluster/Private/Network/Service/ClusterEventsBinary/DisplayClusterClusterEventsBinaryStrings.cpp
@@ -0,0 +1,10 @@
+// Copyright Epic Games, Inc. All Rights Reserved.
+
+#include "Network/Service/ClusterEventsBinary/DisplayClusterClusterEventsBinaryStrings.h"
+
+
+namespace DisplayClusterClusterEventsBinaryStrings
+{
+	const FString ProtocolName = TEXT("EventsBinary");
+	const FString TypeRequest  = TEXT("Request");
+};
diff --git a/Engine/Plugins/Runtime/nDisplay/Source/DisplayCluster/Private/Network/Service/ClusterEventsBinary/DisplayClusterClusterEventsBinaryStrings.h b/Engine/Plugins/Runtime/nDisplay/Source/DisplayCluster/Private/Network/Service/ClusterEventsBinary/DisplayClusterClusterEventsBinaryStrings.h
--- a/Engine/Plugins/Runtime/nDisplay/Source/DisplayCluster/Private/Network/Service/ClusterEventsBinary/DisplayClusterClusterEventsBinaryStrings.h
+++ b/Engine/Plugins/Runtime/nDisplay/Source/DisplayCluster/Private/Network/Service/ClusterEventsBinary/DisplayClusterClusterEventsBinaryStrings.h
@@ -2,7 +2,7 @@
 
 #pragma once
 
-#include "HAL/Platform.h"
+#include "CoreMinimal.h"
 
 
 /**
@@ -10,6 +10,6 @@
  */
 namespace DisplayClusterClusterEventsBinaryStrings
 {
-	constexpr static const TCHAR* ProtocolName = TEXT("EventsBinary");
-	constexpr static const TCHAR* TypeRequest  = TEXT("Request");
+	extern const FString ProtocolName;
+	extern const FString TypeRequest;
 };
diff --git a/Engine/Plugins/Runtime/nDisplay/Source/DisplayCluster/Private/Network/Service/ClusterEventsJson/DisplayClusterClusterEventsJsonService.cpp b/Engine/Plugins/Runtime/nDisplay/Source/DisplayCluster/Private/Network/Service/ClusterEventsJson/DisplayClusterClusterEventsJsonService.cpp
--- a/Engine/Plugins/Runtime/nDisplay/Source/DisplayCluster/Private/Network/Service/ClusterEventsJson/DisplayClusterClusterEventsJsonService.cpp
+++ b/Engine/Plugins/Runtime/nDisplay/Source/DisplayCluster/Private/Network/Service/ClusterEventsJson/DisplayClusterClusterEventsJsonService.cpp
@@ -33,8 +33,7 @@ FDisplayClusterClusterEventsJsonService::~FDisplayClusterClusterEventsJsonServic
 
 FString FDisplayClusterClusterEventsJsonService::GetProtocolName() const
 {
-	static const FString ProtocolName(DisplayClusterClusterEventsJsonStrings::ProtocolName);
-	return ProtocolName;
+	return DisplayClusterClusterEventsJsonStrings::ProtocolName;
 }
 
 TSharedPtr<IDisplayClusterSession> FDisplayClusterClusterEventsJsonService::CreateSession(FDisplayClusterSessionInfo& SessionInfo)
@@ -82,21 +81,21 @@ typename IDisplayClusterSessionPacketHandler<FDisplayClusterPacketJson, false>::
 
 	bool MandatoryFieldsExist = true;
 
-	if (!Request->GetJsonData()->HasField(FString(DisplayClusterClusterEventsJsonStrings::ArgName)))
+	if (!Request->GetJsonData()->HasField(DisplayClusterClusterEventsJsonStrings::ArgName))
 	{
-		UE_LOG(LogDisplayClusterNetwork, Error, TEXT("Json packet doesn't have a mandatory field: %s"), DisplayClusterClusterEventsJsonStrings::ArgName);
+		UE_LOG(LogDisplayClusterNetwork, Error, TEXT("Json packet doesn't have a mandatory field: %s"), *DisplayClusterClusterEventsJsonStrings::ArgName);
 		return IDisplayClusterSessionPacketHandler<FDisplayClusterPacketJson, false>::ReturnType();
 	}
 
-	if (!Request->GetJsonData()->HasField(FString(DisplayClusterClusterEventsJsonStrings::ArgType)))
+	if (!Request->GetJsonData()->HasField(DisplayClusterClusterEventsJsonStrings::ArgType))
 	{
-		UE_LOG(LogDisplayClusterNetwork, Error, TEXT("Json packet doesn't have a mandatory field: %s"), DisplayClusterClusterEventsJsonStrings::ArgType);
+		UE_LOG(LogDisplayClusterNetwork, Error, TEXT("Json packet doesn't have a mandatory field: %s"), *DisplayClusterClusterEventsJsonStrings::ArgType);
 		return IDisplayClusterSessionPacketHandler<FDisplayClusterPacketJson, false>::ReturnType();
 	}
 
-	if (!Request->GetJsonData()->HasField(FString(DisplayClusterClusterEventsJsonStrings::ArgCategory)))
+	if (!Request->GetJsonData()->HasField(DisplayClusterClusterEventsJsonStrings::ArgCategory))
 	{
-		UE_LOG(LogDisplayClusterNetwork, Error, TEXT("Json packet doesn't have a mandatory field: %s"), DisplayClusterClusterEventsJsonStrings::ArgCategory);
+		UE_LOG(LogDisplayClusterNetwork, Error, TEXT("Json packet doesn't have a mandatory field: %s"), *DisplayClusterClusterEventsJsonStrings::ArgCategory);
 		return IDisplayClusterSessionPacketHandler<FDisplayClusterPacketJson, false>::ReturnType();
 	}
 
diff --git a/Engine/Plugins/Runtime/nDisplay/Source/DisplayCluster/Private/Network/Service/ClusterEventsJson/DisplayClusterClusterEventsJsonStrings.cpp b/Engine/Plugins/Runtime/nDisplay/Source/DisplayCluster/Private/Network/Service/ClusterEventsJson/DisplayClusterClusterEventsJsonStrings.cpp
new file mode 100644
--- /dev/null
+++ b/Engine/Plugins/Runtime/nDisplay/Source/DisplayCluster/Private/Network/Service/ClusterEventsJson/DisplayClusterClusterEventsJsonStrings.cpp
@@ -0,0 +1,14 @@
+// Copyright Epic Games, Inc. All Rights Reserved.
+
+#include "Network/Service/ClusterEventsJson/DisplayClusterClusterEventsJsonStrings.h"
+
+
+namespace DisplayClusterClusterEventsJsonStrings
+{
+	const FString ProtocolName = TEXT("EventsJSON");
+	const FString TypeRequest  = TEXT("Request");
+
+	const FString ArgName       = TEXT("Name");
+	const FString ArgType       = TEXT("Type");
+	const FString ArgCategory   = TEXT("Category");
+};
diff --git a/Engine/Plugins/Runtime/nDisplay/Source/DisplayCluster/Private/Network/Service/ClusterEventsJson/DisplayClusterClusterEventsJsonStrings.h b/Engine/Plugins/Runtime/nDisplay/Source/DisplayCluster/Private/Network/Service/ClusterEventsJson/DisplayClusterClusterEventsJsonStrings.h
--- a/Engine/Plugins/Runtime/nDisplay/Source/DisplayCluster/Private/Network/Service/ClusterEventsJson/DisplayClusterClusterEventsJsonStrings.h
+++ b/Engine/Plugins/Runtime/nDisplay/Source/DisplayCluster/Private/Network/Service/ClusterEventsJson/DisplayClusterClusterEventsJsonStrings.h
@@ -2,7 +2,7 @@
 
 #pragma once
 
-#include "HAL/Platform.h"
+#include "CoreMinimal.h"
 
 
 /**
@@ -10,10 +10,10 @@
  */
 namespace DisplayClusterClusterEventsJsonStrings
 {
-	constexpr static const TCHAR* ProtocolName = TEXT("EventsJSON");
-	constexpr static const TCHAR* TypeRequest  = TEXT("Request");
+	extern const FString ProtocolName;
+	extern const FString TypeRequest;
 
-	constexpr static auto ArgName       = TEXT("Name");
-	constexpr static auto ArgType       = TEXT("Type");
-	constexpr static auto ArgCategory   = TEXT("Category");
+	extern const FString ArgName;
+	extern const FString ArgType;
+	extern const FString ArgCategory;
 };
diff --git a/Engine/Plugins/Runtime/nDisplay/Source/DisplayCluster/Private/Network/Service/ClusterSync/DisplayClusterClusterSyncClient.cpp b/Engine/Plugins/Runtime/nDisplay/Source/DisplayCluster/Private/Network/Service/ClusterSync/DisplayClusterClusterSyncClient.cpp
--- a/Engine/Plugins/Runtime/nDisplay/Source/DisplayCluster/Private/Network/Service/ClusterSync/DisplayClusterClusterSyncClient.cpp
+++ b/Engine/Plugins/Runtime/nDisplay/Source/DisplayCluster/Private/Network/Service/ClusterSync/DisplayClusterClusterSyncClient.cpp
@@ -147,7 +147,7 @@ EDisplayClusterCommResult FDisplayClusterClusterSyncClient::GetTimeData(double&
 	FString StrDeltaTime;
 	if (!Response->GetTextArg(DisplayClusterClusterSyncStrings::ArgumentsDefaultCategory, DisplayClusterClusterSyncStrings::GetTimeData::ArgDeltaTime, StrDeltaTime))
 	{
-		UE_LOG(LogDisplayClusterNetwork, Warning, TEXT("Couldn't extract parameter: %s"), DisplayClusterClusterSyncStrings::GetTimeData::ArgDeltaTime);
+		UE_LOG(LogDisplayClusterNetwork, Warning, TEXT("Couldn't extract parameter: %s"), *DisplayClusterClusterSyncStrings::GetTimeData::ArgDeltaTime);
 		return EDisplayClusterCommResult::WrongResponseData;
 	}
 
@@ -157,7 +157,7 @@ EDisplayClusterCommResult FDisplayClusterClusterSyncClient::GetTimeData(double&
 	FString StrGameTime;
 	if (!Response->GetTextArg(DisplayClusterClusterSyncStrings::ArgumentsDefaultCategory, DisplayClusterClusterSyncStrings::GetTimeData::ArgGameTime, StrGameTime))
 	{
-		UE_LOG(LogDisplayClusterNetwork, Warning, TEXT("Couldn't extract parameter: %s"), DisplayClusterClusterSyncStrings::GetTimeData::ArgGameTime);
+		UE_LOG(LogDisplayClusterNetwork, Warning, TEXT("Couldn't extract parameter: %s"), *DisplayClusterClusterSyncStrings::GetTimeData::ArgGameTime);
 		return EDisplayClusterCommResult::WrongResponseData;
 	}
 
@@ -168,7 +168,7 @@ EDisplayClusterCommResult FDisplayClusterClusterSyncClient::GetTimeData(double&
 	bool bIsFrameTimeValid = false;
 	if (!Response->GetTextArg(DisplayClusterClusterSyncStrings::ArgumentsDefaultCategory, DisplayClusterClusterSyncStrings::GetTimeData::ArgIsFrameTimeValid, bIsFrameTimeValid))
 	{
-		UE_LOG(LogDisplayClusterNetwork, Warning, TEXT("Couldn't extract parameter: %s"), DisplayClusterClusterSyncStrings::GetTimeData::ArgIsFrameTimeValid);
+		UE_LOG(LogDisplayClusterNetwork, Warning, TEXT("Couldn't extract parameter: %s"), *DisplayClusterClusterSyncStrings::GetTimeData::ArgIsFrameTimeValid);
 		return EDisplayClusterCommResult::WrongResponseData;
 	}
 
@@ -177,7 +177,7 @@ EDisplayClusterCommResult FDisplayClusterClusterSyncClient::GetTimeData(double&
 		FQualifiedFrameTime NewFrameTime;
 		if (!Response->GetTextArg(DisplayClusterClusterSyncStrings::ArgumentsDefaultCategory, DisplayClusterClusterSyncStrings::GetTimeData::ArgFrameTime, NewFrameTime))
 		{
-			UE_LOG(LogDisplayClusterNetwork, Warning, TEXT("Couldn't extract parameter: %s"), DisplayClusterClusterSyncStrings::GetTimeData::ArgFrameTime);
+			UE_LOG(LogDisplayClusterNetwork, Warning, TEXT("Couldn't extract parameter: %s"), *DisplayClusterClusterSyncStrings::GetTimeData::ArgFrameTime);
 			return EDisplayClusterCommResult::WrongResponseData;
 		}
 
@@ -267,3 +267,31 @@ EDisplayClusterCommResult FDisplayClusterClusterSyncClient::GetNativeInputData(T
 
 	return Response->GetCommResult();
 }
+
+EDisplayClusterCommResult FDisplayClusterClusterSyncClient::PropagateStatesData(const TArray<uint8>& InLocalStatesData, TArray<uint8>& OutClusterStatesData)
+{
+	static const TSharedPtr<FDisplayClusterPacketInternal> Request = MakeShared<FDisplayClusterPacketInternal>(
+		DisplayClusterClusterSyncStrings::PropagateStatesData::Name,
+		DisplayClusterClusterSyncStrings::TypeRequest,
+		DisplayClusterClusterSyncStrings::ProtocolName);
+
+	Request->SetBinArg(DisplayClusterClusterSyncStrings::ArgumentsDefaultCategory, DisplayClusterClusterSyncStrings::PropagateStatesData::ArgLocalStatesData, InLocalStatesData);
+
+	TSharedPtr<FDisplayClusterPacketInternal> Response;
+
+	{
+		TRACE_CPUPROFILER_EVENT_SCOPE(CLN_CS::PropagateStatesData);
+		Response = SendRecvPacket(Request);
+	}
+
+	if (!Response)
+	{
+		UE_LOG(LogDisplayClusterNetwork, Warning, TEXT("Network error on '%s'"), *Request->GetName());
+		return EDisplayClusterCommResult::NetworkError;
+	}
+
+	// Extract data from response packet
+	Response->GetBinArg(DisplayClusterClusterSyncStrings::ArgumentsDefaultCategory, DisplayClusterClusterSyncStrings::PropagateStatesData::ArgClusterStatesData, OutClusterStatesData);
+
+	return Response->GetCommResult();
+}
diff --git a/Engine/Plugins/Runtime/nDisplay/Source/DisplayCluster/Private/Network/Service/ClusterSync/DisplayClusterClusterSyncClient.h b/Engine/Plugins/Runtime/nDisplay/Source/DisplayCluster/Private/Network/Service/ClusterSync/DisplayClusterClusterSyncClient.h
--- a/Engine/Plugins/Runtime/nDisplay/Source/DisplayCluster/Private/Network/Service/ClusterSync/DisplayClusterClusterSyncClient.h
+++ b/Engine/Plugins/Runtime/nDisplay/Source/DisplayCluster/Private/Network/Service/ClusterSync/DisplayClusterClusterSyncClient.h
@@ -34,5 +34,6 @@ public:
 	virtual EDisplayClusterCommResult GetObjectsData(const EDisplayClusterSyncGroup InSyncGroup, TMap<FString, FString>& OutObjectsData) override;
 	virtual EDisplayClusterCommResult GetEventsData(TArray<TSharedPtr<FDisplayClusterClusterEventJson, ESPMode::ThreadSafe>>& OutJsonEvents, TArray<TSharedPtr<FDisplayClusterClusterEventBinary, ESPMode::ThreadSafe>>& OutBinaryEvents) override;
 	virtual EDisplayClusterCommResult GetNativeInputData(TMap<FString, FString>& OutNativeInputData) override;
+	virtual EDisplayClusterCommResult PropagateStatesData(const TArray<uint8>& InLocalStatesData, TArray<uint8>& OutClusterStatesData) override;
 	//~ End IDisplayClusterProtocolClusterSync
 };
diff --git a/Engine/Plugins/Runtime/nDisplay/Source/DisplayCluster/Private/Network/Service/ClusterSync/DisplayClusterClusterSyncService.cpp b/Engine/Plugins/Runtime/nDisplay/Source/DisplayCluster/Private/Network/Service/ClusterSync/DisplayClusterClusterSyncService.cpp
--- a/Engine/Plugins/Runtime/nDisplay/Source/DisplayCluster/Private/Network/Service/ClusterSync/DisplayClusterClusterSyncService.cpp
+++ b/Engine/Plugins/Runtime/nDisplay/Source/DisplayCluster/Private/Network/Service/ClusterSync/DisplayClusterClusterSyncService.cpp
@@ -7,9 +7,12 @@
 #include "DisplayClusterEnums.h"
 #include "IDisplayCluster.h"
 
+#include "Async/ParallelFor.h"
+
 #include "Cluster/DisplayClusterClusterEvent.h"
 #include "Cluster/IPDisplayClusterClusterManager.h"
 #include "Cluster/Controller/IDisplayClusterClusterNodeController.h"
+#include "Cluster/CustomStates/DisplayClusterCustomStateHelpers.h"
 #include "Cluster/Failover/DisplayClusterCommDataCache.h"
 #include "Cluster/RequestHandler/DisplayClusterCommRequestHandlerLocal.h"
 #include "Cluster/RequestHandler/DisplayClusterCommRequestHandlerRemote.h"
@@ -58,6 +61,7 @@ FDisplayClusterClusterSyncService::FDisplayClusterClusterSyncService(const FName
 		BarrierGameStart.Reset( FDisplayClusterBarrierFactory::CreateBarrier(TEXT("GameStart_barrier"),  NodeIds, NetworkSettings.GameStartBarrierTimeout));
 		BarrierFrameStart.Reset(FDisplayClusterBarrierFactory::CreateBarrier(TEXT("FrameStart_barrier"), NodeIds, NetworkSettings.FrameStartBarrierTimeout));
 		BarrierFrameEnd.Reset(  FDisplayClusterBarrierFactory::CreateBarrier(TEXT("FrameEnd_barrier"),   NodeIds, NetworkSettings.FrameEndBarrierTimeout));
+		BarrierStatesSync.Reset(FDisplayClusterBarrierFactory::CreateBarrier(TEXT("StatesSync_barrier"), NodeIds, NetworkSettings.FrameStartBarrierTimeout));
 	}
 	else
 	{
@@ -67,12 +71,14 @@ FDisplayClusterClusterSyncService::FDisplayClusterClusterSyncService(const FName
 		BarrierGameStart.Reset( FDisplayClusterBarrierFactory::CreateBarrier(TEXT("GameStart_barrier"),  NodeIds, DefaultNetworkSettings.GameStartBarrierTimeout));
 		BarrierFrameStart.Reset(FDisplayClusterBarrierFactory::CreateBarrier(TEXT("FrameStart_barrier"), NodeIds, DefaultNetworkSettings.FrameStartBarrierTimeout));
 		BarrierFrameEnd.Reset(  FDisplayClusterBarrierFactory::CreateBarrier(TEXT("FrameEnd_barrier"),   NodeIds, DefaultNetworkSettings.FrameEndBarrierTimeout));
+		BarrierStatesSync.Reset(FDisplayClusterBarrierFactory::CreateBarrier(TEXT("StatesSync_barrier"), NodeIds, DefaultNetworkSettings.FrameStartBarrierTimeout));
 	}
 
 	// Put the barriers into an aux container
 	ServiceBarriers.Emplace(BarrierGameStart->GetName(),  &BarrierGameStart);
 	ServiceBarriers.Emplace(BarrierFrameStart->GetName(), &BarrierFrameStart);
 	ServiceBarriers.Emplace(BarrierFrameEnd->GetName(),   &BarrierFrameEnd);
+	ServiceBarriers.Emplace(BarrierStatesSync->GetName(), &BarrierStatesSync);
 
 	// Subscribe for barrier timeout events
 	for (TPair<FString, TUniquePtr<IDisplayClusterBarrier>*>& BarrierIt : ServiceBarriers)
@@ -82,6 +88,9 @@ FDisplayClusterClusterSyncService::FDisplayClusterClusterSyncService(const FName
 
 	// Subscribe for SessionClosed events
 	OnSessionClosed().AddRaw(this, &FDisplayClusterClusterSyncService::ProcessSessionClosed);
+
+	// Add sync function that will be responsible for parsing requests, generating and propagating state responses
+	BarrierStatesSync->GetPreSyncEndDelegate().BindRaw(this, &FDisplayClusterClusterSyncService::OnStatesSync);
 }
 
 FDisplayClusterClusterSyncService::~FDisplayClusterClusterSyncService()
@@ -92,6 +101,9 @@ FDisplayClusterClusterSyncService::~FDisplayClusterClusterSyncService()
 		(*BarrierIt.Value)->OnBarrierTimeout().RemoveAll(this);
 	}
 
+	// Reset states sync delegate
+	BarrierStatesSync->GetPreSyncEndDelegate().Unbind();
+
 	// Unsubscribe from SessionClosed notifications
 	OnSessionClosed().RemoveAll(this);
 
@@ -125,8 +137,7 @@ void FDisplayClusterClusterSyncService::Shutdown()
 
 FString FDisplayClusterClusterSyncService::GetProtocolName() const
 {
-	static const FString ProtocolName(DisplayClusterClusterSyncStrings::ProtocolName);
-	return ProtocolName;
+	return DisplayClusterClusterSyncStrings::ProtocolName;
 }
 
 void FDisplayClusterClusterSyncService::KillSession(const FString& NodeId)
@@ -221,6 +232,134 @@ void FDisplayClusterClusterSyncService::ShutdownImpl()
 	}
 }
 
+void FDisplayClusterClusterSyncService::OnStatesSync(FDisplayClusterBarrierPreSyncEndDelegateData& SyncData)
+{
+	using namespace UE::nDisplay::Cluster::Private;
+
+	// Cluster nodes num
+	const int32 NodesNum = SyncData.RequestData.Num();
+
+	// Intermediate input data. Array of cluster nodes, each one has its own state-to-data map.
+	TArray<TMap<FName, FCustomStateRequestBlob>> InputData;
+	InputData.AddDefaulted(NodesNum);
+
+	// Generate indexed node Ids for safe access from parallel jobs
+	TArray<FString> NodeIds;
+	SyncData.RequestData.GenerateKeyArray(NodeIds);
+
+	// Pre-generate node names to avoid local names generation
+	TArray<FName> NodeNames;
+	NodeNames.Reserve(NodeIds.Num());
+	for(const FString& NodeId : NodeIds)
+	{
+		NodeNames.Add(*NodeId);
+	}
+
+	// Make sure the response map has pre-allocated all the keys so it can be const while processed in parallel
+	for (const FString& NodeId : NodeIds)
+	{
+		SyncData.ResponseData.Emplace(NodeId);
+	}
+
+	// Per-node request parsing
+	ParallelFor(
+		NodesNum,
+		[
+			NodeIds = static_cast<const TArray<FString>&>(NodeIds),
+			Request = SyncData.RequestData,
+			&InputData
+		](int32 JobIndex)
+		{
+			// Deserialize request from a single cluster node
+			const TArray<uint8>& NodeRequestData = Request.FindChecked(NodeIds[JobIndex]);
+			FMemoryReader DataReader(NodeRequestData);
+
+			// Read states num
+			int32 StatesNum = 0;
+			DataReader << StatesNum;
+
+			// Deserialize all the blobs
+			for (int32 StateIdx = 0; StateIdx < StatesNum; ++StateIdx)
+			{
+				FCustomStateRequestBlob StateBlob;
+				StateBlob.Serialize(DataReader);
+				InputData[JobIndex].Emplace(StateBlob.StateName, MoveTemp(StateBlob));
+			}
+		});
+
+	// Intermediate output data. Array of cluster nodes, each one has its own NodeId-to-StateName-to-StateData map.
+	TArray<TMap<FName, FCustomStateResponseBlob>> OutputData;
+	OutputData.AddDefaulted(NodesNum);
+
+	// Per-node response generating
+	ParallelFor(
+		NodesNum,
+		[
+			NodesNum,
+			NodeIds   = static_cast<const TArray<FString>&>(NodeIds),
+			NodeNames = static_cast<const TArray<FName>&>(NodeNames),
+			InputData = static_cast<const TArray<TMap<FName, FCustomStateRequestBlob>>&>(InputData),
+			&OutputData,
+			&ResponseContainer = SyncData.ResponseData
+		](int32 JobIndex)
+		{
+			// Iterate through the states available on the node that is being processed currently (destination states)
+			const TMap<FName, FCustomStateRequestBlob>& DstNode = InputData[JobIndex];
+			for (const TPair<FName, FCustomStateRequestBlob>& DstState : DstNode)
+			{
+				// Iterate through every node to collect all the data we want to grab
+				for (int32 SrcNodeIdx = 0; SrcNodeIdx < NodesNum; ++SrcNodeIdx)
+				{
+					// Don't send data back. So far, it's not expected that any sender's data
+					// can be changed on the server side. Therefore no need to waste traffic
+					// on sending this data back. But this may change in the future.
+					if (SrcNodeIdx == JobIndex)
+					{
+						continue;
+					}
+
+					// Node ID that we want to grab data from
+					const FName& SrcNodeName = NodeNames[SrcNodeIdx];
+
+					// If destination state wants to get update for the source node
+					if (!DstState.Value.bCustomUpstream || DstState.Value.UpstreamNodes.Contains(SrcNodeName))
+					{
+						// If the source node has registered the same state ...
+						if (const FCustomStateRequestBlob* const FoundSrcBlob = InputData[SrcNodeIdx].Find(DstState.Key))
+						{
+							// ... and exposed it
+							if (FoundSrcBlob->bExposesData)
+							{
+								// Then it needs to be sent to the destination
+								FCustomStateResponseBlob& ResponseBlob = OutputData[JobIndex].FindOrAdd(DstState.Key);
+								ResponseBlob.StateName = DstState.Key;
+								ResponseBlob.NodeStates.Emplace(SrcNodeName, MakeTuple(FoundSrcBlob->StateType, FoundSrcBlob->StateData));
+							}
+						}
+					}
+				}
+			}
+
+			// Response buffer for current node (JobIndex)
+			const FString NodeId = NodeIds[JobIndex];
+			TArray<uint8>& ResponseBuffer = ResponseContainer.FindChecked(NodeId);
+			FMemoryWriter DataWriter(ResponseBuffer);
+
+			// Response data for current node
+			TMap<FName, FCustomStateResponseBlob>& OutputMap = OutputData[JobIndex];
+
+			// Write states num
+			int32 NumStates = OutputMap.Num();
+			DataWriter << NumStates;
+
+			// Write states
+			for (TPair<FName, FCustomStateResponseBlob>& SrcState : OutputMap)
+			{
+				SrcState.Value.Serialize(DataWriter);
+			}
+		});
+}
+
 
 //////////////////////////////////////////////////////////////////////////////////////////////
 // IDisplayClusterSessionPacketHandler
@@ -240,7 +379,8 @@ TSharedPtr<FDisplayClusterPacketInternal> FDisplayClusterClusterSyncService::Pro
 	UE_LOG(LogDisplayClusterNetwork, Verbose, TEXT("%s - Processing packet: %s"), *GetName(), *Request->ToLogString());
 
 	// Check protocol and type
-	if (Request->GetProtocol() != DisplayClusterClusterSyncStrings::ProtocolName || Request->GetType() != DisplayClusterClusterSyncStrings::TypeRequest)
+	if (!Request->GetProtocol().Equals(DisplayClusterClusterSyncStrings::ProtocolName, ESearchCase::IgnoreCase) ||
+		!Request->GetType().Equals(DisplayClusterClusterSyncStrings::TypeRequest, ESearchCase::IgnoreCase))
 	{
 		UE_LOG(LogDisplayClusterNetwork, Error, TEXT("%s - Unsupported packet type: %s"), *GetName(), *Request->ToLogString());
 		return nullptr;
@@ -342,6 +482,21 @@ TSharedPtr<FDisplayClusterPacketInternal> FDisplayClusterClusterSyncService::Pro
 
 		return Response;
 	}
+	else if (ReqName.Equals(DisplayClusterClusterSyncStrings::PropagateStatesData::Name, ESearchCase::IgnoreCase))
+	{
+		TArray<uint8> NodeStatesData;
+
+		Request->GetBinArg(DisplayClusterClusterSyncStrings::ArgumentsDefaultCategory, DisplayClusterClusterSyncStrings::PropagateStatesData::ArgLocalStatesData, NodeStatesData);
+
+		TArray<uint8> ClusterStatesData;
+		const EDisplayClusterCommResult CommResult = PropagateStatesData(NodeStatesData, ClusterStatesData);
+
+		Response->SetBinArg(DisplayClusterClusterSyncStrings::ArgumentsDefaultCategory, DisplayClusterClusterSyncStrings::PropagateStatesData::ArgClusterStatesData, ClusterStatesData);
+
+		Response->SetCommResult(CommResult);
+
+		return Response;
+	}
 
 	// Being here means that we have no appropriate dispatch logic for this packet
 	UE_LOG(LogDisplayClusterNetwork, Warning, TEXT("%s - No dispatcher found for packet '%s'"), *GetName(), *Request->GetName());
@@ -357,7 +512,7 @@ EDisplayClusterCommResult FDisplayClusterClusterSyncService::WaitForGameStart()
 {
 	TRACE_CPUPROFILER_EVENT_SCOPE(SRV_CS::WaitForGameStart);
 
-	const FString CachedNodeId = GetSessionInfoCache().NodeId.Get(TEXT(""));
+	const FString CachedNodeId = GetSessionInfoCache().NodeId.Get(FString());
 	BarrierGameStart->Wait(CachedNodeId);
 	return EDisplayClusterCommResult::Ok;
 }
@@ -366,7 +521,7 @@ EDisplayClusterCommResult FDisplayClusterClusterSyncService::WaitForFrameStart()
 {
 	TRACE_CPUPROFILER_EVENT_SCOPE(SRV_CS::WaitForFrameStart);
 
-	const FString CachedNodeId = GetSessionInfoCache().NodeId.Get(TEXT(""));
+	const FString CachedNodeId = GetSessionInfoCache().NodeId.Get(FString());
 	BarrierFrameStart->Wait(CachedNodeId);
 	return EDisplayClusterCommResult::Ok;
 }
@@ -375,7 +530,7 @@ EDisplayClusterCommResult FDisplayClusterClusterSyncService::WaitForFrameEnd()
 {
 	TRACE_CPUPROFILER_EVENT_SCOPE(SRV_CS::WaitForFrameEnd);
 
-	const FString CachedNodeId = GetSessionInfoCache().NodeId.Get(TEXT(""));
+	const FString CachedNodeId = GetSessionInfoCache().NodeId.Get(FString());
 	BarrierFrameEnd->Wait(CachedNodeId);
 	return EDisplayClusterCommResult::Ok;
 }
@@ -423,3 +578,12 @@ EDisplayClusterCommResult FDisplayClusterClusterSyncService::GetNativeInputData(
 
 	return Handler.GetNativeInputData(OutNativeInputData);
 }
+
+EDisplayClusterCommResult FDisplayClusterClusterSyncService::PropagateStatesData(const TArray<uint8>& InLocalStatesData, TArray<uint8>& OutClusterStatesData)
+{
+	TRACE_CPUPROFILER_EVENT_SCOPE(SRV_CS::PropagateStatesData);
+
+	const FString CachedNodeId = GetSessionInfoCache().NodeId.Get(FString());
+	BarrierStatesSync->WaitWithData(CachedNodeId, InLocalStatesData, OutClusterStatesData);
+	return EDisplayClusterCommResult::Ok;
+}
diff --git a/Engine/Plugins/Runtime/nDisplay/Source/DisplayCluster/Private/Network/Service/ClusterSync/DisplayClusterClusterSyncService.h b/Engine/Plugins/Runtime/nDisplay/Source/DisplayCluster/Private/Network/Service/ClusterSync/DisplayClusterClusterSyncService.h
--- a/Engine/Plugins/Runtime/nDisplay/Source/DisplayCluster/Private/Network/Service/ClusterSync/DisplayClusterClusterSyncService.h
+++ b/Engine/Plugins/Runtime/nDisplay/Source/DisplayCluster/Private/Network/Service/ClusterSync/DisplayClusterClusterSyncService.h
@@ -53,6 +53,12 @@ private:
 	/** Non-virtual shutdown implementation */
 	void ShutdownImpl();
 
+private:
+
+	/** Handles states synchronization. Generates and propagates data according to the propagation strategy. */
+	void OnStatesSync(FDisplayClusterBarrierPreSyncEndDelegateData& SyncData);
+
+
 protected:
 
 	//~ Begin IDisplayClusterSessionPacketHandler
@@ -69,6 +75,7 @@ protected:
 	virtual EDisplayClusterCommResult GetObjectsData(EDisplayClusterSyncGroup InSyncGroup, TMap<FString, FString>& OutObjectsData)  override;
 	virtual EDisplayClusterCommResult GetEventsData(TArray<TSharedPtr<FDisplayClusterClusterEventJson, ESPMode::ThreadSafe>>& OutJsonEvents, TArray<TSharedPtr<FDisplayClusterClusterEventBinary, ESPMode::ThreadSafe>>& OutBinaryEvents) override;
 	virtual EDisplayClusterCommResult GetNativeInputData(TMap<FString, FString>& OutNativeInputData) override;
+	virtual EDisplayClusterCommResult PropagateStatesData(const TArray<uint8>& InLocalStatesData, TArray<uint8>& OutClusterStatesData) override;
 	//~ End IDisplayClusterProtocolClusterSync
 
 private:
@@ -82,6 +89,9 @@ private:
 	/** Frame end barrier */
 	TUniquePtr<IDisplayClusterBarrier> BarrierFrameEnd;
 
+	/** Custom states sync barrier  */
+	TUniquePtr<IDisplayClusterBarrier> BarrierStatesSync;
+
 	/** Auxiliary container that keeps all the barriers */
 	TMap<FString, TUniquePtr<IDisplayClusterBarrier>*> ServiceBarriers;
 };
diff --git a/Engine/Plugins/Runtime/nDisplay/Source/DisplayCluster/Private/Network/Service/ClusterSync/DisplayClusterClusterSyncStrings.cpp b/Engine/Plugins/Runtime/nDisplay/Source/DisplayCluster/Private/Network/Service/ClusterSync/DisplayClusterClusterSyncStrings.cpp
new file mode 100644
--- /dev/null
+++ b/Engine/Plugins/Runtime/nDisplay/Source/DisplayCluster/Private/Network/Service/ClusterSync/DisplayClusterClusterSyncStrings.cpp
@@ -0,0 +1,68 @@
+// Copyright Epic Games, Inc. All Rights Reserved.
+
+#include "Network/Service/ClusterSync/DisplayClusterClusterSyncStrings.h"
+
+
+namespace DisplayClusterClusterSyncStrings
+{
+	const FString ProtocolName = TEXT("ClusterSync");
+
+	const FString TypeRequest  = TEXT("Request");
+	const FString TypeResponse = TEXT("Response");
+
+	const FString ArgumentsDefaultCategory = TEXT("CS");
+	const FString ArgumentsJsonEvents      = TEXT("CS_JE");
+	const FString ArgumentsBinaryEvents    = TEXT("CS_BE");
+
+	namespace WaitForGameStart
+	{
+		const FString Name = TEXT("WaitForGameStart");
+	}
+
+	namespace WaitForFrameStart
+	{
+		const FString Name = TEXT("WaitForFrameStart");
+	}
+
+	namespace WaitForFrameEnd
+	{
+		const FString Name = TEXT("WaitForFrameEnd");
+	}
+
+	namespace GetTimeData
+	{
+		const FString Name = TEXT("GetTimeData");
+
+		const FString ArgDeltaTime        = TEXT("DeltaTime");
+		const FString ArgGameTime         = TEXT("GameTime");
+		const FString ArgIsFrameTimeValid = TEXT("IsFrameTimeValid");
+		const FString ArgFrameTime        = TEXT("FrameTime");
+	}
+
+	namespace GetObjectsData
+	{
+		const FString Name = TEXT("GetObjectsData");
+
+		const FString ArgSyncGroup = TEXT("SyncGroup");
+	}
+
+	namespace GetEventsData
+	{
+		const FString Name = TEXT("GetEventsData");
+	}
+
+	namespace GetNativeInputData
+	{
+		const FString Name = TEXT("GetNativeInputData");
+
+		const FString ArgNativeInputData = TEXT("NativeInputData");
+	}
+
+	namespace PropagateStatesData
+	{
+		const FString Name = TEXT("PropagateStatesData");
+
+		const FString ArgLocalStatesData   = TEXT("LocalData");
+		const FString ArgClusterStatesData = TEXT("ClusterData");
+	}
+};
diff --git a/Engine/Plugins/Runtime/nDisplay/Source/DisplayCluster/Private/Network/Service/ClusterSync/DisplayClusterClusterSyncStrings.h b/Engine/Plugins/Runtime/nDisplay/Source/DisplayCluster/Private/Network/Service/ClusterSync/DisplayClusterClusterSyncStrings.h
--- a/Engine/Plugins/Runtime/nDisplay/Source/DisplayCluster/Private/Network/Service/ClusterSync/DisplayClusterClusterSyncStrings.h
+++ b/Engine/Plugins/Runtime/nDisplay/Source/DisplayCluster/Private/Network/Service/ClusterSync/DisplayClusterClusterSyncStrings.h
@@ -2,7 +2,7 @@
 
 #pragma once
 
-#include "HAL/Platform.h"
+#include "CoreMinimal.h"
 
 
 /**
@@ -10,53 +10,64 @@
  */
 namespace DisplayClusterClusterSyncStrings
 {
-	constexpr static const TCHAR* ProtocolName = TEXT("ClusterSync");
-	
-	constexpr static const TCHAR* TypeRequest  = TEXT("Request");
-	constexpr static const TCHAR* TypeResponse = TEXT("Response");
+	extern const FString ProtocolName;
 
-	constexpr static const TCHAR* ArgumentsDefaultCategory = TEXT("CS");
-	constexpr static const TCHAR* ArgumentsJsonEvents      = TEXT("CS_JE");
-	constexpr static const TCHAR* ArgumentsBinaryEvents    = TEXT("CS_BE");
+	extern const FString TypeRequest;
+	extern const FString TypeResponse;
+
+	extern const FString ArgumentsDefaultCategory;
+	extern const FString ArgumentsJsonEvents;
+	extern const FString ArgumentsBinaryEvents;
 
 	namespace WaitForGameStart
 	{
-		constexpr static const TCHAR* Name = TEXT("WaitForGameStart");
+		extern const FString Name;
 	}
 
 	namespace WaitForFrameStart
 	{
-		constexpr static const TCHAR* Name = TEXT("WaitForFrameStart");
+		extern const FString Name;
 	}
 
 	namespace WaitForFrameEnd
 	{
-		constexpr static const TCHAR* Name = TEXT("WaitForFrameEnd");
+		extern const FString Name;
 	}
 
 	namespace GetTimeData
 	{
-		constexpr static const TCHAR* Name = TEXT("GetTimeData");
-		constexpr static const TCHAR* ArgDeltaTime        = TEXT("DeltaTime");
-		constexpr static const TCHAR* ArgGameTime         = TEXT("GameTime");
-		constexpr static const TCHAR* ArgIsFrameTimeValid = TEXT("IsFrameTimeValid");
-		constexpr static const TCHAR* ArgFrameTime        = TEXT("FrameTime");
+		extern const FString Name;
+
+		extern const FString ArgDeltaTime;
+		extern const FString ArgGameTime;
+		extern const FString ArgIsFrameTimeValid;
+		extern const FString ArgFrameTime;
 	}
 
 	namespace GetObjectsData
 	{
-		constexpr static const TCHAR* Name = TEXT("GetObjectsData");
-		constexpr static const TCHAR* ArgSyncGroup = TEXT("SyncGroup");
+		extern const FString Name;
+
+		extern const FString ArgSyncGroup;
 	}
 
 	namespace GetEventsData
 	{
-		constexpr static const TCHAR* Name = TEXT("GetEventsData");
+		extern const FString Name;
 	}
 
 	namespace GetNativeInputData
 	{
-		constexpr static const TCHAR* Name = TEXT("GetNativeInputData");
-		constexpr static const TCHAR* ArgNativeInputData = TEXT("NativeInputData");
+		extern const FString Name;
+
+		extern const FString ArgNativeInputData;
+	}
+
+	namespace PropagateStatesData
+	{
+		extern const FString Name;
+
+		extern const FString ArgLocalStatesData;
+		extern const FString ArgClusterStatesData;
 	}
 };
diff --git a/Engine/Plugins/Runtime/nDisplay/Source/DisplayCluster/Private/Network/Service/DisplayClusterService.cpp b/Engine/Plugins/Runtime/nDisplay/Source/DisplayCluster/Private/Network/Service/DisplayClusterService.cpp
--- a/Engine/Plugins/Runtime/nDisplay/Source/DisplayCluster/Private/Network/Service/DisplayClusterService.cpp
+++ b/Engine/Plugins/Runtime/nDisplay/Source/DisplayCluster/Private/Network/Service/DisplayClusterService.cpp
@@ -111,6 +111,6 @@ bool FDisplayClusterService::IsLocalRequest() const
 	static const FString ThisNodeId = GDisplayCluster->GetClusterMgr()->GetNodeId();
 
 	// Check if this session belongs to this node. If so, it's a local request.
-	const bool bLocalRequest = ThisNodeId.Equals(GetSessionInfoCache().NodeId.Get(TEXT("")), ESearchCase::IgnoreCase);
+	const bool bLocalRequest = ThisNodeId.Equals(GetSessionInfoCache().NodeId.Get(FString()), ESearchCase::IgnoreCase);
 	return bLocalRequest;
 }
diff --git a/Engine/Plugins/Runtime/nDisplay/Source/DisplayCluster/Private/Network/Service/GenericBarrier/DisplayClusterGenericBarrierService.cpp b/Engine/Plugins/Runtime/nDisplay/Source/DisplayCluster/Private/Network/Service/GenericBarrier/DisplayClusterGenericBarrierService.cpp
--- a/Engine/Plugins/Runtime/nDisplay/Source/DisplayCluster/Private/Network/Service/GenericBarrier/DisplayClusterGenericBarrierService.cpp
+++ b/Engine/Plugins/Runtime/nDisplay/Source/DisplayCluster/Private/Network/Service/GenericBarrier/DisplayClusterGenericBarrierService.cpp
@@ -39,8 +39,7 @@ void FDisplayClusterGenericBarrierService::Shutdown()
 
 FString FDisplayClusterGenericBarrierService::GetProtocolName() const
 {
-	static const FString ProtocolName(DisplayClusterGenericBarrierStrings::ProtocolName);
-	return ProtocolName;
+	return DisplayClusterGenericBarrierStrings::ProtocolName;
 }
 
 void FDisplayClusterGenericBarrierService::KillSession(const FString& NodeId)
@@ -98,7 +97,8 @@ TSharedPtr<FDisplayClusterPacketInternal> FDisplayClusterGenericBarrierService::
 	UE_LOG(LogDisplayClusterNetwork, Verbose, TEXT("%s - Processing packet: %s"), *GetName(), *Request->ToLogString());
 
 	// Check protocol and type
-	if (Request->GetProtocol() != DisplayClusterGenericBarrierStrings::ProtocolName || Request->GetType() != DisplayClusterGenericBarrierStrings::TypeRequest)
+	if (!Request->GetProtocol().Equals(DisplayClusterGenericBarrierStrings::ProtocolName, ESearchCase::IgnoreCase) ||
+		!Request->GetType().Equals(DisplayClusterGenericBarrierStrings::TypeRequest, ESearchCase::IgnoreCase))
 	{
 		UE_LOG(LogDisplayClusterNetwork, Error, TEXT("%s - Unsupported packet type: %s"), *GetName(), *Request->ToLogString());
 		return nullptr;
diff --git a/Engine/Plugins/Runtime/nDisplay/Source/DisplayCluster/Private/Network/Service/GenericBarrier/DisplayClusterGenericBarrierStrings.cpp b/Engine/Plugins/Runtime/nDisplay/Source/DisplayCluster/Private/Network/Service/GenericBarrier/DisplayClusterGenericBarrierStrings.cpp
new file mode 100644
--- /dev/null
+++ b/Engine/Plugins/Runtime/nDisplay/Source/DisplayCluster/Private/Network/Service/GenericBarrier/DisplayClusterGenericBarrierStrings.cpp
@@ -0,0 +1,57 @@
+// Copyright Epic Games, Inc. All Rights Reserved.
+
+#include "Network/Service/GenericBarrier/DisplayClusterGenericBarrierStrings.h"
+
+
+namespace DisplayClusterGenericBarrierStrings
+{
+	const FString ProtocolName = TEXT("GenericBarrier");
+
+	const FString TypeRequest  = TEXT("Request");
+	const FString TypeResponse = TEXT("Response");
+
+	const FString ArgumentsDefaultCategory = TEXT("GB");
+
+	// Shared arguments
+	const FString ArgBarrierId = TEXT("BarrierId");
+	const FString ArgResult    = TEXT("CtrlResult");
+
+	namespace CreateBarrier
+	{
+		const FString Name = TEXT("CreateBarrier");
+
+		const FString ArgCallers = TEXT("SyncCallers");
+		const FString ArgTimeout = TEXT("Timeout");
+	}
+
+	namespace WaitUntilBarrierIsCreated
+	{
+		const FString Name = TEXT("WaitUntilBarrierIsCreated");
+	}
+
+	namespace IsBarrierAvailable
+	{
+		const FString Name = TEXT("IsBarrierAvailable");
+	}
+
+	namespace ReleaseBarrier
+	{
+		const FString Name = TEXT("ReleaseBarrier");
+	}
+
+	namespace SyncOnBarrier
+	{
+		const FString Name = TEXT("Sync");
+
+		const FString ArgCallerId = TEXT("CallerId");
+	}
+
+	namespace SyncOnBarrierWithData
+	{
+		const FString Name = TEXT("SyncWithData");
+
+		const FString ArgCallerId     = TEXT("CallerId");
+		const FString ArgRequestData  = TEXT("ReqData");
+		const FString ArgResponseData = TEXT("RespData");
+	}
+};
diff --git a/Engine/Plugins/Runtime/nDisplay/Source/DisplayCluster/Private/Network/Service/GenericBarrier/DisplayClusterGenericBarrierStrings.h b/Engine/Plugins/Runtime/nDisplay/Source/DisplayCluster/Private/Network/Service/GenericBarrier/DisplayClusterGenericBarrierStrings.h
--- a/Engine/Plugins/Runtime/nDisplay/Source/DisplayCluster/Private/Network/Service/GenericBarrier/DisplayClusterGenericBarrierStrings.h
+++ b/Engine/Plugins/Runtime/nDisplay/Source/DisplayCluster/Private/Network/Service/GenericBarrier/DisplayClusterGenericBarrierStrings.h
@@ -2,7 +2,7 @@
 
 #pragma once
 
-#include "HAL/Platform.h"
+#include "CoreMinimal.h"
 
 
 /**
@@ -10,54 +10,53 @@
  */
 namespace DisplayClusterGenericBarrierStrings
 {
-	constexpr static const TCHAR* ProtocolName = TEXT("GenericBarrier");
+	extern const FString ProtocolName;
 
-	constexpr static const TCHAR* TypeRequest  = TEXT("Request");
-	constexpr static const TCHAR* TypeResponse = TEXT("Response");
+	extern const FString TypeRequest;
+	extern const FString TypeResponse;
 
-	constexpr static const TCHAR* ArgumentsDefaultCategory = TEXT("GB");
+	extern const FString ArgumentsDefaultCategory;
 
 	// Shared arguments
-	constexpr static const TCHAR* ArgBarrierId = TEXT("BarrierId");
-	constexpr static const TCHAR* ArgResult    = TEXT("CtrlResult");
-
+	extern const FString ArgBarrierId;
+	extern const FString ArgResult;
 
 	namespace CreateBarrier
 	{
-		constexpr static const TCHAR* Name = TEXT("CreateBarrier");
+		extern const FString Name;
 
-		constexpr static const TCHAR* ArgCallers = TEXT("SyncCallers");
-		constexpr static const TCHAR* ArgTimeout = TEXT("Timeout");
+		extern const FString ArgCallers;
+		extern const FString ArgTimeout;
 	}
 
 	namespace WaitUntilBarrierIsCreated
 	{
-		constexpr static const TCHAR* Name = TEXT("WaitUntilBarrierIsCreated");
+		extern const FString Name;
 	}
 
 	namespace IsBarrierAvailable
 	{
-		constexpr static const TCHAR* Name = TEXT("IsBarrierAvailable");
+		extern const FString Name;
 	}
 
 	namespace ReleaseBarrier
 	{
-		constexpr static const TCHAR* Name = TEXT("ReleaseBarrier");
+		extern const FString Name;
 	}
 
 	namespace SyncOnBarrier
 	{
-		constexpr static const TCHAR* Name = TEXT("Sync");
+		extern const FString Name;
 
-		constexpr static const TCHAR* ArgCallerId = TEXT("CallerId");
+		extern const FString ArgCallerId;
 	}
 
 	namespace SyncOnBarrierWithData
 	{
-		constexpr static const TCHAR* Name = TEXT("SyncWithData");
+		extern const FString Name;
 
-		constexpr static const TCHAR* ArgCallerId     = TEXT("CallerId");
-		constexpr static const TCHAR* ArgRequestData  = TEXT("ReqData");
-		constexpr static const TCHAR* ArgResponseData = TEXT("RespData");
+		extern const FString ArgCallerId;
+		extern const FString ArgRequestData;
+		extern const FString ArgResponseData;
 	}
 };
diff --git a/Engine/Plugins/Runtime/nDisplay/Source/DisplayCluster/Private/Network/Service/InternalComm/DisplayClusterInternalCommClient.cpp b/Engine/Plugins/Runtime/nDisplay/Source/DisplayCluster/Private/Network/Service/InternalComm/DisplayClusterInternalCommClient.cpp
--- a/Engine/Plugins/Runtime/nDisplay/Source/DisplayCluster/Private/Network/Service/InternalComm/DisplayClusterInternalCommClient.cpp
+++ b/Engine/Plugins/Runtime/nDisplay/Source/DisplayCluster/Private/Network/Service/InternalComm/DisplayClusterInternalCommClient.cpp
@@ -90,7 +90,7 @@ EDisplayClusterCommResult FDisplayClusterInternalCommClient::GatherServicesHosti
 		TArray<uint8> ResponseData;
 		if (!Response->GetBinArg(DisplayClusterInternalCommStrings::ArgumentsDefaultCategory, DisplayClusterInternalCommStrings::GatherServicesHostingInfo::ArgClusterHostingInfo, ResponseData))
 		{
-			UE_LOG(LogDisplayClusterNetwork, Warning, TEXT("Couldn't extract parameter: %s"), DisplayClusterInternalCommStrings::GatherServicesHostingInfo::ArgClusterHostingInfo);
+			UE_LOG(LogDisplayClusterNetwork, Warning, TEXT("Couldn't extract parameter: %s"), *DisplayClusterInternalCommStrings::GatherServicesHostingInfo::ArgClusterHostingInfo);
 			return EDisplayClusterCommResult::WrongResponseData;
 		}
 
@@ -133,16 +133,16 @@ EDisplayClusterCommResult FDisplayClusterInternalCommClient::PostFailureNegotiat
 		Response = SendRecvPacket(Request);
 	}
 
-	// Extract response data
-	InOutRecoveryData.Reset();
-	Response->GetBinArg(DisplayClusterInternalCommStrings::ArgumentsDefaultCategory, DisplayClusterInternalCommStrings::PostFailureNegotiate::ArgRecoveryData, InOutRecoveryData);
-
 	if (!Response)
 	{
 		UE_LOG(LogDisplayClusterNetwork, Warning, TEXT("Network error on '%s'"), *Request->GetName());
 		return EDisplayClusterCommResult::NetworkError;
 	}
 
+	// Extract response data
+	InOutRecoveryData.Reset();
+	Response->GetBinArg(DisplayClusterInternalCommStrings::ArgumentsDefaultCategory, DisplayClusterInternalCommStrings::PostFailureNegotiate::ArgRecoveryData, InOutRecoveryData);
+
 	return Response->GetCommResult();
 }
 
diff --git a/Engine/Plugins/Runtime/nDisplay/Source/DisplayCluster/Private/Network/Service/InternalComm/DisplayClusterInternalCommService.cpp b/Engine/Plugins/Runtime/nDisplay/Source/DisplayCluster/Private/Network/Service/InternalComm/DisplayClusterInternalCommService.cpp
--- a/Engine/Plugins/Runtime/nDisplay/Source/DisplayCluster/Private/Network/Service/InternalComm/DisplayClusterInternalCommService.cpp
+++ b/Engine/Plugins/Runtime/nDisplay/Source/DisplayCluster/Private/Network/Service/InternalComm/DisplayClusterInternalCommService.cpp
@@ -104,8 +104,7 @@ void FDisplayClusterInternalCommService::Shutdown()
 
 FString FDisplayClusterInternalCommService::GetProtocolName() const
 {
-	static const FString ProtocolName(DisplayClusterInternalCommStrings::ProtocolName);
-	return ProtocolName;
+	return DisplayClusterInternalCommStrings::ProtocolName;
 }
 
 void FDisplayClusterInternalCommService::KillSession(const FString& NodeId)
@@ -249,7 +248,8 @@ TSharedPtr<FDisplayClusterPacketInternal> FDisplayClusterInternalCommService::Pr
 	UE_LOG(LogDisplayClusterNetwork, Verbose, TEXT("%s - Processing packet: %s"), *GetName(), *Request->ToLogString());
 
 	// Check protocol and type
-	if (Request->GetProtocol() != DisplayClusterInternalCommStrings::ProtocolName || Request->GetType() != DisplayClusterInternalCommStrings::TypeRequest)
+	if (!Request->GetProtocol().Equals(DisplayClusterInternalCommStrings::ProtocolName, ESearchCase::IgnoreCase) ||
+		!Request->GetType().Equals(DisplayClusterInternalCommStrings::TypeRequest, ESearchCase::IgnoreCase))
 	{
 		UE_LOG(LogDisplayClusterNetwork, Error, TEXT("%s - Unsupported packet type: %s"), *GetName(), *Request->ToLogString());
 		return nullptr;
diff --git a/Engine/Plugins/Runtime/nDisplay/Source/DisplayCluster/Private/Network/Service/InternalComm/DisplayClusterInternalCommStrings.cpp b/Engine/Plugins/Runtime/nDisplay/Source/DisplayCluster/Private/Network/Service/InternalComm/DisplayClusterInternalCommStrings.cpp
new file mode 100644
--- /dev/null
+++ b/Engine/Plugins/Runtime/nDisplay/Source/DisplayCluster/Private/Network/Service/InternalComm/DisplayClusterInternalCommStrings.cpp
@@ -0,0 +1,38 @@
+// Copyright Epic Games, Inc. All Rights Reserved.
+
+#include "Network/Service/InternalComm/DisplayClusterInternalCommStrings.h"
+
+
+namespace DisplayClusterInternalCommStrings
+{
+	const FString ProtocolName = TEXT("InternalComm");
+
+	const FString TypeRequest  = TEXT("Request");
+	const FString TypeResponse = TEXT("Response");
+
+	const FString ArgumentsDefaultCategory = TEXT("IC");
+
+	namespace GatherServicesHostingInfo
+	{
+		const FString Name = TEXT("GatherServicesHostingInfo");
+
+		const FString ArgNodeHostingInfo    = TEXT("NodeHostingInfo");
+		const FString ArgClusterHostingInfo = TEXT("ClusterHostingInfo");
+	}
+
+	namespace PostFailureNegotiate
+	{
+		const FString Name = TEXT("PostFailureNegotiate");
+
+		const FString ArgSyncStateData = TEXT("SyncState");
+		const FString ArgRecoveryData  = TEXT("RecoveryData");
+	}
+
+	namespace RequestNodeDrop
+	{
+		const FString Name = TEXT("RequestNodeDrop");
+
+		const FString ArgNodeId     = TEXT("NodeId");
+		const FString ArgDropReason = TEXT("DropReason");
+	}
+};
diff --git a/Engine/Plugins/Runtime/nDisplay/Source/DisplayCluster/Private/Network/Service/InternalComm/DisplayClusterInternalCommStrings.h b/Engine/Plugins/Runtime/nDisplay/Source/DisplayCluster/Private/Network/Service/InternalComm/DisplayClusterInternalCommStrings.h
--- a/Engine/Plugins/Runtime/nDisplay/Source/DisplayCluster/Private/Network/Service/InternalComm/DisplayClusterInternalCommStrings.h
+++ b/Engine/Plugins/Runtime/nDisplay/Source/DisplayCluster/Private/Network/Service/InternalComm/DisplayClusterInternalCommStrings.h
@@ -2,7 +2,7 @@
 
 #pragma once
 
-#include "HAL/Platform.h"
+#include "CoreMinimal.h"
 
 
 /**
@@ -10,34 +10,34 @@
  */
 namespace DisplayClusterInternalCommStrings
 {
-	constexpr static const TCHAR* ProtocolName = TEXT("InternalComm");
-	
-	constexpr static const TCHAR* TypeRequest  = TEXT("Request");
-	constexpr static const TCHAR* TypeResponse = TEXT("Response");
+	extern const FString ProtocolName;
 
-	constexpr static const TCHAR* ArgumentsDefaultCategory = TEXT("IC");
+	extern const FString TypeRequest;
+	extern const FString TypeResponse;
+
+	extern const FString ArgumentsDefaultCategory;
 
 	namespace GatherServicesHostingInfo
 	{
-		constexpr static const TCHAR* Name = TEXT("GatherServicesHostingInfo");
+		extern const FString Name;
 
-		constexpr static const TCHAR* ArgNodeHostingInfo    = TEXT("NodeHostingInfo");
-		constexpr static const TCHAR* ArgClusterHostingInfo = TEXT("ClusterHostingInfo");
+		extern const FString ArgNodeHostingInfo;
+		extern const FString ArgClusterHostingInfo;
 	}
 
 	namespace PostFailureNegotiate
 	{
-		constexpr static const TCHAR* Name = TEXT("PostFailureNegotiate");
+		extern const FString Name;
 
-		constexpr static const TCHAR* ArgSyncStateData = TEXT("SyncState");
-		constexpr static const TCHAR* ArgRecoveryData  = TEXT("RecoveryData");
+		extern const FString ArgSyncStateData;
+		extern const FString ArgRecoveryData;
 	}
 
 	namespace RequestNodeDrop
 	{
-		constexpr static const TCHAR* Name = TEXT("RequestNodeDrop");
+		extern const FString Name;
 
-		constexpr static const TCHAR* ArgNodeId     = TEXT("NodeId");
-		constexpr static const TCHAR* ArgDropReason = TEXT("DropReason");
+		extern const FString ArgNodeId;
+		extern const FString ArgDropReason;
 	}
 };
diff --git a/Engine/Plugins/Runtime/nDisplay/Source/DisplayCluster/Private/Network/Service/RenderSync/DisplayClusterRenderSyncService.cpp b/Engine/Plugins/Runtime/nDisplay/Source/DisplayCluster/Private/Network/Service/RenderSync/DisplayClusterRenderSyncService.cpp
--- a/Engine/Plugins/Runtime/nDisplay/Source/DisplayCluster/Private/Network/Service/RenderSync/DisplayClusterRenderSyncService.cpp
+++ b/Engine/Plugins/Runtime/nDisplay/Source/DisplayCluster/Private/Network/Service/RenderSync/DisplayClusterRenderSyncService.cpp
@@ -64,8 +64,7 @@ void FDisplayClusterRenderSyncService::Shutdown()
 
 FString FDisplayClusterRenderSyncService::GetProtocolName() const
 {
-	static const FString ProtocolName(DisplayClusterRenderSyncStrings::ProtocolName);
-	return ProtocolName;
+	return DisplayClusterRenderSyncStrings::ProtocolName;
 }
 
 void FDisplayClusterRenderSyncService::KillSession(const FString& NodeId)
@@ -282,7 +281,8 @@ TSharedPtr<FDisplayClusterPacketInternal> FDisplayClusterRenderSyncService::Proc
 	UE_LOG(LogDisplayClusterNetwork, Verbose, TEXT("%s - Processing packet: %s"), *GetName(), *Request->ToLogString());
 
 	// Check protocol and type
-	if (Request->GetProtocol() != DisplayClusterRenderSyncStrings::ProtocolName || Request->GetType() != DisplayClusterRenderSyncStrings::TypeRequest)
+	if (!Request->GetProtocol().Equals(DisplayClusterRenderSyncStrings::ProtocolName, ESearchCase::IgnoreCase) ||
+		!Request->GetType().Equals(DisplayClusterRenderSyncStrings::TypeRequest, ESearchCase::IgnoreCase))
 	{
 		UE_LOG(LogDisplayClusterNetwork, Error, TEXT("%s - Unsupported packet type: %s"), *GetName(), *Request->ToLogString());
 		return nullptr;
@@ -315,7 +315,7 @@ EDisplayClusterCommResult FDisplayClusterRenderSyncService::SynchronizeOnBarrier
 {
 	TRACE_CPUPROFILER_EVENT_SCOPE(SRV_RS::SyncOnBarrier);
 
-	const FString NodeId = GetSessionInfoCache().NodeId.Get(TEXT(""));
+	const FString NodeId = GetSessionInfoCache().NodeId.Get(FString());
 
 	IDisplayClusterBarrier* Barrier = GetBarrierForNode(NodeId);
 	if (ensureMsgf(Barrier, TEXT("%s could not find a barrier for node '%s'"), *GetName(), *NodeId))
diff --git a/Engine/Plugins/Runtime/nDisplay/Source/DisplayCluster/Private/Network/Service/RenderSync/DisplayClusterRenderSyncStrings.cpp b/Engine/Plugins/Runtime/nDisplay/Source/DisplayCluster/Private/Network/Service/RenderSync/DisplayClusterRenderSyncStrings.cpp
new file mode 100644
--- /dev/null
+++ b/Engine/Plugins/Runtime/nDisplay/Source/DisplayCluster/Private/Network/Service/RenderSync/DisplayClusterRenderSyncStrings.cpp
@@ -0,0 +1,19 @@
+// Copyright Epic Games, Inc. All Rights Reserved.
+
+#include "Network/Service/RenderSync/DisplayClusterRenderSyncStrings.h"
+
+
+namespace DisplayClusterRenderSyncStrings
+{
+	const FString ProtocolName = TEXT("RenderSync");
+
+	const FString TypeRequest  = TEXT("Request");
+	const FString TypeResponse = TEXT("Response");
+
+	const FString ArgumentsDefaultCategory = TEXT("RS");
+
+	namespace SynchronizeOnBarrier
+	{
+		const FString Name = TEXT("SyncOnBarrier");
+	};
+};
diff --git a/Engine/Plugins/Runtime/nDisplay/Source/DisplayCluster/Private/Network/Service/RenderSync/DisplayClusterRenderSyncStrings.h b/Engine/Plugins/Runtime/nDisplay/Source/DisplayCluster/Private/Network/Service/RenderSync/DisplayClusterRenderSyncStrings.h
--- a/Engine/Plugins/Runtime/nDisplay/Source/DisplayCluster/Private/Network/Service/RenderSync/DisplayClusterRenderSyncStrings.h
+++ b/Engine/Plugins/Runtime/nDisplay/Source/DisplayCluster/Private/Network/Service/RenderSync/DisplayClusterRenderSyncStrings.h
@@ -2,7 +2,7 @@
 
 #pragma once
 
-#include "HAL/Platform.h"
+#include "CoreMinimal.h"
 
 
 /**
@@ -10,15 +10,15 @@
  */
 namespace DisplayClusterRenderSyncStrings
 {
-	constexpr static const TCHAR* ProtocolName = TEXT("RenderSync");
+	extern const FString ProtocolName;
 
-	constexpr static const TCHAR* TypeRequest  = TEXT("Request");
-	constexpr static const TCHAR* TypeResponse = TEXT("Response");
+	extern const FString TypeRequest;
+	extern const FString TypeResponse;
 
-	constexpr static const TCHAR* ArgumentsDefaultCategory = TEXT("RS");
+	extern const FString ArgumentsDefaultCategory;
 
 	namespace SynchronizeOnBarrier
 	{
-		constexpr static const TCHAR* Name = TEXT("SyncOnBarrier");
+		extern const FString Name;
 	};
 };
diff --git a/Engine/Plugins/Runtime/nDisplay/Source/DisplayCluster/Public/Cluster/CustomStates/DisplayClusterCustomStateBase.h b/Engine/Plugins/Runtime/nDisplay/Source/DisplayCluster/Public/Cluster/CustomStates/DisplayClusterCustomStateBase.h
new file mode 100644
--- /dev/null
+++ b/Engine/Plugins/Runtime/nDisplay/Source/DisplayCluster/Public/Cluster/CustomStates/DisplayClusterCustomStateBase.h
@@ -0,0 +1,372 @@
+// Copyright Epic Games, Inc. All Rights Reserved.
+
+#pragma once
+
+#include "Cluster/CustomStates/IDisplayClusterCustomState.h"
+
+#include "IDisplayCluster.h"
+#include "Cluster/IDisplayClusterClusterManager.h"
+
+
+namespace UE::nDisplay::Private
+{
+	/**
+	 * Type-independent base class for custom states
+	 */
+	class DISPLAYCLUSTER_API FCustomStateBase
+		: public IDisplayClusterCustomState
+	{
+	public:
+
+		FCustomStateBase(const FName& UniqueName, const FName& ClusterNodeId);
+		virtual ~FCustomStateBase() = default;
+
+	public:
+
+		//~ Begin IDisplayClusterCustomState
+		virtual FName GetName() const override;
+		virtual void Lock() const override;
+		virtual void Unlock() const override;
+		//~ End IDisplayClusterCustomState
+
+	protected:
+
+		/**
+		 * Type-independent private implementation of render task scheduler.
+		 * The main idea of this function is to keep ENQUEUE_RENDER_COMMAND in the private
+		 * space, and therefore not introduce any new public dependencies on RenderCore.
+		 */
+		void ExecuteOnRenderThread(TUniqueFunction<void()> FuncRT);
+
+	protected:
+
+		/** Quick access for ClusterMgr cached pointer */
+		IDisplayClusterClusterManager* GetClusterMgr() const
+		{
+			return ClusterMgr;
+		}
+
+		/** Quick access for this cluster node ID */
+		const FName& GetNodeId() const
+		{
+			return ClusterNodeId;
+		}
+
+		/** Access to the main access guard of the state instance */
+		FCriticalSection& GetCritSec() const
+		{
+			return *StateCS;
+		}
+
+		/** Access to the main access guard of the state instance */
+		TSharedRef<FCriticalSection> GetCritSecRef() const
+		{
+			return StateCS;
+		}
+
+	private:
+
+		/** Unique ID of this state */
+		const FName UniqueName;
+
+		/** Id of this cluster node to speed up internal operations */
+		const FName ClusterNodeId;
+
+		/** DC cluster manager to speed up internal operations */
+		IDisplayClusterClusterManager* ClusterMgr = nullptr;
+
+		/**
+		 * Main access guard for this state instance, wrapped in
+		 * a shared reference to ensure thread-safe sharing.
+		 */
+		TSharedRef<FCriticalSection> StateCS;
+	};
+
+
+	/**
+	 * Custom state data class
+	 * 
+	 * Provides general control over local/cluster state data. Also manages game and render thread data.
+	 */
+	template <typename TDataType>
+	class TCustomStateData
+		: public FCustomStateBase
+	{
+	public:
+	
+		template <typename... TArgs>
+		TCustomStateData(const FName& InUniqueName, const FName& InNodeId, TArgs&&... Args)
+			: FCustomStateBase(InUniqueName, InNodeId)
+			, ClusterStatesData(MakeShared<TMap<FName, FCustomStateDataHolder<TDataType>>>())
+		{
+			// Make sure local data of this node always exists
+			ClusterStatesData->Emplace(GetNodeId(), Forward<TArgs>(Args)...);
+		}
+		
+		virtual ~TCustomStateData() = default;
+
+	public:
+
+		/** Returns local state data bound to current thread and frame */
+		const TDataType& GetData() const
+		{
+			FScopeLock Lock(&GetCritSec());
+			return GetThreadData();
+		}
+
+		/** Sets new state data for the next frame. */
+		template <typename TValue>
+		requires std::convertible_to<TValue, TDataType>
+		void SetData(TValue&& NewData)
+		{
+			FScopeLock Lock(&GetCritSec());
+			if (IsSetDataAllowed())
+			{
+				NextFrameData = Forward<TValue>(NewData);
+			}
+		}
+
+	protected:
+
+		/** Ask children whether SetData is allowed for the instance. */
+		virtual bool IsSetDataAllowed() const
+		{
+			return true;
+		}
+
+	protected:
+
+		/** Internal thread sensitive getter of local state data */
+		TDataType& GetThreadData()
+		{
+			return GetThreadDataImpl(ClusterStatesData->FindChecked(GetNodeId()));
+		}
+
+		/** Internal thread sensitive getter of local state data (const version) */
+		const TDataType& GetThreadData() const
+		{
+			return GetThreadDataImpl(ClusterStatesData->FindChecked(GetNodeId()));
+		}
+
+		/** Internal thread sensitive getter of a specific node state data */
+		TDataType& GetThreadData(const FName& NodeId, bool bCreateIfNotExists = false)
+		{
+			FCustomStateDataHolder<TDataType>* const Found = bCreateIfNotExists ? &ClusterStatesData->FindOrAdd(NodeId) : ClusterStatesData->Find(NodeId);
+			return GetThreadDataImpl(Found ? *Found : ClusterStatesData->FindChecked(GetNodeId()));
+		}
+
+		/** Internal thread sensitive getter of a specific node state data (const version) */
+		const TDataType& GetThreadData(const FName& NodeId) const
+		{
+			const FCustomStateDataHolder<TDataType>* const Found = ClusterStatesData->Find(NodeId);
+			return GetThreadDataImpl(Found ? *Found : ClusterStatesData->FindChecked(GetNodeId()));
+		}
+
+		/** Returns node IDs available in the storage */
+		TSet<FName> GetNodes() const
+		{
+			TSet<FName> NodeIds;
+			ClusterStatesData->GetKeys(NodeIds);
+			return NodeIds;
+		}
+
+	protected:
+
+		/** Game thread data update */
+		void AdvanceFrameData_GT()
+		{
+			checkSlow(IsInGameThread());
+
+			// Update local value only
+			ClusterStatesData->FindChecked(GetNodeId()).Data = NextFrameData;
+		}
+
+		/** Render thread data update */
+		void AdvanceFrameData_RT()
+		{
+			checkSlow(IsInGameThread());
+
+			// Update render data for every node
+			for (TPair<FName, FCustomStateDataHolder<TDataType>>& NodeStateData : *ClusterStatesData)
+			{
+				// Current game thread values will replace the render thread ones on the rendering thread
+				ExecuteOnRenderThread(
+					[
+						Data   = ClusterStatesData,
+						NodeId = NodeStateData.Key,
+						Value  = MoveTemp(NodeStateData.Value.Data),
+						CritSecRef = GetCritSecRef()
+					]() mutable
+					{
+						FScopeLock Lock(&CritSecRef.Get());
+						Data->FindChecked(NodeId).DataRT = MoveTemp(Value);
+					});
+			}
+		}
+
+	private:
+
+		/** Internal helper to get data based on the current thread (for const and non-const callers) */
+		template <typename THolderType>
+		static auto& GetThreadDataImpl(THolderType& Holder)
+		{
+			if (IsInRenderingThread())
+			{
+				return Holder.DataRT;
+			}
+
+			checkfSlow(IsInGameThread(), TEXT("Accessed state data from AnyThread"));
+			return Holder.Data;
+		}
+
+	private:
+
+		/**
+		 * Auxiliary structure that holds all necessary data of a single state
+		 */
+		template <typename TDataType>
+		struct FCustomStateDataHolder
+		{
+			/** Game thread data */
+			TDataType Data { };
+
+			/** Render thread data */
+			TDataType DataRT { };
+		};
+
+		/**
+		 * Holds data of this state on all nodes. Its content type depends on the concrete
+		 * state implementation. Wrapped into a TSharedPtr to ensure safe access across threads.
+		 */
+		TSharedRef<TMap<FName, FCustomStateDataHolder<TDataType>>> ClusterStatesData;
+
+		/** Data to use on the next frame */
+		TDataType NextFrameData{ };
+	};
+
+
+	/**
+	 * Factory method implementation for custom state templates
+	 *
+	 * Provides a simple way to instantiate custom states of any types. It also performs
+	 * some additional checks, and automatically registers newly created states.
+	 */
+	template <typename TDataType>
+	struct TCustomStateFactory
+	{
+		/**
+		 * Custom states factory method
+		 */
+		template <typename... TArgs>
+		requires std::derived_from<TDataType, FCustomStateBase>
+		static TSharedPtr<TDataType> Create(const FName& InUniqueName, TArgs&&... Args)
+		{
+			// Invalid names not allowed
+			if (InUniqueName.IsNone())
+			{
+				return nullptr;
+			}
+
+			// Make sure nDisplay is loaded
+			if (!IDisplayCluster::IsAvailable())
+			{
+				return nullptr;
+			}
+
+			// Make sure cluster manager is available
+			IDisplayClusterClusterManager* const ClusterMgr = IDisplayCluster::Get().GetClusterMgr();
+			if (!ClusterMgr)
+			{
+				return nullptr;
+			}
+
+			// Make sure this name is not used already
+			const bool bAlreadyExists = ClusterMgr->IsCustomStateRegistered(InUniqueName);
+			if (bAlreadyExists)
+			{
+				return nullptr;
+			}
+
+			// Instantiate and register
+			TDataType* const NewInstance = new TDataType(InUniqueName, *ClusterMgr->GetNodeId(), Forward<TArgs>(Args)...);
+			TSharedPtr<TDataType> NewState = MakeShareable<TDataType>(NewInstance);
+			TSharedPtr<IDisplayClusterCustomState> NewStatePtr = StaticCastSharedPtr<IDisplayClusterCustomState>(NewState);
+			if (!ClusterMgr->RegisterCustomState(InUniqueName, NewStatePtr))
+			{
+				return nullptr;
+			}
+
+			return NewState;
+		}
+	};
+
+
+	/**
+	 * Type ID generator
+	 * 
+	 * It's an auxiliary class that provides a way to differentiate various state implementations
+	 * (template specialization) in runtime without RTTI. It converts the text based function signatures
+	 * into FName which is quite good for serialization and comparison.
+	 */
+	template<typename TType>
+	class TCustomStateTypeId
+	{
+	public:
+
+		/** Type getter */
+		static const FName& GetTypeId()
+		{
+			static const FName TypeId = GenerateTypeId();
+			return TypeId;
+		}
+
+	private:
+
+		/**
+		 * Internal type name generator. Extracts the actual TType from the full function signature.
+		 * 
+		 * For example, full signature might be like this:
+		 * {Signature=L"class FName __cdecl UE::nDisplay::Private::TCustomStateTypeId<class TDistributedCustomState<struct FMyStruct> >::GenerateTypeId(void)" }
+		 * 
+		 * The type ID would be:
+		 * class TDistributedCustomState<struct FMyStruct>
+		 */
+		static FName GenerateTypeId()
+		{
+#if defined(_MSC_VER)
+			FString Signature(ANSI_TO_TCHAR(__FUNCSIG__));
+#else
+			FString Signature(ANSI_TO_TCHAR(__PRETTY_FUNCTION__));
+#endif
+
+			/** Auxiliary substring extractor */
+			auto SubString = [&Signature](const FString& Left, const FString& Right)
+				{
+					const int32 FoundLeft  = Signature.Find(Left, ESearchCase::CaseSensitive, ESearchDir::FromStart);
+					const int32 FoundRight = Signature.Find(Right, ESearchCase::CaseSensitive, ESearchDir::FromEnd);
+
+					checkSlow(FoundLeft  != INDEX_NONE);
+					checkSlow(FoundRight != INDEX_NONE);
+					checkSlow(FoundRight > FoundLeft);
+
+					if (FoundRight != INDEX_NONE)
+					{
+						Signature = Signature.Left(FoundRight);
+					}
+
+					if (FoundLeft != INDEX_NONE)
+					{
+						Signature = Signature.Right(Signature.Len() - (FoundLeft + Left.Len()));
+					}
+
+					Signature.TrimStartAndEndInline();
+				};
+
+			// Remove unnecessary pieces from the original signature
+			SubString(TEXT("TCustomStateTypeId"), TEXT("::GenerateTypeId"));
+			SubString(TEXT("<"), TEXT(">"));
+
+			return *Signature;
+		}
+	};
+}
diff --git a/Engine/Plugins/Runtime/nDisplay/Source/DisplayCluster/Public/Cluster/CustomStates/DisplayClusterCustomStateDistributed.h b/Engine/Plugins/Runtime/nDisplay/Source/DisplayCluster/Public/Cluster/CustomStates/DisplayClusterCustomStateDistributed.h
new file mode 100644
--- /dev/null
+++ b/Engine/Plugins/Runtime/nDisplay/Source/DisplayCluster/Public/Cluster/CustomStates/DisplayClusterCustomStateDistributed.h
@@ -0,0 +1,277 @@
+// Copyright Epic Games, Inc. All Rights Reserved.
+
+#pragma once
+
+#include "Cluster/CustomStates/DisplayClusterCustomStateBase.h"
+
+#include "Misc/ScopeLock.h"
+#include "Serialization/Archive.h"
+#include "Serialization/MemoryReader.h"
+
+#include "CoreGlobals.h"
+
+
+/**
+ * Distributed custom state
+ *
+ * This state enables data exchange in any direction within a cluster.
+ * By default, any cluster node can access the values of any other node.
+ * An additional API is provided to customize data propagation paths
+ * in arbitrary ways.
+ *
+ * GetData() and GetData(NodeId) are thread-sensitive. The data is
+ * maintained in separate copies for the game and render threads.
+ * Access from other threads is not supported.
+ *
+ * Distributed states are automatically registered for replication
+ * during creation (see the Create() call in the examples below).
+ * Once registered, replication within the cluster begins immediately.
+ * To stop replication, the state must be explicitly unregistered.
+ *
+ * This state type has no leader or explicit data source; it is a
+ * generic mechanism for data exchange. Data propagation is performed
+ * synchronously to preserve determinism, which may introduce update
+ * latency.
+ *
+ * When the state is modified on the game thread, the expected latency
+ * is one frame: a value set on frame N becomes visible via GetData()
+ * on frame N + 1. If the value is updated before state synchronization,
+ * no latency is introduced.
+ *
+ * When modified on the render thread, the latency may be 2–3 frames,
+ * as the value is first propagated to the game thread and then passed
+ * back to the render thread.
+ *
+ * This class uses internal data serialization. Any custom TDataType
+ * must provide its own serialization functions. Refer to the examples
+ * below for details.
+ *
+ *        +---------------------+       +---------------------+
+ *        ¦   NodeA::Value=5    ¦       ¦   NodeB::Value=2    ¦
+ *        +---------------------¦       +---------------------¦
+ *  +---- ¦ NodeA: Value=5, R/W ¦ ----> ¦ NodeA: Value=5, R/O ¦
+ *  ¦     ¦ NodeB: Value=2, R/O ¦ <---- ¦ NodeB: Value=2, R/W ¦ ----+
+ *  ¦ +-> ¦ NodeC: Value=8, R/O ¦       ¦ NodeC: Value=8, R/O ¦ <-+ ¦
+ *  ¦ ¦   +---------------------+       +---------------------+   ¦ ¦
+ *  ¦ ¦                                                           ¦ ¦
+ *  ¦ ¦                  +---------------------+                  ¦ ¦
+ *  ¦ ¦                  ¦   NodeC::Value=8    ¦                  ¦ ¦
+ *  ¦ ¦                  +---------------------¦                  ¦ ¦
+ *  +-+----------------> ¦ NodeA: Value=5, R/O ¦                  ¦ ¦
+ *    ¦                  ¦ NodeB: Value=2, R/O ¦ <----------------+-+
+ *    +----------------- ¦ NodeC: Value=8, R/W ¦ -----------------+
+ *                       +---------------------+
+ *
+ * EXAMPLE - Base type
+ * ------------------------
+ * using FMyState = TDistributedCustomState<float>;
+ *
+ * const FName MyStateName = TEXT("MyStateFloat");  // The name must be unique
+ * TSharedPtr<FMyState> MyState = FMyState::Create(MyStateName);  // Don't forget to nullptr-check.
+ * or
+ * TSharedPtr<FMyState> MyState = FMyState::Create(MyStateName, 0.72f);  // With initialization value
+ *
+ * MyState->SetData(0.75f);  // Set new value
+ * MyState->SetData(0.85f);  // It can be called multiple times, the last one will be used on the next frame
+ *
+ * const TSet<FName> NodeIds = MyState->GetAvailableNodes();
+ * for (const FName& NodeId : NodeIds)
+ * {
+ *     const int32 Value = MyState->GetData(NodeId);
+ *     DoSomething(NodeId, Value);
+ * }
+ *
+ * IDisplayCluster::Get().GetClusterMgr()->UnregisterCustomState(MyStateName);  // Unregister it from replication on this cluster node
+ * IDisplayCluster::Get().GetClusterMgr()->RegisterCustomState(MyStateName, MyState);  // Register again
+ *
+ * MyState->SetCustomUpstreamsEnabled(true); // Enable custom upstream configuration
+ * MyState->SetUpstreams({ TEXT("Node_0"), TEXT("Node_1") }); // Receive updates from Node_0 and Node_1 only
+ *
+ *
+ * EXAMPLE - Custom type
+ * ------------------------
+ * // Any custom types
+ * struct FMyStruct
+ * {
+ *    int32 Items = 0;
+ *    float Progress = 0.f;
+ *    bool  bIsActive = false;
+ * };
+ *
+ * // With custom serialization
+ * inline FArchive& operator<<(FArchive& Ar, FMyStruct& MyStruct)
+ * {
+ *    Ar << MyStruct.Items;
+ *    Ar << MyStruct.Progress;
+ *    Ar << MyStruct.bIsActive;
+ *    return Ar;
+ * }
+ *
+ * using FMyState = TDistributedCustomState<FMyStruct>;
+ *
+ * const FName MyStateName = TEXT("MyStateStruct"); // The name must be unique
+ * TSharedPtr<FMyState> MyState = FMyState::Create(MyStateName); // Don't forget to nullptr-check. It's also auto-registered
+ * or
+ * TSharedPtr<FMyState> MyState = FMyState::Create(MyStateName, FMyStruct{ 3, 2, 0.1f, true });
+ *
+ * MyState->SetData({ 2, 0.3f, true }); // Set new value
+ * MyState->SetData({ 9, 0.5f, true }); // It can be called multiple times, the last one will be used on the next frame
+ *
+ * const TSet<FName> NodeIds = MyState->GetAvailableNodes();
+ * for (const FName& NodeId : NodeIds)
+ * {
+ *     FMyStruct Value = MyState->GetData(NodeId);
+ *     DoSomething(NodeId, Value);
+ * }
+ *
+ * FMyStruct CurrentFrameState = MyState->GetData(); // Get data on current frame (game/render thread sensitive)
+ * CurrentFrameState = { 10, 1.f, false };
+ * MyState->SetData(MoveTemp(CurrentFrameState)); // Move semantics can be useful for complex types
+ *
+ */
+template <typename TDataType>
+class TDistributedCustomState
+	: public UE::nDisplay::Private::TCustomStateData<TDataType>
+	, public UE::nDisplay::Private::TCustomStateFactory<TDistributedCustomState<TDataType>>
+{
+	// Access to the private constructor
+	friend UE::nDisplay::Private::TCustomStateFactory<TDistributedCustomState<TDataType>>;
+
+	using ThisType = TDistributedCustomState<TDataType>;
+
+protected:
+
+	// Bring some parent names explicitly because of the CRTP
+	using UE::nDisplay::Private::FCustomStateBase::GetCritSec;
+	using UE::nDisplay::Private::FCustomStateBase::GetNodeId;
+	using UE::nDisplay::Private::TCustomStateData<TDataType>::AdvanceFrameData_GT;
+	using UE::nDisplay::Private::TCustomStateData<TDataType>::AdvanceFrameData_RT;
+	using UE::nDisplay::Private::TCustomStateData<TDataType>::GetNodes;
+	using UE::nDisplay::Private::TCustomStateData<TDataType>::GetThreadData;
+
+public:
+
+	// Expose some parent names explicitly because of name hiding
+	using UE::nDisplay::Private::TCustomStateData<TDataType>::GetData;
+
+private:
+
+	/** Private constructor. Only factory method can instantiate this class. */
+	template <typename... TArgs>
+	TDistributedCustomState(const FName& InUniqueName, const FName& InNodeId, TArgs&&... Args)
+		: UE::nDisplay::Private::TCustomStateData<TDataType>(InUniqueName, InNodeId, Forward<TArgs>(Args)...)
+		, UE::nDisplay::Private::TCustomStateFactory<TDistributedCustomState<TDataType>>()
+	{ }
+
+public:
+
+	virtual ~TDistributedCustomState() = default;
+
+public:
+
+	//~ Begin IDisplayClusterCustomState
+
+	virtual FName GetType() const override
+	{
+		return UE::nDisplay::Private::TCustomStateTypeId<ThisType>::GetTypeId();
+	}
+
+	virtual void Serialize(FArchive& Ar) override
+	{
+		checkSlow(IsInGameThread());
+		FScopeLock Lock(&GetCritSec());
+		Ar << GetThreadData();
+	}
+
+	virtual void AdvanceFrame() override
+	{
+		checkSlow(IsInGameThread());
+		FScopeLock Lock(&GetCritSec());
+
+		// Update render thread data
+		AdvanceFrameData_RT();
+
+		// Distributed states always update their game thread data because
+		// it's allowed for them to have individual values.
+		AdvanceFrameData_GT();
+	}
+
+	virtual bool HasCustomUpstreamConfiguration() const override
+	{
+		FScopeLock Lock(&GetCritSec());
+		return bUseCustomUpstreamConfiguration;
+	}
+
+	virtual TSet<FName> GetUpstreams() const
+	{
+		FScopeLock Lock(&GetCritSec());
+		return UpstreamNodes;
+	}
+
+	virtual void Update(const TMap<FName, TArray<uint8>>& ClusterStates) override
+	{
+		checkSlow(IsInGameThread());
+
+		// Ignore invalid input
+		if (ClusterStates.Num() <= 0)
+		{
+			return;
+		}
+
+		FScopeLock Lock(&GetCritSec());
+
+		// Deserialize and store new data
+		for (const TPair<FName, TArray<uint8>>& NodeState : ClusterStates)
+		{
+			FMemoryReader DataReader(NodeState.Value);
+			// Always create new space for the nodes that have just joined
+			constexpr bool bShouldCreateIfNotExists = true;
+			DataReader << GetThreadData(NodeState.Key, bShouldCreateIfNotExists);
+		}
+	}
+
+	//~ End IDisplayClusterCustomState
+
+public:
+	
+	/** Returns state data of a specific node. The data is bound to current thread and frame. */
+	const TDataType& GetData(const FName& NodeId) const
+	{
+		FScopeLock Lock(&GetCritSec());
+		return GetThreadData(NodeId);
+	}
+
+	/** Returns node IDs available in the storage  */
+	TSet<FName> GetAvailableNodes() const
+	{
+		return GetNodes();
+	}
+
+public:
+
+	/** Enable/disable custom upstream configuration */
+	void SetCustomUpstreamsEnabled(bool bEnabled)
+	{
+		FScopeLock Lock(&GetCritSec());
+		bUseCustomUpstreamConfiguration = bEnabled;
+	}
+
+	/**
+	 * Configure upstream nodes that should propagate their updates to this state instance
+	 * 
+	 * @param NewUpstreamNodes - new set of cluster node IDs
+	 */
+	void SetUpstreams(const TSet<FName>& NewUpstreamNodes)
+	{
+		FScopeLock Lock(&GetCritSec());
+		UpstreamNodes = NewUpstreamNodes;
+	}
+
+private:
+
+	/** Whether custom upstream configuration is enabled */
+	bool bUseCustomUpstreamConfiguration = false;
+
+	/** Custom upstream cluster node IDs */
+	TSet<FName> UpstreamNodes;
+};
diff --git a/Engine/Plugins/Runtime/nDisplay/Source/DisplayCluster/Public/Cluster/CustomStates/DisplayClusterCustomStateShared.h b/Engine/Plugins/Runtime/nDisplay/Source/DisplayCluster/Public/Cluster/CustomStates/DisplayClusterCustomStateShared.h
new file mode 100644
--- /dev/null
+++ b/Engine/Plugins/Runtime/nDisplay/Source/DisplayCluster/Public/Cluster/CustomStates/DisplayClusterCustomStateShared.h
@@ -0,0 +1,199 @@
+// Copyright Epic Games, Inc. All Rights Reserved.
+
+#pragma once
+
+#include "Cluster/CustomStates/DisplayClusterCustomStateBase.h"
+#include "Cluster/IDisplayClusterClusterManager.h"
+#include "Misc/ScopeLock.h"
+#include "Serialization/Archive.h"
+#include "Serialization/MemoryReader.h"
+
+#include "CoreGlobals.h"
+
+
+/**
+ * Shared custom state.
+ *
+ * "Shared" means there is a single source of truth. This implementation
+ * always propagates the state data from the primary node to all other
+ * nodes. Regardless of which node GetData() is called on, the returned
+ * data will always match the primary node's value. Consequently, SetData()
+ * is ignored on all non-primary nodes.
+ *
+ * GetData() is thread-sensitive. The state maintains separate copies
+ * for the game and render threads. Access from other threads is not
+ * supported.
+ *
+ * Shared states are automatically registered for replication
+ * during creation (see the Create() call in the examples below).
+ * Once registered, replication within the cluster begins immediately.
+ * To stop replication, the state must be explicitly unregistered.
+ *
+ * Data propagation is performed synchronously to preserve determinism,
+ * which may introduce update latency. When the state is modified on
+ * the game thread, the expected latency is one frame: a value set on
+ * frame N becomes visible via GetData() on frame N + 1. If the value
+ * is updated before state synchronization, no latency is introduced.
+ *
+ * When updated on the render thread, latency may be 2–3 frames, as the
+ * value is first propagated to the game thread and then back to the
+ * render thread.
+ *
+ * This class uses internal data serialization. Any custom TDataType
+ * must provide its own serialization functions. Refer to the examples
+ * below for details.
+ *
+ *    +----------------+                    +--------------------+
+ *    ¦  Primary node  ¦                    ¦   Secondary node   ¦
+ *    +----------------¦   Replication      +--------------------¦
+ *    ¦  Value=5, R/W  ¦  --------------->  ¦ Value=5, Read only ¦
+ *    +----------------+             ¦      +--------------------+
+ *                                   ¦
+ *                                   ¦
+ *                                   ¦      +--------------------+
+ *                                   ¦      ¦   Secondary node   ¦
+ *                                   ¦      +--------------------¦
+ *                                   +--->  ¦ Value=5, Read only ¦
+ *                                          +--------------------+
+ *
+ * EXAMPLE - Base type
+ * ------------------------
+ * using FMyState = TSharedCustomState<float>;
+ *
+ * const FName MyStateName = TEXT("MyStateFloat");  // The name must be unique
+ * TSharedPtr<FMyState> MyState = FMyState::Create(MyStateName);  // Don't forget to nullptr-check.
+ * or
+ * TSharedPtr<FMyState> MyState = FMyState::Create(MyStateName, 0.72f);  // With initialization value
+ *
+ * MyState->SetData(0.75f);  // Set new value
+ * MyState->SetData(0.85f);  // It can be called multiple times, the last one will be used on the next frame
+ *
+ * IDisplayCluster::Get().GetClusterMgr()->UnregisterCustomState(MyStateName);  // Unregister it from replication on this cluster node
+ * IDisplayCluster::Get().GetClusterMgr()->RegisterCustomState(MyStateName, MyState);  // Register again
+ *
+ *
+ * EXAMPLE - Custom type
+ * ------------------------
+ * // Any custom types
+ * struct FMyStruct
+ * {
+ *    int32 Items = 0;
+ *    float Progress = 0.f;
+ *    bool  bIsActive = false;
+ * };
+ *
+ * // With custom serialization
+ * inline FArchive& operator<<(FArchive& Ar, FMyStruct& MyStruct)
+ * {
+ *    Ar << MyStruct.Items;
+ *    Ar << MyStruct.Progress;
+ *    Ar << MyStruct.bIsActive;
+ *    return Ar;
+ * }
+ *
+ * using FMyState = TSharedCustomState<FMyStruct>;
+ *
+ * const FName MyStateName = TEXT("MyStateStruct"); // The name must be unique
+ * TSharedPtr<FMyState> MyState = FMyState::Create(MyStateName); // Don't forget to nullptr-check. It's also auto-registered
+ *
+ * MyState->SetData({ 2, 0.3f, true }); // Set new value
+ * MyState->SetData({ 9, 0.5f, true }); // It can be called multiple times, the last one will be used on the next frame
+ *
+ * FMyStruct CurrentFrameState = MyState->GetData(); // Get data on current frame (game/render thread sensitive)
+ * CurrentFrameState = { 10, 1.f, false };
+ * MyState->SetData(MoveTemp(CurrentFrameState)); // Move semantics can be useful for complex types
+ *
+ */
+template <typename TDataType>
+class TSharedCustomState
+	: public UE::nDisplay::Private::TCustomStateData<TDataType>
+	, public UE::nDisplay::Private::TCustomStateFactory<TSharedCustomState<TDataType>>
+{
+	// Access to the private constructor
+	friend UE::nDisplay::Private::TCustomStateFactory<TSharedCustomState<TDataType>>;
+
+	using ThisType = TSharedCustomState<TDataType>;
+
+	// Bring some parent names explicitly because of the CRTP
+	using UE::nDisplay::Private::FCustomStateBase::GetCritSec;
+	using UE::nDisplay::Private::FCustomStateBase::GetClusterMgr;
+	using UE::nDisplay::Private::FCustomStateBase::GetNodeId;
+	using UE::nDisplay::Private::TCustomStateData<TDataType>::AdvanceFrameData_GT;
+	using UE::nDisplay::Private::TCustomStateData<TDataType>::AdvanceFrameData_RT;
+	using UE::nDisplay::Private::TCustomStateData<TDataType>::GetThreadData;
+
+private:
+
+	/** Private constructor. Only factory method can instantiate this class. */
+	template <typename... TArgs>
+	TSharedCustomState(const FName& InUniqueName, const FName& InNodeId, TArgs&&... Args)
+		: UE::nDisplay::Private::TCustomStateData<TDataType>(InUniqueName, InNodeId, Forward<TArgs>(Args)...)
+		, UE::nDisplay::Private::TCustomStateFactory<TSharedCustomState<TDataType>>()
+	{ }
+
+public:
+
+	virtual ~TSharedCustomState() = default;
+
+public:
+
+	//~ Begin IDisplayClusterCustomState
+
+	virtual FName GetType() const override
+	{
+		return UE::nDisplay::Private::TCustomStateTypeId<ThisType>::GetTypeId();
+	}
+
+	virtual bool ShouldPropagate() const override
+	{
+		// Propagate from primary nodes only
+		return GetClusterMgr() ? GetClusterMgr()->IsPrimary() : false;
+	}
+
+	virtual void Serialize(FArchive& Ar) override
+	{
+		checkSlow(IsInGameThread());
+		FScopeLock Lock(&GetCritSec());
+		Ar << GetThreadData();
+	}
+
+	virtual void AdvanceFrame() override
+	{
+		checkSlow(IsInGameThread());
+		FScopeLock Lock(&GetCritSec());
+
+		// Update render thread data
+		AdvanceFrameData_RT();
+
+		// Update game thread data on primary nodes only. That's because shared custom
+		// staes propagate update from p-nodes only.
+		if (GetClusterMgr() && GetClusterMgr()->IsPrimary())
+		{
+			AdvanceFrameData_GT();
+		}
+	}
+
+	virtual void Update(const TMap<FName, TArray<uint8>>& ClusterStates) override
+	{
+		FScopeLock Lock(&GetCritSec());
+
+		// It's expected to have a single state from the primary node,
+		// but still possible the input would be empty
+		if (ClusterStates.Num() > 0)
+		{
+			auto It = ClusterStates.CreateConstIterator();
+			FMemoryReader DataReader(It.Value());
+			DataReader << GetThreadData();
+		}
+	}
+
+	//~ End IDisplayClusterCustomState
+
+protected:
+
+	virtual bool IsSetDataAllowed() const override final
+	{
+		// Value updates are allowed only on primary nodes
+		return GetClusterMgr() && GetClusterMgr()->IsPrimary();
+	}
+};
diff --git a/Engine/Plugins/Runtime/nDisplay/Source/DisplayCluster/Public/Cluster/CustomStates/IDisplayClusterCustomState.h b/Engine/Plugins/Runtime/nDisplay/Source/DisplayCluster/Public/Cluster/CustomStates/IDisplayClusterCustomState.h
new file mode 100644
--- /dev/null
+++ b/Engine/Plugins/Runtime/nDisplay/Source/DisplayCluster/Public/Cluster/CustomStates/IDisplayClusterCustomState.h
@@ -0,0 +1,111 @@
+// Copyright Epic Games, Inc. All Rights Reserved.
+
+#pragma once
+
+#include "CoreMinimal.h"
+
+class FArchive;
+
+
+/**
+ * Custom state interface
+ */
+class IDisplayClusterCustomState
+{
+public:
+
+	virtual ~IDisplayClusterCustomState() = default;
+
+public:
+
+	/**
+	 * Returns custom state's unique name
+	 * 
+	 * @return Unique state name
+	 */
+	virtual FName GetName() const = 0;
+
+	/**
+	 * Returns type identificator
+	 * 
+	 * Custom states are implemented using C++ templates, which allows them to carry
+	 * arbitrary payload types. Since RTTI is not used, it is not possible to determine
+	 * at runtime whether two states with the same name also share the same payload type.
+	 * If such a type mismatch occurs, serialization will fail and the application will
+	 * likely crash. To prevent this, each state is associated with a unique type ID.
+	 * 
+	 * @return Type ID
+	 */
+	virtual FName GetType() const
+	{
+		return NAME_None;
+	}
+
+	/**
+	 * Whether this state should propagate its data to any other cluster nodes
+	 * 
+	 * @return True if we need to propagate it within a cluster
+	 */
+	virtual bool ShouldPropagate() const
+	{
+		return true;
+	}
+
+	/**
+	 * Serialization
+	 * 
+	 * It's used to serialize data of a custom state into an output buffer
+	 * 
+	 * @param Ar - An archive to write data in
+	 */
+	virtual void Serialize(FArchive& Ar) = 0;
+
+	/**
+	 * Whether this state requires some custom update path. If returns true,
+	 * then GetUpstreams() will be called next in order to configure data senders.
+	 * 
+	 * @return - True if has its own list of upstream nodes
+	 */
+	virtual bool HasCustomUpstreamConfiguration() const
+	{
+		return false;
+	}
+
+	/**
+	 * Returns set of node IDs that this custom state is expecting to get updates from.
+	 * Based on the data propagation technique, this may disable any data senders, even all of them.
+	 * It can also be used to reduce traffic while replicationg huge blobs.
+	 * 
+	 * @return - A set of cluster node IDs that should propagate their updates to this instance
+	 */
+	virtual TSet<FName> GetUpstreams() const
+	{
+		return { };
+	}
+
+	/**
+	 * Progress to the next frame
+	 */
+	virtual void AdvanceFrame()
+	{ }
+
+	/**
+	 * Deserialization
+	 * 
+	 * It's used to deserialize and process post-synchronization response data,
+	 * and probably update internal data.
+	 * 
+	 * @param ClusterStates - ClusterNodeId-to-StateData blobs
+	 */
+	virtual void Update(const TMap<FName, TArray<uint8>>& ClusterStates) = 0;
+
+	/**
+	 * Lock this state to prevent any changes
+	 */
+	virtual void Lock() const = 0;
+
+	/**
+	 * Enable any changes after Lock() call
+	 */
+	virtual void Unlock() const = 0;
+};
diff --git a/Engine/Plugins/Runtime/nDisplay/Source/DisplayCluster/Public/Cluster/IDisplayClusterClusterManager.h b/Engine/Plugins/Runtime/nDisplay/Source/DisplayCluster/Public/Cluster/IDisplayClusterClusterManager.h
--- a/Engine/Plugins/Runtime/nDisplay/Source/DisplayCluster/Public/Cluster/IDisplayClusterClusterManager.h
+++ b/Engine/Plugins/Runtime/nDisplay/Source/DisplayCluster/Public/Cluster/IDisplayClusterClusterManager.h
@@ -15,6 +15,7 @@ class IDisplayClusterClusterEventBinaryListener;
 class IDisplayClusterClusterEventJsonListener;
 class IDisplayClusterClusterEventListener;
 class IDisplayClusterClusterSyncObject;
+class IDisplayClusterCustomState;
 
 
 DECLARE_MULTICAST_DELEGATE_OneParam(FOnClusterEventJson, const FDisplayClusterClusterEventJson& /* Event */);
@@ -64,6 +65,40 @@ public:
 	/** Unregisters synchronization object. */
 	virtual void UnregisterSyncObject(IDisplayClusterClusterSyncObject* SyncObj) = 0;
 
+	/**
+	 * Registers custom state for auto-replication
+	 * 
+	 * @param UniqueStateName - Unique name to distinguish different states
+	 * @param StateVariable   - A variable that references the state
+	 * 
+	 * @return True if registered successfully, otherwise false (has been registered already)
+	 */
+	virtual bool RegisterCustomState(const FName& UniqueStateName, TSharedPtr<IDisplayClusterCustomState>& StateVariable) = 0;
+
+	/**
+	 * Unregisters specific custom state
+	 *
+	 * @param UniqueStateName - Unique name to distinguish different states
+	 *
+	 * @return True if unregistered successfully, otherwise false (not found)
+	 */
+	virtual bool UnregisterCustomState(const FName& UniqueStateName) = 0;
+
+	/**
+	 * Returns the names of all registered states
+	 *
+	 * @return Set of custom state names currently registered
+	 */
+	virtual TSet<FName> GetRegisteredCustomStateNames() const = 0;
+
+	/**
+	 * Checks whether a custom state with the specified name has already been registered
+	 *
+	 * @return True if such state exists
+	 */
+	virtual bool IsCustomStateRegistered(const FName& UniqueStateName) const = 0;
+
+
 	/** Creates a general purpose barriers interface. */
 	virtual TSharedRef<IDisplayClusterGenericBarriersClient> CreateGenericBarriersClient() = 0;
 
