[nDisplay] Binary replication of some internal data
- Use binary replication for the sync objects (IDisplayClusterClusterSyncObject), including Livelink.
- Use binary replication for native input data

#rnx

#Synced-CL 0

diff --git a/Engine/Plugins/Runtime/LiveLinkOvernDisplay/Source/LiveLinkOverNDisplay/Private/NDisplayLiveLinkSubjectReplicator.cpp b/Engine/Plugins/Runtime/LiveLinkOvernDisplay/Source/LiveLinkOverNDisplay/Private/NDisplayLiveLinkSubjectReplicator.cpp
--- a/Engine/Plugins/Runtime/LiveLinkOvernDisplay/Source/LiveLinkOverNDisplay/Private/NDisplayLiveLinkSubjectReplicator.cpp
+++ b/Engine/Plugins/Runtime/LiveLinkOvernDisplay/Source/LiveLinkOverNDisplay/Private/NDisplayLiveLinkSubjectReplicator.cpp
@@ -38,32 +38,20 @@ FString FNDisplayLiveLinkSubjectReplicator::GetSyncId() const
 	return SyncId;
 }
 
-FString FNDisplayLiveLinkSubjectReplicator::SerializeToString() const
+void FNDisplayLiveLinkSubjectReplicator::SerializeDC(FArchive& Ar)
 {
-	//Serialization is done on the controller. Payload was created on LiveLinkTicked callback
-	FScopeLock Lock(&PayloadCriticalSection);
-	const FString ThisFramePayload = LiveLinkPayload;
-	return ThisFramePayload;
-}
-
-bool FNDisplayLiveLinkSubjectReplicator::DeserializeFromString(const FString& Str)
-{
-	// On the primary node (controller), we should skip this. The underlying logic
-	// is not supposed to be called on the primary nodes (controllers).
-	if (IDisplayCluster::Get().GetClusterMgr()->IsPrimary())
+	// Don't deserialize on the primary nodes as those are basically the data sources.
+	// Moreover, the underlying logic is not supposed to be called there.
+	if (Ar.IsLoading())
 	{
-		return true;
+		if (IDisplayCluster::Get().GetClusterMgr()->IsPrimary())
+		{
+			return;
+		}
 	}
 
-	//Deserialization is done on each agent
-	TArray<uint8> TempBytes;
-	TempBytes.AddUninitialized(Str.Len());
-	StringToBytes(Str, TempBytes.GetData(), Str.Len());
-
-	FMemoryReader Ar(TempBytes);
+	// Serialize/deserialize internally
 	OnDataSynchronization(Ar);
-
-	return true;
 }
 
 void FNDisplayLiveLinkSubjectReplicator::AddReferencedObjects(FReferenceCollector& Collector)
@@ -84,12 +72,7 @@ void FNDisplayLiveLinkSubjectReplicator::Initialize()
 
 		if (IDisplayCluster::IsAvailable())
 		{
-			if (IDisplayCluster::Get().GetClusterMgr()->IsPrimary())
-			{
-				//If we're a controller, listen to LiveLinkTicked callback to bundle up the SyncObject. Subjects frames will have been processed at that point.
-				LiveLinkClient->OnLiveLinkTicked().AddRaw(this, &FNDisplayLiveLinkSubjectReplicator::OnLiveLinkTicked);
-			}
-			else if (IDisplayCluster::Get().GetClusterMgr()->IsSecondary())
+			if (IDisplayCluster::Get().GetClusterMgr()->IsSecondary())
 			{
 				IDisplayCluster::Get().GetCallbacks().OnDisplayClusterFailoverPrimaryNodeChanged().AddRaw(this, &FNDisplayLiveLinkSubjectReplicator::OnDisplayClusterPrimaryNodeChanged);
 
@@ -120,7 +103,6 @@ void FNDisplayLiveLinkSubjectReplicator::Release()
 
 	if (LiveLinkClient && IModularFeatures::Get().IsModularFeatureAvailable(ILiveLinkClient::ModularFeatureName))
 	{
-		LiveLinkClient->OnLiveLinkTicked().RemoveAll(this);
 		LiveLinkClient->OnLiveLinkSourceRemoved().RemoveAll(this);
 		LiveLinkClient->RemoveSource(LiveLinkSourceGuid);
 	}
@@ -168,16 +150,6 @@ void FNDisplayLiveLinkSubjectReplicator::Deactivate()
 	}
 }
 
-void FNDisplayLiveLinkSubjectReplicator::OnLiveLinkTicked()
-{
-	TArray<uint8> TempBytes;
-	FMemoryWriter Ar(TempBytes, /*bIsPersistent=*/ true);
-	OnDataSynchronization(Ar);
-
-	FScopeLock Lock(&PayloadCriticalSection);
-	LiveLinkPayload = BytesToString(TempBytes.GetData(), TempBytes.Num());
-}
-
 void FNDisplayLiveLinkSubjectReplicator::OnEngineBeginFrame()
 {
 	if (bBecamePrimary)
diff --git a/Engine/Plugins/Runtime/LiveLinkOvernDisplay/Source/LiveLinkOverNDisplay/Private/NDisplayLiveLinkSubjectReplicator.h b/Engine/Plugins/Runtime/LiveLinkOvernDisplay/Source/LiveLinkOverNDisplay/Private/NDisplayLiveLinkSubjectReplicator.h
--- a/Engine/Plugins/Runtime/LiveLinkOvernDisplay/Source/LiveLinkOverNDisplay/Private/NDisplayLiveLinkSubjectReplicator.h
+++ b/Engine/Plugins/Runtime/LiveLinkOvernDisplay/Source/LiveLinkOverNDisplay/Private/NDisplayLiveLinkSubjectReplicator.h
@@ -45,9 +45,7 @@ public:
 	virtual bool IsActive() const override;
 	virtual FString GetSyncId() const override;
 	virtual bool IsDirty() const override { return true; };
-	virtual void ClearDirty() override {};
-	virtual FString SerializeToString() const override;
-	virtual bool DeserializeFromString(const FString& Ar) override;
+	virtual void SerializeDC(FArchive& Ar) override;
 	//~ End IDisplayClusterClusterSyncObject interface
 
 	//~ Begin FGCObject interface
@@ -72,12 +70,6 @@ public:
 
 protected:
 
-	/**
-     * Callback to be notified when LiveLinkClient was ticked so SyncObject can be serialized
-     * @note Only called on Controller
-     */
-	void OnLiveLinkTicked();
-
 	/**
 	 * Hook to beginning of frames to make sure our replicated virtual subjects are the ones enabled
 	 * @note Only called on Agent
@@ -116,12 +108,6 @@ private:
 	/** Cached LiveLinkClient when modular feature is registered */
 	ILiveLinkClient* LiveLinkClient = nullptr;
 
-	//CriticalSection to protect the payload. SerializeToString can be called from other threads than Main Thread
-	mutable FCriticalSection PayloadCriticalSection;
-
-	/** SyncObject uses string to pass around data. This represents binary LiveLinkSubject information converted to string */
-	FString LiveLinkPayload;
-
 	/** List of Subjects that we are replicating across cluster. On Agents, it will actually be added to the LiveLink subject's list */
 	TArray<TObjectPtr<UNDisplayAgentVirtualSubject>> TrackedSubjects;
 
diff --git a/Engine/Plugins/Runtime/nDisplay/Source/DisplayCluster/Private/Cluster/Controller/DisplayClusterClusterNodeCtrlDisabled.cpp b/Engine/Plugins/Runtime/nDisplay/Source/DisplayCluster/Private/Cluster/Controller/DisplayClusterClusterNodeCtrlDisabled.cpp
--- a/Engine/Plugins/Runtime/nDisplay/Source/DisplayCluster/Private/Cluster/Controller/DisplayClusterClusterNodeCtrlDisabled.cpp
+++ b/Engine/Plugins/Runtime/nDisplay/Source/DisplayCluster/Private/Cluster/Controller/DisplayClusterClusterNodeCtrlDisabled.cpp
@@ -66,7 +66,7 @@ EDisplayClusterCommResult FDisplayClusterClusterNodeCtrlDisabled::GetTimeData(do
 	return EDisplayClusterCommResult::NotAllowed;
 }
 
-EDisplayClusterCommResult FDisplayClusterClusterNodeCtrlDisabled::GetObjectsData(const EDisplayClusterSyncGroup InSyncGroup, TMap<FString, FString>& OutObjectsData)
+EDisplayClusterCommResult FDisplayClusterClusterNodeCtrlDisabled::GetObjectsData(EDisplayClusterSyncGroup InSyncGroup, TMap<FString, TArray<uint8>>& OutObjectsData)
 {
 	return EDisplayClusterCommResult::NotAllowed;
 }
@@ -78,7 +78,7 @@ EDisplayClusterCommResult FDisplayClusterClusterNodeCtrlDisabled::GetEventsData(
 	return EDisplayClusterCommResult::NotAllowed;
 }
 
-EDisplayClusterCommResult FDisplayClusterClusterNodeCtrlDisabled::GetNativeInputData(TMap<FString, FString>& OutNativeInputData)
+EDisplayClusterCommResult FDisplayClusterClusterNodeCtrlDisabled::GetNativeInputData(TMap<FString, TArray<uint8>>& OutNativeInputData)
 {
 	return EDisplayClusterCommResult::NotAllowed;
 }
diff --git a/Engine/Plugins/Runtime/nDisplay/Source/DisplayCluster/Private/Cluster/Controller/DisplayClusterClusterNodeCtrlDisabled.h b/Engine/Plugins/Runtime/nDisplay/Source/DisplayCluster/Private/Cluster/Controller/DisplayClusterClusterNodeCtrlDisabled.h
--- a/Engine/Plugins/Runtime/nDisplay/Source/DisplayCluster/Private/Cluster/Controller/DisplayClusterClusterNodeCtrlDisabled.h
+++ b/Engine/Plugins/Runtime/nDisplay/Source/DisplayCluster/Private/Cluster/Controller/DisplayClusterClusterNodeCtrlDisabled.h
@@ -36,9 +36,9 @@ public:
 	virtual EDisplayClusterCommResult WaitForFrameStart() override;
 	virtual EDisplayClusterCommResult WaitForFrameEnd() override;
 	virtual EDisplayClusterCommResult GetTimeData(double& OutDeltaTime, double& OutGameTime, TOptional<FQualifiedFrameTime>& OutFrameTime) override;
-	virtual EDisplayClusterCommResult GetObjectsData(const EDisplayClusterSyncGroup InSyncGroup, TMap<FString, FString>& OutObjectsData) override;
+	virtual EDisplayClusterCommResult GetObjectsData(EDisplayClusterSyncGroup InSyncGroup, TMap<FString, TArray<uint8>>& OutObjectsData) override;
 	virtual EDisplayClusterCommResult GetEventsData(TArray<TSharedPtr<FDisplayClusterClusterEventJson>>& OutJsonEvents, TArray<TSharedPtr<FDisplayClusterClusterEventBinary>>& OutBinaryEvents) override;
-	virtual EDisplayClusterCommResult GetNativeInputData(TMap<FString, FString>& OutNativeInputData) override;
+	virtual EDisplayClusterCommResult GetNativeInputData(TMap<FString, TArray<uint8>>& OutNativeInputData) override;
 	virtual EDisplayClusterCommResult PropagateStatesData(const TArray<uint8>& InLocalStatesData, TArray<uint8>& OutClusterStatesData) override;
 	//~ End IDisplayClusterProtocolClusterSync
 
diff --git a/Engine/Plugins/Runtime/nDisplay/Source/DisplayCluster/Private/Cluster/Controller/DisplayClusterClusterNodeCtrlEditor.cpp b/Engine/Plugins/Runtime/nDisplay/Source/DisplayCluster/Private/Cluster/Controller/DisplayClusterClusterNodeCtrlEditor.cpp
--- a/Engine/Plugins/Runtime/nDisplay/Source/DisplayCluster/Private/Cluster/Controller/DisplayClusterClusterNodeCtrlEditor.cpp
+++ b/Engine/Plugins/Runtime/nDisplay/Source/DisplayCluster/Private/Cluster/Controller/DisplayClusterClusterNodeCtrlEditor.cpp
@@ -96,7 +96,7 @@ EDisplayClusterCommResult FDisplayClusterClusterNodeCtrlEditor::GetTimeData(doub
 	return EDisplayClusterCommResult::Ok;
 }
 
-EDisplayClusterCommResult FDisplayClusterClusterNodeCtrlEditor::GetObjectsData(const EDisplayClusterSyncGroup InSyncGroup, TMap<FString, FString>& OutObjectsData)
+EDisplayClusterCommResult FDisplayClusterClusterNodeCtrlEditor::GetObjectsData(EDisplayClusterSyncGroup InSyncGroup, TMap<FString, TArray<uint8>>& OutObjectsData)
 {
 	// PIE stub
 	return EDisplayClusterCommResult::Ok;
@@ -114,7 +114,7 @@ EDisplayClusterCommResult FDisplayClusterClusterNodeCtrlEditor::GetEventsData(TA
 	return EDisplayClusterCommResult::InternalError;
 }
 
-EDisplayClusterCommResult FDisplayClusterClusterNodeCtrlEditor::GetNativeInputData(TMap<FString, FString>& OutNativeInputData)
+EDisplayClusterCommResult FDisplayClusterClusterNodeCtrlEditor::GetNativeInputData(TMap<FString, TArray<uint8>>& OutNativeInputData)
 {
 	// PIE stub
 	return EDisplayClusterCommResult::Ok;
diff --git a/Engine/Plugins/Runtime/nDisplay/Source/DisplayCluster/Private/Cluster/Controller/DisplayClusterClusterNodeCtrlEditor.h b/Engine/Plugins/Runtime/nDisplay/Source/DisplayCluster/Private/Cluster/Controller/DisplayClusterClusterNodeCtrlEditor.h
--- a/Engine/Plugins/Runtime/nDisplay/Source/DisplayCluster/Private/Cluster/Controller/DisplayClusterClusterNodeCtrlEditor.h
+++ b/Engine/Plugins/Runtime/nDisplay/Source/DisplayCluster/Private/Cluster/Controller/DisplayClusterClusterNodeCtrlEditor.h
@@ -38,9 +38,9 @@ public:
 	virtual EDisplayClusterCommResult WaitForFrameStart() override;
 	virtual EDisplayClusterCommResult WaitForFrameEnd() override;
 	virtual EDisplayClusterCommResult GetTimeData(double& OutDeltaTime, double& OutGameTime, TOptional<FQualifiedFrameTime>& OutFrameTime) override;
-	virtual EDisplayClusterCommResult GetObjectsData(const EDisplayClusterSyncGroup InSyncGroup, TMap<FString, FString>& OutObjectsData) override;
+	virtual EDisplayClusterCommResult GetObjectsData(EDisplayClusterSyncGroup InSyncGroup, TMap<FString, TArray<uint8>>& OutObjectsData) override;
 	virtual EDisplayClusterCommResult GetEventsData(TArray<TSharedPtr<FDisplayClusterClusterEventJson>>& OutJsonEvents, TArray<TSharedPtr<FDisplayClusterClusterEventBinary>>& OutBinaryEvents) override;
-	virtual EDisplayClusterCommResult GetNativeInputData(TMap<FString, FString>& OutNativeInputData) override;
+	virtual EDisplayClusterCommResult GetNativeInputData(TMap<FString, TArray<uint8>>& OutNativeInputData) override;
 	virtual EDisplayClusterCommResult PropagateStatesData(const TArray<uint8>& InLocalStatesData, TArray<uint8>& OutClusterStatesData) override;
 	//~ End IDisplayClusterProtocolClusterSync
 
diff --git a/Engine/Plugins/Runtime/nDisplay/Source/DisplayCluster/Private/Cluster/Controller/DisplayClusterClusterNodeCtrlMain.cpp b/Engine/Plugins/Runtime/nDisplay/Source/DisplayCluster/Private/Cluster/Controller/DisplayClusterClusterNodeCtrlMain.cpp
--- a/Engine/Plugins/Runtime/nDisplay/Source/DisplayCluster/Private/Cluster/Controller/DisplayClusterClusterNodeCtrlMain.cpp
+++ b/Engine/Plugins/Runtime/nDisplay/Source/DisplayCluster/Private/Cluster/Controller/DisplayClusterClusterNodeCtrlMain.cpp
@@ -287,7 +287,7 @@ EDisplayClusterCommResult FDisplayClusterClusterNodeCtrlMain::GetTimeData(double
 	return EDisplayClusterCommResult::InternalError;
 }
 
-EDisplayClusterCommResult FDisplayClusterClusterNodeCtrlMain::GetObjectsData(const EDisplayClusterSyncGroup InSyncGroup, TMap<FString, FString>& OutObjectsData)
+EDisplayClusterCommResult FDisplayClusterClusterNodeCtrlMain::GetObjectsData(EDisplayClusterSyncGroup InSyncGroup, TMap<FString, TArray<uint8>>& OutObjectsData)
 {
 	IPDisplayClusterClusterManager* const ClusterMgr = GDisplayCluster->GetPrivateClusterMgr();
 
@@ -329,7 +329,7 @@ EDisplayClusterCommResult FDisplayClusterClusterNodeCtrlMain::GetEventsData(TArr
 	return EDisplayClusterCommResult::InternalError;
 }
 
-EDisplayClusterCommResult FDisplayClusterClusterNodeCtrlMain::GetNativeInputData(TMap<FString, FString>& OutNativeInputData)
+EDisplayClusterCommResult FDisplayClusterClusterNodeCtrlMain::GetNativeInputData(TMap<FString, TArray<uint8>>& OutNativeInputData)
 {
 	IPDisplayClusterClusterManager* const ClusterMgr = GDisplayCluster->GetPrivateClusterMgr();
 
diff --git a/Engine/Plugins/Runtime/nDisplay/Source/DisplayCluster/Private/Cluster/Controller/DisplayClusterClusterNodeCtrlMain.h b/Engine/Plugins/Runtime/nDisplay/Source/DisplayCluster/Private/Cluster/Controller/DisplayClusterClusterNodeCtrlMain.h
--- a/Engine/Plugins/Runtime/nDisplay/Source/DisplayCluster/Private/Cluster/Controller/DisplayClusterClusterNodeCtrlMain.h
+++ b/Engine/Plugins/Runtime/nDisplay/Source/DisplayCluster/Private/Cluster/Controller/DisplayClusterClusterNodeCtrlMain.h
@@ -48,9 +48,9 @@ public:
 	virtual EDisplayClusterCommResult WaitForFrameStart() override;
 	virtual EDisplayClusterCommResult WaitForFrameEnd() override;
 	virtual EDisplayClusterCommResult GetTimeData(double& OutDeltaTime, double& OutGameTime, TOptional<FQualifiedFrameTime>& OutFrameTime) override;
-	virtual EDisplayClusterCommResult GetObjectsData(const EDisplayClusterSyncGroup InSyncGroup, TMap<FString, FString>& OutObjectsData) override;
+	virtual EDisplayClusterCommResult GetObjectsData(EDisplayClusterSyncGroup InSyncGroup, TMap<FString, TArray<uint8>>& OutObjectsData) override;
 	virtual EDisplayClusterCommResult GetEventsData(TArray<TSharedPtr<FDisplayClusterClusterEventJson>>& OutJsonEvents, TArray<TSharedPtr<FDisplayClusterClusterEventBinary>>& OutBinaryEvents) override;
-	virtual EDisplayClusterCommResult GetNativeInputData(TMap<FString, FString>& OutNativeInputData) override;
+	virtual EDisplayClusterCommResult GetNativeInputData(TMap<FString, TArray<uint8>>& OutNativeInputData) override;
 	virtual EDisplayClusterCommResult PropagateStatesData(const TArray<uint8>& InLocalStatesData, TArray<uint8>& OutClusterStatesData) override;
 	//~ End IDisplayClusterProtocolClusterSync
 
diff --git a/Engine/Plugins/Runtime/nDisplay/Source/DisplayCluster/Private/Cluster/DisplayClusterClusterManager.cpp b/Engine/Plugins/Runtime/nDisplay/Source/DisplayCluster/Private/Cluster/DisplayClusterClusterManager.cpp
--- a/Engine/Plugins/Runtime/nDisplay/Source/DisplayCluster/Private/Cluster/DisplayClusterClusterManager.cpp
+++ b/Engine/Plugins/Runtime/nDisplay/Source/DisplayCluster/Private/Cluster/DisplayClusterClusterManager.cpp
@@ -248,7 +248,7 @@ void FDisplayClusterClusterManager::EndFrame(uint64 FrameNum)
 	NativeInputCache.Reset();
 
 	// Reset objects sync cache for all sync groups
-	for (TPair<EDisplayClusterSyncGroup, TMap<FString, FString>>& It : ObjectsToSyncCache)
+	for (TPair<EDisplayClusterSyncGroup, TMap<FString, TArray<uint8>>>& It : ObjectsToSyncCache)
 	{
 		It.Value.Reset();
 	}
@@ -710,7 +710,7 @@ void FDisplayClusterClusterManager::SyncStates()
 
 void FDisplayClusterClusterManager::SyncObjects(EDisplayClusterSyncGroup InSyncGroup)
 {
-	TMap<FString, FString> ObjectsData;
+	TMap<FString, TArray<uint8>> ObjectsData;
 
 	UE_LOG(LogDisplayClusterCluster, Verbose, TEXT("Downloading synchronization data (objects)..."));
 	GetNetApi().GetClusterSyncAPI()->GetObjectsData(InSyncGroup, ObjectsData);
@@ -725,7 +725,7 @@ void FDisplayClusterClusterManager::CacheObjects(EDisplayClusterSyncGroup SyncGr
 	FScopeLock Lock(&ObjectsToSyncCS);
 
 	// Cache data for requested sync group
-	if (TMap<FString, FString>* GroupCache = ObjectsToSyncCache.Find(SyncGroup))
+	if (TMap<FString, TArray<uint8>>* GroupCache = ObjectsToSyncCache.Find(SyncGroup))
 	{
 		UE_LOG(LogDisplayClusterCluster, Verbose, TEXT("Exporting sync data for sync group: %u, items to sync: %d"), (uint8)SyncGroup, GroupCache->Num());
 
@@ -733,15 +733,23 @@ void FDisplayClusterClusterManager::CacheObjects(EDisplayClusterSyncGroup SyncGr
 		{
 			if (SyncObj && SyncObj->IsActive() && SyncObj->IsDirty())
 			{
-				UE_LOG(LogDisplayClusterCluster, Verbose, TEXT("Adding object to sync: %s"), *SyncObj->GetSyncId());
+				// Get ID of the object being serialized
+				FString SyncId = SyncObj->GetSyncId();
+				UE_LOG(LogDisplayClusterCluster, Verbose, TEXT("Adding object to sync: %s"), *SyncId);
 
-				const FString SyncId   = SyncObj->GetSyncId();
-				const FString SyncData = SyncObj->SerializeToString();
+				// Initialize output buffer
+				TArray<uint8> SyncData;
+				constexpr int32 SyncDataReservedBufferSize = 256;
+				SyncData.Reserve(SyncDataReservedBufferSize);
 
-				UE_LOG(LogDisplayClusterCluster, VeryVerbose, TEXT("Sync object: %s - %s"), *SyncId, *SyncData);
+				// Serialize the object
+				FMemoryWriter Writer(SyncData);
+				SyncObj->SerializeDC(Writer);
+
+				UE_LOG(LogDisplayClusterCluster, VeryVerbose, TEXT("Sync object: %s serialized into %d bytes"), *SyncId, SyncData.Num());
 
 				// Cache the object
-				GroupCache->Emplace(SyncId, SyncData);
+				GroupCache->Emplace(MoveTemp(SyncId), MoveTemp(SyncData));
 
 				SyncObj->ClearDirty();
 			}
@@ -754,7 +762,7 @@ void FDisplayClusterClusterManager::CacheObjects(EDisplayClusterSyncGroup SyncGr
 	ObjectsToSyncCacheReadySignals[SyncGroup]->Trigger();
 }
 
-void FDisplayClusterClusterManager::ExportObjectsData(const EDisplayClusterSyncGroup InSyncGroup, TMap<FString, FString>& OutObjectsData)
+void FDisplayClusterClusterManager::ExportObjectsData(const EDisplayClusterSyncGroup InSyncGroup, TMap<FString, TArray<uint8>>& OutObjectsData)
 {
 	// Wait until primary node provides data
 	ObjectsToSyncCacheReadySignals[InSyncGroup]->Wait();
@@ -762,13 +770,13 @@ void FDisplayClusterClusterManager::ExportObjectsData(const EDisplayClusterSyncG
 	OutObjectsData = ObjectsToSyncCache[InSyncGroup];
 }
 
-void FDisplayClusterClusterManager::ImportObjectsData(const EDisplayClusterSyncGroup InSyncGroup, const TMap<FString, FString>& InObjectsData)
+void FDisplayClusterClusterManager::ImportObjectsData(const EDisplayClusterSyncGroup InSyncGroup, const TMap<FString, TArray<uint8>>& InObjectsData)
 {
 	if (InObjectsData.Num() > 0)
 	{
 		for (auto It = InObjectsData.CreateConstIterator(); It; ++It)
 		{
-			UE_LOG(LogDisplayClusterCluster, VeryVerbose, TEXT("sync-data: %s=%s"), *It->Key, *It->Value);
+			UE_LOG(LogDisplayClusterCluster, VeryVerbose, TEXT("sync-data: %s has %d bytes"), *It->Key, It->Value.Num());
 		}
 
 		FScopeLock Lock(&ObjectsToSyncCS);
@@ -778,20 +786,17 @@ void FDisplayClusterClusterManager::ImportObjectsData(const EDisplayClusterSyncG
 			if (SyncObj && SyncObj->IsActive())
 			{
 				const FString SyncId = SyncObj->GetSyncId();
-				if (!InObjectsData.Contains(SyncId))
+				const TArray<uint8>* const FoundSerializedData = InObjectsData.Find(SyncId);
+
+				// Determine whether any data exists for this object
+				if (!FoundSerializedData)
 				{
 					UE_LOG(LogDisplayClusterCluster, VeryVerbose, TEXT("%s has nothing to update"), *SyncId);
 					continue;
 				}
 
-				if (SyncObj->DeserializeFromString(InObjectsData[SyncId]))
-				{
-					UE_LOG(LogDisplayClusterCluster, VeryVerbose, TEXT("Synchronized object: %s"), *SyncId);
-				}
-				else
-				{
-					UE_LOG(LogDisplayClusterCluster, Error, TEXT("Couldn't apply sync data for sync object %s"), *SyncId);
-				}
+				FMemoryReader Reader(*FoundSerializedData);
+				SyncObj->SerializeDC(Reader);
 			}
 		}
 	}
@@ -896,7 +901,7 @@ void FDisplayClusterClusterManager::ImportEventsData(const TArray<TSharedPtr<FDi
 	}
 }
 
-void FDisplayClusterClusterManager::ImportNativeInputData(TMap<FString, FString>& InNativeInputData)
+void FDisplayClusterClusterManager::ImportNativeInputData(TMap<FString, TArray<uint8>>& InNativeInputData)
 {
 	// Cache input data
 	NativeInputCache = MoveTemp(InNativeInputData);
@@ -907,7 +912,7 @@ void FDisplayClusterClusterManager::ImportNativeInputData(TMap<FString, FString>
 	NativeInputCacheReadySignal->Trigger();
 }
 
-void FDisplayClusterClusterManager::ExportNativeInputData(TMap<FString, FString>& OutNativeInputData)
+void FDisplayClusterClusterManager::ExportNativeInputData(TMap<FString, TArray<uint8>>& OutNativeInputData)
 {
 	// Wait for data cache to be ready
 	NativeInputCacheReadySignal->Wait();
diff --git a/Engine/Plugins/Runtime/nDisplay/Source/DisplayCluster/Private/Cluster/DisplayClusterClusterManager.h b/Engine/Plugins/Runtime/nDisplay/Source/DisplayCluster/Private/Cluster/DisplayClusterClusterManager.h
--- a/Engine/Plugins/Runtime/nDisplay/Source/DisplayCluster/Private/Cluster/DisplayClusterClusterManager.h
+++ b/Engine/Plugins/Runtime/nDisplay/Source/DisplayCluster/Private/Cluster/DisplayClusterClusterManager.h
@@ -109,8 +109,8 @@ public:
 	// Sync objects
 	virtual void SyncObjects(EDisplayClusterSyncGroup SyncGroup) override;
 	virtual void CacheObjects(EDisplayClusterSyncGroup SyncGroup) override;
-	virtual void ExportObjectsData(const EDisplayClusterSyncGroup InSyncGroup,       TMap<FString, FString>& OutObjectsData) override;
-	virtual void ImportObjectsData(const EDisplayClusterSyncGroup InSyncGroup, const TMap<FString, FString>& InObjectsData) override;
+	virtual void ExportObjectsData(const EDisplayClusterSyncGroup InSyncGroup,       TMap<FString, TArray<uint8>>& OutObjectsData) override;
+	virtual void ImportObjectsData(const EDisplayClusterSyncGroup InSyncGroup, const TMap<FString, TArray<uint8>>& InObjectsData) override;
 
 	// Sync cluster events
 	virtual void SyncEvents()  override;
@@ -119,8 +119,8 @@ public:
 	virtual void ImportEventsData(const TArray<TSharedPtr<FDisplayClusterClusterEventJson>>& InJsonEvents, const TArray<TSharedPtr<FDisplayClusterClusterEventBinary>>& InBinaryEvents) override;
 
 	// Sync native input
-	virtual void ExportNativeInputData(TMap<FString, FString>& OutNativeInputData) override;
-	virtual void ImportNativeInputData(TMap<FString, FString>& InNativeInputData) override;
+	virtual void ExportNativeInputData(TMap<FString, TArray<uint8>>& OutNativeInputData) override;
+	virtual void ImportNativeInputData(TMap<FString, TArray<uint8>>& InNativeInputData) override;
 
 	//~ End IPDisplayClusterClusterManager
 
@@ -242,11 +242,11 @@ private:
 	mutable FCriticalSection ObjectsToSyncCS;
 	// Sync objects - replication
 	TMap<EDisplayClusterSyncGroup, FEvent*> ObjectsToSyncCacheReadySignals;
-	TMap<EDisplayClusterSyncGroup, TMap<FString, FString>> ObjectsToSyncCache;
+	TMap<EDisplayClusterSyncGroup, TMap<FString, TArray<uint8>>> ObjectsToSyncCache;
 
 	// Native input - replication
 	FEventRef NativeInputCacheReadySignal{ EEventMode::ManualReset };
-	TMap<FString, FString> NativeInputCache;
+	TMap<FString, TArray<uint8>> NativeInputCache;
 
 	// JSON events
 	TMap<bool, TMap<FString, TSharedPtr<FDisplayClusterClusterEventJson>>> ClusterEventsJson;
diff --git a/Engine/Plugins/Runtime/nDisplay/Source/DisplayCluster/Private/Cluster/Failover/DisplayClusterCommDataCache.cpp b/Engine/Plugins/Runtime/nDisplay/Source/DisplayCluster/Private/Cluster/Failover/DisplayClusterCommDataCache.cpp
--- a/Engine/Plugins/Runtime/nDisplay/Source/DisplayCluster/Private/Cluster/Failover/DisplayClusterCommDataCache.cpp
+++ b/Engine/Plugins/Runtime/nDisplay/Source/DisplayCluster/Private/Cluster/Failover/DisplayClusterCommDataCache.cpp
@@ -100,12 +100,12 @@ void FDisplayClusterCommDataCache::Initialize_GetTimeData()
 
 void FDisplayClusterCommDataCache::Initialize_GetObjectsData()
 {
-	GetObjectsData_OpLoad = [this](const EDisplayClusterSyncGroup InSyncGroup, TMap<FString, FString>& OutObjectsData)
+	GetObjectsData_OpLoad = [this](const EDisplayClusterSyncGroup InSyncGroup, TMap<FString, TArray<uint8>>& OutObjectsData)
 		{
 			OpLoadImpl<FCache_GetObjectsData>(GetObjectsDataSlotName(InSyncGroup), OutObjectsData);
 		};
 
-	GetObjectsData_OpSave = [this](const EDisplayClusterSyncGroup InSyncGroup, TMap<FString, FString>& InObjectsData)
+	GetObjectsData_OpSave = [this](const EDisplayClusterSyncGroup InSyncGroup, TMap<FString, TArray<uint8>>& InObjectsData)
 		{
 			OpSaveImpl<FCache_GetObjectsData>(GetObjectsDataSlotName(InSyncGroup), InObjectsData);
 		};
@@ -144,12 +144,12 @@ void FDisplayClusterCommDataCache::Initialize_GetEventsData()
 
 void FDisplayClusterCommDataCache::Initialize_GetNativeInputData()
 {
-	GetNativeInputData_OpLoad = [this](TMap<FString, FString>& OutNativeInputData)
+	GetNativeInputData_OpLoad = [this](TMap<FString, TArray<uint8>>& OutNativeInputData)
 		{
 			OpLoadImpl<FCache_GetNativeInputData>(UE::nDisplay::Failover::Private::Name_SlotNativeInput, OutNativeInputData);
 		};
 
-	GetNativeInputData_OpSave = [this](TMap<FString, FString>& InNativeInputData)
+	GetNativeInputData_OpSave = [this](TMap<FString, TArray<uint8>>& InNativeInputData)
 		{
 			OpSaveImpl<FCache_GetNativeInputData>(UE::nDisplay::Failover::Private::Name_SlotNativeInput, InNativeInputData);
 		};
@@ -545,9 +545,9 @@ FString FDisplayClusterCommDataCache::FCache_GetObjectsData::ToLogString() const
 
 	Output = FString::Printf(TEXT("OD[%s]: "), bCached ? TEXT("+") : TEXT("-"));
 	int32 Idx = 0;
-	for (const TPair<FString, FString>& DataIt : ObjData)
+	for (const TPair<FString, TArray<uint8>>& DataIt : ObjData)
 	{
-		Output += FString::Printf(TEXT("<%d : %s=%s> "), Idx++, *DataIt.Key, *DataIt.Value);
+		Output += FString::Printf(TEXT("<%d : %s=[%d bytes]> "), Idx++, *DataIt.Key, DataIt.Value.Num());
 	}
 
 	return Output;
@@ -628,9 +628,9 @@ FString FDisplayClusterCommDataCache::FCache_GetNativeInputData::ToLogString() c
 
 	Output = FString::Printf(TEXT("ID[%s]: "), bCached ? TEXT("+") : TEXT("-"));
 	int32 Idx = 0;
-	for (const TPair<FString, FString>& DataIt : NativeInputData)
+	for (const TPair<FString, TArray<uint8>>& DataIt : NativeInputData)
 	{
-		Output += FString::Printf(TEXT("<%d : %s=%s> "), Idx++, *DataIt.Key, *DataIt.Value);
+		Output += FString::Printf(TEXT("<%d : %s -> %d bytes> "), Idx++, *DataIt.Key, DataIt.Value.Num());
 	}
 
 	return Output;
diff --git a/Engine/Plugins/Runtime/nDisplay/Source/DisplayCluster/Private/Cluster/Failover/DisplayClusterCommDataCache.h b/Engine/Plugins/Runtime/nDisplay/Source/DisplayCluster/Private/Cluster/Failover/DisplayClusterCommDataCache.h
--- a/Engine/Plugins/Runtime/nDisplay/Source/DisplayCluster/Private/Cluster/Failover/DisplayClusterCommDataCache.h
+++ b/Engine/Plugins/Runtime/nDisplay/Source/DisplayCluster/Private/Cluster/Failover/DisplayClusterCommDataCache.h
@@ -52,10 +52,10 @@ public:
 	TFunction<FOpIsCached(const EDisplayClusterSyncGroup InSyncGroup)> GetObjectsData_OpIsCached;
 
 	/** Loads sync objects data from cache */
-	TFunction<void(const EDisplayClusterSyncGroup InSyncGroup, TMap<FString, FString>& OutObjectsData)> GetObjectsData_OpLoad;
+	TFunction<void(const EDisplayClusterSyncGroup InSyncGroup, TMap<FString, TArray<uint8>>& OutObjectsData)> GetObjectsData_OpLoad;
 
 	/** Caches sync objects data */
-	TFunction<void(const EDisplayClusterSyncGroup InSyncGroup, TMap<FString, FString>& InObjectsData)> GetObjectsData_OpSave;
+	TFunction<void(const EDisplayClusterSyncGroup InSyncGroup, TMap<FString, TArray<uint8>>& InObjectsData)> GetObjectsData_OpSave;
 
 public:
 
@@ -74,10 +74,10 @@ public:
 	FOpIsCached GetNativeInputData_OpIsCached;
 
 	/** Loads input data from cache */
-	TFunction<void(TMap<FString, FString>& OutNativeInputData)> GetNativeInputData_OpLoad;
+	TFunction<void(TMap<FString, TArray<uint8>>& OutNativeInputData)> GetNativeInputData_OpLoad;
 
 	/** Caches input data */
-	TFunction<void(TMap<FString, FString>& InNativeInputData)> GetNativeInputData_OpSave;
+	TFunction<void(TMap<FString, TArray<uint8>>& InNativeInputData)> GetNativeInputData_OpSave;
 
 public:
 
@@ -250,12 +250,12 @@ private: // GetObjectsData
 	public:
 
 		/** Object data of a dedicated synchronization group */
-		TMap<FString, FString> ObjData;
+		TMap<FString, TArray<uint8>> ObjData;
 
 	public:
 
 		/** Copies data TO or FROM */
-		void CopyData(bool bCopyOutside, TMap<FString, FString>& ObjDataRef)
+		void CopyData(bool bCopyOutside, TMap<FString, TArray<uint8>>& ObjDataRef)
 		{
 			if (bCopyOutside)
 			{
@@ -341,12 +341,12 @@ private: // GetNativeInputData
 	public:
 
 		/** Native input data cached for current cycle (frame) */
-		TMap<FString, FString> NativeInputData;
+		TMap<FString, TArray<uint8>> NativeInputData;
 
 	public:
 
 		/** Copies data TO or FROM */
-		void CopyData(bool bCopyOutside, TMap<FString, FString>& NativeInputDataRef)
+		void CopyData(bool bCopyOutside, TMap<FString, TArray<uint8>>& NativeInputDataRef)
 		{
 			if (bCopyOutside)
 			{
diff --git a/Engine/Plugins/Runtime/nDisplay/Source/DisplayCluster/Private/Cluster/Failover/DisplayClusterFailoverNodeCtrlDisabled.cpp b/Engine/Plugins/Runtime/nDisplay/Source/DisplayCluster/Private/Cluster/Failover/DisplayClusterFailoverNodeCtrlDisabled.cpp
--- a/Engine/Plugins/Runtime/nDisplay/Source/DisplayCluster/Private/Cluster/Failover/DisplayClusterFailoverNodeCtrlDisabled.cpp
+++ b/Engine/Plugins/Runtime/nDisplay/Source/DisplayCluster/Private/Cluster/Failover/DisplayClusterFailoverNodeCtrlDisabled.cpp
@@ -48,7 +48,7 @@ EDisplayClusterCommResult FDisplayClusterFailoverNodeCtrlDisabled::GetTimeData(d
 	return EDisplayClusterCommResult::NotAllowed;
 }
 
-EDisplayClusterCommResult FDisplayClusterFailoverNodeCtrlDisabled::GetObjectsData(const EDisplayClusterSyncGroup InSyncGroup, TMap<FString, FString>& OutObjectsData)
+EDisplayClusterCommResult FDisplayClusterFailoverNodeCtrlDisabled::GetObjectsData(EDisplayClusterSyncGroup InSyncGroup, TMap<FString, TArray<uint8>>& OutObjectsData)
 {
 	return EDisplayClusterCommResult::NotAllowed;
 }
@@ -58,7 +58,7 @@ EDisplayClusterCommResult FDisplayClusterFailoverNodeCtrlDisabled::GetEventsData
 	return EDisplayClusterCommResult::NotAllowed;
 }
 
-EDisplayClusterCommResult FDisplayClusterFailoverNodeCtrlDisabled::GetNativeInputData(TMap<FString, FString>& OutNativeInputData)
+EDisplayClusterCommResult FDisplayClusterFailoverNodeCtrlDisabled::GetNativeInputData(TMap<FString, TArray<uint8>>& OutNativeInputData)
 {
 	return EDisplayClusterCommResult::NotAllowed;
 }
diff --git a/Engine/Plugins/Runtime/nDisplay/Source/DisplayCluster/Private/Cluster/Failover/DisplayClusterFailoverNodeCtrlDisabled.h b/Engine/Plugins/Runtime/nDisplay/Source/DisplayCluster/Private/Cluster/Failover/DisplayClusterFailoverNodeCtrlDisabled.h
--- a/Engine/Plugins/Runtime/nDisplay/Source/DisplayCluster/Private/Cluster/Failover/DisplayClusterFailoverNodeCtrlDisabled.h
+++ b/Engine/Plugins/Runtime/nDisplay/Source/DisplayCluster/Private/Cluster/Failover/DisplayClusterFailoverNodeCtrlDisabled.h
@@ -34,9 +34,9 @@ public:
 	virtual EDisplayClusterCommResult WaitForFrameStart() override;
 	virtual EDisplayClusterCommResult WaitForFrameEnd() override;
 	virtual EDisplayClusterCommResult GetTimeData(double& OutDeltaTime, double& OutGameTime, TOptional<FQualifiedFrameTime>& OutFrameTime) override;
-	virtual EDisplayClusterCommResult GetObjectsData(const EDisplayClusterSyncGroup InSyncGroup, TMap<FString, FString>& OutObjectsData) override;
+	virtual EDisplayClusterCommResult GetObjectsData(EDisplayClusterSyncGroup InSyncGroup, TMap<FString, TArray<uint8>>& OutObjectsData) override;
 	virtual EDisplayClusterCommResult GetEventsData(TArray<TSharedPtr<FDisplayClusterClusterEventJson>>& OutJsonEvents, TArray<TSharedPtr<FDisplayClusterClusterEventBinary>>& OutBinaryEvents) override;
-	virtual EDisplayClusterCommResult GetNativeInputData(TMap<FString, FString>& OutNativeInputData) override;
+	virtual EDisplayClusterCommResult GetNativeInputData(TMap<FString, TArray<uint8>>& OutNativeInputData) override;
 	virtual EDisplayClusterCommResult PropagateStatesData(const TArray<uint8>& InLocalStatesData, TArray<uint8>& OutClusterStatesData) override;
 	//~ End IDisplayClusterProtocolClusterSync
 
diff --git a/Engine/Plugins/Runtime/nDisplay/Source/DisplayCluster/Private/Cluster/Failover/DisplayClusterFailoverNodeCtrlEditor.cpp b/Engine/Plugins/Runtime/nDisplay/Source/DisplayCluster/Private/Cluster/Failover/DisplayClusterFailoverNodeCtrlEditor.cpp
--- a/Engine/Plugins/Runtime/nDisplay/Source/DisplayCluster/Private/Cluster/Failover/DisplayClusterFailoverNodeCtrlEditor.cpp
+++ b/Engine/Plugins/Runtime/nDisplay/Source/DisplayCluster/Private/Cluster/Failover/DisplayClusterFailoverNodeCtrlEditor.cpp
@@ -57,7 +57,7 @@ EDisplayClusterCommResult FDisplayClusterFailoverNodeCtrlEditor::GetTimeData(dou
 	return GetNodeController()->GetTimeData(OutDeltaTime, OutGameTime, OutFrameTime);
 }
 
-EDisplayClusterCommResult FDisplayClusterFailoverNodeCtrlEditor::GetObjectsData(const EDisplayClusterSyncGroup InSyncGroup, TMap<FString, FString>& OutObjectsData)
+EDisplayClusterCommResult FDisplayClusterFailoverNodeCtrlEditor::GetObjectsData(EDisplayClusterSyncGroup InSyncGroup, TMap<FString, TArray<uint8>>& OutObjectsData)
 {
 	// `Editor` failover controller doesn't perform any validations, just forward the call
 	return GetNodeController()->GetObjectsData(InSyncGroup, OutObjectsData);
@@ -69,7 +69,7 @@ EDisplayClusterCommResult FDisplayClusterFailoverNodeCtrlEditor::GetEventsData(T
 	return GetNodeController()->GetEventsData(OutJsonEvents, OutBinaryEvents);
 }
 
-EDisplayClusterCommResult FDisplayClusterFailoverNodeCtrlEditor::GetNativeInputData(TMap<FString, FString>& OutNativeInputData)
+EDisplayClusterCommResult FDisplayClusterFailoverNodeCtrlEditor::GetNativeInputData(TMap<FString, TArray<uint8>>& OutNativeInputData)
 {
 	// `Editor` failover controller doesn't perform any validations, just forward the call
 	return GetNodeController()->GetNativeInputData(OutNativeInputData);
diff --git a/Engine/Plugins/Runtime/nDisplay/Source/DisplayCluster/Private/Cluster/Failover/DisplayClusterFailoverNodeCtrlEditor.h b/Engine/Plugins/Runtime/nDisplay/Source/DisplayCluster/Private/Cluster/Failover/DisplayClusterFailoverNodeCtrlEditor.h
--- a/Engine/Plugins/Runtime/nDisplay/Source/DisplayCluster/Private/Cluster/Failover/DisplayClusterFailoverNodeCtrlEditor.h
+++ b/Engine/Plugins/Runtime/nDisplay/Source/DisplayCluster/Private/Cluster/Failover/DisplayClusterFailoverNodeCtrlEditor.h
@@ -36,9 +36,9 @@ public:
 	virtual EDisplayClusterCommResult WaitForFrameStart() override;
 	virtual EDisplayClusterCommResult WaitForFrameEnd() override;
 	virtual EDisplayClusterCommResult GetTimeData(double& OutDeltaTime, double& OutGameTime, TOptional<FQualifiedFrameTime>& OutFrameTime) override;
-	virtual EDisplayClusterCommResult GetObjectsData(const EDisplayClusterSyncGroup InSyncGroup, TMap<FString, FString>& OutObjectsData) override;
+	virtual EDisplayClusterCommResult GetObjectsData(EDisplayClusterSyncGroup InSyncGroup, TMap<FString, TArray<uint8>>& OutObjectsData) override;
 	virtual EDisplayClusterCommResult GetEventsData(TArray<TSharedPtr<FDisplayClusterClusterEventJson>>& OutJsonEvents, TArray<TSharedPtr<FDisplayClusterClusterEventBinary>>& OutBinaryEvents) override;
-	virtual EDisplayClusterCommResult GetNativeInputData(TMap<FString, FString>& OutNativeInputData) override;
+	virtual EDisplayClusterCommResult GetNativeInputData(TMap<FString, TArray<uint8>>& OutNativeInputData) override;
 	virtual EDisplayClusterCommResult PropagateStatesData(const TArray<uint8>& InLocalStatesData, TArray<uint8>& OutClusterStatesData) override;
 	//~ End IDisplayClusterProtocolClusterSync
 
diff --git a/Engine/Plugins/Runtime/nDisplay/Source/DisplayCluster/Private/Cluster/Failover/DisplayClusterFailoverNodeCtrlMain.cpp b/Engine/Plugins/Runtime/nDisplay/Source/DisplayCluster/Private/Cluster/Failover/DisplayClusterFailoverNodeCtrlMain.cpp
--- a/Engine/Plugins/Runtime/nDisplay/Source/DisplayCluster/Private/Cluster/Failover/DisplayClusterFailoverNodeCtrlMain.cpp
+++ b/Engine/Plugins/Runtime/nDisplay/Source/DisplayCluster/Private/Cluster/Failover/DisplayClusterFailoverNodeCtrlMain.cpp
@@ -290,7 +290,7 @@ EDisplayClusterCommResult FDisplayClusterFailoverNodeCtrlMain::GetTimeData(doubl
 	return CommResult;
 }
 
-EDisplayClusterCommResult FDisplayClusterFailoverNodeCtrlMain::GetObjectsData(const EDisplayClusterSyncGroup InSyncGroup, TMap<FString, FString>& OutObjectsData)
+EDisplayClusterCommResult FDisplayClusterFailoverNodeCtrlMain::GetObjectsData(EDisplayClusterSyncGroup InSyncGroup, TMap<FString, TArray<uint8>>& OutObjectsData)
 {
 	static const TMap<EDisplayClusterSyncGroup, FName> TxnNames =
 	{
@@ -347,7 +347,7 @@ EDisplayClusterCommResult FDisplayClusterFailoverNodeCtrlMain::GetEventsData(
 	return CommResult;
 }
 
-EDisplayClusterCommResult FDisplayClusterFailoverNodeCtrlMain::GetNativeInputData(TMap<FString, FString>& OutNativeInputData)
+EDisplayClusterCommResult FDisplayClusterFailoverNodeCtrlMain::GetNativeInputData(TMap<FString, TArray<uint8>>& OutNativeInputData)
 {
 	static const FName TxnName = TEXT("GetNativeInputData");
 
diff --git a/Engine/Plugins/Runtime/nDisplay/Source/DisplayCluster/Private/Cluster/Failover/DisplayClusterFailoverNodeCtrlMain.h b/Engine/Plugins/Runtime/nDisplay/Source/DisplayCluster/Private/Cluster/Failover/DisplayClusterFailoverNodeCtrlMain.h
--- a/Engine/Plugins/Runtime/nDisplay/Source/DisplayCluster/Private/Cluster/Failover/DisplayClusterFailoverNodeCtrlMain.h
+++ b/Engine/Plugins/Runtime/nDisplay/Source/DisplayCluster/Private/Cluster/Failover/DisplayClusterFailoverNodeCtrlMain.h
@@ -47,9 +47,9 @@ public:
 	virtual EDisplayClusterCommResult WaitForFrameStart() override;
 	virtual EDisplayClusterCommResult WaitForFrameEnd() override;
 	virtual EDisplayClusterCommResult GetTimeData(double& OutDeltaTime, double& OutGameTime, TOptional<FQualifiedFrameTime>& OutFrameTime) override;
-	virtual EDisplayClusterCommResult GetObjectsData(const EDisplayClusterSyncGroup InSyncGroup, TMap<FString, FString>& OutObjectsData) override;
+	virtual EDisplayClusterCommResult GetObjectsData(EDisplayClusterSyncGroup InSyncGroup, TMap<FString, TArray<uint8>>& OutObjectsData) override;
 	virtual EDisplayClusterCommResult GetEventsData(TArray<TSharedPtr<FDisplayClusterClusterEventJson>>& OutJsonEvents, TArray<TSharedPtr<FDisplayClusterClusterEventBinary>>& OutBinaryEvents) override;
-	virtual EDisplayClusterCommResult GetNativeInputData(TMap<FString, FString>& OutNativeInputData) override;
+	virtual EDisplayClusterCommResult GetNativeInputData(TMap<FString, TArray<uint8>>& OutNativeInputData) override;
 	virtual EDisplayClusterCommResult PropagateStatesData(const TArray<uint8>& InLocalStatesData, TArray<uint8>& OutClusterStatesData) override;
 	//~ End IDisplayClusterProtocolClusterSync
 
diff --git a/Engine/Plugins/Runtime/nDisplay/Source/DisplayCluster/Private/Cluster/IPDisplayClusterClusterManager.h b/Engine/Plugins/Runtime/nDisplay/Source/DisplayCluster/Private/Cluster/IPDisplayClusterClusterManager.h
--- a/Engine/Plugins/Runtime/nDisplay/Source/DisplayCluster/Private/Cluster/IPDisplayClusterClusterManager.h
+++ b/Engine/Plugins/Runtime/nDisplay/Source/DisplayCluster/Private/Cluster/IPDisplayClusterClusterManager.h
@@ -79,10 +79,10 @@ public: // Objects sync
 	virtual void CacheObjects(EDisplayClusterSyncGroup SyncGroup) = 0;
 
 	/** Exports objects data from this node */
-	virtual void ExportObjectsData(const EDisplayClusterSyncGroup InSyncGroup,       TMap<FString, FString>& OutObjectsData) = 0;
+	virtual void ExportObjectsData(const EDisplayClusterSyncGroup InSyncGroup,       TMap<FString, TArray<uint8>>& OutObjectsData) = 0;
 
 	/** Imports objects data from the primary node */
-	virtual void ImportObjectsData(const EDisplayClusterSyncGroup InSyncGroup, const TMap<FString, FString>& InObjectsData) = 0;
+	virtual void ImportObjectsData(const EDisplayClusterSyncGroup InSyncGroup, const TMap<FString, TArray<uint8>>& InObjectsData) = 0;
 
 public: // Cluster events sync
 
@@ -101,8 +101,8 @@ public: // Cluster events sync
 public: // Native input sync
 
 	/** Exports native input data from the local PlayerInput on the primary node */
-	virtual void ExportNativeInputData(TMap<FString, FString>& OutNativeInputData) = 0;
+	virtual void ExportNativeInputData(TMap<FString, TArray<uint8>>& OutNativeInputData) = 0;
 
 	/** Imports and applies the native input data from the the primary node */
-	virtual void ImportNativeInputData(TMap<FString, FString>& InNativeInputData) = 0;
+	virtual void ImportNativeInputData(TMap<FString, TArray<uint8>>& InNativeInputData) = 0;
 };
diff --git a/Engine/Plugins/Runtime/nDisplay/Source/DisplayCluster/Private/Cluster/NetAPI/DisplayClusterNetApiFacade.cpp b/Engine/Plugins/Runtime/nDisplay/Source/DisplayCluster/Private/Cluster/NetAPI/DisplayClusterNetApiFacade.cpp
--- a/Engine/Plugins/Runtime/nDisplay/Source/DisplayCluster/Private/Cluster/NetAPI/DisplayClusterNetApiFacade.cpp
+++ b/Engine/Plugins/Runtime/nDisplay/Source/DisplayCluster/Private/Cluster/NetAPI/DisplayClusterNetApiFacade.cpp
@@ -79,7 +79,7 @@ public:
 	}
 
 	/** Failover transaction for GetObjectsData */
-	virtual EDisplayClusterCommResult GetObjectsData(const EDisplayClusterSyncGroup InSyncGroup, TMap<FString, FString>& OutObjectsData) override
+	virtual EDisplayClusterCommResult GetObjectsData(EDisplayClusterSyncGroup InSyncGroup, TMap<FString, TArray<uint8>>& OutObjectsData) override
 	{
 		checkSlow(IsInGameThread());
 		return GetFailoverController()->GetObjectsData(InSyncGroup, OutObjectsData);
@@ -95,7 +95,7 @@ public:
 	}
 
 	/** Failover transaction for GetNativeInputData */
-	virtual EDisplayClusterCommResult GetNativeInputData(TMap<FString, FString>& OutNativeInputData) override
+	virtual EDisplayClusterCommResult GetNativeInputData(TMap<FString, TArray<uint8>>& OutNativeInputData) override
 	{
 		checkSlow(IsInGameThread());
 		return GetFailoverController()->GetNativeInputData(OutNativeInputData);
diff --git a/Engine/Plugins/Runtime/nDisplay/Source/DisplayCluster/Private/Cluster/RequestHandler/DisplayClusterCommRequestHandlerLocal.cpp b/Engine/Plugins/Runtime/nDisplay/Source/DisplayCluster/Private/Cluster/RequestHandler/DisplayClusterCommRequestHandlerLocal.cpp
--- a/Engine/Plugins/Runtime/nDisplay/Source/DisplayCluster/Private/Cluster/RequestHandler/DisplayClusterCommRequestHandlerLocal.cpp
+++ b/Engine/Plugins/Runtime/nDisplay/Source/DisplayCluster/Private/Cluster/RequestHandler/DisplayClusterCommRequestHandlerLocal.cpp
@@ -49,7 +49,7 @@ EDisplayClusterCommResult FDisplayClusterCommRequestHandlerLocal::GetTimeData(do
 	return EDisplayClusterCommResult::Ok;
 }
 
-EDisplayClusterCommResult FDisplayClusterCommRequestHandlerLocal::GetObjectsData(const EDisplayClusterSyncGroup InSyncGroup, TMap<FString, FString>& OutObjectsData)
+EDisplayClusterCommResult FDisplayClusterCommRequestHandlerLocal::GetObjectsData(EDisplayClusterSyncGroup InSyncGroup, TMap<FString, TArray<uint8>>& OutObjectsData)
 {
 	const TSharedRef<FDisplayClusterCommDataCache> DataCache = GDisplayCluster->GetPrivateClusterMgr()->GetDataCache();
 
@@ -85,7 +85,7 @@ EDisplayClusterCommResult FDisplayClusterCommRequestHandlerLocal::GetEventsData(
 	return EDisplayClusterCommResult::Ok;
 }
 
-EDisplayClusterCommResult FDisplayClusterCommRequestHandlerLocal::GetNativeInputData(TMap<FString, FString>& OutNativeInputData)
+EDisplayClusterCommResult FDisplayClusterCommRequestHandlerLocal::GetNativeInputData(TMap<FString, TArray<uint8>>& OutNativeInputData)
 {
 	const TSharedRef<FDisplayClusterCommDataCache> DataCache = GDisplayCluster->GetPrivateClusterMgr()->GetDataCache();
 
diff --git a/Engine/Plugins/Runtime/nDisplay/Source/DisplayCluster/Private/Cluster/RequestHandler/DisplayClusterCommRequestHandlerLocal.h b/Engine/Plugins/Runtime/nDisplay/Source/DisplayCluster/Private/Cluster/RequestHandler/DisplayClusterCommRequestHandlerLocal.h
--- a/Engine/Plugins/Runtime/nDisplay/Source/DisplayCluster/Private/Cluster/RequestHandler/DisplayClusterCommRequestHandlerLocal.h
+++ b/Engine/Plugins/Runtime/nDisplay/Source/DisplayCluster/Private/Cluster/RequestHandler/DisplayClusterCommRequestHandlerLocal.h
@@ -37,9 +37,9 @@ public:
 	virtual EDisplayClusterCommResult WaitForFrameStart() override;
 	virtual EDisplayClusterCommResult WaitForFrameEnd() override;
 	virtual EDisplayClusterCommResult GetTimeData(double& OutDeltaTime, double& OutGameTime, TOptional<FQualifiedFrameTime>& OutFrameTime) override;
-	virtual EDisplayClusterCommResult GetObjectsData(const EDisplayClusterSyncGroup InSyncGroup, TMap<FString, FString>& OutObjectsData) override;
+	virtual EDisplayClusterCommResult GetObjectsData(EDisplayClusterSyncGroup InSyncGroup, TMap<FString, TArray<uint8>>& OutObjectsData) override;
 	virtual EDisplayClusterCommResult GetEventsData(TArray<TSharedPtr<FDisplayClusterClusterEventJson>>& OutJsonEvents, TArray<TSharedPtr<FDisplayClusterClusterEventBinary>>& OutBinaryEvents) override;
-	virtual EDisplayClusterCommResult GetNativeInputData(TMap<FString, FString>& OutNativeInputData) override;
+	virtual EDisplayClusterCommResult GetNativeInputData(TMap<FString, TArray<uint8>>& OutNativeInputData) override;
 	virtual EDisplayClusterCommResult PropagateStatesData(const TArray<uint8>& InLocalStatesData, TArray<uint8>& OutClusterStatesData) override;
 	//~ End IDisplayClusterProtocolClusterSync
 };
diff --git a/Engine/Plugins/Runtime/nDisplay/Source/DisplayCluster/Private/Cluster/RequestHandler/DisplayClusterCommRequestHandlerRemote.cpp b/Engine/Plugins/Runtime/nDisplay/Source/DisplayCluster/Private/Cluster/RequestHandler/DisplayClusterCommRequestHandlerRemote.cpp
--- a/Engine/Plugins/Runtime/nDisplay/Source/DisplayCluster/Private/Cluster/RequestHandler/DisplayClusterCommRequestHandlerRemote.cpp
+++ b/Engine/Plugins/Runtime/nDisplay/Source/DisplayCluster/Private/Cluster/RequestHandler/DisplayClusterCommRequestHandlerRemote.cpp
@@ -48,7 +48,7 @@ EDisplayClusterCommResult FDisplayClusterCommRequestHandlerRemote::GetTimeData(d
 	return EDisplayClusterCommResult::Ok;
 }
 
-EDisplayClusterCommResult FDisplayClusterCommRequestHandlerRemote::GetObjectsData(const EDisplayClusterSyncGroup InSyncGroup, TMap<FString, FString>& OutObjectsData)
+EDisplayClusterCommResult FDisplayClusterCommRequestHandlerRemote::GetObjectsData(EDisplayClusterSyncGroup InSyncGroup, TMap<FString, TArray<uint8>>& OutObjectsData)
 {
 	const TSharedRef<FDisplayClusterCommDataCache> DataCache = GDisplayCluster->GetPrivateClusterMgr()->GetDataCache();
 
@@ -82,7 +82,7 @@ EDisplayClusterCommResult FDisplayClusterCommRequestHandlerRemote::GetEventsData
 	return EDisplayClusterCommResult::Ok;
 }
 
-EDisplayClusterCommResult FDisplayClusterCommRequestHandlerRemote::GetNativeInputData(TMap<FString, FString>& OutNativeInputData)
+EDisplayClusterCommResult FDisplayClusterCommRequestHandlerRemote::GetNativeInputData(TMap<FString, TArray<uint8>>& OutNativeInputData)
 {
 	const TSharedRef<FDisplayClusterCommDataCache> DataCache = GDisplayCluster->GetPrivateClusterMgr()->GetDataCache();
 
diff --git a/Engine/Plugins/Runtime/nDisplay/Source/DisplayCluster/Private/Cluster/RequestHandler/DisplayClusterCommRequestHandlerRemote.h b/Engine/Plugins/Runtime/nDisplay/Source/DisplayCluster/Private/Cluster/RequestHandler/DisplayClusterCommRequestHandlerRemote.h
--- a/Engine/Plugins/Runtime/nDisplay/Source/DisplayCluster/Private/Cluster/RequestHandler/DisplayClusterCommRequestHandlerRemote.h
+++ b/Engine/Plugins/Runtime/nDisplay/Source/DisplayCluster/Private/Cluster/RequestHandler/DisplayClusterCommRequestHandlerRemote.h
@@ -37,9 +37,9 @@ public:
 	virtual EDisplayClusterCommResult WaitForFrameStart() override;
 	virtual EDisplayClusterCommResult WaitForFrameEnd() override;
 	virtual EDisplayClusterCommResult GetTimeData(double& OutDeltaTime, double& OutGameTime, TOptional<FQualifiedFrameTime>& OutFrameTime) override;
-	virtual EDisplayClusterCommResult GetObjectsData(const EDisplayClusterSyncGroup InSyncGroup, TMap<FString, FString>& OutObjectsData) override;
+	virtual EDisplayClusterCommResult GetObjectsData(EDisplayClusterSyncGroup InSyncGroup, TMap<FString, TArray<uint8>>& OutObjectsData) override;
 	virtual EDisplayClusterCommResult GetEventsData(TArray<TSharedPtr<FDisplayClusterClusterEventJson>>& OutJsonEvents, TArray<TSharedPtr<FDisplayClusterClusterEventBinary>>& OutBinaryEvents) override;
-	virtual EDisplayClusterCommResult GetNativeInputData(TMap<FString, FString>& OutNativeInputData) override;
+	virtual EDisplayClusterCommResult GetNativeInputData(TMap<FString, TArray<uint8>>& OutNativeInputData) override;
 	virtual EDisplayClusterCommResult PropagateStatesData(const TArray<uint8>& InLocalStatesData, TArray<uint8>& OutClusterStatesData) override;
 	//~ End IDisplayClusterProtocolClusterSync
 };
diff --git a/Engine/Plugins/Runtime/nDisplay/Source/DisplayCluster/Private/Game/EngineClasses/Basics/DisplayClusterPlayerInput.cpp b/Engine/Plugins/Runtime/nDisplay/Source/DisplayCluster/Private/Game/EngineClasses/Basics/DisplayClusterPlayerInput.cpp
--- a/Engine/Plugins/Runtime/nDisplay/Source/DisplayCluster/Private/Game/EngineClasses/Basics/DisplayClusterPlayerInput.cpp
+++ b/Engine/Plugins/Runtime/nDisplay/Source/DisplayCluster/Private/Game/EngineClasses/Basics/DisplayClusterPlayerInput.cpp
@@ -6,18 +6,20 @@
 #include "DisplayClusterConfigurationTypes.h"
 #include "DisplayClusterEnums.h"
 
-#include "Components/InputComponent.h"
-
 #include "Cluster/IPDisplayClusterClusterManager.h"
 #include "Cluster/NetAPI/DisplayClusterNetApiFacade.h"
+#include "Components/InputComponent.h"
 #include "Config/IPDisplayClusterConfigManager.h"
 
+#include "Engine/World.h"
+
 #include "Misc/DisplayClusterGlobals.h"
 #include "Misc/DisplayClusterHelpers.h"
 #include "Misc/DisplayClusterLog.h"
 #include "Misc/DisplayClusterTypesConverter.h"
 
-#include "Engine/World.h"
+#include "Serialization/MemoryReader.h"
+#include "Serialization/MemoryWriter.h"
 
 
 UDisplayClusterPlayerInput::UDisplayClusterPlayerInput()
@@ -56,89 +58,100 @@ void UDisplayClusterPlayerInput::ProcessInputStack(const TArray<UInputComponent*
 	Super::ProcessInputStack(InputComponentStack, DeltaTime, bGamePaused);
 }
 
-bool UDisplayClusterPlayerInput::SerializeKeyStateMap(TMap<FString, FString>& OutKeyStateMap)
+bool UDisplayClusterPlayerInput::SerializeKeyStateMap(TMap<FString, TArray<uint8>>& OutKeyStateMap)
 {
-	const TMap<FKey, FKeyState>& StateMap = GetKeyStateMap();
-	for (auto it = StateMap.CreateConstIterator(); it; ++it)
+	// Serialize each key state individually
+	TMap<FKey, FKeyState>& StateMap = GetKeyStateMap();
+	for (TPair<FKey, FKeyState>& It : StateMap)
 	{
-		const FString KeyName = it->Key.ToString();
-
-		FString StrKeyState;
-		StrKeyState.Reserve(2048);
-
-		StrKeyState = FString::Printf(TEXT("%s;%s;%s;%s;%s;%s;%s;%s;%s;"),
-			*DisplayClusterTypesConverter::template ToHexString(it->Value.RawValue),
-			*DisplayClusterTypesConverter::template ToHexString(it->Value.Value),
-			*DisplayClusterTypesConverter::template ToHexString(it->Value.LastUpDownTransitionTime),
-			*DisplayClusterHelpers::str::BoolToStr(it->Value.bDown ? true : false, false),
-			*DisplayClusterHelpers::str::BoolToStr(it->Value.bDownPrevious ? true : false, false),
-			*DisplayClusterHelpers::str::BoolToStr(it->Value.bConsumed ? true : false, false),
-			*DisplayClusterTypesConverter::template ToString(it->Value.PairSampledAxes),
-			*DisplayClusterTypesConverter::template ToString(it->Value.SampleCountAccumulator),
-			*DisplayClusterTypesConverter::template ToHexString(it->Value.RawValueAccumulator));
-
-		for (int i = 0; i < IE_MAX; ++i)
+		// Output buffer
+		TArray<uint8> KeyStateBuffer;
+		FMemoryWriter KeyStateWriter(KeyStateBuffer);
+
+		// Serialized simple member variables
+		KeyStateWriter << It.Value.RawValue;
+		KeyStateWriter << It.Value.Value;
+		KeyStateWriter << It.Value.LastUpDownTransitionTime;
+		KeyStateWriter << It.Value.SampleCountAccumulator;
+		KeyStateWriter << It.Value.RawValueAccumulator;
+
+		// Serialize the bitfield
+		uint8 BitField = 0;
+		BitField |= (It.Value.bDown           & 0x01) << 0;
+		BitField |= (It.Value.bDownPrevious   & 0x01) << 1;
+		BitField |= (It.Value.bConsumed       & 0x01) << 2;
+		BitField |= (It.Value.bWasJustFlushed & 0x01) << 3;
+		BitField |= (It.Value.PairSampledAxes & 0x07) << 4;
+		KeyStateWriter << BitField;
+
+		// Serialize EventCounts arrays
+		for (int ArrayIdx = 0; ArrayIdx < IE_MAX; ++ArrayIdx)
 		{
-			StrKeyState += FString::Printf(TEXT("%s;"), *DisplayClusterHelpers::str::template ArrayToStr(it->Value.EventCounts[i], FString(","), false));
+			KeyStateWriter << It.Value.EventCounts[ArrayIdx];
 		}
 
-		for (int i = 0; i < IE_MAX; ++i)
+		// Serialize EventAccumulator arrays
+		for (int ArrayIdx = 0; ArrayIdx < IE_MAX; ++ArrayIdx)
 		{
-			StrKeyState += FString::Printf(TEXT("%s;"), *DisplayClusterHelpers::str::template ArrayToStr(it->Value.EventAccumulator[i], FString(","), false));
+			KeyStateWriter << It.Value.EventAccumulator[ArrayIdx];
 		}
 
-		OutKeyStateMap.Emplace(KeyName, StrKeyState);
+		FString KeyName = It.Key.ToString();
+		UE_LOG(LogDisplayClusterGame, VeryVerbose, TEXT("Input data serialized: %s @ %d bytes"), *KeyName, KeyStateBuffer.Num());
+		OutKeyStateMap.Emplace(MoveTemp(KeyName), MoveTemp(KeyStateBuffer));
 	}
 
 	return true;
 }
 
-bool UDisplayClusterPlayerInput::DeserializeKeyStateMap(const TMap<FString, FString>& InKeyStateMap)
+bool UDisplayClusterPlayerInput::DeserializeKeyStateMap(const TMap<FString, TArray<uint8>>& InKeyStateMap)
 {
 	TMap<FKey, FKeyState>& StateMap = GetKeyStateMap();
 
 	// Reset local key state map
 	StateMap.Reset();
 
-	int idx = 0;
-	for (auto it = InKeyStateMap.CreateConstIterator(); it; ++it)
+	// Deserialize each key state individually
+	for (const TPair<FString, TArray<uint8>>& It : InKeyStateMap)
 	{
-		UE_LOG(LogDisplayClusterGame, VeryVerbose, TEXT("Input data [%d]: %s = %s"), idx, *it->Key, *it->Value);
-		++idx;
+		UE_LOG(LogDisplayClusterGame, VeryVerbose, TEXT("Deserializing input data: %s @ %d bytes"), *It.Key, It.Value.Num());
 
-		FKey Key(*it->Key);
+		// Deserialized data
+		FKey Key(*It.Key);
 		FKeyState KeyState;
 
-		TArray<FString> Fields;
-		DisplayClusterHelpers::str::template StrToArray<FString>(it->Value, FString(";"), Fields, false);
-
-		KeyState.RawValue                 = DisplayClusterTypesConverter::template FromHexString<FVector>(Fields[0]);
-		KeyState.Value                    = DisplayClusterTypesConverter::template FromHexString<FVector>(Fields[1]);
-		KeyState.LastUpDownTransitionTime = DisplayClusterTypesConverter::template FromHexString<float>(Fields[2]);
-		KeyState.bDown                    = DisplayClusterTypesConverter::template FromString<bool>(Fields[3]) ? 1 : 0;
-		KeyState.bDownPrevious            = DisplayClusterTypesConverter::template FromString<bool>(Fields[4]) ? 1 : 0;
-		KeyState.bConsumed                = DisplayClusterTypesConverter::template FromString<bool>(Fields[5]) ? 1 : 0;
-		KeyState.PairSampledAxes          = DisplayClusterTypesConverter::template FromString<uint8>(Fields[6]);
-		KeyState.SampleCountAccumulator   = DisplayClusterTypesConverter::template FromString<uint8>(Fields[7]);
-		KeyState.RawValueAccumulator      = DisplayClusterTypesConverter::template FromHexString<FVector>(Fields[8]);
-
-		for (int i = 0; i < IE_MAX; ++i)
+		FMemoryReader KeyStateReader(It.Value);
+
+		// Deserialize simple member variables
+		KeyStateReader << KeyState.RawValue;
+		KeyStateReader << KeyState.Value;
+		KeyStateReader << KeyState.LastUpDownTransitionTime;
+		KeyStateReader << KeyState.SampleCountAccumulator;
+		KeyStateReader << KeyState.RawValueAccumulator;
+
+		// Deserialize the bitfield
+		uint8 BitField = 0;
+		KeyStateReader << BitField;
+		KeyState.bDown           = (BitField >> 0) & 0x01;
+		KeyState.bDownPrevious   = (BitField >> 1) & 0x01;
+		KeyState.bConsumed       = (BitField >> 2) & 0x01;
+		KeyState.bWasJustFlushed = (BitField >> 3) & 0x01;
+		KeyState.PairSampledAxes = (BitField >> 4) & 0x07;
+
+		// Deserialize EventCounts arrays
+		for (int ArrayIdx = 0; ArrayIdx < IE_MAX; ++ArrayIdx)
 		{
-			const FString& EventCountsStr = Fields[9 + i];
-			if (!EventCountsStr.IsEmpty())
-			{
-				DisplayClusterHelpers::str::template StrToArray<uint32>(EventCountsStr, FString(","), KeyState.EventCounts[i]);
-			}
+			KeyStateReader << KeyState.EventCounts[ArrayIdx];
+		}
 
-			const FString& EventAccumulatorStr = Fields[9 + IE_MAX + i];
-			if (!EventAccumulatorStr.IsEmpty())
-			{
-				DisplayClusterHelpers::str::template StrToArray<uint32>(Fields[9 + IE_MAX + i], FString(","), KeyState.EventAccumulator[i]);
-			}
+		// Deserialize EventAccumulator arrays
+		for (int ArrayIdx = 0; ArrayIdx < IE_MAX; ++ArrayIdx)
+		{
+			KeyStateReader << KeyState.EventAccumulator[ArrayIdx];
 		}
 
 		// Add incoming data to the local map
-		StateMap.Add(Key, KeyState);
+		StateMap.Emplace(MoveTemp(Key), MoveTemp(KeyState));
 	}
 
 	return true;
@@ -150,7 +163,7 @@ void UDisplayClusterPlayerInput::ProcessPolicy_ReplicatePrimary(const TArray<UIn
 
 	if (IPDisplayClusterClusterManager* const ClusterMgr = GDisplayCluster->GetPrivateClusterMgr())
 	{
-		TMap<FString, FString> KeyStates;
+		TMap<FString, TArray<uint8>> KeyStates;
 
 		// Always cache local inputs whether we're primary or not. If we're secondary now,
 		// we may turn into primary in a moment, and the data would be available for others.
diff --git a/Engine/Plugins/Runtime/nDisplay/Source/DisplayCluster/Private/Game/EngineClasses/Scene/DisplayClusterSceneComponentSync.cpp b/Engine/Plugins/Runtime/nDisplay/Source/DisplayCluster/Private/Game/EngineClasses/Scene/DisplayClusterSceneComponentSync.cpp
--- a/Engine/Plugins/Runtime/nDisplay/Source/DisplayCluster/Private/Game/EngineClasses/Scene/DisplayClusterSceneComponentSync.cpp
+++ b/Engine/Plugins/Runtime/nDisplay/Source/DisplayCluster/Private/Game/EngineClasses/Scene/DisplayClusterSceneComponentSync.cpp
@@ -5,18 +5,14 @@
 #include "DisplayClusterEnums.h"
 
 #include "GameFramework/Actor.h"
-
 #include "Cluster/IDisplayClusterClusterManager.h"
 
 #include "Misc/DisplayClusterGlobals.h"
 #include "Misc/DisplayClusterLog.h"
 #include "Misc/DisplayClusterTypesConverter.h"
 
+#include "Serialization/Archive.h"
 
-UDisplayClusterSceneComponentSync::UDisplayClusterSceneComponentSync(const FObjectInitializer& ObjectInitializer)
-	: Super(ObjectInitializer)
-{
-}
 
 void UDisplayClusterSceneComponentSync::BeginPlay()
 {
@@ -57,34 +53,29 @@ void UDisplayClusterSceneComponentSync::EndPlay(const EEndPlayReason::Type EndPl
 	Super::EndPlay(EndPlayReason);
 }
 
-
-//////////////////////////////////////////////////////////////////////////////////////////////
-// IDisplayClusterClusterSyncObject
-//////////////////////////////////////////////////////////////////////////////////////////////
 bool UDisplayClusterSceneComponentSync::IsActive() const
 {
 	return IsValidChecked(this);
 }
 
-
-//////////////////////////////////////////////////////////////////////////////////////////////
-// IDisplayClusterStringSerializable
-//////////////////////////////////////////////////////////////////////////////////////////////
 FString UDisplayClusterSceneComponentSync::GenerateSyncId()
 {
 	return FString::Printf(TEXT("S_%s"), *GetFullName());
 }
 
-FString UDisplayClusterSceneComponentSync::SerializeToString() const
-{
-	return DisplayClusterTypesConverter::template ToHexString(GetSyncTransform());
-}
-
-bool UDisplayClusterSceneComponentSync::DeserializeFromString(const FString& data)
+void UDisplayClusterSceneComponentSync::SerializeDC(FArchive& Ar)
 {
-	FTransform NewTransform = DisplayClusterTypesConverter::template FromHexString<FTransform>(data);
-	UE_LOG(LogDisplayClusterGame, Verbose, TEXT("%s: applying transform data <%s>"), *SyncId, *NewTransform.ToHumanReadableString());
-	SetSyncTransform(NewTransform);
+	FTransform Transform;
 
-	return true;
+	if (Ar.IsLoading())
+	{
+		Ar << Transform;
+		UE_LOG(LogDisplayClusterGame, Verbose, TEXT("%s: applying transform data <%s>"), *SyncId, *Transform.ToHumanReadableString());
+		SetSyncTransform(Transform);
+	}
+	else
+	{
+		Transform = GetSyncTransform();
+		Ar << Transform;
+	}
 }
diff --git a/Engine/Plugins/Runtime/nDisplay/Source/DisplayCluster/Private/Network/Packet/DisplayClusterPacketInternal.cpp b/Engine/Plugins/Runtime/nDisplay/Source/DisplayCluster/Private/Network/Packet/DisplayClusterPacketInternal.cpp
--- a/Engine/Plugins/Runtime/nDisplay/Source/DisplayCluster/Private/Network/Packet/DisplayClusterPacketInternal.cpp
+++ b/Engine/Plugins/Runtime/nDisplay/Source/DisplayCluster/Private/Network/Packet/DisplayClusterPacketInternal.cpp
@@ -9,46 +9,23 @@
 #include "Misc/DisplayClusterLog.h"
 
 
-bool FDisplayClusterPacketInternal::Serialize(FMemoryWriter& Arch)
+void FDisplayClusterPacketInternal::SerializeDC(FArchive& Ar)
 {
 	// Header
-	Arch << Name;
-	Arch << Type;
-	Arch << Protocol;
+	Ar << Name;
+	Ar << Type;
+	Ar << Protocol;
 
 	// Comm result
-	Arch << CommResult;
+	Ar << CommResult;
 
 	// Arguments
-	Arch << TextArguments;
-	Arch << BinaryArguments;
+	Ar << TextArguments;
+	Ar << BinaryArguments;
 
 	// Objects
-	Arch << TextObjects;
-	Arch << BinaryObjects;
-
-	return true;
-}
-
-bool FDisplayClusterPacketInternal::Deserialize(FMemoryReader& Arch)
-{
-	// Header
-	Arch << Name;
-	Arch << Type;
-	Arch << Protocol;
-
-	// Comm result
-	Arch << CommResult;
-
-	// Arguments
-	Arch << TextArguments;
-	Arch << BinaryArguments;
-
-	// Objects
-	Arch << TextObjects;
-	Arch << BinaryObjects;
-
-	return true;
+	Ar << TextObjects;
+	Ar << BinaryObjects;
 }
 
 FString FDisplayClusterPacketInternal::ToString() const
@@ -96,11 +73,7 @@ bool FDisplayClusterPacketInternal::SendPacket(FDisplayClusterSocketOperations&
 	MemoryWriter.Seek(sizeof(FPacketHeader));
 
 	// Serialize the packet body
-	if (!Serialize(MemoryWriter))
-	{
-		UE_LOG(LogDisplayClusterNetwork, Error, TEXT("%s - couldn't serialize a packet"), *SocketOps.GetConnectionName());
-		return false;
-	}
+	SerializeDC(MemoryWriter);
 
 	// Initialize the packet header
 	FPacketHeader PacketHeader;
@@ -111,7 +84,8 @@ bool FDisplayClusterPacketInternal::SendPacket(FDisplayClusterSocketOperations&
 	FMemory::Memcpy(DataBuffer.GetData(), &PacketHeader, sizeof(FPacketHeader));
 
 	// Send the header
-	if (!SocketOps.SendChunk(DataBuffer, DataBuffer.Num(), FString("send-internal-msg")))
+	static const FString SendMsgLogStr(TEXT("send-internal-msg"));
+	if (!SocketOps.SendChunk(DataBuffer, DataBuffer.Num(), SendMsgLogStr))
 	{
 		UE_LOG(LogDisplayClusterNetwork, Warning, TEXT("%s - Couldn't send a packet"), *SocketOps.GetConnectionName());
 		return false;
@@ -137,7 +111,8 @@ bool FDisplayClusterPacketInternal::RecvPacket(FDisplayClusterSocketOperations&
 
 	// Read packet header
 	DataBuffer.Reset();
-	if (!SocketOps.RecvChunk(DataBuffer, sizeof(FPacketHeader), FString("recv-internal-chunk-header")))
+	static const FString RecvChunkHeaderLogStr(TEXT("recv-internal-chunk-header"));
+	if (!SocketOps.RecvChunk(DataBuffer, sizeof(FPacketHeader), RecvChunkHeaderLogStr))
 	{
 		UE_LOG(LogDisplayClusterNetwork, Warning, TEXT("%s couldn't receive packet header"), *SocketOps.GetConnectionName());
 		return false;
@@ -152,7 +127,8 @@ bool FDisplayClusterPacketInternal::RecvPacket(FDisplayClusterSocketOperations&
 	check(PacketHeader.PacketBodyLength > 0);
 
 	// Read packet body
-	if (!SocketOps.RecvChunk(DataBuffer, PacketHeader.PacketBodyLength, FString("recv-internal-chunk-body")))
+	static const FString RecvChunkBodyLogStr(TEXT("recv-internal-chunk-body"));
+	if (!SocketOps.RecvChunk(DataBuffer, PacketHeader.PacketBodyLength, RecvChunkBodyLogStr))
 	{
 		UE_LOG(LogDisplayClusterNetwork, Warning, TEXT("%s couldn't receive packet body"), *SocketOps.GetConnectionName());
 		return false;
@@ -164,12 +140,8 @@ bool FDisplayClusterPacketInternal::RecvPacket(FDisplayClusterSocketOperations&
 	UE_LOG(LogDisplayClusterNetwork, VeryVerbose, TEXT("%s - packet body received"), *SocketOps.GetConnectionName());
 
 	// Deserialize packet from buffer
-	FMemoryReader Arch = FMemoryReader(DataBuffer, false);
-	if (!Deserialize(Arch))
-	{
-		UE_LOG(LogDisplayClusterNetwork, Error, TEXT("%s couldn't deserialize a packet"), *SocketOps.GetConnectionName());
-		return false;
-	}
+	FMemoryReader MemoryReader(DataBuffer);
+	SerializeDC(MemoryReader);
 
 	// Succeeded
 	return true;
diff --git a/Engine/Plugins/Runtime/nDisplay/Source/DisplayCluster/Private/Network/Packet/DisplayClusterPacketInternal.h b/Engine/Plugins/Runtime/nDisplay/Source/DisplayCluster/Private/Network/Packet/DisplayClusterPacketInternal.h
--- a/Engine/Plugins/Runtime/nDisplay/Source/DisplayCluster/Private/Network/Packet/DisplayClusterPacketInternal.h
+++ b/Engine/Plugins/Runtime/nDisplay/Source/DisplayCluster/Private/Network/Packet/DisplayClusterPacketInternal.h
@@ -3,11 +3,12 @@
 #pragma once
 
 #include "Network/Packet/IDisplayClusterPacket.h"
-#include "Network/DisplayClusterNetworkTypes.h"
+#include "IDisplayClusterSerializable.h"
+
 #include "Misc/DisplayClusterTypesConverter.h"
+#include "Network/DisplayClusterNetworkTypes.h"
 
-class FMemoryWriter;
-class FMemoryReader;
+class FArchive;
 class FDisplayClusterSocketOperations;
 
 
@@ -16,6 +17,7 @@ class FDisplayClusterSocketOperations;
  */
 class FDisplayClusterPacketInternal
 	: public IDisplayClusterPacket
+	, public IDisplayClusterSerializable
 {
 public:
 	FDisplayClusterPacketInternal()
@@ -109,41 +111,37 @@ public:
 	}
 
 	// Set all text arguments to a specified section
-	void SetTextArgs(const FString& SectionName, const TMap<FString, FString>& Data)
+	void SetTextArgs(const FString& SectionName, TMap<FString, FString> Data)
 	{
-		TextArguments.Emplace(SectionName, Data);
+		TextArguments.Emplace(SectionName, MoveTemp(Data));
 	}
 
 public:
 	// Get specific binary argument from the packet
-	bool GetBinArg(const FString& SectionName, const FString& ArgName, TArray<uint8>& ArgVal)
+	bool GetBinArg(const FString& SectionName, const FString& ArgName, TArray<uint8>& ArgVal, bool bKeepData = false)
 	{
-		const TMap<FString, TArray<uint8>>* Section = BinaryArguments.Find(SectionName);
+		TMap<FString, TArray<uint8>>* Section = BinaryArguments.Find(SectionName);
 		if (!Section)
 		{
 			return false;
 		}
 
-		const TArray<uint8>* Value = Section->Find(ArgName);
-		if (!Value)
+		if (bKeepData)
 		{
-			return false;
+			const TArray<uint8>* Value = Section->Find(ArgName);
+			ArgVal = Value ? *Value : TArray<uint8>();
+			return Value != nullptr;
+		}
+		else
+		{
+			return Section->RemoveAndCopyValue(ArgName, ArgVal);
 		}
-
-		ArgVal = *Value;
-		return true;
 	}
 
 	// Set new binary argument to the packet
-	void SetBinArg(const FString& SectionName, const FString& ArgName, const TArray<uint8>& ArgVal)
+	void SetBinArg(const FString& SectionName, const FString& ArgName, TArray<uint8> ArgVal)
 	{
-		TMap<FString, TArray<uint8>>* Section = BinaryArguments.Find(SectionName);
-		if (!Section)
-		{
-			Section = &BinaryArguments.Emplace(SectionName);
-		}
-
-		Section->Emplace(ArgName, ArgVal);
+		BinaryArguments.FindOrAdd(SectionName).Emplace(ArgName, MoveTemp(ArgVal));
 	}
 
 	// Remove binary argument
@@ -157,16 +155,25 @@ public:
 	}
 
 	// Get all binary arguments of a specified section
-	const TMap<FString, TArray<uint8>> GetBinArgs(const FString& SectionName)
+	TMap<FString, TArray<uint8>> GetBinArgs(const FString& SectionName, bool bKeepData = false)
 	{
-		const TMap<FString, TArray<uint8>>* Section = BinaryArguments.Find(SectionName);
-		return Section ? *Section : TMap<FString, TArray<uint8>>();
+		if (bKeepData)
+		{
+			TMap<FString, TArray<uint8>>* Section = BinaryArguments.Find(SectionName);
+			return Section ? *Section : TMap<FString, TArray<uint8>>();
+		}
+		else
+		{
+			TMap<FString, TArray<uint8>> SectionData;
+			BinaryArguments.RemoveAndCopyValue(SectionName, SectionData);
+			return MoveTemp(SectionData);
+		}
 	}
 
 	// Set all binary arguments to a specified section
-	void SetBinArgs(const FString& SectionName, const TMap<FString, TArray<uint8>>& Data)
+	void SetBinArgs(const FString& SectionName, TMap<FString, TArray<uint8>> Data)
 	{
-		BinaryArguments.Emplace(SectionName, Data);
+		BinaryArguments.Emplace(SectionName, MoveTemp(Data));
 	}
 
 public:
@@ -192,22 +199,15 @@ public:
 		}
 		else
 		{
-			TArray<FString>* Objects = TextObjects.Find(SectionName);
-			OutObjects = (Objects ? MoveTemp(*Objects) : TArray<FString>());
+			TextObjects.RemoveAndCopyValue(SectionName, OutObjects);
 		}
 	}
 
 public:
 	// Adds a binary object to the packet
-	void AddBinObject(const FString& SectionName, const TArray<uint8>& Object)
+	void AddBinObject(const FString& SectionName, TArray<uint8> Object)
 	{
-		TArray<TArray<uint8>>* Section = BinaryObjects.Find(SectionName);
-		if (!Section)
-		{
-			Section = &BinaryObjects.Emplace(SectionName);
-		}
-
-		Section->Add(Object);
+		BinaryObjects.FindOrAdd(SectionName).Add(MoveTemp(Object));
 	}
 
 	// Get all binary objects
@@ -220,23 +220,25 @@ public:
 		}
 		else
 		{
-			TArray<TArray<uint8>>* Objects = BinaryObjects.Find(SectionName);
-			OutObjects = (Objects ? MoveTemp(*Objects) : TArray<TArray<uint8>>());
+			BinaryObjects.RemoveAndCopyValue(SectionName, OutObjects);
 		}
 	}
 
 public:
-	//////////////////////////////////////////////////////////////////////////////////////////////
-	// IDisplayClusterPacket
-	//////////////////////////////////////////////////////////////////////////////////////////////
+
+	//~ Begin IDisplayClusterPacket
 	virtual bool SendPacket(FDisplayClusterSocketOperations& SocketOps) override;
 	virtual bool RecvPacket(FDisplayClusterSocketOperations& SocketOps) override;
 	virtual FString ToLogString(bool bDetailed = false) const override;
+	//~ End IDisplayClusterPacket
+
+public:
+
+	//~ Begin IDisplayClusterSerializable
+	virtual void SerializeDC(FArchive& Ar) override;
+	//~ End IDisplayClusterSerializable
 
 protected:
-	// Serialization
-	bool Serialize(FMemoryWriter& Arch);
-	bool Deserialize(FMemoryReader& Arch);
 
 	// Text representation of the whole meessage
 	FString ToString() const;
diff --git a/Engine/Plugins/Runtime/nDisplay/Source/DisplayCluster/Private/Network/Protocol/IDisplayClusterProtocolClusterSync.h b/Engine/Plugins/Runtime/nDisplay/Source/DisplayCluster/Private/Network/Protocol/IDisplayClusterProtocolClusterSync.h
--- a/Engine/Plugins/Runtime/nDisplay/Source/DisplayCluster/Private/Network/Protocol/IDisplayClusterProtocolClusterSync.h
+++ b/Engine/Plugins/Runtime/nDisplay/Source/DisplayCluster/Private/Network/Protocol/IDisplayClusterProtocolClusterSync.h
@@ -37,13 +37,13 @@ public:
 	virtual EDisplayClusterCommResult GetTimeData(double& OutDeltaTime, double& OutGameTime, TOptional<FQualifiedFrameTime>& OutFrameTime) = 0;
 
 	/** Custom objects syncrhonization */
-	virtual EDisplayClusterCommResult GetObjectsData(const EDisplayClusterSyncGroup InSyncGroup, TMap<FString, FString>& OutObjectsData) = 0;
+	virtual EDisplayClusterCommResult GetObjectsData(EDisplayClusterSyncGroup InSyncGroup, TMap<FString, TArray<uint8>>& OutObjectsData) = 0;
 
 	/** Cluster events syncrhonization */
 	virtual EDisplayClusterCommResult GetEventsData(TArray<TSharedPtr<FDisplayClusterClusterEventJson, ESPMode::ThreadSafe>>& OutJsonEvents, TArray<TSharedPtr<FDisplayClusterClusterEventBinary, ESPMode::ThreadSafe>>& OutBinaryEvents) = 0;
 
 	/** Native UE input syncrhonization */
-	virtual EDisplayClusterCommResult GetNativeInputData(TMap<FString, FString>& OutNativeInputData) = 0;
+	virtual EDisplayClusterCommResult GetNativeInputData(TMap<FString, TArray<uint8>>& OutNativeInputData) = 0;
 
 	/** Custom states synchronization */
 	virtual EDisplayClusterCommResult PropagateStatesData(const TArray<uint8>& InLocalStatesData, TArray<uint8>& OutClusterStatesData) = 0;
diff --git a/Engine/Plugins/Runtime/nDisplay/Source/DisplayCluster/Private/Network/Service/ClusterSync/DisplayClusterClusterSyncClient.cpp b/Engine/Plugins/Runtime/nDisplay/Source/DisplayCluster/Private/Network/Service/ClusterSync/DisplayClusterClusterSyncClient.cpp
--- a/Engine/Plugins/Runtime/nDisplay/Source/DisplayCluster/Private/Network/Service/ClusterSync/DisplayClusterClusterSyncClient.cpp
+++ b/Engine/Plugins/Runtime/nDisplay/Source/DisplayCluster/Private/Network/Service/ClusterSync/DisplayClusterClusterSyncClient.cpp
@@ -187,7 +187,7 @@ EDisplayClusterCommResult FDisplayClusterClusterSyncClient::GetTimeData(double&
 	return Response->GetCommResult();
 }
 
-EDisplayClusterCommResult FDisplayClusterClusterSyncClient::GetObjectsData(const EDisplayClusterSyncGroup InSyncGroup, TMap<FString, FString>& OutObjectsData)
+EDisplayClusterCommResult FDisplayClusterClusterSyncClient::GetObjectsData(EDisplayClusterSyncGroup InSyncGroup, TMap<FString, TArray<uint8>>& OutObjectsData)
 {
 	static TSharedPtr<FDisplayClusterPacketInternal> Request = MakeShared<FDisplayClusterPacketInternal>(
 			DisplayClusterClusterSyncStrings::GetObjectsData::Name,
@@ -210,7 +210,7 @@ EDisplayClusterCommResult FDisplayClusterClusterSyncClient::GetObjectsData(const
 	}
 
 	// Extract data from response packet
-	OutObjectsData = Response->GetTextArgs(DisplayClusterClusterSyncStrings::ArgumentsDefaultCategory);
+	OutObjectsData = Response->GetBinArgs(DisplayClusterClusterSyncStrings::ArgumentsDefaultCategory);
 
 	return Response->GetCommResult();
 }
@@ -242,7 +242,7 @@ EDisplayClusterCommResult FDisplayClusterClusterSyncClient::GetEventsData(TArray
 	return Response->GetCommResult();
 }
 
-EDisplayClusterCommResult FDisplayClusterClusterSyncClient::GetNativeInputData(TMap<FString, FString>& OutNativeInputData)
+EDisplayClusterCommResult FDisplayClusterClusterSyncClient::GetNativeInputData(TMap<FString, TArray<uint8>>& OutNativeInputData)
 {
 	static const TSharedPtr<FDisplayClusterPacketInternal> Request = MakeShared<FDisplayClusterPacketInternal>(
 			DisplayClusterClusterSyncStrings::GetNativeInputData::Name,
@@ -263,7 +263,7 @@ EDisplayClusterCommResult FDisplayClusterClusterSyncClient::GetNativeInputData(T
 	}
 
 	// Extract data from response packet
-	OutNativeInputData = Response->GetTextArgs(DisplayClusterClusterSyncStrings::ArgumentsDefaultCategory);
+	OutNativeInputData = Response->GetBinArgs(DisplayClusterClusterSyncStrings::ArgumentsDefaultCategory);
 
 	return Response->GetCommResult();
 }
diff --git a/Engine/Plugins/Runtime/nDisplay/Source/DisplayCluster/Private/Network/Service/ClusterSync/DisplayClusterClusterSyncClient.h b/Engine/Plugins/Runtime/nDisplay/Source/DisplayCluster/Private/Network/Service/ClusterSync/DisplayClusterClusterSyncClient.h
--- a/Engine/Plugins/Runtime/nDisplay/Source/DisplayCluster/Private/Network/Service/ClusterSync/DisplayClusterClusterSyncClient.h
+++ b/Engine/Plugins/Runtime/nDisplay/Source/DisplayCluster/Private/Network/Service/ClusterSync/DisplayClusterClusterSyncClient.h
@@ -31,9 +31,9 @@ public:
 	virtual EDisplayClusterCommResult WaitForFrameStart() override;
 	virtual EDisplayClusterCommResult WaitForFrameEnd() override;
 	virtual EDisplayClusterCommResult GetTimeData(double& OutDeltaTime, double& OutGameTime, TOptional<FQualifiedFrameTime>& OutFrameTime) override;
-	virtual EDisplayClusterCommResult GetObjectsData(const EDisplayClusterSyncGroup InSyncGroup, TMap<FString, FString>& OutObjectsData) override;
+	virtual EDisplayClusterCommResult GetObjectsData(EDisplayClusterSyncGroup InSyncGroup, TMap<FString, TArray<uint8>>& OutObjectsData) override;
 	virtual EDisplayClusterCommResult GetEventsData(TArray<TSharedPtr<FDisplayClusterClusterEventJson, ESPMode::ThreadSafe>>& OutJsonEvents, TArray<TSharedPtr<FDisplayClusterClusterEventBinary, ESPMode::ThreadSafe>>& OutBinaryEvents) override;
-	virtual EDisplayClusterCommResult GetNativeInputData(TMap<FString, FString>& OutNativeInputData) override;
+	virtual EDisplayClusterCommResult GetNativeInputData(TMap<FString, TArray<uint8>>& OutNativeInputData) override;
 	virtual EDisplayClusterCommResult PropagateStatesData(const TArray<uint8>& InLocalStatesData, TArray<uint8>& OutClusterStatesData) override;
 	//~ End IDisplayClusterProtocolClusterSync
 };
diff --git a/Engine/Plugins/Runtime/nDisplay/Source/DisplayCluster/Private/Network/Service/ClusterSync/DisplayClusterClusterSyncService.cpp b/Engine/Plugins/Runtime/nDisplay/Source/DisplayCluster/Private/Network/Service/ClusterSync/DisplayClusterClusterSyncService.cpp
--- a/Engine/Plugins/Runtime/nDisplay/Source/DisplayCluster/Private/Network/Service/ClusterSync/DisplayClusterClusterSyncService.cpp
+++ b/Engine/Plugins/Runtime/nDisplay/Source/DisplayCluster/Private/Network/Service/ClusterSync/DisplayClusterClusterSyncService.cpp
@@ -443,14 +443,14 @@ TSharedPtr<FDisplayClusterPacketInternal> FDisplayClusterClusterSyncService::Pro
 	}
 	else if (ReqName.Equals(DisplayClusterClusterSyncStrings::GetObjectsData::Name, ESearchCase::IgnoreCase))
 	{
-		TMap<FString, FString> ObjectsData;
 		uint8 SyncGroupNum = 0;
 		Request->GetTextArg(DisplayClusterClusterSyncStrings::ArgumentsDefaultCategory, DisplayClusterClusterSyncStrings::GetObjectsData::ArgSyncGroup, SyncGroupNum);
 		EDisplayClusterSyncGroup SyncGroup = (EDisplayClusterSyncGroup)SyncGroupNum;
 
+		TMap<FString, TArray<uint8>> ObjectsData;
 		const EDisplayClusterCommResult CommResult = GetObjectsData(SyncGroup, ObjectsData);
 
-		Response->SetTextArgs(DisplayClusterClusterSyncStrings::ArgumentsDefaultCategory, ObjectsData);
+		Response->SetBinArgs(DisplayClusterClusterSyncStrings::ArgumentsDefaultCategory, MoveTemp(ObjectsData));
 
 		Response->SetCommResult(CommResult);
 
@@ -472,11 +472,11 @@ TSharedPtr<FDisplayClusterPacketInternal> FDisplayClusterClusterSyncService::Pro
 	}
 	else if (ReqName.Equals(DisplayClusterClusterSyncStrings::GetNativeInputData::Name, ESearchCase::IgnoreCase))
 	{
-		TMap<FString, FString> NativeInputData;
+		TMap<FString, TArray<uint8>> NativeInputData;
 
 		const EDisplayClusterCommResult CommResult = GetNativeInputData(NativeInputData);
 
-		Response->SetTextArgs(DisplayClusterClusterSyncStrings::ArgumentsDefaultCategory, NativeInputData);
+		Response->SetBinArgs(DisplayClusterClusterSyncStrings::ArgumentsDefaultCategory, MoveTemp(NativeInputData));
 
 		Response->SetCommResult(CommResult);
 
@@ -546,7 +546,7 @@ EDisplayClusterCommResult FDisplayClusterClusterSyncService::GetTimeData(double&
 	return Handler.GetTimeData(OutDeltaTime, OutGameTime, OutFrameTime);
 }
 
-EDisplayClusterCommResult FDisplayClusterClusterSyncService::GetObjectsData(EDisplayClusterSyncGroup InSyncGroup, TMap<FString, FString>& OutObjectsData)
+EDisplayClusterCommResult FDisplayClusterClusterSyncService::GetObjectsData(EDisplayClusterSyncGroup InSyncGroup, TMap<FString, TArray<uint8>>& OutObjectsData)
 {
 	TRACE_CPUPROFILER_EVENT_SCOPE(SRV_CS::GetObjectsData);
 
@@ -568,7 +568,7 @@ EDisplayClusterCommResult FDisplayClusterClusterSyncService::GetEventsData(TArra
 	return Handler.GetEventsData(OutJsonEvents, OutBinaryEvents);
 }
 
-EDisplayClusterCommResult FDisplayClusterClusterSyncService::GetNativeInputData(TMap<FString, FString>& OutNativeInputData)
+EDisplayClusterCommResult FDisplayClusterClusterSyncService::GetNativeInputData(TMap<FString, TArray<uint8>>& OutNativeInputData)
 {
 	TRACE_CPUPROFILER_EVENT_SCOPE(SRV_CS::GetNativeInputData);
 
diff --git a/Engine/Plugins/Runtime/nDisplay/Source/DisplayCluster/Private/Network/Service/ClusterSync/DisplayClusterClusterSyncService.h b/Engine/Plugins/Runtime/nDisplay/Source/DisplayCluster/Private/Network/Service/ClusterSync/DisplayClusterClusterSyncService.h
--- a/Engine/Plugins/Runtime/nDisplay/Source/DisplayCluster/Private/Network/Service/ClusterSync/DisplayClusterClusterSyncService.h
+++ b/Engine/Plugins/Runtime/nDisplay/Source/DisplayCluster/Private/Network/Service/ClusterSync/DisplayClusterClusterSyncService.h
@@ -72,9 +72,9 @@ protected:
 	virtual EDisplayClusterCommResult WaitForFrameStart() override;
 	virtual EDisplayClusterCommResult WaitForFrameEnd() override;
 	virtual EDisplayClusterCommResult GetTimeData(double& OutDeltaTime, double& OutGameTime, TOptional<FQualifiedFrameTime>& OutFrameTime) override;
-	virtual EDisplayClusterCommResult GetObjectsData(EDisplayClusterSyncGroup InSyncGroup, TMap<FString, FString>& OutObjectsData)  override;
+	virtual EDisplayClusterCommResult GetObjectsData(EDisplayClusterSyncGroup InSyncGroup, TMap<FString, TArray<uint8>>& OutObjectsData) override;
 	virtual EDisplayClusterCommResult GetEventsData(TArray<TSharedPtr<FDisplayClusterClusterEventJson, ESPMode::ThreadSafe>>& OutJsonEvents, TArray<TSharedPtr<FDisplayClusterClusterEventBinary, ESPMode::ThreadSafe>>& OutBinaryEvents) override;
-	virtual EDisplayClusterCommResult GetNativeInputData(TMap<FString, FString>& OutNativeInputData) override;
+	virtual EDisplayClusterCommResult GetNativeInputData(TMap<FString, TArray<uint8>>& OutNativeInputData) override;
 	virtual EDisplayClusterCommResult PropagateStatesData(const TArray<uint8>& InLocalStatesData, TArray<uint8>& OutClusterStatesData) override;
 	//~ End IDisplayClusterProtocolClusterSync
 
diff --git a/Engine/Plugins/Runtime/nDisplay/Source/DisplayCluster/Public/Cluster/IDisplayClusterClusterSyncObject.h b/Engine/Plugins/Runtime/nDisplay/Source/DisplayCluster/Public/Cluster/IDisplayClusterClusterSyncObject.h
--- a/Engine/Plugins/Runtime/nDisplay/Source/DisplayCluster/Public/Cluster/IDisplayClusterClusterSyncObject.h
+++ b/Engine/Plugins/Runtime/nDisplay/Source/DisplayCluster/Public/Cluster/IDisplayClusterClusterSyncObject.h
@@ -2,6 +2,7 @@
 
 #pragma once
 
+#include "IDisplayClusterSerializable.h"
 #include "IDisplayClusterStringSerializable.h"
 
 
@@ -9,18 +10,28 @@
  * Synchronizable object interface
  */
 class IDisplayClusterClusterSyncObject
-	: public IDisplayClusterStringSerializable
+	: public IDisplayClusterSerializable
+	, public IDisplayClusterStringSerializable // Deprecated, don't override any of its functions
 {
 public:
+
 	virtual ~IDisplayClusterClusterSyncObject() = default;
 
 public:
-	// Need to sync this object?
+
+	/** Whether an object needs to be syncrhonized on current frame */
 	virtual bool IsActive() const = 0;
-	// Unique ID of synced object
+
+	/** Unique ID of the object being synchronized */
 	virtual FString GetSyncId() const = 0;
-	// Check if object has changed since last ClearDirty call
-	virtual bool IsDirty() const = 0;
-	// Cleans dirty flag making it 'not changed yet'
-	virtual void ClearDirty() = 0;
+
+	/** Whether the object's state has changed since last ClearDirty call */
+	virtual bool IsDirty() const
+	{
+		return false;
+	}
+
+	/** Resets the dirty flag, marking the state as not yet changed */
+	virtual void ClearDirty()
+	{ }
 };
diff --git a/Engine/Plugins/Runtime/nDisplay/Source/DisplayCluster/Public/Components/DisplayClusterSceneComponentSync.h b/Engine/Plugins/Runtime/nDisplay/Source/DisplayCluster/Public/Components/DisplayClusterSceneComponentSync.h
--- a/Engine/Plugins/Runtime/nDisplay/Source/DisplayCluster/Public/Components/DisplayClusterSceneComponentSync.h
+++ b/Engine/Plugins/Runtime/nDisplay/Source/DisplayCluster/Public/Components/DisplayClusterSceneComponentSync.h
@@ -18,15 +18,9 @@ class DISPLAYCLUSTER_API UDisplayClusterSceneComponentSync
 	GENERATED_BODY()
 
 public:
-	UDisplayClusterSceneComponentSync(const FObjectInitializer& ObjectInitializer);
-	
-	virtual ~UDisplayClusterSceneComponentSync()
-	{ }
 
-public:
-	//////////////////////////////////////////////////////////////////////////////////////////////
-	// IDisplayClusterClusterSyncObject
-	//////////////////////////////////////////////////////////////////////////////////////////////
+	//~ Begin IDisplayClusterClusterSyncObject
+
 	virtual bool IsActive() const override;
 	
 	virtual FString GetSyncId() const override
@@ -39,35 +33,45 @@ public:
 		return true;
 	}
 
-	virtual void ClearDirty() override
-	{ }
+	//~ End IDisplayClusterClusterSyncObject
 
 public:
-	//////////////////////////////////////////////////////////////////////////////////////////////
-	// IDisplayClusterStringSerializable
-	//////////////////////////////////////////////////////////////////////////////////////////////
-	virtual FString SerializeToString() const override;
-	virtual bool    DeserializeFromString(const FString& data) override;
+
+	//~ Begin IDisplayClusterSerializable
+	virtual void SerializeDC(FArchive& Ar) override;
+	//~ End IDisplayClusterSerializable
 
 public:
+
+	//~ Begin USceneComponent
 	virtual void BeginPlay() override;
 	virtual void EndPlay(const EEndPlayReason::Type EndPlayReason) override;
+	//~ End USceneComponent
 
 protected:
+
+	// Generates sync ID of this instance
 	virtual FString GenerateSyncId();
 
+	// Lets children provide their transforms for synchronization
 	virtual FTransform GetSyncTransform() const
-	{ return FTransform(); }
+	{
+		return FTransform();
+	}
 
+	// Lets children update thir transforms with new data
 	virtual void SetSyncTransform(const FTransform& t)
 	{ }
 
 protected:
+
 	// Caching state
 	FVector  LastSyncLoc;
 	FRotator LastSyncRot;
 	FVector  LastSyncScale;
 
 private:
+
+	/** Unique synchronization ID of this instance */
 	FString SyncId;
 };
diff --git a/Engine/Plugins/Runtime/nDisplay/Source/DisplayCluster/Public/DisplayClusterPlayerInput.h b/Engine/Plugins/Runtime/nDisplay/Source/DisplayCluster/Public/DisplayClusterPlayerInput.h
--- a/Engine/Plugins/Runtime/nDisplay/Source/DisplayCluster/Public/DisplayClusterPlayerInput.h
+++ b/Engine/Plugins/Runtime/nDisplay/Source/DisplayCluster/Public/DisplayClusterPlayerInput.h
@@ -35,10 +35,10 @@ protected:
 protected:
 
 	/** Auxiliary function that serializes whole input data into a byte stream */
-	bool SerializeKeyStateMap(TMap<FString, FString>& OutKeyStateMap);
+	bool SerializeKeyStateMap(TMap<FString, TArray<uint8>>& OutKeyStateMap);
 
 	/** Aux function that deserializes whole input data from a map */
-	bool DeserializeKeyStateMap(const TMap<FString, FString>& InKeyStateMap);
+	bool DeserializeKeyStateMap(const TMap<FString, TArray<uint8>>& InKeyStateMap);
 
 private:
 
diff --git a/Engine/Plugins/Runtime/nDisplay/Source/DisplayCluster/Public/IDisplayClusterSerializable.h b/Engine/Plugins/Runtime/nDisplay/Source/DisplayCluster/Public/IDisplayClusterSerializable.h
--- a/Engine/Plugins/Runtime/nDisplay/Source/DisplayCluster/Public/IDisplayClusterSerializable.h
+++ b/Engine/Plugins/Runtime/nDisplay/Source/DisplayCluster/Public/IDisplayClusterSerializable.h
@@ -2,19 +2,39 @@
 
 #pragma once
 
-#include "Serialization/MemoryReader.h"
-#include "Serialization/MemoryWriter.h"
+class FArchive;
 
 
 /**
- * Memory serialization interface
+ * nDisplay serialization interface
  */
 class IDisplayClusterSerializable
 {
 public:
+
 	virtual ~IDisplayClusterSerializable() = default;
 
 public:
-	virtual bool Serialize  (FMemoryWriter& ar) = 0;
-	virtual bool Deserialize(FMemoryReader& ar) = 0;
+
+	/**
+	 * Serialize / deserialize method
+	 * 
+	 * Like the common UE serialization approach, this one is used for both serialization and
+	 * deserialization. To avoid any method ambiguity issues with UObject::Serialize, this one
+	 * has the 'DC' postfix. For example, UObject subobjects that also inherit from
+	 * IDisplayClusterSerializable, can safely go like this:
+	 * 
+	 * virtual void SomeUObject::Serialize(FArchive& Ar) override
+	 * {
+	 *    // Serialization implementation
+	 * }
+	 * 
+	 * virtual void SomeUObject::SerializeDC(FArchive& Ar) override
+	 * {
+	 *    Serialize(Ar);
+	 * }
+	 * 
+	 * Or vice versa, call 'SerializeDC' from 'Serialize'.
+	 */
+	virtual void SerializeDC(FArchive& Ar) = 0;
 };
diff --git a/Engine/Plugins/Runtime/nDisplay/Source/DisplayCluster/Public/IDisplayClusterStringSerializable.h b/Engine/Plugins/Runtime/nDisplay/Source/DisplayCluster/Public/IDisplayClusterStringSerializable.h
--- a/Engine/Plugins/Runtime/nDisplay/Source/DisplayCluster/Public/IDisplayClusterStringSerializable.h
+++ b/Engine/Plugins/Runtime/nDisplay/Source/DisplayCluster/Public/IDisplayClusterStringSerializable.h
@@ -5,14 +5,28 @@
 #include "CoreMinimal.h"
 
 /**
- * String serialization interface
+ * DEPRECATED
+ * 
+ * This interface is deprecated and no longer referenced. Its methods are never invoked and will be removed soon.
+ * Refer to the IDisplayClusterSerializable interface, which provides common UE binary serialization.
  */
 class IDisplayClusterStringSerializable
 {
 public:
+
 	virtual ~IDisplayClusterStringSerializable() = default;
 
 public:
-	virtual FString SerializeToString() const = 0;
-	virtual bool DeserializeFromString(const FString& ar) = 0;
+
+	UE_DEPRECATED(5.8, "This method is deprecated and no longer used. Refer to IDisplayClusterSerializable instead.")
+	virtual FString SerializeToString() const
+	{
+		return FString();
+	}
+
+	UE_DEPRECATED(5.8, "This method is deprecated and no longer used. Refer to IDisplayClusterSerializable instead.")
+	virtual bool DeserializeFromString(const FString& ar)
+	{
+		return false;
+	}
 };
