[nDisplay] Fixed CIS issue
- declaration of 'TDataType' shadows template parameter

#rnx
#Synced-CL 0

diff --git a/Engine/Plugins/Runtime/nDisplay/Source/DisplayCluster/Public/Cluster/CustomStates/DisplayClusterCustomStateBase.h b/Engine/Plugins/Runtime/nDisplay/Source/DisplayCluster/Public/Cluster/CustomStates/DisplayClusterCustomStateBase.h
--- a/Engine/Plugins/Runtime/nDisplay/Source/DisplayCluster/Public/Cluster/CustomStates/DisplayClusterCustomStateBase.h
+++ b/Engine/Plugins/Runtime/nDisplay/Source/DisplayCluster/Public/Cluster/CustomStates/DisplayClusterCustomStateBase.h
@@ -224,14 +224,14 @@ namespace UE::nDisplay::Private
 		/**
 		 * Auxiliary structure that holds all necessary data of a single state
 		 */
-		template <typename TDataType>
+		template <typename THolderDataType>
 		struct FCustomStateDataHolder
 		{
 			/** Game thread data */
-			TDataType Data { };
+			THolderDataType Data { };
 
 			/** Render thread data */
-			TDataType DataRT { };
+			THolderDataType DataRT { };
 		};
 
 		/**
