[Gizmo Transform API] Refactor UGizmoEdModeInterface cancelation API to be a past-tense action that is called after `EndTransform()` resolves normally. The inner transaction is completed as if everything was happening normally, then the outer transaction is canceled. `OnTransformCanceled()` allows for post-cancellation clean up logic.

#rnx

#Synced-CL 49993678

diff --git a/Engine/Plugins/2D/Paper2D/Source/Paper2DEditor/Private/TileSetEditor/SingleTileEditorViewportClient.cpp b/Engine/Plugins/2D/Paper2D/Source/Paper2DEditor/Private/TileSetEditor/SingleTileEditorViewportClient.cpp
--- a/Engine/Plugins/2D/Paper2D/Source/Paper2DEditor/Private/TileSetEditor/SingleTileEditorViewportClient.cpp
+++ b/Engine/Plugins/2D/Paper2D/Source/Paper2DEditor/Private/TileSetEditor/SingleTileEditorViewportClient.cpp
@@ -105,7 +105,7 @@ bool FSingleTileEditorViewportClient::EndTransform(const FGizmoState& InState)
 	return HandleEndTransform();
 }
 
-bool FSingleTileEditorViewportClient::CancelTransform(const FGizmoState& InState)
+bool FSingleTileEditorViewportClient::OnTransformCanceled(const FGizmoState& InState)
 {
 	return HandleCancelTransform();
 }
@@ -135,12 +135,11 @@ bool FSingleTileEditorViewportClient::HandleEndTransform()
 
 bool FSingleTileEditorViewportClient::HandleCancelTransform()
 {
-	if (ScopedTransaction)
+	if (TileSet)
 	{
-		ScopedTransaction->Cancel();
+		TileSet->PostEditChange();
 	}
-
-	return HandleEndTransform();
+	return true;
 }
 
 void FSingleTileEditorViewportClient::DrawCanvas(FViewport& InViewport, FSceneView& View, FCanvas& Canvas)
diff --git a/Engine/Plugins/2D/Paper2D/Source/Paper2DEditor/Private/TileSetEditor/SingleTileEditorViewportClient.h b/Engine/Plugins/2D/Paper2D/Source/Paper2DEditor/Private/TileSetEditor/SingleTileEditorViewportClient.h
--- a/Engine/Plugins/2D/Paper2D/Source/Paper2DEditor/Private/TileSetEditor/SingleTileEditorViewportClient.h
+++ b/Engine/Plugins/2D/Paper2D/Source/Paper2DEditor/Private/TileSetEditor/SingleTileEditorViewportClient.h
@@ -35,7 +35,7 @@ public:
 	virtual void TrackingStopped() override;
 	virtual bool BeginTransform(const FGizmoState& InState) override;
 	virtual bool EndTransform(const FGizmoState& InState) override;
-	virtual bool CancelTransform(const FGizmoState& InState) override;
+	virtual bool OnTransformCanceled(const FGizmoState& InState) override;
 	virtual void DrawCanvas(FViewport& Viewport, FSceneView& View, FCanvas& Canvas) override;
 	// End of FEditorViewportClient interface
 
diff --git a/Engine/Plugins/Runtime/nDisplay/Source/DisplayClusterConfigurator/Private/Views/Viewport/DisplayClusterConfiguratorSCSEditorViewportClient.cpp b/Engine/Plugins/Runtime/nDisplay/Source/DisplayClusterConfigurator/Private/Views/Viewport/DisplayClusterConfiguratorSCSEditorViewportClient.cpp
--- a/Engine/Plugins/Runtime/nDisplay/Source/DisplayClusterConfigurator/Private/Views/Viewport/DisplayClusterConfiguratorSCSEditorViewportClient.cpp
+++ b/Engine/Plugins/Runtime/nDisplay/Source/DisplayClusterConfigurator/Private/Views/Viewport/DisplayClusterConfiguratorSCSEditorViewportClient.cpp
@@ -740,9 +740,8 @@ bool FDisplayClusterConfiguratorSCSEditorViewportClient::EndTransform(const FGiz
 	return true;
 }
 
-bool FDisplayClusterConfiguratorSCSEditorViewportClient::CancelTransform(const FGizmoState& InState)
+bool FDisplayClusterConfiguratorSCSEditorViewportClient::OnTransformCanceled(const FGizmoState& InState)
 {
-	CancelTransaction();
 	return EndTransform(InState);
 }
 
@@ -1311,14 +1310,6 @@ void FDisplayClusterConfiguratorSCSEditorViewportClient::EndTransaction()
 	}
 }
 
-void FDisplayClusterConfiguratorSCSEditorViewportClient::CancelTransaction()
-{
-	if (ScopedTransaction)
-	{
-		ScopedTransaction->Cancel();
-	}
-}
-
 void FDisplayClusterConfiguratorSCSEditorViewportClient::RefreshPreviewBounds()
 {
 	AActor* PreviewActor = GetPreviewActor();
diff --git a/Engine/Plugins/Runtime/nDisplay/Source/DisplayClusterConfigurator/Private/Views/Viewport/DisplayClusterConfiguratorSCSEditorViewportClient.h b/Engine/Plugins/Runtime/nDisplay/Source/DisplayClusterConfigurator/Private/Views/Viewport/DisplayClusterConfiguratorSCSEditorViewportClient.h
--- a/Engine/Plugins/Runtime/nDisplay/Source/DisplayClusterConfigurator/Private/Views/Viewport/DisplayClusterConfiguratorSCSEditorViewportClient.h
+++ b/Engine/Plugins/Runtime/nDisplay/Source/DisplayClusterConfigurator/Private/Views/Viewport/DisplayClusterConfiguratorSCSEditorViewportClient.h
@@ -41,7 +41,7 @@ public:
 	virtual void TrackingStopped() override;
 	virtual bool BeginTransform(const FGizmoState& InState) override;
 	virtual bool EndTransform(const FGizmoState& InState) override;
-	virtual bool CancelTransform(const FGizmoState& InState) override;
+	virtual bool OnTransformCanceled(const FGizmoState& InState) override;
 	virtual UE::Widget::EWidgetMode GetWidgetMode() const override;
 	virtual void SetWidgetMode(UE::Widget::EWidgetMode NewMode) override;
 	virtual void SetWidgetCoordSystemSpace(ECoordSystem NewCoordSystem) override;
@@ -210,11 +210,6 @@ protected:
 	 */
 	void EndTransaction();
 
-	/**
-	 * Cancels the current transaction, if one exists.
-	 */
-	void CancelTransaction();
-
 	/**
 	 * Updates preview bounds and floor positioning
 	 */
diff --git a/Engine/Source/Editor/Experimental/EditorInteractiveToolsFramework/Private/EditorGizmos/EditorGizmoStateTarget.cpp b/Engine/Source/Editor/Experimental/EditorInteractiveToolsFramework/Private/EditorGizmos/EditorGizmoStateTarget.cpp
--- a/Engine/Source/Editor/Experimental/EditorInteractiveToolsFramework/Private/EditorGizmos/EditorGizmoStateTarget.cpp
+++ b/Engine/Source/Editor/Experimental/EditorInteractiveToolsFramework/Private/EditorGizmos/EditorGizmoStateTarget.cpp
@@ -48,6 +48,12 @@ void UEditorGizmoStateTarget::EndUpdate()
 
 void UEditorGizmoStateTarget::CancelUpdate()
 {
+	// For compatibility with users of FSnappingUtils - previously executed in MouseDeltaTracker for legacy gizmos
+	{
+		constexpr bool bClearImmediately = true;
+		FSnappingUtils::ClearSnappingHelpers(bClearImmediately);	
+	}
+
 	const TSharedPtr<FEditorModeTools> ModeTools = WeakModeTools.IsValid() ? WeakModeTools.Pin() : nullptr;
 	if (ensure(ModeTools))
 	{
@@ -59,6 +65,8 @@ void UEditorGizmoStateTarget::CancelUpdate()
 		{
 			TransactionManager->CancelUndoTransaction();
 		}
+		
+		(void)ModeTools->OnTransformCanceled(State);
 	}
 }
 
diff --git a/Engine/Source/Editor/Experimental/EditorInteractiveToolsFramework/Public/GizmoEdModeInterface.h b/Engine/Source/Editor/Experimental/EditorInteractiveToolsFramework/Public/GizmoEdModeInterface.h
--- a/Engine/Source/Editor/Experimental/EditorInteractiveToolsFramework/Public/GizmoEdModeInterface.h
+++ b/Engine/Source/Editor/Experimental/EditorInteractiveToolsFramework/Public/GizmoEdModeInterface.h
@@ -37,8 +37,7 @@ public:
 	virtual bool EndTransform(const FGizmoState& InState) = 0;
 
 	/**
-	 * Called when a transform operation is canceled.
-	 * If Begin & End do not interact with transactions, `EndTransform()` is usually the correct behavior.
+	 * Called after a transform operation is canceled.
 	 */
-	virtual bool CancelTransform(const FGizmoState& InState) { return EndTransform(InState); }
+	virtual bool OnTransformCanceled(const FGizmoState& InState) { return false; }
 };
diff --git a/Engine/Source/Editor/Kismet/Private/SCSEditorViewportClient.cpp b/Engine/Source/Editor/Kismet/Private/SCSEditorViewportClient.cpp
--- a/Engine/Source/Editor/Kismet/Private/SCSEditorViewportClient.cpp
+++ b/Engine/Source/Editor/Kismet/Private/SCSEditorViewportClient.cpp
@@ -780,11 +780,6 @@ bool FSCSEditorViewportClient::EndTransform(const FGizmoState& InState)
 	return HandleEndTransform();
 }
 
-bool FSCSEditorViewportClient::CancelTransform(const FGizmoState& InState)
-{
-	return HandleCancelTransform();
-}
-
 bool FSCSEditorViewportClient::HandleBeginTransform()
 {
 	if (!bIsManipulating)
@@ -829,21 +824,6 @@ bool FSCSEditorViewportClient::HandleEndTransform()
 	return false;
 }
 
-bool FSCSEditorViewportClient::HandleCancelTransform()
-{
-	if (bIsManipulating)
-	{
-		if (ScopedTransaction)
-		{
-			ScopedTransaction->Cancel();
-		}
-		
-		return HandleEndTransform();
-	}
-	
-	return false;
-}
-
 UE::Widget::EWidgetMode FSCSEditorViewportClient::GetWidgetMode() const
 {
 	// Default to not drawing the widget
diff --git a/Engine/Source/Editor/Kismet/Private/SCSEditorViewportClient.h b/Engine/Source/Editor/Kismet/Private/SCSEditorViewportClient.h
--- a/Engine/Source/Editor/Kismet/Private/SCSEditorViewportClient.h
+++ b/Engine/Source/Editor/Kismet/Private/SCSEditorViewportClient.h
@@ -60,7 +60,6 @@ public:
 	virtual void TrackingStopped() override;
 	virtual bool BeginTransform(const FGizmoState& InState) override;
 	virtual bool EndTransform(const FGizmoState& InState) override;
-	virtual bool CancelTransform(const FGizmoState& InState) override;
 	virtual UE::Widget::EWidgetMode GetWidgetMode() const override;
 	virtual void SetWidgetMode( UE::Widget::EWidgetMode NewMode ) override;
 	virtual void SetWidgetCoordSystemSpace( ECoordSystem NewCoordSystem ) override;
@@ -186,7 +185,6 @@ private:
 
 	bool HandleBeginTransform();
 	bool HandleEndTransform();
-	bool HandleCancelTransform();
 
 	/** Hovered primitives and their last overlay color before we apply the hover overlay */
 	TMap<TWeakObjectPtr<UPrimitiveComponent>, TOptional<FColor>> HoveredPrimitiveComponents;
diff --git a/Engine/Source/Editor/UnrealEd/Private/EditorModeManager.cpp b/Engine/Source/Editor/UnrealEd/Private/EditorModeManager.cpp
--- a/Engine/Source/Editor/UnrealEd/Private/EditorModeManager.cpp
+++ b/Engine/Source/Editor/UnrealEd/Private/EditorModeManager.cpp
@@ -2185,19 +2185,19 @@ bool FEditorModeTools::EndTransform(const FGizmoState& InState) const
 	return bHandled;
 }
 
-bool FEditorModeTools::CancelTransform(const FGizmoState& InState) const
+bool FEditorModeTools::OnTransformCanceled(const FGizmoState& InState) const
 {
 	bool bHandled = false;
 	ForEachEdMode<IGizmoEdModeInterface>([&bHandled, &InState](IGizmoEdModeInterface* GizmoInterface)
 	{
-		bHandled |= GizmoInterface->CancelTransform(InState);
+		bHandled |= GizmoInterface->OnTransformCanceled(InState);
 		return true;
 	});
 
 	// Give the focused VPC an opportunity to close the transform transacting if it has not been handled before 
 	if (!bHandled && FocusedViewportClient)
 	{
-		bHandled = FocusedViewportClient->CancelTransform(InState);
+		bHandled = FocusedViewportClient->OnTransformCanceled(InState);
 	}
 
 	// NOTE bHasOngoingTransform is not set to false here but in EndTracking as its needed there.
diff --git a/Engine/Source/Editor/UnrealEd/Private/LevelEditorViewport.cpp b/Engine/Source/Editor/UnrealEd/Private/LevelEditorViewport.cpp
--- a/Engine/Source/Editor/UnrealEd/Private/LevelEditorViewport.cpp
+++ b/Engine/Source/Editor/UnrealEd/Private/LevelEditorViewport.cpp
@@ -5612,18 +5612,9 @@ bool FLevelEditorViewportClient::EndTransform(const FGizmoState& InState)
 	return true;
 }
 
-bool FLevelEditorViewportClient::CancelTransform(const FGizmoState& InState)
+bool FLevelEditorViewportClient::OnTransformCanceled(const FGizmoState& InState)
 {
-	if (TrackingTransaction.IsActive())
-	{
-		if (GUndo)
-		{
-			GUndo->Apply();
-		}
-		TrackingTransaction.Cancel();
-		EndTransform(InState);
-	}
-
+	GUnrealEd->UpdatePivotLocationForSelection();
 	return true;
 }
 
diff --git a/Engine/Source/Editor/UnrealEd/Public/EditorModeManager.h b/Engine/Source/Editor/UnrealEd/Public/EditorModeManager.h
--- a/Engine/Source/Editor/UnrealEd/Public/EditorModeManager.h
+++ b/Engine/Source/Editor/UnrealEd/Public/EditorModeManager.h
@@ -563,7 +563,7 @@ public:
 	UNREALED_API IGizmoStateTarget* GetGizmoStateTarget();
 	UNREALED_API bool BeginTransform(const FGizmoState& InState);
 	UNREALED_API bool EndTransform(const FGizmoState& InState) const;
-	UNREALED_API bool CancelTransform(const FGizmoState& InState) const;
+	UNREALED_API bool OnTransformCanceled(const FGizmoState& InState) const;
 	UNREALED_API bool HasOngoingTransform() const;
 	UNREALED_API UEditorTransformGizmoContextObject* GetGizmoContext() const;
 	UNREALED_API const FRotationContext& GetRotationContext() const;
diff --git a/Engine/Source/Editor/UnrealEd/Public/EditorViewportClient.h b/Engine/Source/Editor/UnrealEd/Public/EditorViewportClient.h
--- a/Engine/Source/Editor/UnrealEd/Public/EditorViewportClient.h
+++ b/Engine/Source/Editor/UnrealEd/Public/EditorViewportClient.h
@@ -645,7 +645,7 @@ public:
 
 	virtual bool BeginTransform(const FGizmoState& InState) { return false; }
 	virtual bool EndTransform(const FGizmoState& InState) { return false; }
-	virtual bool CancelTransform(const FGizmoState& InState) { return EndTransform(InState); }
+	virtual bool OnTransformCanceled(const FGizmoState& InState) { return EndTransform(InState); }
 	
 protected:
 
diff --git a/Engine/Source/Editor/UnrealEd/Public/LevelEditorViewport.h b/Engine/Source/Editor/UnrealEd/Public/LevelEditorViewport.h
--- a/Engine/Source/Editor/UnrealEd/Public/LevelEditorViewport.h
+++ b/Engine/Source/Editor/UnrealEd/Public/LevelEditorViewport.h
@@ -239,7 +239,7 @@ public:
 
 	UNREALED_API virtual bool BeginTransform(const FGizmoState& InState) override;
 	UNREALED_API virtual bool EndTransform(const FGizmoState& InState) override;
-	UNREALED_API virtual bool CancelTransform(const FGizmoState& InState) override; 
+	UNREALED_API virtual bool OnTransformCanceled(const FGizmoState& InState) override; 
 	
 
 	/** Sets a flag for this frame indicating that the camera has been cut, and temporal effects (such as motion blur) should be reset */
