[nDisplay] Automatic control of late OCIO
- Late OCIO is automatically managed based on per-node OCIO settings
- PQ transfer is always used for late OCIO

#Synced-CL 0

diff --git a/Engine/Plugins/Runtime/nDisplay/Source/DisplayClusterConfiguration/Public/DisplayClusterConfigurationTypes_Media.h b/Engine/Plugins/Runtime/nDisplay/Source/DisplayClusterConfiguration/Public/DisplayClusterConfigurationTypes_Media.h
--- a/Engine/Plugins/Runtime/nDisplay/Source/DisplayClusterConfiguration/Public/DisplayClusterConfigurationTypes_Media.h
+++ b/Engine/Plugins/Runtime/nDisplay/Source/DisplayClusterConfiguration/Public/DisplayClusterConfigurationTypes_Media.h
@@ -282,6 +282,17 @@ struct DISPLAYCLUSTERCONFIGURATION_API FDisplayClusterConfigurationMediaICVFX
 {
 	GENERATED_BODY()
 
+public:
+
+	PRAGMA_DISABLE_DEPRECATION_WARNINGS
+	FDisplayClusterConfigurationMediaICVFX() = default;
+	~FDisplayClusterConfigurationMediaICVFX() = default;
+	FDisplayClusterConfigurationMediaICVFX(const FDisplayClusterConfigurationMediaICVFX&) = default;
+	FDisplayClusterConfigurationMediaICVFX(FDisplayClusterConfigurationMediaICVFX&&) = default;
+	FDisplayClusterConfigurationMediaICVFX& operator=(const FDisplayClusterConfigurationMediaICVFX&) = default;
+	FDisplayClusterConfigurationMediaICVFX& operator=(FDisplayClusterConfigurationMediaICVFX&&) = default;
+	PRAGMA_ENABLE_DEPRECATION_WARNINGS
+
 public:
 	/** Enable/disable media */
 	UPROPERTY(EditAnywhere, BlueprintReadWrite, Category = "Media", meta = (DisplayName = "Enable Media"))
@@ -323,13 +334,13 @@ public:
 	UPROPERTY(EditAnywhere, BlueprintReadWrite, Category = "Media", meta = (DisplayName = "Media Output Groups"))
 	TArray<FDisplayClusterConfigurationMediaTiledOutputGroup> TiledMediaOutputGroups;
 
-	/** Allows the receviers to apply their own OCIO transformations (per node OCIO override). Requires the media to support PF_FloatRGBA pixel format if PQ transfer is not used. */
-	UPROPERTY(EditAnywhere, BlueprintReadWrite, Category = "Media", meta = (DisplayName = "Receivers Apply OCIO"))
-	bool bLateOCIOPass = false;
+	UE_DEPRECATED(5.8, "This property has been depreacted. Late OCIO is automatically detected in runtime.")
+	UPROPERTY(meta = (DeprecatedProperty))
+	bool bLateOCIOPass_DEPRECATED = false;
 
-	/** Use perceptual quantizer (PQ) transfer function to reduce bandwith. Requires PF_A2B10G10R10 pixel format configured on media side. */
-	UPROPERTY(EditAnywhere, BlueprintReadWrite, Category = "Media", meta = (DisplayName = "Use PQ-Encoding for Transfer", EditCondition = "bLateOCIOPass"))
-	bool bTransferPQ = false;
+	UE_DEPRECATED(5.8, "This property has been depreacted. PQ is allways used now.")
+	UPROPERTY(meta = (DeprecatedProperty))
+	bool bTransferPQ_DEPRECATED = true;
 
 public:
 	/** Returns true if these media settings can be used for tiling. */
diff --git a/Engine/Plugins/Runtime/nDisplay/Source/DisplayClusterMedia/Private/Common/DisplayClusterMediaCameraCommon.cpp b/Engine/Plugins/Runtime/nDisplay/Source/DisplayClusterMedia/Private/Common/DisplayClusterMediaCameraCommon.cpp
--- a/Engine/Plugins/Runtime/nDisplay/Source/DisplayClusterMedia/Private/Common/DisplayClusterMediaCameraCommon.cpp
+++ b/Engine/Plugins/Runtime/nDisplay/Source/DisplayClusterMedia/Private/Common/DisplayClusterMediaCameraCommon.cpp
@@ -20,7 +20,7 @@ UDisplayClusterICVFXCameraComponent* FDisplayClusterMediaCameraCommon::GetCamera
 	{
 		if (UDisplayClusterICVFXCameraComponent* ICVFXCameraComponent = DCRA->GetComponentByName<UDisplayClusterICVFXCameraComponent>(CameraId))
 		{
-			return ICVFXCameraComponent;
+			return IsValid(ICVFXCameraComponent) ? ICVFXCameraComponent : nullptr;
 		}
 	}
 
@@ -35,11 +35,35 @@ void FDisplayClusterMediaCameraCommon::GetLateOCIOParameters(bool& bOutLateOCIOE
 		const bool bEnabledOCIO = ICVFXCameraComponent->CameraSettings.CameraOCIO.AllNodesOCIOConfiguration.bIsEnabled;
 		// Check if media is enabled on this camera
 		const bool bEnabledMedia = ICVFXCameraComponent->CameraSettings.RenderSettings.Media.bEnable;
-		// Check if the late OCIO flag is on
-		const bool bEnabledLateOCIO = ICVFXCameraComponent->CameraSettings.RenderSettings.Media.bLateOCIOPass;
+		// Check if late OCIO is required by current per-node OCIO configuration
+		const bool bValidPerNodeOCIO = HasAnyValidPerNodeOCIOConfiguration();
 
-		// Return OCIO configuration parameters
-		bOutLateOCIOEnabled = bEnabledOCIO && bEnabledMedia && bEnabledLateOCIO;
-		bOutTransferPQ = ICVFXCameraComponent->CameraSettings.RenderSettings.Media.bTransferPQ;
+		// Late OCIO
+		bOutLateOCIOEnabled = bEnabledOCIO && bEnabledMedia && bValidPerNodeOCIO;
+		// Always PQ
+		bOutTransferPQ = true;
 	}
+	else
+	{
+		// Reset if we get here
+		bOutLateOCIOEnabled = false;
+		bOutTransferPQ = true;
+	}
+}
+
+bool FDisplayClusterMediaCameraCommon::HasAnyValidPerNodeOCIOConfiguration() const
+{
+	if (const UDisplayClusterICVFXCameraComponent* const ICVFXCamera = GetCameraComponent())
+	{
+		// See if there is at least one valid configuration
+		const bool bHasValidPerNodeOCIO = ICVFXCamera->GetCameraSettingsICVFX().CameraOCIO.PerNodeOCIOProfiles.ContainsByPredicate(
+			[](const FDisplayClusterConfigurationOCIOProfile& Item)
+			{
+				return Item.bIsEnabled && Item.ColorConfiguration.IsValid() && !Item.ApplyOCIOToObjects.IsEmpty();
+			});
+
+		return bHasValidPerNodeOCIO;
+	}
+
+	return false;
 }
diff --git a/Engine/Plugins/Runtime/nDisplay/Source/DisplayClusterMedia/Private/Common/DisplayClusterMediaCameraCommon.h b/Engine/Plugins/Runtime/nDisplay/Source/DisplayClusterMedia/Private/Common/DisplayClusterMediaCameraCommon.h
--- a/Engine/Plugins/Runtime/nDisplay/Source/DisplayClusterMedia/Private/Common/DisplayClusterMediaCameraCommon.h
+++ b/Engine/Plugins/Runtime/nDisplay/Source/DisplayClusterMedia/Private/Common/DisplayClusterMediaCameraCommon.h
@@ -25,6 +25,11 @@ protected:
 	/** Returns late OCIO parameters of the current camera */
 	void GetLateOCIOParameters(bool& bOutLateOCIOEnabled, bool& bOutTransferPQ) const;
 
+private:
+
+	/** Checks whether any valid per-node OCIO configuration is set */
+	bool HasAnyValidPerNodeOCIOConfiguration() const;
+
 private:
 
 	/** Holds ICVFX camera component name associated with this media adapter */
