[nDisplay] Removed old 'dc_gpu' cmdline parameter
- It actually duplicates the 'r.GraphicsAdapter' cvar
- It's not used anywhere

#rnx
#Synced-CL 0

diff --git a/Engine/Plugins/Runtime/nDisplay/Source/DisplayClusterConfiguration/Private/DisplayClusterConfigurationModule.cpp b/Engine/Plugins/Runtime/nDisplay/Source/DisplayClusterConfiguration/Private/DisplayClusterConfigurationModule.cpp
--- a/Engine/Plugins/Runtime/nDisplay/Source/DisplayClusterConfiguration/Private/DisplayClusterConfigurationModule.cpp
+++ b/Engine/Plugins/Runtime/nDisplay/Source/DisplayClusterConfiguration/Private/DisplayClusterConfigurationModule.cpp
@@ -15,26 +15,6 @@
 #include "Misc/DisplayClusterHelpers.h"
 
 
-//////////////////////////////////////////////////////////////////////////////////////////////
-// IModuleInterface
-//////////////////////////////////////////////////////////////////////////////////////////////
-void FDisplayClusterConfigurationModule::StartupModule()
-{
-	FString ConfigLineStr = FCommandLine::Get();
-
-	int32 GraphicsAdapter;
-	if (DisplayClusterHelpers::str::ExtractValue(ConfigLineStr, DisplayClusterConfigurationStrings::args::Gpu, GraphicsAdapter))
-	{
-		IConsoleVariable* const GpuCVar = IConsoleManager::Get().FindConsoleVariable(TEXT("r.GraphicsAdapter"));
-		if (GpuCVar)
-		{
-			UE_LOG(LogDisplayClusterConfiguration, Log, TEXT("Set custom GPU selection policy - r.GraphicsAdapter=%d"), GraphicsAdapter);
-			GpuCVar->Set(GraphicsAdapter);
-		}
-	}
-}
-
-
 //////////////////////////////////////////////////////////////////////////////////////////////
 // IDisplayClusterConfiguration
 //////////////////////////////////////////////////////////////////////////////////////////////
diff --git a/Engine/Plugins/Runtime/nDisplay/Source/DisplayClusterConfiguration/Private/DisplayClusterConfigurationModule.h b/Engine/Plugins/Runtime/nDisplay/Source/DisplayClusterConfiguration/Private/DisplayClusterConfigurationModule.h
--- a/Engine/Plugins/Runtime/nDisplay/Source/DisplayClusterConfiguration/Private/DisplayClusterConfigurationModule.h
+++ b/Engine/Plugins/Runtime/nDisplay/Source/DisplayClusterConfiguration/Private/DisplayClusterConfigurationModule.h
@@ -13,12 +13,6 @@ class UDisplayClusterConfigurationData;
 class FDisplayClusterConfigurationModule :
 	public IDisplayClusterConfiguration
 {
-public:
-	//////////////////////////////////////////////////////////////////////////////////////////////
-	// IModuleInterface
-	//////////////////////////////////////////////////////////////////////////////////////////////
-	virtual void StartupModule() override;
-
 public:
 	//////////////////////////////////////////////////////////////////////////////////////////////
 	// IDisplayClusterConfiguration
diff --git a/Engine/Plugins/Runtime/nDisplay/Source/DisplayClusterConfiguration/Public/DisplayClusterConfigurationStrings.h b/Engine/Plugins/Runtime/nDisplay/Source/DisplayClusterConfiguration/Public/DisplayClusterConfigurationStrings.h
--- a/Engine/Plugins/Runtime/nDisplay/Source/DisplayClusterConfiguration/Public/DisplayClusterConfigurationStrings.h
+++ b/Engine/Plugins/Runtime/nDisplay/Source/DisplayClusterConfiguration/Public/DisplayClusterConfigurationStrings.h
@@ -53,12 +53,6 @@ namespace DisplayClusterConfigurationStrings
 		static constexpr const TCHAR* ViewPointInFrustumProjectionCategory = TEXT("Frustum Fit");
 	}
 
-	// Command line arguments
-	namespace args
-	{
-		static constexpr const TCHAR* Gpu    = TEXT("dc_gpu");
-	}
-
 	// Config file extensions
 	namespace file
 	{
