[nDisplay] Removed 'simple' projection policy

[RN] Removed 'Simple' projection policy
#Synced-CL 0

diff --git a/Engine/Plugins/Runtime/nDisplay/Source/DisplayClusterConfiguration/Private/DisplayClusterConfigurationMgr.cpp b/Engine/Plugins/Runtime/nDisplay/Source/DisplayClusterConfiguration/Private/DisplayClusterConfigurationMgr.cpp
--- a/Engine/Plugins/Runtime/nDisplay/Source/DisplayClusterConfiguration/Private/DisplayClusterConfigurationMgr.cpp
+++ b/Engine/Plugins/Runtime/nDisplay/Source/DisplayClusterConfiguration/Private/DisplayClusterConfigurationMgr.cpp
@@ -11,6 +11,7 @@
 
 #include "DisplayClusterConfigurationLog.h"
 #include "DisplayClusterConfigurationTypes.h"
+#include "DisplayClusterConfigurationUtils.h"
 #include "DisplayClusterConfigurationVersion.h"
 
 #include "Misc/DisplayClusterHelpers.h"
@@ -80,7 +81,16 @@ UDisplayClusterConfigurationData* FDisplayClusterConfigurationMgr::LoadConfig(co
 		return nullptr;
 	}
 
-	return Parser->LoadData(ConfigFile);
+	// Load configuration data
+	UDisplayClusterConfigurationData* Config = Parser->LoadData(ConfigFile);
+
+	// Apply any upgrades if necessary
+	if (IsValid(Config))
+	{
+		FDisplayClusterConfigurationUtils::UpdateToLatest(Config);
+	}
+
+	return Config;
 }
 
 bool FDisplayClusterConfigurationMgr::SaveConfig(const UDisplayClusterConfigurationData* Config, const FString& FilePath)
diff --git a/Engine/Plugins/Runtime/nDisplay/Source/DisplayClusterConfiguration/Private/DisplayClusterConfigurationTypes.cpp b/Engine/Plugins/Runtime/nDisplay/Source/DisplayClusterConfiguration/Private/DisplayClusterConfigurationTypes.cpp
--- a/Engine/Plugins/Runtime/nDisplay/Source/DisplayClusterConfiguration/Private/DisplayClusterConfigurationTypes.cpp
+++ b/Engine/Plugins/Runtime/nDisplay/Source/DisplayClusterConfiguration/Private/DisplayClusterConfigurationTypes.cpp
@@ -8,6 +8,7 @@
 #include "DisplayClusterConfigurationTypes_Viewport.h"
 #include "DisplayClusterConfigurationTypes_ICVFX.h"
 #include "DisplayClusterConfigurationTypes_PostRender.h"
+#include "DisplayClusterConfigurationUtils.h"
 #include "DisplayClusterProjectionStrings.h"
 
 #include "Engine/StaticMesh.h"
@@ -184,7 +185,6 @@ const TSet<FString> UDisplayClusterConfigurationData::InputSyncPolicies =
 
 const TSet<FString> UDisplayClusterConfigurationData::ProjectionPolicies =
 {
-	TEXT("Simple"),
 	TEXT("Camera"),
 	TEXT("Mesh"),
 	TEXT("MPCDI"),
@@ -199,7 +199,7 @@ const TSet<FString> UDisplayClusterConfigurationData::ProjectionPolicies =
 
 FDisplayClusterConfigurationProjection::FDisplayClusterConfigurationProjection()
 {
-	Type = TEXT("simple");
+	Type = TEXT("Mesh");
 }
 
 bool FDisplayClusterConfigurationProjection::ShouldDisplayICVFXCategory() const
@@ -240,6 +240,9 @@ void UDisplayClusterConfigurationViewport::PostLoad()
 		}
 	PRAGMA_ENABLE_DEPRECATION_WARNINGS
 #endif // WITH_EDITOR
+
+	// 5.8+ backward compatibility path for 'Simple' projection policy
+	FDisplayClusterConfigurationUtils::RedirectSimpleToMeshPolicy(this);
 }
 
 #if WITH_EDITOR
diff --git a/Engine/Plugins/Runtime/nDisplay/Source/DisplayClusterConfiguration/Private/DisplayClusterConfigurationUtils.cpp b/Engine/Plugins/Runtime/nDisplay/Source/DisplayClusterConfiguration/Private/DisplayClusterConfigurationUtils.cpp
--- a/Engine/Plugins/Runtime/nDisplay/Source/DisplayClusterConfiguration/Private/DisplayClusterConfigurationUtils.cpp
+++ b/Engine/Plugins/Runtime/nDisplay/Source/DisplayClusterConfiguration/Private/DisplayClusterConfigurationUtils.cpp
@@ -2,6 +2,9 @@
 
 #include "DisplayClusterConfigurationUtils.h"
 
+#include "DisplayClusterConfigurationTypes.h"
+#include "DisplayClusterProjectionStrings.h"
+
 #include "UObject/Object.h"
 #include "UObject/UObjectThreadContext.h"
 
@@ -14,3 +17,71 @@ bool FDisplayClusterConfigurationUtils::IsSerializingTemplate(const FArchive& Ar
 	// We assume that if the SerializedObject is null, it is indicative of a template.
 	return SerializedObject ? SerializedObject->IsTemplate() : true;
 }
+
+void FDisplayClusterConfigurationUtils::UpdateToLatest(UDisplayClusterConfigurationData* Config)
+{
+	// Ignore invalid input
+	if (!IsValid(Config) || !IsValid(Config->Cluster))
+	{
+		return;
+	}
+
+	// Iterate through every cluster node
+	for (TPair<FString, TObjectPtr<UDisplayClusterConfigurationClusterNode>>& ClusterNode : Config->Cluster->Nodes)
+	{
+		if (!IsValid(ClusterNode.Value))
+		{
+			continue;
+		}
+
+		// Iterate though every viewport
+		for (TPair<FString, TObjectPtr<UDisplayClusterConfigurationViewport>>& Viewport : ClusterNode.Value->Viewports)
+		{
+			if (!IsValid(Viewport.Value))
+			{
+				continue;
+			}
+
+			// 'simple' -> 'mesh' policy migration
+			RedirectSimpleToMeshPolicy(Viewport.Value);
+		}
+	}
+}
+
+
+void FDisplayClusterConfigurationUtils::RedirectSimpleToMeshPolicy(UDisplayClusterConfigurationViewport* ViewportCfg)
+{
+	// Ignore inalid input
+	if (!IsValid(ViewportCfg))
+	{
+		return;
+	}
+
+	// 5.8+ backward compatibility path for 'Simple' projection policy. Replace it with 'Mesh' policy.
+	{
+		FDisplayClusterConfigurationProjection& PolicyCfg = ViewportCfg->ProjectionPolicy;
+		const bool bIsSimplePolicy = PolicyCfg.Type.Equals(TEXT("simple"), ESearchCase::IgnoreCase);
+		const bool bIsMeshPolicy   = PolicyCfg.Type.Equals(DisplayClusterProjectionStrings::projection::Mesh, ESearchCase::IgnoreCase);
+
+		// Check both 'Mesh' and 'Simple' as 'Mesh' also replaces the default value
+		if (bIsSimplePolicy || bIsMeshPolicy)
+		{
+			// 'Screen' property was the property of 'Simple' policy. We can use it as a reliable sign
+			// that this instance needs to be upgraded.
+			FString ScreenId;
+			if (PolicyCfg.Parameters.RemoveAndCopyValue(TEXT("screen"), ScreenId))
+			{
+				// Rename the key. Screen name will simply become a mesh component name used in 'Mesh' policy
+				PolicyCfg.Parameters.Emplace(DisplayClusterProjectionStrings::cfg::mesh::Component, MoveTemp(ScreenId));
+
+				// Only 'Mesh' policy since 5.8
+				PolicyCfg.Type = DisplayClusterProjectionStrings::projection::Mesh;
+				if (PolicyCfg.Type.Len() > 0)
+				{
+					// Explicitly capitalize the first character to improve UI presentation.
+					PolicyCfg.Type[0] = FChar::ToUpper(PolicyCfg.Type[0]);
+				}
+			}
+		}
+	}
+}
diff --git a/Engine/Plugins/Runtime/nDisplay/Source/DisplayClusterConfiguration/Private/DisplayClusterConfigurationUtils.h b/Engine/Plugins/Runtime/nDisplay/Source/DisplayClusterConfiguration/Private/DisplayClusterConfigurationUtils.h
--- a/Engine/Plugins/Runtime/nDisplay/Source/DisplayClusterConfiguration/Private/DisplayClusterConfigurationUtils.h
+++ b/Engine/Plugins/Runtime/nDisplay/Source/DisplayClusterConfiguration/Private/DisplayClusterConfigurationUtils.h
@@ -4,6 +4,8 @@
 
 
 class FArchive;
+class UDisplayClusterConfigurationData;
+class UDisplayClusterConfigurationViewport;
 
 
 /**
@@ -18,4 +20,18 @@ public:
 	// Returns true if we're serializing a template or archetype. False otherwise.
 	static bool IsSerializingTemplate(const FArchive& Ar);
 
+	/**
+	 * Updates the config data container to the latest version. It's called internally
+	 * after loading the configuration from an external source.
+	 * 
+	 * @param Config - Configuration data container
+	 */
+	static void UpdateToLatest(UDisplayClusterConfigurationData* Config);
+
+	/**
+	 * Migrates 'simple' to 'mesh' projection policy
+	 * 
+	 * @param ViewportCfg - Viewport configuration to modify
+	 */
+	static void RedirectSimpleToMeshPolicy(UDisplayClusterConfigurationViewport* ViewportCfg);
 };
diff --git a/Engine/Plugins/Runtime/nDisplay/Source/DisplayClusterConfigurator/Private/DisplayClusterConfiguratorBlueprintEditor.cpp b/Engine/Plugins/Runtime/nDisplay/Source/DisplayClusterConfigurator/Private/DisplayClusterConfiguratorBlueprintEditor.cpp
--- a/Engine/Plugins/Runtime/nDisplay/Source/DisplayClusterConfigurator/Private/DisplayClusterConfiguratorBlueprintEditor.cpp
+++ b/Engine/Plugins/Runtime/nDisplay/Source/DisplayClusterConfigurator/Private/DisplayClusterConfiguratorBlueprintEditor.cpp
@@ -441,12 +441,7 @@ void FDisplayClusterConfiguratorBlueprintEditor::SelectAncillaryViewports(const
 				UDisplayClusterConfigurationViewport* Viewport = ViewportPair.Value;
 
 				FString ComponentName;
-				if (Viewport->ProjectionPolicy.Type.Equals(DisplayClusterProjectionStrings::projection::Simple, ESearchCase::IgnoreCase)
-					&& Viewport->ProjectionPolicy.Parameters.Contains(DisplayClusterProjectionStrings::cfg::simple::Screen))
-				{
-					ComponentName = Viewport->ProjectionPolicy.Parameters[DisplayClusterProjectionStrings::cfg::simple::Screen];
-				}
-				else if (Viewport->ProjectionPolicy.Type.Equals(DisplayClusterProjectionStrings::projection::Mesh, ESearchCase::IgnoreCase)
+				if (Viewport->ProjectionPolicy.Type.Equals(DisplayClusterProjectionStrings::projection::Mesh, ESearchCase::IgnoreCase)
 					&& Viewport->ProjectionPolicy.Parameters.Contains(DisplayClusterProjectionStrings::cfg::mesh::Component))
 				{
 					ComponentName = Viewport->ProjectionPolicy.Parameters[DisplayClusterProjectionStrings::cfg::mesh::Component];
diff --git a/Engine/Plugins/Runtime/nDisplay/Source/DisplayClusterConfigurator/Private/Views/Details/Policies/DisplayClusterConfiguratorPolicyDetailCustomization.cpp b/Engine/Plugins/Runtime/nDisplay/Source/DisplayClusterConfigurator/Private/Views/Details/Policies/DisplayClusterConfiguratorPolicyDetailCustomization.cpp
--- a/Engine/Plugins/Runtime/nDisplay/Source/DisplayClusterConfigurator/Private/Views/Details/Policies/DisplayClusterConfiguratorPolicyDetailCustomization.cpp
+++ b/Engine/Plugins/Runtime/nDisplay/Source/DisplayClusterConfigurator/Private/Views/Details/Policies/DisplayClusterConfiguratorPolicyDetailCustomization.cpp
@@ -349,17 +349,10 @@ void FDisplayClusterConfiguratorProjectionCustomization::BuildParametersForPolic
 	}
 
 	const FString PolicyLower = Policy.ToLower();
-	/*
-	 * Simple
-	 */
-	if (PolicyLower == DisplayClusterProjectionStrings::projection::Simple)
-	{
-		CreateSimplePolicy(Blueprint);
-	}
 	/*
 	 * Camera
 	 */
-	else if (PolicyLower == DisplayClusterProjectionStrings::projection::Camera)
+	if (PolicyLower == DisplayClusterProjectionStrings::projection::Camera)
 	{
 		CreateCameraPolicy(Blueprint);
 	}
@@ -420,20 +413,6 @@ void FDisplayClusterConfiguratorProjectionCustomization::BuildParametersForPolic
 	}
 }
 
-void FDisplayClusterConfiguratorProjectionCustomization::CreateSimplePolicy(UDisplayClusterBlueprint* Blueprint)
-{
-	TSharedPtr<FPolicyParameterInfoComponentCombo> ScreenCombo = MakeShared<FPolicyParameterInfoComponentCombo>(
-		"Screen",
-		DisplayClusterProjectionStrings::cfg::simple::Screen,
-		Blueprint,
-		ConfigurationViewports,
-		TArray<TSubclassOf<UActorComponent>>{ UDisplayClusterScreenComponent::StaticClass() });
-
-	ScreenCombo->SetParameterTooltip(LOCTEXT("SimplePolicyScreenTooltip", "Target Screen or Display"));
-
-	CustomPolicyParameters.Add(ScreenCombo);
-}
-
 void FDisplayClusterConfiguratorProjectionCustomization::CreateCameraPolicy(UDisplayClusterBlueprint* Blueprint)
 {
 	CustomPolicyParameters.Add(MakeShared<FPolicyParameterInfoComponentCombo>(
diff --git a/Engine/Plugins/Runtime/nDisplay/Source/DisplayClusterConfigurator/Private/Views/TreeViews/Cluster/TreeItems/DisplayClusterConfiguratorTreeItemViewport.cpp b/Engine/Plugins/Runtime/nDisplay/Source/DisplayClusterConfigurator/Private/Views/TreeViews/Cluster/TreeItems/DisplayClusterConfiguratorTreeItemViewport.cpp
--- a/Engine/Plugins/Runtime/nDisplay/Source/DisplayClusterConfigurator/Private/Views/TreeViews/Cluster/TreeItems/DisplayClusterConfiguratorTreeItemViewport.cpp
+++ b/Engine/Plugins/Runtime/nDisplay/Source/DisplayClusterConfigurator/Private/Views/TreeViews/Cluster/TreeItems/DisplayClusterConfiguratorTreeItemViewport.cpp
@@ -62,12 +62,7 @@ void FDisplayClusterConfiguratorTreeItemViewport::OnSelection()
 	UDisplayClusterConfigurationViewport* Viewport = GetObjectChecked<UDisplayClusterConfigurationViewport>();
 
 	TArray<FString> AssociatedComponents;
-	if (Viewport->ProjectionPolicy.Type.Equals(DisplayClusterProjectionStrings::projection::Simple, ESearchCase::IgnoreCase)
-		&& Viewport->ProjectionPolicy.Parameters.Contains(DisplayClusterProjectionStrings::cfg::simple::Screen))
-	{
-		AssociatedComponents.Add(Viewport->ProjectionPolicy.Parameters[DisplayClusterProjectionStrings::cfg::simple::Screen]);
-	}
-	else if (Viewport->ProjectionPolicy.Type.Equals(DisplayClusterProjectionStrings::projection::Mesh, ESearchCase::IgnoreCase)
+	if (Viewport->ProjectionPolicy.Type.Equals(DisplayClusterProjectionStrings::projection::Mesh, ESearchCase::IgnoreCase)
 		&& Viewport->ProjectionPolicy.Parameters.Contains(DisplayClusterProjectionStrings::cfg::mesh::Component))
 	{
 		AssociatedComponents.Add(Viewport->ProjectionPolicy.Parameters[DisplayClusterProjectionStrings::cfg::mesh::Component]);
diff --git a/Engine/Plugins/Runtime/nDisplay/Source/DisplayClusterConfigurator/Private/Views/Viewport/DisplayClusterConfiguratorSCSEditorViewport.cpp b/Engine/Plugins/Runtime/nDisplay/Source/DisplayClusterConfigurator/Private/Views/Viewport/DisplayClusterConfiguratorSCSEditorViewport.cpp
--- a/Engine/Plugins/Runtime/nDisplay/Source/DisplayClusterConfigurator/Private/Views/Viewport/DisplayClusterConfiguratorSCSEditorViewport.cpp
+++ b/Engine/Plugins/Runtime/nDisplay/Source/DisplayClusterConfigurator/Private/Views/Viewport/DisplayClusterConfiguratorSCSEditorViewport.cpp
@@ -738,29 +738,20 @@ FReply SDisplayClusterConfiguratorSCSEditorViewport::OnDrop(const FGeometry& MyG
 					FString ParameterKey;
 					FString ComponentName;
 
-					if (const UDisplayClusterScreenComponent* ScreenComponent = Cast<UDisplayClusterScreenComponent>(MeshComponent))
+					USceneComponent* AttachParent = MeshComponent->GetAttachParent();
+					const bool bIsCamera = AttachParent && (AttachParent->IsA<UCameraComponent>());
+
+					if (bIsCamera)
 					{
-						PolicyType = DisplayClusterProjectionStrings::projection::Simple;
-						ParameterKey = DisplayClusterProjectionStrings::cfg::simple::Screen;
-						ComponentName = ScreenComponent->GetName();
+						PolicyType = DisplayClusterProjectionStrings::projection::Camera;
+						ParameterKey = DisplayClusterProjectionStrings::cfg::camera::Component;
+						ComponentName = AttachParent->GetName();
 					}
-					else
+					else if (!MeshComponent->IsVisualizationComponent())
 					{
-						USceneComponent* AttachParent = MeshComponent->GetAttachParent();
-						const bool bIsCamera = AttachParent && (AttachParent->IsA<UCameraComponent>());
-						
-						if (bIsCamera)
-						{
-							PolicyType = DisplayClusterProjectionStrings::projection::Camera;
-							ParameterKey = DisplayClusterProjectionStrings::cfg::camera::Component;
-							ComponentName = AttachParent->GetName();
-						}
-						else if (!MeshComponent->IsVisualizationComponent())
-						{
-							PolicyType = DisplayClusterProjectionStrings::projection::Mesh;
-							ParameterKey = DisplayClusterProjectionStrings::cfg::mesh::Component;
-							ComponentName = MeshComponent->GetName();
-						}
+						PolicyType = DisplayClusterProjectionStrings::projection::Mesh;
+						ParameterKey = DisplayClusterProjectionStrings::cfg::mesh::Component;
+						ComponentName = MeshComponent->GetName();
 					}
 
 					if (!PolicyType.IsEmpty() && !ParameterKey.IsEmpty() && !ComponentName.IsEmpty())
diff --git a/Engine/Plugins/Runtime/nDisplay/Source/DisplayClusterLightCardEditor/Private/Viewport/DisplayClusterLightCardEditorViewportClient.cpp b/Engine/Plugins/Runtime/nDisplay/Source/DisplayClusterLightCardEditor/Private/Viewport/DisplayClusterLightCardEditorViewportClient.cpp
--- a/Engine/Plugins/Runtime/nDisplay/Source/DisplayClusterLightCardEditor/Private/Viewport/DisplayClusterLightCardEditorViewportClient.cpp
+++ b/Engine/Plugins/Runtime/nDisplay/Source/DisplayClusterLightCardEditor/Private/Viewport/DisplayClusterLightCardEditorViewportClient.cpp
@@ -2295,12 +2295,7 @@ UDisplayClusterConfigurationViewport* FDisplayClusterLightCardEditorViewportClie
 				UDisplayClusterConfigurationViewport* CfgViewport = ViewportPair.Value;
 
 				FString ComponentName;
-				if (CfgViewport->ProjectionPolicy.Type.Equals(DisplayClusterProjectionStrings::projection::Simple, ESearchCase::IgnoreCase)
-					&& CfgViewport->ProjectionPolicy.Parameters.Contains(DisplayClusterProjectionStrings::cfg::simple::Screen))
-				{
-					ComponentName = CfgViewport->ProjectionPolicy.Parameters[DisplayClusterProjectionStrings::cfg::simple::Screen];
-				}
-				else if (CfgViewport->ProjectionPolicy.Type.Equals(DisplayClusterProjectionStrings::projection::Mesh, ESearchCase::IgnoreCase)
+				if (CfgViewport->ProjectionPolicy.Type.Equals(DisplayClusterProjectionStrings::projection::Mesh, ESearchCase::IgnoreCase)
 					&& CfgViewport->ProjectionPolicy.Parameters.Contains(DisplayClusterProjectionStrings::cfg::mesh::Component))
 				{
 					ComponentName = CfgViewport->ProjectionPolicy.Parameters[DisplayClusterProjectionStrings::cfg::mesh::Component];
diff --git a/Engine/Plugins/Runtime/nDisplay/Source/DisplayClusterProjection/Private/DisplayClusterProjectionModule.cpp b/Engine/Plugins/Runtime/nDisplay/Source/DisplayClusterProjection/Private/DisplayClusterProjectionModule.cpp
--- a/Engine/Plugins/Runtime/nDisplay/Source/DisplayClusterProjection/Private/DisplayClusterProjectionModule.cpp
+++ b/Engine/Plugins/Runtime/nDisplay/Source/DisplayClusterProjection/Private/DisplayClusterProjectionModule.cpp
@@ -12,7 +12,6 @@
 #include "Policy/Manual/DisplayClusterProjectionManualPolicyFactory.h"
 #include "Policy/MPCDI/DisplayClusterProjectionMPCDIPolicyFactory.h"
 #include "Policy/Reference/DisplayClusterProjectionReferencePolicyFactory.h"
-#include "Policy/Simple/DisplayClusterProjectionSimplePolicyFactory.h"
 #include "Policy/VIOSO/DisplayClusterProjectionVIOSOPolicyFactory.h"
 
 #include "Policy/Mesh/DisplayClusterProjectionMeshPolicy.h"
@@ -64,10 +63,6 @@ FDisplayClusterProjectionModule::FDisplayClusterProjectionModule()
 	ProjectionPolicyFactories.Emplace(DisplayClusterProjectionStrings::projection::MPCDI, Factory);
 	ProjectionPolicyFactories.Emplace(DisplayClusterProjectionStrings::projection::Mesh, Factory);
 
-	// Simple projection
-	Factory = MakeShared<FDisplayClusterProjectionSimplePolicyFactory>();
-	ProjectionPolicyFactories.Emplace(DisplayClusterProjectionStrings::projection::Simple, Factory);
-
 	// VIOSO projection
 	Factory = MakeShared<FDisplayClusterProjectionVIOSOPolicyFactory>();
 	ProjectionPolicyFactories.Emplace(DisplayClusterProjectionStrings::projection::VIOSO, Factory);
diff --git a/Engine/Plugins/Runtime/nDisplay/Source/DisplayClusterProjection/Private/Policy/Simple/DisplayClusterProjectionSimplePolicy.cpp b/Engine/Plugins/Runtime/nDisplay/Source/DisplayClusterProjection/Private/Policy/Simple/DisplayClusterProjectionSimplePolicy.cpp
deleted file mode 100644
diff --git a/Engine/Plugins/Runtime/nDisplay/Source/DisplayClusterProjection/Private/Policy/Simple/DisplayClusterProjectionSimplePolicy.h b/Engine/Plugins/Runtime/nDisplay/Source/DisplayClusterProjection/Private/Policy/Simple/DisplayClusterProjectionSimplePolicy.h
deleted file mode 100644
diff --git a/Engine/Plugins/Runtime/nDisplay/Source/DisplayClusterProjection/Private/Policy/Simple/DisplayClusterProjectionSimplePolicyFactory.cpp b/Engine/Plugins/Runtime/nDisplay/Source/DisplayClusterProjection/Private/Policy/Simple/DisplayClusterProjectionSimplePolicyFactory.cpp
deleted file mode 100644
diff --git a/Engine/Plugins/Runtime/nDisplay/Source/DisplayClusterProjection/Private/Policy/Simple/DisplayClusterProjectionSimplePolicyFactory.h b/Engine/Plugins/Runtime/nDisplay/Source/DisplayClusterProjection/Private/Policy/Simple/DisplayClusterProjectionSimplePolicyFactory.h
deleted file mode 100644
diff --git a/Engine/Plugins/Runtime/nDisplay/Source/DisplayClusterProjection/Public/DisplayClusterProjectionStrings.h b/Engine/Plugins/Runtime/nDisplay/Source/DisplayClusterProjection/Public/DisplayClusterProjectionStrings.h
--- a/Engine/Plugins/Runtime/nDisplay/Source/DisplayClusterProjection/Public/DisplayClusterProjectionStrings.h
+++ b/Engine/Plugins/Runtime/nDisplay/Source/DisplayClusterProjection/Public/DisplayClusterProjectionStrings.h
@@ -11,11 +11,6 @@ namespace DisplayClusterProjectionStrings
 {
 	namespace cfg
 	{
-		namespace simple
-		{
-			static constexpr const TCHAR* Screen = TEXT("screen");
-		}
-
 		namespace easyblend
 		{
 			static constexpr const TCHAR* File   = TEXT("file");
@@ -180,7 +175,6 @@ namespace DisplayClusterProjectionStrings
 		static constexpr const TCHAR* Mesh           = TEXT("mesh");
 		static constexpr const TCHAR* MPCDI          = TEXT("mpcdi");
 		static constexpr const TCHAR* Reference      = TEXT("reference");
-		static constexpr const TCHAR* Simple         = TEXT("simple");
 		static constexpr const TCHAR* VIOSO          = TEXT("vioso");
 	}
 };
