[nDisplay] PR #8125 - Add hotfix for ExitApplication() on Linux
- Target code has been changed a lot, so the fix was re-written

#rnx

#Synced-CL 0

diff --git a/Engine/Plugins/Runtime/nDisplay/Source/DisplayCluster/Private/Misc/DisplayClusterAppExit.cpp b/Engine/Plugins/Runtime/nDisplay/Source/DisplayCluster/Private/Misc/DisplayClusterAppExit.cpp
--- a/Engine/Plugins/Runtime/nDisplay/Source/DisplayCluster/Private/Misc/DisplayClusterAppExit.cpp
+++ b/Engine/Plugins/Runtime/nDisplay/Source/DisplayCluster/Private/Misc/DisplayClusterAppExit.cpp
@@ -49,8 +49,13 @@ void FDisplayClusterAppExit::ExitApplication(const FString& Msg, EExitType ExitT
 		{
 			UE_LOG(LogDisplayClusterModule, Error, TEXT("KillImmediately requested - %s"), *Msg);
 
+			// Linux platform layer doesn't support force kill. Use normal exit instead.
+#if PLATFORM_LINUX
+			ExitApplication(Msg, EExitType::Normal);
+#else
 			FProcHandle hProc = FPlatformProcess::OpenProcess(FPlatformProcess::GetCurrentProcessId());
 			FPlatformProcess::TerminateProc(hProc, true);
+#endif
 		}
 	}
 }
