[nDisplay] Fixed crash when adding DCCamera component to non-DCRA actors
- Removed some unnecessary checks to avoid crashes on non-DCRA actors

#rnx

#Synced-CL 0

diff --git a/Engine/Plugins/Runtime/nDisplay/Source/DisplayClusterConfigurator/Private/Views/Details/Components/DisplayClusterCameraComponentDetailsCustomization.cpp b/Engine/Plugins/Runtime/nDisplay/Source/DisplayClusterConfigurator/Private/Views/Details/Components/DisplayClusterCameraComponentDetailsCustomization.cpp
--- a/Engine/Plugins/Runtime/nDisplay/Source/DisplayClusterConfigurator/Private/Views/Details/Components/DisplayClusterCameraComponentDetailsCustomization.cpp
+++ b/Engine/Plugins/Runtime/nDisplay/Source/DisplayClusterConfigurator/Private/Views/Details/Components/DisplayClusterCameraComponentDetailsCustomization.cpp
@@ -87,7 +87,10 @@ void FDisplayClusterCameraComponentDetailsCustomization::RebuildCameraOptions()
 	check(DestCameraComponent != nullptr);
 
 	const AActor* RootActor = GetRootActor();
-	check(RootActor);
+	if (!IsValid(RootActor))
+	{
+		return;
+	}
 
 	TArray<UActorComponent*> ActorComponents;
 	RootActor->GetComponents(GetCameraComponentClass(), ActorComponents);
@@ -179,7 +182,10 @@ void FDisplayClusterCameraComponentDetailsCustomization::OnSelectComponentButton
 		UActorComponent* RootActorComponent = nullptr;
 
 		const AActor* RootActor = GetRootActor();
-		check(RootActor);
+		if (!IsValid(RootActor))
+		{
+			return;
+		}
 
 		TArray<UActorComponent*> ActorComponents;
 		RootActor->GetComponents(GetCameraComponentClass(), ActorComponents);
diff --git a/Engine/Plugins/Runtime/nDisplay/Source/DisplayClusterConfigurator/Private/Views/Details/DisplayClusterConfiguratorBaseDetailCustomization.cpp b/Engine/Plugins/Runtime/nDisplay/Source/DisplayClusterConfigurator/Private/Views/Details/DisplayClusterConfiguratorBaseDetailCustomization.cpp
--- a/Engine/Plugins/Runtime/nDisplay/Source/DisplayClusterConfigurator/Private/Views/Details/DisplayClusterConfiguratorBaseDetailCustomization.cpp
+++ b/Engine/Plugins/Runtime/nDisplay/Source/DisplayClusterConfigurator/Private/Views/Details/DisplayClusterConfiguratorBaseDetailCustomization.cpp
@@ -129,7 +129,6 @@ ADisplayClusterRootActor* FDisplayClusterConfiguratorBaseDetailCustomization::Ge
 		RootActor = RootActorPtr.Get();
 	}
 
-	check(RootActor);
 	return RootActor;
 }
 
