[nDisplay] Implemented AutoExposure for the entire cluster.

IDisplayClusterViewport/IDisplayClusterViewportProxy:
Added usage flags and base id: GetViewportFlags(), GetBaseId()

Added new data to FDisplayClusterViewport_Context:
OverscanInnerRenderTargetRect -  Inner rectangle inside RenderTargetRect

Optimize memory consumption for DCRA preview:
Release viewport ViewStates (DCRA preview) when preview is disabled.
Note: The functions have been renamed:
ReleaseTextures() -> ReleaseAllResources() in FDisplayClusterViewport and FDisplayClusterViewportManager
ReleaseTextures_RenderThread() -> ReleaseAllResources_RenderThread() in FDisplayClusterViewportProxy and FDisplayClusterViewportManagerProxy

Unified viewport name generation, see new function FDisplayClusterViewportConfiguration::BuildViewportName().
Reworked the constructor for FDisplayClusterViewport - the unique viewport name is now generated inside this constructor function.

Created a new folder 'DisplayCluster/Private/Render/Features':
Purpose: consolidate rendering-related functionality (e.g., Upscaler, AutoExposure, etc.) into the same Features folder for better organization.
Removed folder 'DisplayCluster/Private/Render/Upscaler': files have been moved to 'DisplayCluster/Private/Render/Features'.

#Synced-CL 0

diff --git a/Engine/Plugins/Runtime/nDisplay/Source/DisplayCluster/Private/Game/EngineClasses/Scene/DisplayClusterRootActorEditor.cpp b/Engine/Plugins/Runtime/nDisplay/Source/DisplayCluster/Private/Game/EngineClasses/Scene/DisplayClusterRootActorEditor.cpp
--- a/Engine/Plugins/Runtime/nDisplay/Source/DisplayCluster/Private/Game/EngineClasses/Scene/DisplayClusterRootActorEditor.cpp
+++ b/Engine/Plugins/Runtime/nDisplay/Source/DisplayCluster/Private/Game/EngineClasses/Scene/DisplayClusterRootActorEditor.cpp
@@ -21,7 +21,6 @@
 #include "Misc/DisplayClusterLog.h"
 #include "Misc/DisplayClusterStrings.h"
 #include "Render/Viewport/DisplayClusterViewportManager.h"
-#include "Render/Viewport/DisplayClusterViewportStrings.h"
 #include "Render/Viewport/IDisplayClusterViewport.h"
 #include "Render/Viewport/IDisplayClusterViewportManager.h"
 #include "Render/Viewport/IDisplayClusterViewportProxy.h"
diff --git a/Engine/Plugins/Runtime/nDisplay/Source/DisplayCluster/Private/Render/Features/AutoExposure/DisplayClusterAutoExposure.cpp b/Engine/Plugins/Runtime/nDisplay/Source/DisplayCluster/Private/Render/Features/AutoExposure/DisplayClusterAutoExposure.cpp
new file mode 100644
--- /dev/null
+++ b/Engine/Plugins/Runtime/nDisplay/Source/DisplayCluster/Private/Render/Features/AutoExposure/DisplayClusterAutoExposure.cpp
@@ -0,0 +1,850 @@
+// Copyright Epic Games, Inc. All Rights Reserved.
+
+#include "Render/Features/AutoExposure/DisplayClusterAutoExposure.h"
+
+#include "Misc/DisplayClusterLog.h"
+
+#include "Render/Viewport/Configuration/DisplayClusterViewportConfiguration.h"
+#include "Render/Viewport/DisplayClusterViewport.h"
+#include "Render/Viewport/DisplayClusterViewportManager.h"
+#include "Render/Viewport/DisplayClusterViewportProxy.h"
+
+#include "RenderUtils.h"
+
+// When enabled, compiles in extra UE_LOG output for nDisplay auto-exposure diagnostics.
+#define ENABLE_DISPLAYCLUSTER_AUTOEXPOSURE_LOG 0
+#define ENABLE_DISPLAYCLUSTER_AUTOEXPOSURE_LOG_MATH 0
+
+/** CPU version for functions from Engine\Shaders\Private\PostProcessHistogramCommon.ush */
+namespace UE::DisplayClusterRender::Features::AutoExposure::Math
+{
+	/** CPU version of Eye Adaptation parameters.
+	 * Based on structure from \Engine\Source\Runtime\Renderer\Private\PostProcess\PostProcessEyeAdaptation.h
+	 * CPU dosn't support textures. Remove related parameters (Histogram, CompensationCurve, MeterMask).
+	 */
+	struct FEyeAdaptationParametersForCPU
+	{
+		float MinAverageLuminance = 0;
+		float MaxAverageLuminance = 0;
+		float DeltaWorldTime = 0;
+		float ExposureSpeedUp = 3;
+		float ExposureSpeedDown = 1;
+		float ExponentialUpM = 0;
+		float ExponentialDownM = 0;
+		float StartDistance = 3;
+		bool  ForceTarget = false;
+	};
+
+	// Logic based on ExponentialAdaption() from Engine\Shaders\Private\PostProcessHistogramCommon.ush
+	static float ExponentialAdaption(float Current, float Target, float FrameTime, float AdaptionSpeed, float M)
+	{
+		const float Factor = 1.0f - FMath::Exp2(-FrameTime * AdaptionSpeed);
+		const float Value = Current + (Target - Current) * Factor * M;
+
+		return Value;
+	}
+
+	// Logic based on LinearAdaption() from Engine\Shaders\Private\PostProcessHistogramCommon.ush
+	static float LinearAdaption(float Current, float Target, float FrameTime, float AdaptionSpeed)
+	{
+		const float Offset = FrameTime * AdaptionSpeed;
+
+		const float Value = (Current < Target)
+			? FMath::Min(Target, Current + Offset)
+			: FMath::Max(Target, Current - Offset);
+
+		return Value;
+	}
+
+	// Logic based on the ComputeEyeAdaptation() from Engine\Shaders\Private\PostProcessHistogramCommon.ush
+	static float ComputeEyeAdaptation(const FEyeAdaptationParametersForCPU& EyeAdaptation, float OldExposure, float TargetExposure, float FrameTime)
+	{
+		// for manual mode or camera cuts, just lerp to the target
+		if (EyeAdaptation.ForceTarget)
+		{
+			return TargetExposure;
+		}
+
+		// Avoid log2(0) -inf and later NaNs.
+		const float SafeOld = FMath::Max(OldExposure, SMALL_NUMBER);
+		const float SafeTarget = FMath::Max(TargetExposure, SMALL_NUMBER);
+
+		const float LogTargetExposure = FMath::Log2(SafeTarget);
+		const float LogOldExposure = FMath::Log2(SafeOld);
+
+		const float LogDiff = LogTargetExposure - LogOldExposure;
+
+		const bool bGoingUp = (LogDiff > 0.0f);
+		const float AdaptionSpeed = bGoingUp ? EyeAdaptation.ExposureSpeedUp : EyeAdaptation.ExposureSpeedDown;
+		const float M = bGoingUp ? EyeAdaptation.ExponentialUpM : EyeAdaptation.ExponentialDownM;
+
+		const float AbsLogDiff = FMath::Abs(LogDiff);
+
+		// blended exposure
+		float LogAdaptedExposure = LogTargetExposure;
+		if (AbsLogDiff > EyeAdaptation.StartDistance)
+		{
+			// Linear
+			LogAdaptedExposure = LinearAdaption(LogOldExposure, LogTargetExposure, FrameTime, AdaptionSpeed);
+		}
+		else
+		{
+			// Exponential adaptation on the close distance
+			LogAdaptedExposure = ExponentialAdaption(LogOldExposure, LogTargetExposure, FrameTime, AdaptionSpeed, M);
+		}
+
+#if ENABLE_DISPLAYCLUSTER_AUTOEXPOSURE_LOG_MATH
+		// Debug: log auto-exposure adaptation per frame.
+		UE_LOG(LogDisplayClusterViewport, Log, TEXT("[%d] AutoExposure %.3f -> %.3f = %.3f (S=%.3f, M=%.3f, T=%.4f)"),
+			GFrameNumber, LogOldExposure, LogTargetExposure, LogAdaptedExposure, AdaptionSpeed, M, FrameTime);
+#endif /** ENABLE_DISPLAYCLUSTER_AUTOEXPOSURE_LOG_MATH */
+
+		// Note: no clamping here. Target exposure should be clamped upstream.
+		const float AdaptedExposure = FMath::Exp2(LogAdaptedExposure);
+
+		return AdaptedExposure;
+	}
+
+	/** Return eye adaptation changes over time
+	* Logic based on the EyeAdaptationCommon() from '\\Engine\Shaders\Private\PostProcessEyeAdaptation.usf'
+	*
+	* return estimated luminance for current frame
+	*/
+	static float GetEstimatedLuminance(const FEyeAdaptationParametersForCPU& EyeAdaptation, float TargetLuminance, float PrevFrameLuminance)
+	{
+		const float TargetExposure = FMath::Clamp(TargetLuminance,    EyeAdaptation.MinAverageLuminance, EyeAdaptation.MaxAverageLuminance);
+		const float OldExposure    = FMath::Clamp(PrevFrameLuminance, EyeAdaptation.MinAverageLuminance, EyeAdaptation.MaxAverageLuminance);
+
+		// eye adaptation changes over time
+		const float EstimatedExposure = ComputeEyeAdaptation(EyeAdaptation, OldExposure, TargetExposure, EyeAdaptation.DeltaWorldTime);
+		const float FinalExposure = FMath::Clamp(EstimatedExposure, EyeAdaptation.MinAverageLuminance, EyeAdaptation.MaxAverageLuminance);
+
+		return FinalExposure;
+	}
+};
+
+namespace UE::DisplayClusterRender::Features::AutoExposure
+{
+	using namespace UE::DisplayClusterRender::Features::AutoExposure::Math;
+
+	/**
+	* Checks whether the extended luminance range (EV100) is enabled.
+	*
+	* When enabled, certain post-processing parameters are expressed in
+	* exposure values (EV100). When disabled, they are expressed in pixel
+	* luminance (cd/m²).
+	*
+	* @return true  if the extended luminance range (EV100) is enabled
+	* @return false if parameters are expressed in pixel luminance (cd/m²)
+	*/
+	static bool IsExtendLuminanceRangeEnabled()
+	{
+		// Notes:
+		//    CVar ("r.DefaultFeature.AutoExposure.ExtendDefaultLuminanceRange"),
+		//    Whether the default values for AutoExposure should support an extended range of scene luminance.
+		//    This also change the PostProcessSettings.Exposure.MinBrightness, MaxBrightness, HistogramLogMin and HisogramLogMax
+		//    to be expressed in EV100 values instead of in Luminance and Log2 Luminance
+		//    0: Legacy range (UE4 default)
+		//    1: Extended range (UE5 default)
+		static const TConsoleVariableData<int32>* const ExtendDefaultLuminanceRangeCVar =
+			IConsoleManager::Get().FindTConsoleVariableDataInt(TEXT("r.DefaultFeature.AutoExposure.ExtendDefaultLuminanceRange"));
+		if (ExtendDefaultLuminanceRangeCVar)
+		{
+			return ExtendDefaultLuminanceRangeCVar->GetValueOnGameThread() != 0;
+		}
+
+		// UE5 default use extended range
+		return true;
+	}
+
+	// Function type used to calculate a per-viewport value with an associated weight.
+	using FAutoExposureMetricFunc = TFunction<void(
+		const FDisplayClusterViewportAutoExposureData& ViewportData,
+		const double ViewportWeight)>;
+
+	/**
+	* Iterates over viewports and invokes a user-supplied function with each viewport and its normalized weight.
+	* 
+	* @param InViewports  Array of viewports to process.
+	* @param MetricFunc   Function to execute per viewport with its normalized weight.
+	*/
+	static void ForEachViewportWeighted(
+		const TArray<FDisplayClusterViewportAutoExposureData>& InViewports,
+		FAutoExposureMetricFunc&& MetricFunc)
+	{
+		uint32 TotalPixels = 0;
+		for (const FDisplayClusterViewportAutoExposureData& ViewportIt : InViewports)
+		{
+			const uint32 VisiblePixels = FMath::Max(ViewportIt.NumVisiblePixels, 0);
+
+			// Prevent uint32 overflow when accumulating visible pixels across viewports.
+			if (TotalPixels > UINT32_MAX - VisiblePixels)
+			{
+				// The weights of viewports cannot be calculated.
+				return;
+			}
+
+			TotalPixels += VisiblePixels;
+		}
+
+		if (TotalPixels == 0)
+		{
+			// The weights of viewports cannot be calculated.
+			return;
+		}
+
+		for (const FDisplayClusterViewportAutoExposureData& ViewportIt : InViewports)
+		{
+			// Viewport Weight in 0..1 range
+			const double NumVisiblePixels = static_cast<double>(FMath::Max(ViewportIt.NumVisiblePixels, 0));
+			const double NumTotalPixels = static_cast<double>(TotalPixels);
+			const double ViewportWeight = NumVisiblePixels / NumTotalPixels;
+
+			MetricFunc(ViewportIt, ViewportWeight);
+		}
+	}
+
+	/**
+	* Copied from \\Engine\Source\Runtime\Renderer\Private\PostProcess\PostProcessEyeAdaptation.cpp
+	*/
+	static float LuminanceMaxFromLensAttenuation()
+	{
+		const bool bExtendedLuminanceRange = IsExtendLuminanceRangeEnabled();
+
+		static const auto VarEyeAdaptationLensAttenuation = IConsoleManager::Get().FindTConsoleVariableDataFloat(TEXT("r.EyeAdaptation.LensAttenuation"));
+		check(VarEyeAdaptationLensAttenuation);
+		float LensAttenuation = VarEyeAdaptationLensAttenuation->GetValueOnGameThread();
+
+		// 78 is defined in the ISO 12232:2006 standard.
+		const float kISOSaturationSpeedConstant = 0.78f;
+
+		const float LuminanceMax = kISOSaturationSpeedConstant / FMath::Max<float>(LensAttenuation, .01f);
+
+		// if we do not have luminance range extended, the math is hardcoded to 1.0 scale.
+		return bExtendedLuminanceRange ? LuminanceMax : 1.0f;
+	}
+
+	/**
+	* Logic from \\Engine\Source\Runtime\Renderer\Private\PostProcess\PostProcessEyeAdaptation.cpp -> GetEyeAdaptationParameters()
+	*/
+	static FEyeAdaptationParametersForCPU GetEyeAdaptationParameters(const FSceneView& View)
+	{
+		const bool bExtendedLuminanceRange = IsExtendLuminanceRangeEnabled();
+		const float LuminanceMax = LuminanceMaxFromLensAttenuation();
+		const float kMiddleGrey = 0.18f;
+
+		const FPostProcessSettings& Settings = View.FinalPostProcessSettings;
+
+		// These clamp the average luminance computed from the scene color. We are going to calculate the white point first, and then
+		// figure out the average grey point later. I.e. if the white point is 1.0, the middle grey point should be 0.18.
+		float MinWhitePointLuminance = 1.0f;
+		float MaxWhitePointLuminance = 1.0f;
+		{
+			if (bExtendedLuminanceRange)
+			{
+				MinWhitePointLuminance = EV100ToLuminance(LuminanceMax, Settings.AutoExposureMinBrightness);
+				MaxWhitePointLuminance = EV100ToLuminance(LuminanceMax, Settings.AutoExposureMaxBrightness);
+			}
+			else
+			{
+				MinWhitePointLuminance = Settings.AutoExposureMinBrightness;
+				MaxWhitePointLuminance = Settings.AutoExposureMaxBrightness;
+			}
+			MinWhitePointLuminance = FMath::Min(MinWhitePointLuminance, MaxWhitePointLuminance);
+		}
+
+		static const auto CVarEyeAdaptationExponentialTransitionDistance =
+			IConsoleManager::Get().FindTConsoleVariableDataFloat(TEXT("r.EyeAdaptation.ExponentialTransitionDistance"));
+		check(CVarEyeAdaptationExponentialTransitionDistance);
+
+		// Logic from PostProcessEyeAdaptation.cpp -> GetEyeAdaptationParameters()
+
+		// The distance at which we switch from linear to exponential. I.e. at StartDistance=1.5, when linear is 1.5 f-stops away from hitting the 
+		// target, we switch to exponential.
+		const float StartDistance = FMath::Max(CVarEyeAdaptationExponentialTransitionDistance->GetValueOnGameThread(), 0.001f);
+		const float StartTimeUp = StartDistance / FMath::Max(Settings.AutoExposureSpeedUp, 0.001f);
+		const float StartTimeDown = StartDistance / FMath::Max(Settings.AutoExposureSpeedDown, 0.001f);
+
+		// We want to ensure that at time=StartT, that the derivative of the exponential curve is the same as the derivative of the linear curve.
+		// For the linear curve, the step will be AdaptationSpeed * FrameTime.
+		// For the exponential curve, the step will be at t=StartT, M is slope modifier:
+		//      slope(t) = M * (1.0f - exp2(-FrameTime * AdaptionSpeed)) * AdaptionSpeed * StartT
+		//      AdaptionSpeed * FrameTime = M * (1.0f - exp2(-FrameTime * AdaptionSpeed)) * AdaptationSpeed * StartT
+		//      M = FrameTime / (1.0f - exp2(-FrameTime * AdaptionSpeed)) * StartT
+		//
+		// To be technically correct, we should take the limit as FrameTime->0, but for simplicity we can make FrameTime a small number. So:
+		const float kFrameTimeEps = 1.0f / 60.0f;
+		const float ExponentialUpM = kFrameTimeEps / ((1.0f - exp2(-kFrameTimeEps * Settings.AutoExposureSpeedUp)) * StartTimeUp);
+		const float ExponentialDownM = kFrameTimeEps / ((1.0f - exp2(-kFrameTimeEps * Settings.AutoExposureSpeedDown)) * StartTimeDown);
+
+		// If the white point luminance is 1.0, then the middle grey luminance should be 0.18.
+		const float MinAverageLuminance = MinWhitePointLuminance * kMiddleGrey;
+		const float MaxAverageLuminance = MaxWhitePointLuminance * kMiddleGrey;
+
+		const bool bValidRange = Settings.AutoExposureMinBrightness < Settings.AutoExposureMaxBrightness;
+		const bool bValidSpeeds = Settings.AutoExposureSpeedDown >= 0.f && Settings.AutoExposureSpeedUp >= 0.f;
+
+		// if it is a camera cut we force the exposure to go all the way to the target exposure without blending.
+		// if it is manual mode, we also force the exposure to hit the target, which matters for HDR Visualization
+		// if we don't have a valid range (AutoExposureMinBrightness == AutoExposureMaxBrightness) then force it like Manual as well.
+		const bool ForceTarget = (View.bCameraCut || !bValidRange || !bValidSpeeds);
+
+		FEyeAdaptationParametersForCPU Parameters;
+		Parameters.MinAverageLuminance = MinAverageLuminance;
+		Parameters.MaxAverageLuminance = MaxAverageLuminance;
+		Parameters.DeltaWorldTime = View.Family->Time.GetDeltaWorldTimeSeconds();
+		Parameters.ExposureSpeedUp = Settings.AutoExposureSpeedUp;
+		Parameters.ExposureSpeedDown = Settings.AutoExposureSpeedDown;
+		Parameters.ExponentialDownM = ExponentialDownM;
+		Parameters.ExponentialUpM = ExponentialUpM;
+		Parameters.StartDistance = StartDistance;
+		Parameters.ForceTarget = ForceTarget;
+
+		return Parameters;
+	}
+
+	/**
+	 * Computes the Auto Exposure (AE) brightness value to override
+	 * AutoExposureMinBrightness/MaxBrightness for a scene view.
+	 *
+	 * Uses the most recent eye-adaptation measurements from the specified
+	 * metering viewports, together with the configured AE settings, to
+	 * calculate the brightness.
+	 *
+	 * @param InMeteringViewports Viewports contributing to the average luminance used for metering.
+	 * @return                    Brightness, suitable for AutoExposureMinBrightness/MaxBrightness override.
+	 */
+	static float ComputeDestBrightness(
+		const TArray<FDisplayClusterViewportAutoExposureData>& InMeteringViewports)
+	{
+		// Average luminance
+		double DestLuminance = 0;
+
+		// Iterate over metering (dest) viewports
+		ForEachViewportWeighted(InMeteringViewports,
+			[&](const FDisplayClusterViewportAutoExposureData& ViewportData, const double ViewportWeight)
+			{
+				DestLuminance += ViewportData.LastAverageSceneLuminance * ViewportWeight;
+			});
+
+		const float LuminanceMax = LuminanceMaxFromLensAttenuation();
+
+		const float DestLuminanceEV100 = LuminanceToEV100(LuminanceMax, DestLuminance);
+		const float MidGrayLuminanceEV100 = LuminanceToEV100(LuminanceMax, 0.18f);
+
+		const float FinalDestExposureEV100 = DestLuminanceEV100 - MidGrayLuminanceEV100;
+
+		// All viewports should use same exposure
+		const float DestBrightness = IsExtendLuminanceRangeEnabled()
+			? FinalDestExposureEV100
+			: EV100ToLuminance(LuminanceMax, FinalDestExposureEV100);
+
+		return DestBrightness;
+	}
+
+	/**
+	 * Computes the next eye-adaptation brightness value for this view.
+	 *
+	 * Moves the adaptation state from OldBrightness toward TargetBrightness using the view's
+	 * exposure/eye-adaptation settings (e.g., speed up/down, clamping).
+	 *
+	 * @param View              Scene view providing exposure/adaptation parameters.
+	 * @param TargetBrightness  Desired brightness after adaptation (the goal value).
+	 * @param OldBrightness     Previous frame's adapted brightness (the current state).
+	 * @return                  New adapted brightness for this frame.
+	 */
+	static float ComputeEyeAdaptationBrightness(const FSceneView & View, const float TargetBrightness, const float OldBrightness)
+	{
+		const bool bExtendLuminanceRangeEnabled = IsExtendLuminanceRangeEnabled();
+		const float LuminanceMax = LuminanceMaxFromLensAttenuation();
+
+		// Get EyeAdaptation parameters from View
+		const FEyeAdaptationParametersForCPU EyeAdaptationParameters = GetEyeAdaptationParameters(View);
+
+		const float TargetLuminance = bExtendLuminanceRangeEnabled
+			? EV100ToLuminance(LuminanceMax, TargetBrightness)
+			: TargetBrightness;
+
+		const float OldLuminance = bExtendLuminanceRangeEnabled
+			? EV100ToLuminance(LuminanceMax, OldBrightness)
+			: OldBrightness;
+
+		// Calculate new brightness value based on eye adaptation changes over time
+		const float NewLuminance = GetEstimatedLuminance(EyeAdaptationParameters, TargetLuminance, OldLuminance);
+
+		const float NewBrightness = bExtendLuminanceRangeEnabled
+			? LuminanceToEV100(LuminanceMax, NewLuminance)
+			: NewLuminance;
+
+		return NewBrightness;
+	}
+
+	/**
+	 * Checks whether the specified viewport proxy is allowed to contribute
+	 * to auto-exposure metering on the render thread.
+	 *
+	 * @param InViewportProxy  The viewport proxy to test for AE metering eligibility.
+	 * @return                 True if this viewport is allowed to contribute to AE metering.
+	 */
+	static bool IsViewportAllowedForMetering_RenderThread(
+		const TSharedPtr<FDisplayClusterViewportProxy, ESPMode::ThreadSafe>& InViewportProxy)
+	{
+		check(IsInRenderingThread());
+
+		const TArray<FDisplayClusterViewport_Context>& Contexts = InViewportProxy->GetContexts_RenderThread();
+		if (Contexts.IsEmpty())
+		{
+			return false;
+		}
+
+		const FDisplayClusterViewport_Context& InContext = InViewportProxy->GetContexts_RenderThread()[0];
+		const FDisplayClusterViewport_RenderSettings& InRenderSettings = InViewportProxy->GetRenderSettings_RenderThread();
+
+		if (InContext.bDisableRender || InRenderSettings.bSkipRendering)
+		{
+			// Ignore viewports that are not used for rendering.
+			return false;
+		}
+
+		if (!InContext.RenderThreadData.EyeAdaptationData.CanBeUsed())
+		{
+			// Ignore viewports without EyeAdaptation data
+			return false;
+		}
+
+		return true;
+	}
+
+	/** Convert FDisplayClusterNodeAutoExposureData to text. */
+	static FString ToString(const FDisplayClusterNodeAutoExposureData& InNodeAutoExposureData)
+	{
+		FString OutText;
+		for (const FDisplayClusterViewportAutoExposureData& ViewportIt : InNodeAutoExposureData.Viewports)
+		{
+			OutText += FString::Printf(TEXT(" [%d, %.5f] "), ViewportIt.FrameNumber, ViewportIt.LastAverageSceneLuminance);
+		}
+
+		return OutText;
+	}
+
+	/** Convert TMap<FName, FDisplayClusterNodeAutoExposureData> to text. */
+	static FString ToString(const TMap<FName, FDisplayClusterNodeAutoExposureData>& InPerNodeAutoExposureData)
+	{
+		FString OutText;
+
+		// Sort nodes:
+		TArray<FName> Keys;
+		InPerNodeAutoExposureData.GetKeys(Keys);
+
+		Keys.Sort([](const FName& A, const FName& B)
+			{
+				return A.LexicalLess(B); // alphabetical
+			});
+
+		for (const FName& Key : Keys)
+		{
+			OutText += FString::Printf(TEXT("%s={%s} "), *Key.ToString(), *ToString(InPerNodeAutoExposureData[Key]));
+		}
+
+		return OutText;
+	}
+
+	// The unique name for the FAutoExposureState object.
+	static const FName AutoExposureStateName = TEXT("AutoExposure");
+};
+
+float FDisplayClusterAutoExposure::ComputeViewportAutoExposureBrightness(
+	const FDisplayClusterViewport& InViewport,
+	const uint32 InContextNum,
+	FSceneView& View)
+{
+	// Finds an existing cache entry keyed by (InViewport, InContextNum) or creates one if missing.
+	TOptional<float>* PrevFrameBrightness = nullptr;
+	{
+
+		// Opportunistic cleanup: drop entries whose viewport is gone.
+		PrevFrameBrightnessCache.RemoveAllSwap([](const FPrevFrameBrightnessCacheEntry& Entry)
+			{
+				return !Entry.ViewportWeakPtr.IsValid();
+			}, EAllowShrinking::No);
+
+		// Find an existing entry keyed by (Viewport, Context).
+		for (FPrevFrameBrightnessCacheEntry& Entry : PrevFrameBrightnessCache)
+		{
+			if (Entry.ContextIndex != InContextNum)
+			{
+				continue;
+			}
+
+			const TSharedPtr<const FDisplayClusterViewport, ESPMode::ThreadSafe> ViewportPinned = Entry.ViewportWeakPtr.Pin();
+			if (!ViewportPinned.IsValid())
+			{
+				continue;
+			}
+
+			// Key match: compare actual viewport object identity.
+			if (ViewportPinned.Get() == &InViewport)
+			{
+				PrevFrameBrightness = &Entry.PrevFrameBrightness;
+				break;
+			}
+		}
+
+		if (!PrevFrameBrightness)
+		{
+			// Not found: create a new cache entry.
+			FPrevFrameBrightnessCacheEntry& NewEntry = PrevFrameBrightnessCache.AddDefaulted_GetRef();
+			NewEntry.ViewportWeakPtr = InViewport.AsShared();
+			NewEntry.ContextIndex = InContextNum;
+
+			PrevFrameBrightness = &NewEntry.PrevFrameBrightness;
+		}
+	}
+
+	using namespace UE::DisplayClusterRender::Features::AutoExposure;
+
+	// Compute the brightness for this frame. If we have a previous value, apply eye adaptation,
+	// otherwise fall back to the target brightness for the first frame.
+	const float NewBrightness = PrevFrameBrightness && PrevFrameBrightness->IsSet()
+		? ComputeEyeAdaptationBrightness(View, *TargetBrightness, PrevFrameBrightness->GetValue())
+		: *TargetBrightness;
+
+	// Update the cache with the brightness from the current frame.
+	if (PrevFrameBrightness)
+	{
+		*PrevFrameBrightness = NewBrightness;
+	}
+
+	return NewBrightness;
+}
+
+void FDisplayClusterAutoExposure::ConfigureAutoExposureForSceneView(const FDisplayClusterViewport& InViewport, const uint32 InContextNum, FSceneView& View)
+{
+
+	if (!TargetBrightness.IsSet())
+	{
+		// Skip viewport if CurrentFrameBrightness isn't calculated.
+		return;
+	}
+
+	if (!View.Family || !View.Family->EngineShowFlags.EyeAdaptation)
+	{
+		// Skip viewports that not use eye adptation.
+		return;
+	}
+
+	// Ignore excluded nodes and viewports
+	if(Settings->IsNodeExcludedFromAutoExposure(InViewport.ClusterNodeId)
+	|| Settings->IsViewportExcludedFromAutoExposure(InViewport.GetBaseId(), InViewport.GetViewportFlags()))
+	{
+		return;
+	}
+
+	// Override brightness value for this viewport:
+	FPostProcessSettings& PPSettings = View.FinalPostProcessSettings;
+
+	// Only AEM_Basic and AEM_Histogram can be used to override the brightness value.
+	if (View.FinalPostProcessSettings.AutoExposureMethod == EAutoExposureMethod::AEM_Manual)
+	{
+		// Override Manual method to AEM_Basic.
+		PPSettings.bOverride_AutoExposureMethod = true;
+		PPSettings.AutoExposureMethod = EAutoExposureMethod::AEM_Basic;
+	}
+
+	// Compute the brightness for this frame. If we have a previous value, apply eye adaptation,
+	// otherwise fall back to the target brightness for the first frame.
+	const float NewBrightness = ComputeViewportAutoExposureBrightness(InViewport, InContextNum, View);
+
+	// Override Min/Max brightness value
+	PPSettings.bOverride_AutoExposureMinBrightness = true;
+	PPSettings.bOverride_AutoExposureMaxBrightness = true;
+	PPSettings.AutoExposureMinBrightness = NewBrightness;
+	PPSettings.AutoExposureMaxBrightness = NewBrightness;
+
+	// Disable AE blending: apply actual brightness value in this frame.
+	const float CustomAutoExposureSpeed = -1;
+	PPSettings.bOverride_AutoExposureSpeedUp = true;
+	PPSettings.bOverride_AutoExposureSpeedDown = true;
+	PPSettings.AutoExposureSpeedUp = CustomAutoExposureSpeed;
+	PPSettings.AutoExposureSpeedDown = CustomAutoExposureSpeed;	
+}
+
+void FDisplayClusterAutoExposure::ReleaseAutoExposureState()
+{
+	if (AutoExposureState.IsValid())
+	{
+#if ENABLE_DISPLAYCLUSTER_AUTOEXPOSURE_LOG
+		UE_LOG(LogDisplayClusterViewport, Log, TEXT("Delete AutoExposureState."));
+#endif /*ENABLE_DISPLAYCLUSTER_AUTOEXPOSURE_LOG*/
+
+		using namespace UE::DisplayClusterRender::Features::AutoExposure;
+
+		// Unregister it from replication on this cluster node
+		IDisplayCluster::Get().GetClusterMgr()->UnregisterCustomState(AutoExposureStateName);
+
+		AutoExposureState.Reset();
+	}
+}
+
+bool FDisplayClusterAutoExposure::EnsureAutoExposureStateUsable()
+{
+	if (!Configuration->IsClusterRendering())
+	{
+		// Use TDistributedCustomState<> only for cluster mode
+		return false;
+	}
+
+	// Create TDistributedCustomState<>
+	if (!AutoExposureState.IsValid())
+	{
+		if (!bDisableAutoExposureState)
+		{
+			using namespace UE::DisplayClusterRender::Features::AutoExposure;
+
+			// Create TDistributedCustomState<> object for AutoExposure.
+			AutoExposureState = FAutoExposureState::Create(AutoExposureStateName);
+
+#if ENABLE_DISPLAYCLUSTER_AUTOEXPOSURE_LOG
+			UE_LOG(LogDisplayClusterViewport, Log, TEXT("Create AutoExposureState."));
+#endif /*ENABLE_DISPLAYCLUSTER_AUTOEXPOSURE_LOG*/
+		}
+
+		// nullptr-check
+		if (!AutoExposureState.IsValid())
+		{
+			// Call ::Create() only once
+			bDisableAutoExposureState = true;
+
+			return false;
+		}
+	}
+
+	return true;
+}
+
+bool FDisplayClusterAutoExposure::UpdatePerNodeAutoExposureDataForCluster(const FDisplayClusterNodeAutoExposureData& CurrentNodeAutoExposureData)
+{
+	using namespace UE::DisplayClusterRender::Features::AutoExposure;
+
+	// Reset the the data from prev frame.
+	PerNodeAutoExposureData.Reset();
+
+	if (!EnsureAutoExposureStateUsable())
+	{
+		return false;
+	}
+
+	// Receive AutoExposure data from all cluster nodes
+	const TSet<FName> ClusterNodesIds = AutoExposureState->GetAvailableNodes();
+	for (const FName& NodeId : ClusterNodesIds)
+	{
+		PerNodeAutoExposureData.Emplace(NodeId, AutoExposureState->GetData(NodeId));
+	}
+
+	// Send current node AutoExposure data to cluster
+	if (Settings->IsNodeAllowedForMetering(Configuration->GetClusterNodeId()))
+	{
+		AutoExposureState->SetData(CurrentNodeAutoExposureData);
+	}
+	else
+	{
+		// Publish empty data for excluded node
+		AutoExposureState->SetData(FDisplayClusterNodeAutoExposureData());
+	}
+
+	return true;
+}
+
+void FDisplayClusterAutoExposure::UpdatePerNodeAutoExposureDataForPreview(const FDisplayClusterNodeAutoExposureData& CurrentNodeAutoExposureData)
+{
+	// Get cluster node from configuration (DCRA preview rendering)
+	const FName ClusterNodeId(Configuration->GetClusterNodeId());
+
+	if (Settings->IsNodeAllowedForMetering(Configuration->GetClusterNodeId()))
+	{
+		// Add CurrentNodeAutoExposureData to the PerNodeAutoExposureData.
+		PerNodeAutoExposureData.Emplace(ClusterNodeId, CurrentNodeAutoExposureData);
+	}
+	else
+	{
+		PerNodeAutoExposureData.Remove(ClusterNodeId);
+	}
+}
+
+void FDisplayClusterAutoExposure::OnEntireClusterPreviewRendered()
+{
+	if (Configuration->IsPreviewRendering())
+	{
+		ComputeAutoExposure();
+	}
+}
+
+void FDisplayClusterAutoExposure::SyncAndComputeAutoExposure()
+{
+	check(IsInGameThread());
+
+	// Get CurrentNodeAutoExposureData from GT queue
+	FDisplayClusterNodeAutoExposureData CurrentNodeAutoExposureData;
+	{
+		FScopeLock Lock(&DataQueueCS);
+		if (!DataQueue.IsEmpty())
+		{
+			CurrentNodeAutoExposureData = DataQueue[0];
+			DataQueue.RemoveAt(0);
+		}
+	}
+
+	// nDisplay AutoExposure disabled in settings
+	if (!Settings->bEnable)
+	{
+		Release();
+
+		return;
+	}
+
+	// Sync NodeAutoExposureData
+	if (Configuration->IsPreviewRendering())
+	{
+		// Note: Use TDistributedCustomState<> only for cluster
+		ReleaseAutoExposureState();
+
+		// Logic for preview rendering
+		UpdatePerNodeAutoExposureDataForPreview(CurrentNodeAutoExposureData);
+
+		// For preview, the ComputeAutoExposure() function must be called after the last cluster node has been rendered.
+		// see OnEntireClusterPreviewRendered().
+	}
+	else if (Configuration->IsClusterRendering())
+	{
+		// Use TDistributedCustomState<> to exchange AutoExposure data between cluster nodes.
+		if (!UpdatePerNodeAutoExposureDataForCluster(CurrentNodeAutoExposureData))
+		{
+			Release();
+
+			return;
+		}
+
+		// In a cluster, each node calculates data separately.
+		ComputeAutoExposure();
+	}
+	else
+	{
+		// Other rendering modes do not support nDisplay AutoExposure.
+		Release();
+
+		return;
+	}	
+}
+
+void FDisplayClusterAutoExposure::ComputeAutoExposure()
+{
+	using namespace UE::DisplayClusterRender::Features::AutoExposure;
+
+	// Reset target brightness value
+	TargetBrightness.Reset();
+
+	// Compute the global exposure by averaging contributions from the
+	//    selected viewports(or all, depending on settings).
+	TArray<FDisplayClusterViewportAutoExposureData> MeteringViewports;
+	for (const TPair<FName, FDisplayClusterNodeAutoExposureData>& NodeIt : PerNodeAutoExposureData)
+	{
+		MeteringViewports.Append(NodeIt.Value.Viewports);
+	}
+
+	if (!MeteringViewports.IsEmpty())
+	{
+		// Calculate the new brightness value.
+		TargetBrightness = ComputeDestBrightness(MeteringViewports);
+
+#if ENABLE_DISPLAYCLUSTER_AUTOEXPOSURE_LOG
+		UE_LOG(LogDisplayClusterViewport, Log, TEXT("[%d] AutoExposure %.5f, Data %s."),
+			GFrameNumber,
+			TargetBrightness.IsSet() ? *TargetBrightness : 0.0f,
+			*ToString(PerNodeAutoExposureData));
+#endif /*ENABLE_DISPLAYCLUSTER_AUTOEXPOSURE_LOG*/
+	}
+}
+
+void FDisplayClusterAutoExposure::GatherMeterViewportsAutoExposureDataAndSync_RenderThread(const TArray<TSharedPtr<FDisplayClusterViewportProxy, ESPMode::ThreadSafe>>& InViewportProxies)
+{
+	TArray<FDisplayClusterViewportAutoExposureData> MeterViewports;
+
+	if (SettingsProxy->bEnable
+		&& !SettingsProxy->IsNodeExcludedFromAutoExposure(Configuration->GetClusterNodeId()))
+	{
+		using namespace UE::DisplayClusterRender::Features::AutoExposure;
+
+		// Collect meter viewports
+		for (const TSharedPtr<FDisplayClusterViewportProxy, ESPMode::ThreadSafe>& ViewportProxy : InViewportProxies)
+		{
+			if (!ViewportProxy.IsValid())
+			{
+				continue;
+			}
+
+			// Check if this viewport supports AutoExposure
+			if (!IsViewportAllowedForMetering_RenderThread(ViewportProxy))
+			{
+				continue;
+			}
+
+			// Determine if current settings allow this viewport to contribute to AE metering
+			if (!SettingsProxy->IsViewportEligibleForMetering(ViewportProxy->GetBaseId(), ViewportProxy->GetViewportFlags()))
+			{
+				continue;
+			}
+
+			const FDisplayClusterViewport_Context& InContext = ViewportProxy->GetContexts_RenderThread()[0];
+			const FIntPoint VisibleSize = InContext.OverscanInnerRenderTargetRect.Size();
+			const int32 NumVisiblePixels = VisibleSize.X * VisibleSize.Y;
+			if (NumVisiblePixels == 0)
+			{
+				continue;
+			}
+
+			// Collect AutoExposure data of this viewport.
+			MeterViewports.Add(
+				FDisplayClusterViewportAutoExposureData(
+					InContext.RenderThreadData.EyeAdaptationData,
+					NumVisiblePixels,
+					InContext.RenderThreadData.FrameNumber));
+		}
+	}
+
+	// RT->GT: Send data to game thread
+	{
+		FScopeLock Lock(&DataQueueCS);
+		DataQueue.Add(MeterViewports);
+	}
+}
+
+void FDisplayClusterAutoExposure::ApplyAutoExposureSettings(
+	const FDisplayClusterConfiguration_AutoExposureSettings& InAutoExposureSettings)
+{
+	using namespace UE::DisplayClusterRender::Features::AutoExposure;
+
+	// Update GT settings
+	Settings = MakeShared<FDisplayClusterAutoExposureSettings>(InAutoExposureSettings);
+	
+	// Update RT settings
+	ENQUEUE_RENDER_COMMAND(DisplayClusterUpdateAutoExposureSettingsProxy)([
+		This = AsShared(), NewSettings = Settings](FRHICommandListImmediate& RHICmdList)
+		{
+			This->SettingsProxy = NewSettings;
+		});
+}
+
+void FDisplayClusterAutoExposure::Release()
+{
+#if ENABLE_DISPLAYCLUSTER_AUTOEXPOSURE_LOG
+	UE_LOG(LogDisplayClusterViewport, Log, TEXT("FDisplayClusterAutoExposure: Release internal data."));
+#endif /*ENABLE_DISPLAYCLUSTER_AUTOEXPOSURE_LOG*/
+
+	PrevFrameBrightnessCache.Reset();
+	TargetBrightness.Reset();
+
+	PerNodeAutoExposureData.Reset();
+
+	ReleaseAutoExposureState();
+	bDisableAutoExposureState = false;
+}
\ No newline at end of file
diff --git a/Engine/Plugins/Runtime/nDisplay/Source/DisplayCluster/Private/Render/Features/AutoExposure/DisplayClusterAutoExposure.h b/Engine/Plugins/Runtime/nDisplay/Source/DisplayCluster/Private/Render/Features/AutoExposure/DisplayClusterAutoExposure.h
new file mode 100644
--- /dev/null
+++ b/Engine/Plugins/Runtime/nDisplay/Source/DisplayCluster/Private/Render/Features/AutoExposure/DisplayClusterAutoExposure.h
@@ -0,0 +1,245 @@
+// Copyright Epic Games, Inc. All Rights Reserved.
+
+#pragma once
+
+#include "CoreMinimal.h"
+#include "Cluster/CustomStates/DisplayClusterCustomStateDistributed.h"
+
+#include "Render/Viewport/Containers/DisplayClusterViewport_EyeAdaptation.h"
+#include "Render/Features/AutoExposure/DisplayClusterAutoExposureSettings.h"
+
+#include "Templates/SharedPointer.h"
+
+class FDisplayClusterViewport;
+class FDisplayClusterViewportProxy;
+
+/**
+ * Auto-exposure data for a single nDisplay viewport.
+ */
+struct FDisplayClusterViewportAutoExposureData
+	: public FDisplayClusterViewport_EyeAdaptationData
+{
+	FDisplayClusterViewportAutoExposureData() = default;
+	FDisplayClusterViewportAutoExposureData(
+		const FDisplayClusterViewport_EyeAdaptationData& In,
+		const int32 InNumVisiblePixels,
+		const uint32 InFrameNumber)
+		: FDisplayClusterViewport_EyeAdaptationData(In)
+		, NumVisiblePixels(InNumVisiblePixels)
+		, FrameNumber(InFrameNumber)
+	{ }
+
+	// Number of visible pixels covered by this viewport.
+	// Used as an input when computing the eye-adaptation weight for this viewport.
+	int32 NumVisiblePixels = 0;
+
+	// RT frame number of this data
+	uint32 FrameNumber = 0;
+};
+
+/** Auto-exposure data for all viewports on a single nDisplay node. */
+struct FDisplayClusterNodeAutoExposureData
+{
+	FDisplayClusterNodeAutoExposureData() = default;
+
+	FDisplayClusterNodeAutoExposureData(const TArray<FDisplayClusterViewportAutoExposureData>& InViewports)
+		: Viewports(InViewports)
+	{ }
+
+	// Per-viewport auto-exposure data for this node.
+	TArray<FDisplayClusterViewportAutoExposureData> Viewports;
+};
+
+/** Custom serialization for FDisplayClusterViewportAutoExposureData */
+inline FArchive& operator<<(FArchive& Ar, FDisplayClusterViewportAutoExposureData& In)
+{
+	Ar << In.LastAverageSceneLuminance;
+	Ar << In.NumVisiblePixels;
+	Ar << In.FrameNumber;
+
+	return Ar;
+}
+
+/** Custom serialization for FDisplayClusterNodeAutoExposureData */
+inline FArchive& operator<<(FArchive& Ar, FDisplayClusterNodeAutoExposureData& In)
+{
+	Ar << In.Viewports;
+
+	return Ar;
+}
+
+/**
+* Implements the nDisplay multi-viewport AutoExposure feature.
+* Collects exposure data from all viewports across the cluster,
+* synchronizes the result via the cluster network,
+* and applies a unified exposure state to each node’s viewports.
+*/
+class FDisplayClusterAutoExposure
+	: public TSharedFromThis<FDisplayClusterAutoExposure, ESPMode::ThreadSafe>
+{
+public:
+	FDisplayClusterAutoExposure(const TSharedRef<FDisplayClusterViewportConfiguration, ESPMode::ThreadSafe>& InConfiguration)
+		: Configuration(InConfiguration)
+		, Settings(MakeShared<FDisplayClusterAutoExposureSettings, ESPMode::ThreadSafe>())
+		, SettingsProxy(MakeShared<FDisplayClusterAutoExposureSettings, ESPMode::ThreadSafe>())
+	{ }
+
+	~FDisplayClusterAutoExposure() = default;
+
+public:
+	/**
+	* Apply new auto-exposure settings to the current viewport configuration.
+	*
+	* @param InAutoExposureSettings  The auto-exposure settings to apply.
+	*/
+	void ApplyAutoExposureSettings(const FDisplayClusterConfiguration_AutoExposureSettings& InAutoExposureSettings);
+
+	/**
+	 * Configures AutoExposure settings for a given viewport scene view.
+	 *
+	 * @param InViewport   The nDisplay viewport to configure.
+	 * @param InContextNum View context index within the viewport.
+	 * @param InOutView    Scene view to apply AutoExposure settings to.
+	 */
+	void ConfigureAutoExposureForSceneView(
+		const FDisplayClusterViewport& InViewport,
+		uint32 InContextNum,
+		FSceneView& InOutView);
+
+	/**
+	 * Collects Auto Exposure (AE) metering data from the given viewport proxies
+	 * on this node and forwards it to SyncAndComputeAutoExposure() for cluster-wide
+	 * synchronization and global exposure computation.
+	 *
+	 * Executed on the render thread once per frame.
+	 *
+	 * @param InViewportProxies  The viewport proxies on this node that provide
+	 *                           AE metering input.
+	 */
+	void GatherMeterViewportsAutoExposureDataAndSync_RenderThread(
+		const TArray<TSharedPtr<FDisplayClusterViewportProxy, ESPMode::ThreadSafe>>& InViewportProxies);
+
+	/**
+	 * Synchronizes Auto Exposure (AE) data across the nDisplay cluster and computes
+	 * a unified global exposure value.
+	 *
+	 * Process:
+	 *   - For cluster: call UpdatePerNodeAutoExposureDataForCluster()
+	 *   - For preview: call UpdatePerNodeAutoExposureDataForPreview()
+	 * 
+	 *   - Compute the global exposure by averaging contributions from the
+	 *      selected viewports (or all, depending on settings).
+	 *
+	 * The resulting exposure value is consistent across the cluster, preventing
+	 * brightness mismatches between displays.
+	 */
+	void SyncAndComputeAutoExposure();
+
+	/** Release inner data. */
+	void Release();
+
+	/** Called every time previews for all cluster nodes are rendered. */
+	void OnEntireClusterPreviewRendered();
+
+protected:
+	/**
+	 * Compute the global exposure by averaging contributions from the
+	 * selected viewports (or all, depending on settings).
+	 */
+	void ComputeAutoExposure();
+
+	/**
+	 * Synchronizes Auto Exposure (AE) data across the nDisplay cluster and update ClusterData.
+	 *
+	 * Process:
+	 *   1. Publish this node's AE data (NodeId + viewports) to the cluster.
+	 *   2. Receive AE data from other nodes.
+	 *
+	 * @param CurrentNodeAutoExposureData  The viewports on this node that contribute AE metrics.
+	 */
+	bool UpdatePerNodeAutoExposureDataForCluster(const FDisplayClusterNodeAutoExposureData& CurrentNodeAutoExposureData);
+
+	/**
+	 * Update ClusterData for preview rendering.
+	 * 
+	 * @param CurrentNodeAutoExposureData  The viewports on this node that contribute AE metrics.
+	 */
+	void UpdatePerNodeAutoExposureDataForPreview(const FDisplayClusterNodeAutoExposureData& CurrentNodeAutoExposureData);
+
+	/**
+	 * Ensures the AutoExposureState exists and is usable.
+	 * Creates the AutoExposureState instance if it does not already exist.
+	 *
+	 * @return true if AutoExposureState is available after this call, false otherwise.
+	 */
+	bool EnsureAutoExposureStateUsable();
+
+	/** Release AutoExposureState instance. */
+	void ReleaseAutoExposureState();
+
+private:
+	// Cached brightness value from the previous frame, keyed by (Viewport, ContextIndex).
+	struct FPrevFrameBrightnessCacheEntry
+	{
+		// Weak reference to the viewport this entry belongs to (no ownership).
+		TWeakPtr<const FDisplayClusterViewport, ESPMode::ThreadSafe> ViewportWeakPtr;
+
+		// Viewport context index within the viewport (e.g. stereo eye / context).
+		uint32 ContextIndex = 0;
+
+		// Previous frame brightness for this viewport/context. Unset until a value is produced.
+		TOptional<float> PrevFrameBrightness;
+	};
+
+	// Cache of previous-frame brightness values per viewport/context.
+	// This variable is exclusively used only by ComputeViewportAutoExposureBrightness()
+	TArray<FPrevFrameBrightnessCacheEntry> PrevFrameBrightnessCache;
+
+	/**
+	 * Computes the viewport brightness.
+	 * Call only once per frame per viewport context.
+	 *
+	 * @param InViewport    nDisplay viewport to evaluate.
+	 * @param InContextNum  Viewport context index (eye/view pass) within the nDisplay viewport.
+	 * @param View          Scene view associated with the specified viewport context.
+	 *
+	 * @return Computed brightness value used by the auto-exposure update.
+	 */
+	float ComputeViewportAutoExposureBrightness(
+		const FDisplayClusterViewport& InViewport,
+		const uint32 InContextNum,
+		FSceneView& View);
+
+private:
+	// Configuration of the current cluster node
+	const TSharedRef<FDisplayClusterViewportConfiguration, ESPMode::ThreadSafe> Configuration;
+
+	// Target Brightness for the current frame.
+	TOptional<float> TargetBrightness;
+
+	// Per-node AutoExposure data collected from all cluster viewports.
+	TMap<FName, FDisplayClusterNodeAutoExposureData> PerNodeAutoExposureData;
+
+	// Runtime settings for Game thread
+	TSharedRef<FDisplayClusterAutoExposureSettings, ESPMode::ThreadSafe> Settings;
+
+	// Runtime settings for Rendering thread
+	TSharedRef<FDisplayClusterAutoExposureSettings, ESPMode::ThreadSafe> SettingsProxy;
+
+private:
+	// RT->GT AE data queue
+	TArray<FDisplayClusterNodeAutoExposureData> DataQueue;
+
+	FCriticalSection DataQueueCS;
+
+private:
+	// Declare distribute state type
+	using FAutoExposureState = TDistributedCustomState<FDisplayClusterNodeAutoExposureData>;
+
+	// API to distribute auto exposure data between nodes.
+	// Note: Use TDistributedCustomState<> only for cluster
+	TSharedPtr<FAutoExposureState> AutoExposureState;
+
+	// This flag disables the use of AutoExposureState.
+	bool bDisableAutoExposureState = false;
+};
diff --git a/Engine/Plugins/Runtime/nDisplay/Source/DisplayCluster/Private/Render/Features/AutoExposure/DisplayClusterAutoExposureSettings.cpp b/Engine/Plugins/Runtime/nDisplay/Source/DisplayCluster/Private/Render/Features/AutoExposure/DisplayClusterAutoExposureSettings.cpp
new file mode 100644
--- /dev/null
+++ b/Engine/Plugins/Runtime/nDisplay/Source/DisplayCluster/Private/Render/Features/AutoExposure/DisplayClusterAutoExposureSettings.cpp
@@ -0,0 +1,102 @@
+// Copyright Epic Games, Inc. All Rights Reserved.
+
+#include "Render/Features/AutoExposure/DisplayClusterAutoExposureSettings.h"
+#include "Render/Viewport/Configuration/DisplayClusterViewportConfiguration.h"
+#include "DisplayClusterConfigurationTypes_AutoExposure.h"
+
+FDisplayClusterAutoExposureSettings::FDisplayClusterAutoExposureSettings(
+	const FDisplayClusterConfiguration_AutoExposureSettings& InAutoExposureSettings)
+{
+	bEnable = InAutoExposureSettings.bEnabled;
+
+	OuterViewportsForMetering = InAutoExposureSettings.MeteringSettings.Viewports;
+	ICVFXCamerasForMetering = InAutoExposureSettings.MeteringSettings.ICVFXCameras;
+	ClusterNodesForMetering = InAutoExposureSettings.MeteringSettings.Nodes;
+
+	ExcludedOuterViewports = InAutoExposureSettings.ExcludeSettings.Viewports;
+	ExcludedICVFXCameras = InAutoExposureSettings.ExcludeSettings.ICVFXCameras;
+	ExcludedClusterNodes = InAutoExposureSettings.ExcludeSettings.Nodes;
+}
+
+bool FDisplayClusterAutoExposureSettings::IsNodeExcludedFromAutoExposure(const FString& NodeId) const
+{
+	if (ExcludedClusterNodes.Contains(NodeId))
+	{
+		return true;
+	}
+
+	return false;
+}
+
+bool FDisplayClusterAutoExposureSettings::IsNodeAllowedForMetering(const FString& NodeId) const
+{
+	if (IsNodeExcludedFromAutoExposure(NodeId))
+	{
+		return false;
+	}
+
+	if (ClusterNodesForMetering.IsEmpty() || ClusterNodesForMetering.Contains(NodeId))
+	{
+		return true;
+	}
+
+	return false;
+}
+
+bool FDisplayClusterAutoExposureSettings::IsViewportExcludedFromAutoExposure(
+	const FString& InViewportBaseId,
+	const EDisplayClusterViewportFlags InViewportFlags) const
+{
+	// Outer viewport rules
+	if (!EnumHasAnyFlags(InViewportFlags, EDisplayClusterViewportFlags::ICVFX_AnyRole))
+	{
+		if (ExcludedOuterViewports.Contains(InViewportBaseId))
+		{
+			return true;
+		}
+	}
+	// Inner Frustum rules
+	else if (EnumHasAnyFlags(InViewportFlags, EDisplayClusterViewportFlags::ICVFX_InnerFrustum))
+	{
+		if (ExcludedICVFXCameras.Contains(InViewportBaseId))
+		{
+			return true;
+		}
+	}
+	// Exclude all other ICVFX viewports (uv/LC/CK)
+	else if (EnumHasAnyFlags(InViewportFlags, EDisplayClusterViewportFlags::ICVFX_AnyRole))
+	{
+		return true;
+	}
+
+	return false;
+}
+
+bool FDisplayClusterAutoExposureSettings::IsViewportEligibleForMetering(
+	const FString& InViewportBaseId,
+	const EDisplayClusterViewportFlags InViewportFlags) const
+{
+	if (IsViewportExcludedFromAutoExposure(InViewportBaseId, InViewportFlags))
+	{
+		return false;
+	}
+
+	// Outer viewport rules
+	if (!EnumHasAnyFlags(InViewportFlags, EDisplayClusterViewportFlags::ICVFX_AnyRole))
+	{
+		if (OuterViewportsForMetering.IsEmpty() || OuterViewportsForMetering.Contains(InViewportBaseId))
+		{
+			return true;
+		}
+	}
+	// Inner Frustum rules
+	else if (EnumHasAnyFlags(InViewportFlags, EDisplayClusterViewportFlags::ICVFX_InnerFrustum))
+	{
+		if (ICVFXCamerasForMetering.IsEmpty() || ICVFXCamerasForMetering.Contains(InViewportBaseId))
+		{
+			return true;
+		}
+	}
+
+	return false;
+}
diff --git a/Engine/Plugins/Runtime/nDisplay/Source/DisplayCluster/Private/Render/Features/AutoExposure/DisplayClusterAutoExposureSettings.h b/Engine/Plugins/Runtime/nDisplay/Source/DisplayCluster/Private/Render/Features/AutoExposure/DisplayClusterAutoExposureSettings.h
new file mode 100644
--- /dev/null
+++ b/Engine/Plugins/Runtime/nDisplay/Source/DisplayCluster/Private/Render/Features/AutoExposure/DisplayClusterAutoExposureSettings.h
@@ -0,0 +1,83 @@
+// Copyright Epic Games, Inc. All Rights Reserved.
+
+#pragma once
+
+#include "CoreMinimal.h"
+
+#include "Render/Viewport/Containers/DisplayClusterViewport_Enums.h"
+#include "Render/Viewport/Configuration/DisplayClusterViewportConfiguration.h"
+
+struct FDisplayClusterConfiguration_AutoExposureSettings;
+
+/**
+ * Runtime AutoExposure settings.
+ */
+struct FDisplayClusterAutoExposureSettings
+{
+	FDisplayClusterAutoExposureSettings() = default;
+	FDisplayClusterAutoExposureSettings(
+		const FDisplayClusterConfiguration_AutoExposureSettings& InAutoExposureSettings);
+
+	/**
+	 * Checks whether the specified cluster node is allowed to contribute its
+	 * viewports to auto-exposure metering, based on the given settings.
+	 *
+	 * @param NodeId     The identifier of the cluster node being tested.
+	 * @return           True if viewports from this node can be used for metering.
+	 */
+	bool IsNodeAllowedForMetering(const FString& NodeId) const;
+
+	/**
+	 * Checks whether the specified cluster node is excluded from auto-exposure.
+	 *
+	 * A node marked as excluded prevents all of its viewports from participating
+	 * in auto-exposure metering and from receiving auto-exposure customization.
+	 *
+	 * @param NodeId  Identifier of the cluster node to check.
+	 * @return        True if the node is excluded from auto-exposure.
+	 */
+	bool IsNodeExcludedFromAutoExposure(const FString& NodeId) const;
+
+	/**
+	* Checks whether the specified viewport is eligible to contribute to auto-exposure metering.
+	*
+	* @param InViewportBaseId  Base viewport identifier (shared across derived/tiled viewports).
+	* @param InViewportFlags   Viewport usage flags (e.g. Regular, InCamera, Chromakey, LightCard).
+	* @return                  True if this viewport is allowed to participate in AE metering.
+	*/
+	bool IsViewportEligibleForMetering(const FString& InViewportBaseId, const EDisplayClusterViewportFlags InViewportFlags) const;
+
+	/**
+	* Checks whether the specified viewport is excluded from auto-exposure.
+	*
+	* A viewport marked as excluded will not contribute to auto-exposure
+	* metering and will not receive auto-exposure customization.
+	*
+	* @param InViewportBaseId  Base viewport identifier (shared across derived/tiled viewports).
+	* @param InViewportFlags   Viewport usage flags (e.g. Regular, InCamera, Chromakey, LightCard).
+	* @return                  True if the viewport is excluded from auto-exposure.
+	*/
+	bool IsViewportExcludedFromAutoExposure(const FString& InViewportBaseId, const EDisplayClusterViewportFlags InViewportFlags) const;
+
+public:
+	// Enables or disables the AutoExposure feature at runtime.
+	bool bEnable = false;
+
+	// Outer viewport names used for auto - exposure metering.
+	TArray<FString> OuterViewportsForMetering;
+
+	// ICVFX camera Id's used for auto-exposure metering.
+	TArray<FString> ICVFXCamerasForMetering;
+
+	// Only viewports from these nodes can be used for metering.
+	TArray<FString> ClusterNodesForMetering;
+
+	// Outer viewpots excluded from auto-exposure.
+	TArray<FString> ExcludedOuterViewports;
+
+	// ICVFX cameras excluded from auto-exposure.
+	TArray<FString> ExcludedICVFXCameras;
+
+	// Cluster nodes excluded from auto-exposure.
+	TArray<FString> ExcludedClusterNodes;
+};
\ No newline at end of file
diff --git a/Engine/Plugins/Runtime/nDisplay/Source/DisplayCluster/Private/Render/Features/DisplayClusterUpscaler.cpp b/Engine/Plugins/Runtime/nDisplay/Source/DisplayCluster/Private/Render/Features/DisplayClusterUpscaler.cpp
new file mode 100644
--- /dev/null
+++ b/Engine/Plugins/Runtime/nDisplay/Source/DisplayCluster/Private/Render/Features/DisplayClusterUpscaler.cpp
@@ -0,0 +1,577 @@
+// Copyright Epic Games, Inc. All Rights Reserved.
+
+#include "Render/Features/DisplayClusterUpscaler.h"
+#include "Render/Viewport/DisplayClusterViewport.h"
+#include "DisplayClusterSceneViewExtensions.h"
+
+#include "IUpscalerModularFeature.h"
+#include "Features/IModularFeatures.h"
+#include "SceneViewExtensionContext.h"
+
+#include "LegacyScreenPercentageDriver.h"
+#include "SceneView.h"
+
+using namespace UE::VirtualProduction;
+
+namespace UE::DisplayClusterRender::Features::Upscaler
+{
+	int32 GDisplayClusterRenderUpscaling_HideDisabledMethods = 1;
+	static FAutoConsoleVariableRef CVarDisplayClusterRenderUpscaling_HideDisabledMethods(
+		TEXT("nDisplay.render.upscaling.HideDisabledMethods"),
+		GDisplayClusterRenderUpscaling_HideDisabledMethods,
+		TEXT("Hide disabled upscaling methods (default = 1).\n")
+		TEXT("0 - Show all upsccaling methods.\n"),
+		ECVF_Default
+	);
+
+	int32 GDisplayClusterRenderUpscaling_Enable = 1;
+	static FAutoConsoleVariableRef CVarDisplayClusterRenderUpscaling_Enable(
+		TEXT("nDisplay.render.upscaling.Enable"),
+		GDisplayClusterRenderUpscaling_Enable,
+		TEXT("Allows to use the upscaling settings (default = 1).\n")
+		TEXT("0 - Ignore upscaling settings.\n"),
+		ECVF_Default
+	);
+
+	int32 GDisplayClusterRenderUpscaling_EnableCustomUpscalers = 1;
+	static FAutoConsoleVariableRef CVarDisplayClusterRenderUpscaling_EnableCustomUpscalers(
+		TEXT("nDisplay.render.upscaling.EnableCustomUpscalers"),
+		GDisplayClusterRenderUpscaling_EnableCustomUpscalers,
+		TEXT("Allows the use of custom upscalers (default = 1).\n")
+		TEXT("0 - Don't use custom upscalers and render by default.\n"),
+		ECVF_Default
+	);
+
+	/** A custom upscaler may be used. */
+	static inline bool UseCustomUpscalers()
+	{
+		return GDisplayClusterRenderUpscaling_EnableCustomUpscalers && GDisplayClusterRenderUpscaling_Enable;
+	}
+
+	/** Iterate over all upscaler modular features. */
+	static void ForeachUpscaler(TFunction<void(IUpscalerModularFeature&)> IteratorFunc)
+	{
+		static IModularFeatures& ModularFeatures = IModularFeatures::Get();
+		ModularFeatures.LockModularFeatureList();
+		ON_SCOPE_EXIT
+		{
+			ModularFeatures.UnlockModularFeatureList();
+		};
+
+		// Iterate over enabled upscalers
+		const TArray<IUpscalerModularFeature*> UpscalerModularFeatures = ModularFeatures.GetModularFeatureImplementations<IUpscalerModularFeature>(IUpscalerModularFeature::ModularFeatureName);
+		for (IUpscalerModularFeature* UpscalerModularFeature : UpscalerModularFeatures)
+		{
+			if (UpscalerModularFeature == nullptr)
+			{
+				continue;
+			}
+
+			if (UpscalerModularFeature->IsFeatureEnabled())
+			{
+				IteratorFunc(*UpscalerModularFeature);
+			}
+		}
+	}
+
+	/** Get global fraction. */
+	static float GetGlobalScreenPercentage()
+	{
+		static const TConsoleVariableData<float>* CVarScreenPercentage = IConsoleManager::Get().FindTConsoleVariableDataFloat(TEXT("r.ScreenPercentage"));
+		if (CVarScreenPercentage)
+		{
+			const float GlobalScreenPercentage = CVarScreenPercentage->GetValueOnGameThread() / 100.0f;
+			if (GlobalScreenPercentage > 0)
+			{
+				return GlobalScreenPercentage;
+			}
+		}
+
+		return 1.0f;
+	}
+
+	/** Get secondary screen percentage fraction. */
+	static float GetSecondaryViewFraction(const float InDPIScale)
+	{
+		static IConsoleVariable* CVarCustomSecondaryScreenPercentage = IConsoleManager::Get().FindConsoleVariable(TEXT("r.SecondaryScreenPercentage.GameViewport"), false);
+		if (CVarCustomSecondaryScreenPercentage)
+		{
+			const float CustomSecondaryScreenPercentage = CVarCustomSecondaryScreenPercentage->GetFloat();
+			if (CustomSecondaryScreenPercentage > 0.0)
+			{
+				// Override secondary resolution fraction with CVar.
+				return FMath::Min(CustomSecondaryScreenPercentage / 100.0f, 1.0f);
+			}
+		}
+
+#if WITH_EDITOR
+			// Automatically compute secondary resolution fraction from DPI.
+			// When in high res screenshot do not modify screen percentage based on dpi scale
+			if (GIsEditor && !GIsHighResScreenshot && InDPIScale > 0.0)
+			{
+				static IConsoleVariable* CVarEditorViewportHighDPIPtr = IConsoleManager::Get().FindConsoleVariable(TEXT("r.Editor.Viewport.HighDPI"));
+				if (CVarEditorViewportHighDPIPtr && CVarEditorViewportHighDPIPtr->GetInt() == 0)
+				{
+
+					const float DPIDerivedResolutionFraction = FMath::Min(1.0f / InDPIScale, 1.0f);
+
+					return DPIDerivedResolutionFraction;
+				}
+			}
+#endif
+
+		return 1.0f;
+	}
+
+
+	/** Check if temporal upscaler can be used. */
+	static bool CanUseTemporalUpscaler()
+	{
+		static const IConsoleVariable* CVarTemporalAAUpscaler = IConsoleManager::Get().FindConsoleVariable(TEXT("r.TemporalAA.Upscaler"));
+		return CVarTemporalAAUpscaler && (CVarTemporalAAUpscaler->GetInt() != 0);
+	}
+
+	/** Check if MSAA can be used. */
+	static bool CanUseMSAA(const FStaticFeatureLevel InFeatureLevel)
+	{
+		static auto* MSAACountCVar = IConsoleManager::Get().FindTConsoleVariableDataInt(TEXT("r.MSAACount"));
+		if (MSAACountCVar->GetValueOnAnyThread() <= 0)
+		{
+			return false;
+		}
+
+		EShaderPlatform ShaderPlatform = GetFeatureLevelShaderPlatform(InFeatureLevel);
+		if ((InFeatureLevel >= ERHIFeatureLevel::SM5 && !IsForwardShadingEnabled(ShaderPlatform)))
+		{
+			// MSAA requires forward shading renderer
+			return false;
+		}
+	
+		return true;
+	}
+
+	/** Temporal upscaler interface requires special AA method. */
+	static bool SceneViewSupportsTemporalUpscaler(const FSceneView& InSceneView)
+	{
+		// check if current AA method legitable for temporal upscaler
+		switch (InSceneView.AntiAliasingMethod)
+		{
+		case EAntiAliasingMethod::AAM_TSR:
+		case EAntiAliasingMethod::AAM_TemporalAA: //@TODO check AAM_TemporalAA
+			return true;
+
+		default:
+			break;
+		}
+
+		return false;
+	}
+
+	/** Return default AntiAliasing method. */
+	static EDisplayClusterUpscalerAntiAliasingMethod GetDefaultAntiAliasingMethod(const FStaticFeatureLevel InFeatureLevel)
+	{
+		static auto* DefaultAntiAliasingCvar = IConsoleManager::Get().FindTConsoleVariableDataInt(TEXT("r.AntiAliasingMethod"));
+		static EAntiAliasingMethod DefaultAntiAliasingMethod = EAntiAliasingMethod(FMath::Clamp<int32>(DefaultAntiAliasingCvar->GetValueOnAnyThread(), 0, AAM_MAX));
+
+		switch (DefaultAntiAliasingMethod)
+		{
+		case EAntiAliasingMethod::AAM_FXAA:
+			return EDisplayClusterUpscalerAntiAliasingMethod::FXAA;
+
+		case EAntiAliasingMethod::AAM_MSAA:
+			if (CanUseMSAA(InFeatureLevel))
+			{
+				return EDisplayClusterUpscalerAntiAliasingMethod::MSAA;
+			}
+			break;
+
+		case EAntiAliasingMethod::AAM_TemporalAA:
+			return EDisplayClusterUpscalerAntiAliasingMethod::TAA;
+
+		case EAntiAliasingMethod::AAM_TSR:
+			{
+				EShaderPlatform ShaderPlatform = GetFeatureLevelShaderPlatform(InFeatureLevel);
+				if (!SupportsTSR(ShaderPlatform))
+				{
+					// Fallback to UE4's TAA if TSR isn't supported on that platform
+					return EDisplayClusterUpscalerAntiAliasingMethod::TAA;
+				}
+			}
+			return EDisplayClusterUpscalerAntiAliasingMethod::TSR;
+
+		default:
+			break;
+		}
+
+		return EDisplayClusterUpscalerAntiAliasingMethod::None;
+	}
+
+	/** 
+	* Get the antialiasing method that can be used for the ViewFamily.
+	* Note: the logic of this function is based on the FSceneView::SetupAntiAliasingMethod() function.
+	* 
+	* @param InViewport         - (in) the DC viewport
+	* @param InUpscalerSettings - (in) the upscaler settings that we want to use.
+	* @param InViewFamily       - (in) ViewFamily thyat own SceneView.
+	* @param InSceneView        - (in) SceneView to which the upscaler settings should be applied.
+	* 
+	* @return AntiAliasing method that can be used.
+	*/
+	static EDisplayClusterUpscalerAntiAliasingMethod GetSuitableAntiAliasingMethod(
+		const FDisplayClusterViewport& InViewport,
+		const FDisplayClusterUpscalerSettings& InUpscalerSettings,
+		const FSceneViewFamily& InViewFamily,
+		const FSceneView& InSceneView)
+	{
+		const bool bWillApplyTemporalAA = InViewFamily.EngineShowFlags.PostProcessing || InSceneView.bIsPlanarReflection;
+
+		if (!bWillApplyTemporalAA || !InViewFamily.EngineShowFlags.AntiAliasing)
+		{
+			return EDisplayClusterUpscalerAntiAliasingMethod::None;
+		}
+
+		// Get AntiAliasing method from the upscaler settings
+		EDisplayClusterUpscalerAntiAliasingMethod AntiAliasingMethod = InUpscalerSettings.AntiAliasingMethod;
+		
+		// Use the default method if the upscaler settings is disabled
+		if (!GDisplayClusterRenderUpscaling_Enable)
+		{
+			AntiAliasingMethod = EDisplayClusterUpscalerAntiAliasingMethod::Default;
+		}
+
+		// Back to default method if TemporalUpscaling not supported
+		if (AntiAliasingMethod == EDisplayClusterUpscalerAntiAliasingMethod::TAA
+			|| AntiAliasingMethod == EDisplayClusterUpscalerAntiAliasingMethod::TSR)
+		{
+			if (!CanUseTemporalUpscaler())
+			{
+				AntiAliasingMethod = EDisplayClusterUpscalerAntiAliasingMethod::Default;
+			}
+		}
+
+		// Back to default method is MSAA not supported
+		if (AntiAliasingMethod == EDisplayClusterUpscalerAntiAliasingMethod::MSAA && !CanUseMSAA(InSceneView.GetFeatureLevel()))
+		{
+			AntiAliasingMethod = EDisplayClusterUpscalerAntiAliasingMethod::Default;
+		}
+
+		// Get default method from the project settings
+		if (AntiAliasingMethod == EDisplayClusterUpscalerAntiAliasingMethod::Default)
+		{
+			AntiAliasingMethod = GetDefaultAntiAliasingMethod(InSceneView.GetFeatureLevel());
+		}
+
+		if (AntiAliasingMethod == EDisplayClusterUpscalerAntiAliasingMethod::TAA)
+		{
+			// Downgrade TAA->FXAA
+			if (!InViewFamily.EngineShowFlags.TemporalAA || !InViewFamily.bRealtimeUpdate || !SupportsGen4TAA(InSceneView.GetShaderPlatform()))
+			{
+				return EDisplayClusterUpscalerAntiAliasingMethod::FXAA;
+			}
+		}
+		else if (AntiAliasingMethod == EDisplayClusterUpscalerAntiAliasingMethod::TSR)
+		{
+			// Downgrade TSR->FXAA
+			if (!InViewFamily.EngineShowFlags.TemporalAA || !InViewFamily.bRealtimeUpdate || !SupportsTSR(InSceneView.GetShaderPlatform()))
+			{
+				return EDisplayClusterUpscalerAntiAliasingMethod::FXAA;
+			}
+		}
+
+		if(AntiAliasingMethod == EDisplayClusterUpscalerAntiAliasingMethod::TAA
+		|| AntiAliasingMethod == EDisplayClusterUpscalerAntiAliasingMethod::TSR)
+		{
+			// TemporalAA/TSR requires the view to have a valid state.
+			if (InSceneView.State == nullptr)
+			{	
+				return EDisplayClusterUpscalerAntiAliasingMethod::FXAA;
+			}
+
+			// Disable TSR/TAA for preview
+			if (!InViewport.Configuration->GetPreviewSettings().bPreviewEnableTSR)
+			{
+				if (InViewport.Configuration->IsPreviewRendering())
+				{
+					return EDisplayClusterUpscalerAntiAliasingMethod::FXAA;
+				}
+			}
+		}
+
+		return AntiAliasingMethod;
+	}
+};
+
+void FDisplayClusterUpscaler::SetupSceneView(
+	const FDisplayClusterViewport& InViewport,
+	const FDisplayClusterUpscalerSettings& InUpscalerSettings,
+	const FSceneViewFamily& InViewFamily,
+	FSceneView& InOutView)
+{
+	using namespace UE::DisplayClusterRender::Features::Upscaler;
+
+	// Get AntiAliasingMethod that can be used.
+	const EDisplayClusterUpscalerAntiAliasingMethod SuitableAntiAliasingMethod = GetSuitableAntiAliasingMethod(InViewport, InUpscalerSettings, InViewFamily, InOutView);
+	switch (SuitableAntiAliasingMethod)
+	{
+	case EDisplayClusterUpscalerAntiAliasingMethod::None:
+		InOutView.AntiAliasingMethod = EAntiAliasingMethod::AAM_None;
+		InOutView.PrimaryScreenPercentageMethod = EPrimaryScreenPercentageMethod::SpatialUpscale;
+		break;
+
+	case EDisplayClusterUpscalerAntiAliasingMethod::FXAA:
+		InOutView.AntiAliasingMethod = EAntiAliasingMethod::AAM_FXAA;
+		InOutView.PrimaryScreenPercentageMethod = EPrimaryScreenPercentageMethod::SpatialUpscale;
+		break;
+
+	case EDisplayClusterUpscalerAntiAliasingMethod::MSAA:
+		InOutView.AntiAliasingMethod = EAntiAliasingMethod::AAM_MSAA;
+		InOutView.PrimaryScreenPercentageMethod = EPrimaryScreenPercentageMethod::SpatialUpscale;
+		break;
+
+	case EDisplayClusterUpscalerAntiAliasingMethod::TAA:
+		InOutView.AntiAliasingMethod = EAntiAliasingMethod::AAM_TemporalAA;
+		InOutView.PrimaryScreenPercentageMethod = EPrimaryScreenPercentageMethod::TemporalUpscale;
+		break;
+
+	case EDisplayClusterUpscalerAntiAliasingMethod::TSR:
+		InOutView.AntiAliasingMethod = EAntiAliasingMethod::AAM_TSR;
+		InOutView.PrimaryScreenPercentageMethod = EPrimaryScreenPercentageMethod::TemporalUpscale;
+		break;
+
+	default:
+		break;
+	}
+
+	// Custom upscaler should configure view for the settings
+	if (!InUpscalerSettings.CustomUpscalerName.IsNone() && UseCustomUpscalers())
+	{
+		ForeachUpscaler(
+			[&](IUpscalerModularFeature& Upscaler)
+			{
+				if (Upscaler.GetName() == InUpscalerSettings.CustomUpscalerName)
+				{
+					Upscaler.SetupSceneView(InUpscalerSettings.CustomUpscalerSettings, InOutView);
+				}
+			});
+	}
+}
+
+FName FDisplayClusterUpscaler::PostConfigureViewFamily(
+	const FDisplayClusterUpscalerSettings& InUpscalerSettings,
+	const float InScreenPercentage,
+	const float InDPIScale,
+	FSceneViewFamilyContext& InOutViewFamily,
+	const TArray<FSceneView*>& InOutViews)
+{
+	using namespace UE::DisplayClusterRender::Features::Upscaler;
+
+	// One view per viewfamily
+	check(InOutViewFamily.Views.Num() == 1);
+	check(InOutViewFamily.Views[0]);
+
+#if WITH_EDITOR
+	if (GIsEditor)
+	{
+		// Force enable view family show flag for HighDPI derived's screen percentage.
+		InOutViewFamily.EngineShowFlags.ScreenPercentage = true;
+	}
+#endif
+
+	// Force screen percentage show flag to be turned off if not supported.
+	if (!InOutViewFamily.SupportsScreenPercentage())
+	{
+		InOutViewFamily.EngineShowFlags.ScreenPercentage = false;
+	}
+
+	if (!InOutViewFamily.EngineShowFlags.ScreenPercentage)
+	{
+		return NAME_None;
+	}
+
+	// If a screen percentage interface was set by dynamic resolution, then do nothing
+	if (InOutViewFamily.GetScreenPercentageInterface() != nullptr)
+	{
+		return NAME_None;
+	}
+
+	FUpscalerModularFeatureParameters UpscalerParam;
+	UpscalerParam.UpscalerScreenPercentage = InScreenPercentage;
+	UpscalerParam.SecondaryScreenPercentage = GetSecondaryViewFraction(InDPIScale);
+
+	check(UpscalerParam.UpscalerScreenPercentage > 0.0f);
+	check(UpscalerParam.SecondaryScreenPercentage > 0.0f);
+
+	// Use custom upscaler
+	if (!InUpscalerSettings.CustomUpscalerName.IsNone() && UseCustomUpscalers())
+	{
+		// Find  upscaler modular feature
+		FName UpscalerName = NAME_None;
+		ForeachUpscaler([&](IUpscalerModularFeature& Upscaler)
+			{
+				if (Upscaler.GetName() == InUpscalerSettings.CustomUpscalerName)
+				{
+					if (Upscaler.PostConfigureViewFamily(
+						InUpscalerSettings.CustomUpscalerSettings, UpscalerParam,
+						InOutViewFamily))
+					{
+						// Return upscaler name on success
+						UpscalerName = Upscaler.GetName();
+					}
+				}
+			});
+
+		if (UpscalerName != NAME_None)
+		{
+			// Upscaler modular feature expected to set screen percentage interface
+			check(InOutViewFamily.GetScreenPercentageInterface())
+
+			return UpscalerName;
+		}
+	}
+
+	// Validating that InOutViews is non-empty.
+	if (InOutViews.IsEmpty() || InOutViews[0] == nullptr)
+	{
+		return NAME_None;
+	}
+
+	// Get the upscaler settings from the first SceneView.
+	const EPrimaryScreenPercentageMethod PrimaryScreenPercentageMethod = InOutViews[0]->PrimaryScreenPercentageMethod;
+	const EAntiAliasingMethod AntiAliasingMethod = InOutViews[0]->AntiAliasingMethod;
+
+	// Get global view fraction set by r.ScreenPercentage.
+	const float GlobalScreenPercentage = GetGlobalScreenPercentage();
+
+	// Get TSR or TAA resolution fraction range
+	FVector2f TemporalUpscaleResolutionFractionRange = FVector2f::ZeroVector;
+	if (PrimaryScreenPercentageMethod == EPrimaryScreenPercentageMethod::TemporalUpscale)
+	{
+		if (AntiAliasingMethod == EAntiAliasingMethod::AAM_TemporalAA)
+		{
+			// Sets the minimal and maximal screen percentage for TAAU. [0.5 .. 2.0]
+			TemporalUpscaleResolutionFractionRange = FVector2f(
+				ISceneViewFamilyScreenPercentage::kMinTAAUpsampleResolutionFraction,
+				ISceneViewFamilyScreenPercentage::kMaxTAAUpsampleResolutionFraction);
+		}
+		else if (AntiAliasingMethod == EAntiAliasingMethod::AAM_TSR)
+		{
+			// Sets the minimal and maximal screen percentage for TSR.  [0.25 .. 2.0]
+			TemporalUpscaleResolutionFractionRange = FVector2f(
+				ISceneViewFamilyScreenPercentage::kMinTSRResolutionFraction,
+				ISceneViewFamilyScreenPercentage::kMaxTSRResolutionFraction);
+		}
+	}
+
+	// We need to split the screen percentage value if it is outside the TSR/TAA range.
+	// All TemporalUpscaleResolutionFractionRange values are greater than zero.
+	if(TemporalUpscaleResolutionFractionRange.GetMin() > 0)
+	{
+		// Adjust for global screen% in math bellow
+		const float FinalResolutionFraction = UpscalerParam.UpscalerScreenPercentage * GlobalScreenPercentage;
+		if (FinalResolutionFraction < TemporalUpscaleResolutionFractionRange.X)
+		{
+			// Move the remaining ScreenPercentage to the spatial upscaler
+			const float RemainingScreenPercentage = UpscalerParam.UpscalerScreenPercentage / TemporalUpscaleResolutionFractionRange.X;
+			UpscalerParam.SecondaryScreenPercentage *= RemainingScreenPercentage;
+
+			// Adjust min value (0.5f) for the global screen%. Expectations that the actual TSR_Fraction=0.5 
+			UpscalerParam.UpscalerScreenPercentage = TemporalUpscaleResolutionFractionRange.X / GlobalScreenPercentage;
+		}
+		else if (FinalResolutionFraction > TemporalUpscaleResolutionFractionRange.Y)
+		{
+			// Move the remaining ScreenPercentage to the spatial upscaler
+			const float RemainingScreenPercentage = UpscalerParam.UpscalerScreenPercentage / TemporalUpscaleResolutionFractionRange.Y;
+			UpscalerParam.SecondaryScreenPercentage *= RemainingScreenPercentage;
+
+			// Adjust min value (0.5f) for the global screen%. Expectations that the actual TSR_Fraction=0.5 
+			UpscalerParam.UpscalerScreenPercentage = TemporalUpscaleResolutionFractionRange.Y / GlobalScreenPercentage;
+		}
+	}
+
+	const float FinalResolutionFraction = UpscalerParam.UpscalerScreenPercentage * GlobalScreenPercentage;
+
+	// Clamp final resolution fraction
+	// Sets the minimal and max screen percentage to [0.01 ... 4.0]
+	float ClampedFinalResolutionFraction = FMath::Clamp(
+		FinalResolutionFraction,
+		ISceneViewFamilyScreenPercentage::kMinResolutionFraction,
+		ISceneViewFamilyScreenPercentage::kMaxResolutionFraction);
+
+	UpscalerParam.UpscalerScreenPercentage = ClampedFinalResolutionFraction / GlobalScreenPercentage;
+
+	// Setup screen percentage
+	InOutViewFamily.SetScreenPercentageInterface(new FLegacyScreenPercentageDriver(
+		InOutViewFamily, UpscalerParam.UpscalerScreenPercentage));
+
+	InOutViewFamily.SecondaryViewFraction = UpscalerParam.SecondaryScreenPercentage;
+
+	return NAME_None;
+}
+
+void FDisplayClusterUpscaler::InitializeNewFrame()
+{
+	using namespace UE::DisplayClusterRender::Features::Upscaler;
+
+	// Store the upscaler name with function GUID
+	static TMap<FName, FGuid> SceneViewExtensionIsActiveFunctorsMap;
+
+	// If the modular function will be disabled at runtime, we should forget about it
+	// And then restore again when it is enabled.
+	TMap<FName, FGuid> NewSceneViewExtensionIsActiveFunctorsMap;
+
+	ForeachUpscaler(
+		[&](IUpscalerModularFeature& Upscaler)
+		{
+			// Register IsActive Functors
+			const FName& UpscalerName = Upscaler.GetName();
+			if (SceneViewExtensionIsActiveFunctorsMap.Contains(UpscalerName))
+			{
+				// Use exist Guid
+				NewSceneViewExtensionIsActiveFunctorsMap.Emplace(UpscalerName, SceneViewExtensionIsActiveFunctorsMap[UpscalerName]);
+			}
+			else
+			{
+				FSceneViewExtensionIsActiveFunctor IsActiveFunctor;
+				IsActiveFunctor.IsActiveFunction = [UpscalerName](const ISceneViewExtension* SceneViewExtension, const FSceneViewExtensionContext& Context)
+					{
+						// Note: SceneViewExtension this pointer is always passed to functors from FSceneViewExtensionBase::IsActiveThisFrame, we confirm this remains so.
+						check(SceneViewExtension);
+
+						// Is nDisplay context
+						static const FDisplayClusterSceneViewExtensionContext DCViewExtensionContext;
+						if (Context.IsA(MoveTempIfPossible(DCViewExtensionContext)))
+						{
+							// Is viewport exist
+							const FDisplayClusterSceneViewExtensionContext& DisplayContext = static_cast<const FDisplayClusterSceneViewExtensionContext&>(Context);
+							if (DisplayContext.DisplayClusterViewport.IsValid())
+							{
+								const FDisplayClusterUpscalerSettings& UpscalerSettings = DisplayContext.DisplayClusterViewport->GetRenderSettings().UpscalerSettings;
+								if (!UpscalerSettings.bIsActive)
+								{
+									// Custom upscalers are disabled for this viewport
+									return TOptional<bool>(false);
+								}
+
+								if (UpscalerSettings.CustomUpscalerName != UpscalerName)
+								{
+									// This viewport uses a different upscaler.
+									return TOptional<bool>(false);
+								}
+
+								return TOptional<bool>(true);
+							}
+						}
+
+						return TOptional<bool>();
+					};
+
+				if (Upscaler.AddSceneViewExtensionIsActiveFunctor(IsActiveFunctor))
+				{
+					NewSceneViewExtensionIsActiveFunctorsMap.Emplace(UpscalerName, IsActiveFunctor.GetGuid());
+				}
+			}
+		});
+
+	// Update
+	SceneViewExtensionIsActiveFunctorsMap = NewSceneViewExtensionIsActiveFunctorsMap;
+}
diff --git a/Engine/Plugins/Runtime/nDisplay/Source/DisplayCluster/Private/Render/Features/DisplayClusterUpscaler.h b/Engine/Plugins/Runtime/nDisplay/Source/DisplayCluster/Private/Render/Features/DisplayClusterUpscaler.h
new file mode 100644
--- /dev/null
+++ b/Engine/Plugins/Runtime/nDisplay/Source/DisplayCluster/Private/Render/Features/DisplayClusterUpscaler.h
@@ -0,0 +1,53 @@
+// Copyright Epic Games, Inc. All Rights Reserved.
+
+#pragma once
+
+#include "CoreMinimal.h"
+
+#include "Render/Upscaler/DisplayClusterUpscalerSettings.h"
+
+class FDisplayClusterViewport;
+class FSceneViewFamilyContext;
+class FSceneViewFamily;
+class FSceneView;
+
+/**
+ * nDisplay upscaling implementation.
+ */
+class FDisplayClusterUpscaler
+{
+public:
+	/**
+	 * Prepares all registered upscaler modular features for the new frame.
+	 * Should be called at the start of each frame.
+	 */
+	static void InitializeNewFrame();
+
+	/** Configure upscaler and viewfamily for ertain settings.
+	* 
+	* @param InUpscalerSettings - (in) the settings to apply for viewport
+	* @param InScreenPercentage - (inj screen percentage of the viewport
+	* @param InDPIScale         - (in) The DPI scale of the window
+	* @param InOutViewFamily    - (in, out) The view family that will be configured.
+	* @param InOutViews         - (in, out) The views in the family that will be configured.
+	*/
+	static FName PostConfigureViewFamily(
+		const FDisplayClusterUpscalerSettings& InUpscalerSettings,
+		const float InScreenPercentage,
+		const float InDPIScale,
+		FSceneViewFamilyContext& InOutViewFamily,
+		const TArray<FSceneView*>& InOutViews);
+
+	/**
+	* Configure scene view for upscalers.
+	* @param InViewport         - (in) the viewport who ask for this setup
+	* @param InUpscalerSettings - (in) the settings to apply for view
+	* @param InViewFamily       - (in) ViewFamily that own this view.
+	* @param InOutView          - (in, out) The view that will be configured. 
+	*/
+	static void SetupSceneView(
+		const FDisplayClusterViewport& InViewport,
+		const FDisplayClusterUpscalerSettings& InUpscalerSettings,
+		const FSceneViewFamily& InViewFamily,
+		FSceneView& InOutView);
+};
diff --git a/Engine/Plugins/Runtime/nDisplay/Source/DisplayCluster/Private/Render/Upscaler/DisplayClusterUpscaler.cpp b/Engine/Plugins/Runtime/nDisplay/Source/DisplayCluster/Private/Render/Upscaler/DisplayClusterUpscaler.cpp
deleted file mode 100644
diff --git a/Engine/Plugins/Runtime/nDisplay/Source/DisplayCluster/Private/Render/Upscaler/DisplayClusterUpscaler.h b/Engine/Plugins/Runtime/nDisplay/Source/DisplayCluster/Private/Render/Upscaler/DisplayClusterUpscaler.h
deleted file mode 100644
diff --git a/Engine/Plugins/Runtime/nDisplay/Source/DisplayCluster/Private/Render/Viewport/Configuration/DisplayClusterViewportConfiguration.cpp b/Engine/Plugins/Runtime/nDisplay/Source/DisplayCluster/Private/Render/Viewport/Configuration/DisplayClusterViewportConfiguration.cpp
--- a/Engine/Plugins/Runtime/nDisplay/Source/DisplayCluster/Private/Render/Viewport/Configuration/DisplayClusterViewportConfiguration.cpp
+++ b/Engine/Plugins/Runtime/nDisplay/Source/DisplayCluster/Private/Render/Viewport/Configuration/DisplayClusterViewportConfiguration.cpp
@@ -13,9 +13,12 @@
 #include "DisplayClusterViewportConfigurationHelpers_RenderFrameSettings.h"
 
 #include "DisplayClusterConfigurationStrings.h"
+#include "DisplayClusterProjectionStrings.h"
 
 #include "DisplayClusterRootActor.h"
 
+#include "Render/Features/AutoExposure/DisplayClusterAutoExposure.h"
+
 #include "Render/Viewport/DisplayClusterViewport.h"
 #include "Render/Viewport/DisplayClusterViewportManager.h"
 #include "Render/Viewport/DisplayClusterViewportManagerProxy.h"
@@ -396,7 +399,7 @@ void FDisplayClusterViewportConfiguration::ReleaseConfiguration()
 	if (FDisplayClusterViewportManager* ViewportManager = GetViewportManagerImpl())
 	{
 		ViewportManager->HandleEndScene();
-		ViewportManager->ReleaseTextures();
+		ViewportManager->ReleaseAllResources();
 	}
 
 	RenderFrameSettings = FDisplayClusterRenderFrameSettings();
@@ -473,6 +476,12 @@ bool FDisplayClusterViewportConfiguration::ImplUpdateConfiguration(EDisplayClust
 
 	FDisplayClusterViewportConfigurationHelpers_RenderFrameSettings::PostUpdateRenderFrameConfiguration(*this);
 
+	// Apply Auto Exposure new settings
+	if (const FDisplayClusterConfigurationICVFX_StageSettings* StageSettings = GetStageSettings())
+	{
+		ViewportManager->AutoExposure->ApplyAutoExposureSettings(StageSettings->AutoExposureSettings);
+	}
+
 	return true;
 }
 
@@ -486,3 +495,203 @@ void FDisplayClusterViewportConfiguration::ImplUpdateConfigurationVisibility()
 		SceneRootActor->GetHiddenInGamePrimitives(RootActorHidePrimitivesList);
 	}
 }
+
+FString FDisplayClusterViewportConfiguration::BuildViewportName(
+	const FString& InViewportBaseId,
+	const EDisplayClusterViewportFlags InViewportFlags,
+	const FIntPoint* InOptTilePos)
+{
+	// Regular viewport use same name
+	FString UniqueName = InViewportBaseId;
+
+	// Build unique ICVFX name
+	if (EnumHasAnyFlags(InViewportFlags, EDisplayClusterViewportFlags::ICVFX_AnyRole))
+	{
+		if (EnumHasAllFlags(InViewportFlags, EDisplayClusterViewportFlags::ICVFX_InnerFrustum))
+		{
+			UniqueName = FString::Printf(TEXT("%s_icvfx_%s_incamera"), *RenderFrameSettings.CurrentNode.Id, *InViewportBaseId);
+		}
+		else if (EnumHasAllFlags(InViewportFlags, EDisplayClusterViewportFlags::ICVFX_ChromaKey))
+		{
+			UniqueName = FString::Printf(TEXT("%s_icvfx_%s_chromakey"), *RenderFrameSettings.CurrentNode.Id, *InViewportBaseId);
+		}
+		else if (EnumHasAnyFlags(InViewportFlags, EDisplayClusterViewportFlags::ICVFX_LightCard))
+		{
+			if (EnumHasAllFlags(InViewportFlags, EDisplayClusterViewportFlags::ICVFX_LightCard | EDisplayClusterViewportFlags::ICVFX_Modifier_RenderUnderInnerFrustum))
+			{
+				UniqueName = FString::Printf(TEXT("%s_icvfx_%s_lightcard_under"), *RenderFrameSettings.CurrentNode.Id, *InViewportBaseId);
+
+			}
+			else if (EnumHasAllFlags(InViewportFlags, EDisplayClusterViewportFlags::ICVFX_LightCard | EDisplayClusterViewportFlags::ICVFX_Modifier_RenderOverInnerFrustum))
+			{
+				UniqueName = FString::Printf(TEXT("%s_icvfx_%s_lightcard_over"), *RenderFrameSettings.CurrentNode.Id, *InViewportBaseId);
+			}
+			else
+			{
+				check(false); // Invalid EDisplayClusterViewportFlags flags combination: ICVFX_LightCard requires ICVFX_Modifier_RenderUnderInnerFrustum or ICVFX_Modifier_RenderOverInnerFrustum.
+
+				UniqueName = FString::Printf(TEXT("%s_icvfx_%s_lightcard_invalid"), *RenderFrameSettings.CurrentNode.Id, *InViewportBaseId);
+			}
+		}
+		else if (EnumHasAnyFlags(InViewportFlags, EDisplayClusterViewportFlags::ICVFX_UVLightCard))
+		{
+			if (EnumHasAllFlags(InViewportFlags, EDisplayClusterViewportFlags::ICVFX_UVLightCard | EDisplayClusterViewportFlags::ICVFX_Modifier_RenderUnderInnerFrustum))
+			{
+				UniqueName = FString::Printf(TEXT("%s_icvfx_%s_uvlightcard_under"), *RenderFrameSettings.CurrentNode.Id, *InViewportBaseId);
+			}
+			else if (EnumHasAllFlags(InViewportFlags, EDisplayClusterViewportFlags::ICVFX_UVLightCard | EDisplayClusterViewportFlags::ICVFX_Modifier_RenderOverInnerFrustum))
+			{
+				UniqueName = FString::Printf(TEXT("%s_icvfx_%s_uvlightcard_over"), *RenderFrameSettings.CurrentNode.Id, *InViewportBaseId);
+			}
+			else
+			{
+				check(false); // Invalid EDisplayClusterViewportFlags flags combination: ICVFX_UVLightCard requires ICVFX_Modifier_RenderUnderInnerFrustum or ICVFX_Modifier_RenderOverInnerFrustum.
+
+				UniqueName = FString::Printf(TEXT("%s_icvfx_%s_uvlightcard_invalid"), *RenderFrameSettings.CurrentNode.Id, *InViewportBaseId);
+			}
+		}
+		else
+		{
+			check(false); // Invalid EDisplayClusterViewportFlags flags combination: Undefined viewport type - one of LC/CK/InCamera is expected.
+
+			UniqueName = FString::Printf(TEXT("%s_icvfx_%s_invalid"), *RenderFrameSettings.CurrentNode.Id, *InViewportBaseId);
+		}
+	}
+
+	// Build the tile viewport's final unique name
+	if (EnumHasAnyFlags(InViewportFlags, EDisplayClusterViewportFlags::TileViewport))
+	{
+		if (!InOptTilePos)
+		{
+			check(false); // Invalid args combination: when using EDisplayClusterViewportFlags::TileViewport, InOptTilePos must not be nullptr.
+
+			return FString::Printf(TEXT("%s_tile_invalid"), *UniqueName);
+		}
+		else
+		{
+			return FString::Printf(TEXT("%s_tile_%d_%d"), *UniqueName, InOptTilePos->X, InOptTilePos->Y);
+		}
+	}
+
+	return UniqueName;
+}
+
+TSharedPtr<IDisplayClusterProjectionPolicy, ESPMode::ThreadSafe> FDisplayClusterViewportConfiguration::CreateProjectionPolicy_ICVFXViewport(
+	const FString& InViewportBaseId,
+	const EDisplayClusterViewportFlags InViewportFlags)
+{
+	if (!EnumHasAnyFlags(InViewportFlags, EDisplayClusterViewportFlags::ICVFX_AnyRole)
+		|| EnumHasAnyFlags(InViewportFlags, EDisplayClusterViewportFlags::TileViewport))
+	{
+		// Create only projection policy only for ICVFX viewports
+		return nullptr;
+	}
+
+	const FString& ClusterNodeId = GetClusterNodeId();
+	if (ClusterNodeId.IsEmpty())
+	{
+		// The node name is requried
+		return nullptr;
+	}
+
+	FDisplayClusterConfigurationProjection ProjectionPolicyConfig;
+
+	// Only InCamera viewports use the 'camera' projection policy.
+	// All other viewport types inherit/copy the projection policy data (use type 'link') from their parent viewport.
+	ProjectionPolicyConfig.Type = EnumHasAnyFlags(InViewportFlags, EDisplayClusterViewportFlags::ICVFX_InnerFrustum)
+		? DisplayClusterProjectionStrings::projection::Camera
+		: DisplayClusterProjectionStrings::projection::Link;
+
+	const FString ProjectionPolicyId = BuildViewportName(InViewportBaseId, InViewportFlags);
+
+	// Create new projection policy for the ICVFX viewport
+	const TSharedPtr<IDisplayClusterProjectionPolicy, ESPMode::ThreadSafe>& NewProjectionPolicy =
+		FDisplayClusterViewportManager::CreateProjectionPolicy(ProjectionPolicyId, &ProjectionPolicyConfig);
+
+	if (!NewProjectionPolicy.IsValid())
+	{
+		UE_LOG(LogDisplayClusterViewport, Error, TEXT("Can't create ICVFX projection policy '%s' for node '%s'."), *ProjectionPolicyId, *ClusterNodeId);
+
+		return nullptr;
+	}
+
+	return NewProjectionPolicy;
+}
+
+TSharedPtr<IDisplayClusterProjectionPolicy, ESPMode::ThreadSafe> FDisplayClusterViewportConfiguration::CreateProjectionPolicy_TileViewport(
+	const FString& InViewportBaseId,
+	const EDisplayClusterViewportFlags InViewportFlags,
+	const FIntPoint& InTilePos)
+{
+	if (!EnumHasAnyFlags(InViewportFlags, EDisplayClusterViewportFlags::TileViewport))
+	{
+		// Create projection policy only for Tile viewports
+		return nullptr;
+	}
+
+	const FString& ClusterNodeId = GetClusterNodeId();
+	if (ClusterNodeId.IsEmpty())
+	{
+		// The node name is requried
+		return nullptr;
+	}
+
+	FDisplayClusterConfigurationProjection ProjectionPolicyConfig;
+
+	// tile viewports always use prj policy with type "link"
+	ProjectionPolicyConfig.Type = DisplayClusterProjectionStrings::projection::Link;
+
+	const FString ProjectionPolicyId = BuildViewportName(InViewportBaseId, InViewportFlags, &InTilePos);
+
+	// Create new projection policy for the tile viewport
+	const TSharedPtr<IDisplayClusterProjectionPolicy, ESPMode::ThreadSafe>& NewProjectionPolicy =
+		FDisplayClusterViewportManager::CreateProjectionPolicy(ProjectionPolicyId, &ProjectionPolicyConfig);
+
+	if (!NewProjectionPolicy.IsValid())
+	{
+		UE_LOG(LogDisplayClusterViewport, Error, TEXT("Can't create Tile viewport projection policy '%s' for node '%s'."), *ProjectionPolicyId, *ClusterNodeId);
+
+		return nullptr;
+	}
+
+	return NewProjectionPolicy;
+}
+
+FDisplayClusterViewport* FDisplayClusterViewportConfiguration::CreateViewport(
+	const FString& InViewportBaseId,
+	const TSharedPtr<IDisplayClusterProjectionPolicy, ESPMode::ThreadSafe>& InProjectionPolicy,
+	const EDisplayClusterViewportFlags InViewportFlags,
+	const FIntPoint* InOptTilePos)
+{
+	if (!InProjectionPolicy.IsValid())
+	{
+		return nullptr;
+	}
+
+	if (FDisplayClusterViewportManager* ViewportManager = GetViewportManagerImpl())
+	{
+		// Create new viewport
+		return ViewportManager->ImplCreateViewport(InViewportBaseId, InProjectionPolicy, InViewportFlags, InOptTilePos).Get();
+	}
+
+	return nullptr;
+}
+
+FDisplayClusterViewport* FDisplayClusterViewportConfiguration::FindViewport(
+	const FString& InViewportBaseId,
+	const EDisplayClusterViewportFlags InViewportFlags,
+	const FIntPoint* InOptTilePos) const
+{
+	if (InViewportBaseId.IsEmpty())
+	{
+		// The base ID of the viewport is mandatory.
+		return nullptr;
+	}
+
+	if (FDisplayClusterViewportManager* ViewportManager = GetViewportManagerImpl())
+	{
+		return ViewportManager->ImplFindViewport(InViewportBaseId, InViewportFlags, InOptTilePos).Get();
+	}
+
+	return nullptr;
+}
diff --git a/Engine/Plugins/Runtime/nDisplay/Source/DisplayCluster/Private/Render/Viewport/Configuration/DisplayClusterViewportConfiguration.h b/Engine/Plugins/Runtime/nDisplay/Source/DisplayCluster/Private/Render/Viewport/Configuration/DisplayClusterViewportConfiguration.h
--- a/Engine/Plugins/Runtime/nDisplay/Source/DisplayCluster/Private/Render/Viewport/Configuration/DisplayClusterViewportConfiguration.h
+++ b/Engine/Plugins/Runtime/nDisplay/Source/DisplayCluster/Private/Render/Viewport/Configuration/DisplayClusterViewportConfiguration.h
@@ -5,6 +5,8 @@
 #include "Render/Viewport/Configuration/DisplayClusterViewportConfigurationProxy.h"
 #include "Render/Viewport/RenderFrame/DisplayClusterRenderFrameSettings.h"
 
+#include "Render/Viewport/Containers/DisplayClusterViewport_InternalEnums.h"
+
 #include "SceneTypes.h"
 #include "Misc/DisplayClusterObjectRef.h"
 #include "Templates/SharedPointer.h"
@@ -12,6 +14,8 @@
 class FDisplayClusterViewportConfigurationProxy;
 class FDisplayClusterViewportManager;
 class FDisplayClusterViewportManagerProxy;
+class FDisplayClusterViewport; 
+class IDisplayClusterProjectionPolicy;
 class ADisplayClusterRootActor;
 class UDisplayClusterConfigurationData;
 struct FDisplayClusterConfigurationRenderFrame;
@@ -90,6 +94,11 @@ public:
 		return RenderFrameSettings.IsPreviewRendering();
 	}
 
+	virtual bool IsClusterRendering() const override
+	{
+		return RenderFrameSettings.IsClusterRendering();
+	}
+
 	virtual bool IsTechvisEnabled() const override
 	{
 		return RenderFrameSettings.IsTechvisEnabled();
@@ -105,6 +114,82 @@ public:
 	virtual const float GetWorldToMeters() const override;
 	// ~~END IDisplayClusterViewportConfiguration
 
+public:
+	/**
+	* Builds a viewport name that is unique within the cluster.
+	*
+	* The name is composed from the current cluster node ID(from this configuration),
+	* the base viewport ID, the viewport usage type, and optionally a tile position
+	* to disambiguate tiled viewports.
+	*
+	* @param InViewportBaseId Base viewport identifier.
+	* @param InViewportFlags  Usage flags of the viewport (e.g. InCamera/LightCard/Tile, etc.).
+	* @param InOptTilePos     Optional parameter for tile viewport position
+	*
+	* @return A viewport name string unique within the cluster.
+	*/
+	FString BuildViewportName(
+		const FString& InViewportBaseId,
+		const EDisplayClusterViewportFlags InViewportFlags,
+		const FIntPoint* InOptTilePos = nullptr);
+
+	/**
+	* Creates a projection policy instance with type "link" or "camera" for ICVFX viewport.
+	*
+	* @param InViewportBaseId The base viewport identifier (before ICVFX unique naming).
+	* @param InViewportFlags  Usage flags of the viewport (e.g. InCamera/LightCard/Tile, etc.).
+	*
+	* @return A thread-safe shared pointer to the created projection policy.
+	*/
+	TSharedPtr<IDisplayClusterProjectionPolicy, ESPMode::ThreadSafe> CreateProjectionPolicy_ICVFXViewport(
+		const FString& InViewportBaseId,
+		const EDisplayClusterViewportFlags InViewportFlags);
+
+	/**
+	* Creates a projection policy instance with type "link" for Tile viewport.
+	*
+	* @param InViewportBaseId The base viewport identifier (before ICVFX unique naming).
+	* @param InViewportFlags  Usage flags of the viewport (e.g. InCamera/LightCard/Tile, etc.).
+	* @param InTilePos        Tile viewport position
+	*
+	* @return A thread-safe shared pointer to the created projection policy.
+	*/
+	TSharedPtr<IDisplayClusterProjectionPolicy, ESPMode::ThreadSafe> CreateProjectionPolicy_TileViewport(
+		const FString& InViewportBaseId,
+		const EDisplayClusterViewportFlags InViewportFlags,
+		const FIntPoint& InTilePos);
+
+	/**
+	*Creates and initializes a new nDisplay viewport.
+	*
+	* @param InViewportBaseId   The base viewport identifier (before ICVFX unique naming).
+	* @param InProjectionPolicy Projection policy assigned to the viewport
+	* @param InViewportFlags    Usage flags of the viewport (e.g. InCamera/LightCard/Tile, etc.).
+	* @param InOptTilePos       Optional parameter for tile viewport position
+	*
+	* @return A raw pointer to the created FDisplayClusterViewport or nullptr in case of failure.
+	*/
+	FDisplayClusterViewport* CreateViewport(
+		const FString& InViewportBaseId,
+		const TSharedPtr<IDisplayClusterProjectionPolicy, ESPMode::ThreadSafe>& InProjectionPolicy,
+		const EDisplayClusterViewportFlags InViewportFlags,
+		const FIntPoint* InOptTilePos = nullptr);
+
+	/**
+	* Finds a viewport instance by identifier and usage type.
+	*
+	* @param InViewportBaseId The base viewport identifier (before ICVFX unique naming).
+	* @param InViewportFlags  Usage flags of the viewport (e.g. InCamera/LightCard/Tile, etc.).
+	* @param InOptTilePos     Optional parameter for tile viewport position
+	*
+	* @return A pointer to the matching FDisplayClusterViewport,
+	*         or nullptr if no viewport matches the given criteria.
+	*/
+	FDisplayClusterViewport* FindViewport(
+		const FString& InViewportBaseId,
+		const EDisplayClusterViewportFlags InViewportFlags = EDisplayClusterViewportFlags::None,
+		const FIntPoint* InOptTilePos = nullptr) const;
+
 public:
 	/** Get a pointer to the DC ViewportManager if it still exists. */
 	FDisplayClusterViewportManager* GetViewportManagerImpl() const
diff --git a/Engine/Plugins/Runtime/nDisplay/Source/DisplayCluster/Private/Render/Viewport/Configuration/DisplayClusterViewportConfigurationHelpers_ICVFX.cpp b/Engine/Plugins/Runtime/nDisplay/Source/DisplayCluster/Private/Render/Viewport/Configuration/DisplayClusterViewportConfigurationHelpers_ICVFX.cpp
--- a/Engine/Plugins/Runtime/nDisplay/Source/DisplayCluster/Private/Render/Viewport/Configuration/DisplayClusterViewportConfigurationHelpers_ICVFX.cpp
+++ b/Engine/Plugins/Runtime/nDisplay/Source/DisplayCluster/Private/Render/Viewport/Configuration/DisplayClusterViewportConfigurationHelpers_ICVFX.cpp
@@ -11,7 +11,6 @@
 
 #include "Render/Viewport/DisplayClusterViewport.h"
 #include "Render/Viewport/DisplayClusterViewportManager.h"
-#include "Render/Viewport/DisplayClusterViewportStrings.h"
 #include "Render/Viewport/DisplayClusterViewportHelpers.h"
 
 #include "DisplayClusterRootActor.h"
@@ -45,55 +44,69 @@
 #include "Misc/DisplayClusterLog.h"
 #include "TextureResource.h"
 
-////////////////////////////////////////////////////////////////////////////////
-namespace UE::DisplayCluster::Viewport::ConfigurationHelpers_ICVFX
-{
-	// Return unique ICVFX name
-	static FString ImplGetNameICVFX(const FString& InClusterNodeId, const FString& InViewportId, const FString& InResourceId, const FString& InResourceSuffixId)
-	{
-		check(!InClusterNodeId.IsEmpty());
-		check(!InViewportId.IsEmpty());
-		check(!InResourceId.IsEmpty());
-
-		return FString::Printf(TEXT("%s_%s_%s_%s%s"), *InClusterNodeId, DisplayClusterViewportStrings::icvfx::prefix, *InViewportId, *InResourceId, *InResourceSuffixId);
-	}
-};
-using namespace UE::DisplayCluster::Viewport::ConfigurationHelpers_ICVFX;
-
 ////////////////////////////////////////////////////////////////////////////////
 // FDisplayClusterViewportConfigurationHelpers_ICVFX
 ////////////////////////////////////////////////////////////////////////////////
-bool FDisplayClusterViewportConfigurationHelpers_ICVFX::CreateProjectionPolicyICVFX(FDisplayClusterViewportConfiguration& InConfiguration, const FString& InViewportId, const FString& InResourceId, const FString& InResourceSuffixId, bool bIsCameraProjection, TSharedPtr<IDisplayClusterProjectionPolicy, ESPMode::ThreadSafe>& OutProjPolicy)
+void FDisplayClusterViewportConfigurationHelpers_ICVFX::UpdateLegacyRuntimeICVFXFlagsFromViewportFlags(FDisplayClusterViewport& InOutViewport)
 {
-	const FString& ClusterNodeId = InConfiguration.GetClusterNodeId();
-	if (ClusterNodeId.IsEmpty())
+	// TODO: Deprecate and remove the duplicate flags from EDisplayClusterViewportRuntimeICVFXFlags.
+	// They have been moved to EDisplayClusterViewportFlags.
+
+	// Temporary bridge: mirror legacy ICVFX flags from EDisplayClusterViewportFlags into
+	// EDisplayClusterViewportRuntimeICVFXFlags until the old flags are removed.
+	FDisplayClusterViewport_RenderSettingsICVFX& RenderSettingsICVFX = InOutViewport.GetRenderSettingsICVFXImpl();
+	if (EnumHasAnyFlags(InOutViewport.ViewportFlags, EDisplayClusterViewportFlags::ICVFX_InnerFrustum))
 	{
-		return false;
+		// Add viewport ICVFX usage as Incamera
+		EnumAddFlags(RenderSettingsICVFX.RuntimeFlags, EDisplayClusterViewportRuntimeICVFXFlags::InCamera);
 	}
-
-	FDisplayClusterConfigurationProjection CameraProjectionPolicyConfig;
-	CameraProjectionPolicyConfig.Type = bIsCameraProjection ? DisplayClusterProjectionStrings::projection::Camera : DisplayClusterProjectionStrings::projection::Link;
-
-	// Create projection policy for viewport
-	OutProjPolicy = FDisplayClusterViewportManager::CreateProjectionPolicy(ImplGetNameICVFX(ClusterNodeId, InViewportId, InResourceId, InResourceSuffixId), &CameraProjectionPolicyConfig);
-
-	if (!OutProjPolicy.IsValid())
+	else if (EnumHasAnyFlags(InOutViewport.ViewportFlags, EDisplayClusterViewportFlags::ICVFX_ChromaKey))
 	{
-		UE_LOG(LogDisplayClusterViewport, Error, TEXT("ICVFX Viewport '%s': projection policy for resource '%s%s' not created for node '%s'."), *InViewportId, *InResourceId, *InResourceSuffixId,*ClusterNodeId);
-		return false;
+		// Add viewport ICVFX usage as Chromakey
+		EnumAddFlags(RenderSettingsICVFX.RuntimeFlags, EDisplayClusterViewportRuntimeICVFXFlags::Chromakey);
 	}
+	else if (EnumHasAnyFlags(InOutViewport.ViewportFlags, EDisplayClusterViewportFlags::ICVFX_LightCard))
+	{
+		// Add viewport ICVFX usage as Lightcard
+		EnumAddFlags(RenderSettingsICVFX.RuntimeFlags, EDisplayClusterViewportRuntimeICVFXFlags::Lightcard);
 
-	return true;
-}
+		// Add LC viewport position flag
+		if (EnumHasAnyFlags(InOutViewport.ViewportFlags, EDisplayClusterViewportFlags::ICVFX_Modifier_RenderOverInnerFrustum))
+		{
+			EnumAddFlags(RenderSettingsICVFX.RuntimeFlags, EDisplayClusterViewportRuntimeICVFXFlags::OverInFrustum);
+		}
+		else if (EnumHasAnyFlags(InOutViewport.ViewportFlags, EDisplayClusterViewportFlags::ICVFX_Modifier_RenderUnderInnerFrustum))
+		{
+			EnumAddFlags(RenderSettingsICVFX.RuntimeFlags, EDisplayClusterViewportRuntimeICVFXFlags::UnderInFrustum);
+		}
+		else
+		{
+			check(false); // invalid flag settings
+		}
+	}
+	else if (EnumHasAnyFlags(InOutViewport.ViewportFlags, EDisplayClusterViewportFlags::ICVFX_UVLightCard))
+	{
+		// Add viewport ICVFX usage as UVLightcard
+		EnumAddFlags(RenderSettingsICVFX.RuntimeFlags, EDisplayClusterViewportRuntimeICVFXFlags::UVLightcard);
 
-FString FDisplayClusterViewportConfigurationHelpers_ICVFX::BuildCameraViewportId(const FString& InClusterNodeId, const FString& InCameraName)
-{
-	return ImplGetNameICVFX(
-		InClusterNodeId,
-		InCameraName,
-		DisplayClusterViewportStrings::icvfx::camera,
-		DisplayClusterViewportStrings::icvfx::suffix::none
-	);
+		// Add LC viewport position flag
+		if (EnumHasAnyFlags(InOutViewport.ViewportFlags, EDisplayClusterViewportFlags::ICVFX_Modifier_RenderOverInnerFrustum))
+		{
+			EnumAddFlags(RenderSettingsICVFX.RuntimeFlags, EDisplayClusterViewportRuntimeICVFXFlags::OverInFrustum);
+		}
+		else if (EnumHasAnyFlags(InOutViewport.ViewportFlags, EDisplayClusterViewportFlags::ICVFX_Modifier_RenderUnderInnerFrustum))
+		{
+			EnumAddFlags(RenderSettingsICVFX.RuntimeFlags, EDisplayClusterViewportRuntimeICVFXFlags::UnderInFrustum);
+		}
+		else
+		{
+			check(false); // invalid flag settings
+		}
+	}
+	else
+	{
+		check(false); // invalid flag settings
+	}
 }
 
 bool FDisplayClusterViewportConfigurationHelpers_ICVFX::UpdateCameraProjectionSettingsICVFX(FDisplayClusterViewportConfiguration& InConfiguration, UDisplayClusterICVFXCameraComponent& InCameraComponent, const FDisplayClusterConfigurationICVFX_CameraSettings& InCameraSettings, const TSharedPtr<IDisplayClusterProjectionPolicy, ESPMode::ThreadSafe>& InProjectionPolicy)
@@ -119,70 +132,30 @@ bool FDisplayClusterViewportConfigurationHelpers_ICVFX::UpdateCameraProjectionSe
 	return DisplayClusterProjectionAPI.CameraPolicySetCamera(InProjectionPolicy, &InCameraComponent, PolicyCameraSettings);
 }
 
-bool FDisplayClusterViewportConfigurationHelpers_ICVFX::CreateProjectionPolicyCameraICVFX(FDisplayClusterViewportConfiguration& InConfiguration, UDisplayClusterICVFXCameraComponent& InCameraComponent, const FDisplayClusterConfigurationICVFX_CameraSettings& InCameraSettings, TSharedPtr<IDisplayClusterProjectionPolicy, ESPMode::ThreadSafe>& OutProjPolicy)
-{
-	return CreateProjectionPolicyICVFX(InConfiguration, InCameraComponent.GetCameraUniqueId(), DisplayClusterViewportStrings::icvfx::camera, DisplayClusterViewportStrings::icvfx::suffix::none, true, OutProjPolicy)
-		&& UpdateCameraProjectionSettingsICVFX(InConfiguration, InCameraComponent, InCameraSettings, OutProjPolicy);
-}
-
-FDisplayClusterViewport* FDisplayClusterViewportConfigurationHelpers_ICVFX::FindViewportICVFX(FDisplayClusterViewportConfiguration& InConfiguration, const FString& InViewportId, const FString& InResourceId, const FString& InResourceSuffixId)
-{
-	if (FDisplayClusterViewportManager* ViewportManager = InConfiguration.GetViewportManagerImpl())
-	{
-		TSharedPtr<FDisplayClusterViewport, ESPMode::ThreadSafe> Viewport = ViewportManager->ImplFindViewport(ImplGetNameICVFX(InConfiguration.GetClusterNodeId(), InViewportId, InResourceId, InResourceSuffixId));
-
-		return Viewport.Get();
-	}
-
-	return nullptr;
-}
-
-FDisplayClusterViewport* FDisplayClusterViewportConfigurationHelpers_ICVFX::CreateViewportICVFX(FDisplayClusterViewportConfiguration& InConfiguration, const FString& InViewportId, const FString& InResourceId, const FString& InResourceSuffixId, TSharedPtr<IDisplayClusterProjectionPolicy, ESPMode::ThreadSafe>& InProjectionPolicy)
+TSharedPtr<IDisplayClusterProjectionPolicy, ESPMode::ThreadSafe> FDisplayClusterViewportConfigurationHelpers_ICVFX::CreateProjectionPolicyCameraICVFX(FDisplayClusterViewportConfiguration& InConfiguration, UDisplayClusterICVFXCameraComponent& InCameraComponent, const FDisplayClusterConfigurationICVFX_CameraSettings& InCameraSettings)
 {
-	check(InProjectionPolicy.IsValid());
+	TSharedPtr<IDisplayClusterProjectionPolicy, ESPMode::ThreadSafe> OutProjPolicy =
+		InConfiguration.CreateProjectionPolicy_ICVFXViewport(InCameraComponent.GetCameraUniqueId(), EDisplayClusterViewportFlags::ICVFX_InnerFrustum);
 
-	if (FDisplayClusterViewportManager* ViewportManager = InConfiguration.GetViewportManagerImpl())
+	if (OutProjPolicy.IsValid() && UpdateCameraProjectionSettingsICVFX(InConfiguration, InCameraComponent, InCameraSettings, OutProjPolicy))
 	{
-		// Create viewport for new projection policy
-		TSharedPtr<FDisplayClusterViewport, ESPMode::ThreadSafe> NewViewport = ViewportManager->ImplCreateViewport(ImplGetNameICVFX(InConfiguration.GetClusterNodeId(), InViewportId, InResourceId, InResourceSuffixId), InProjectionPolicy);
-		if (NewViewport.IsValid())
-		{
-			// Gain direct access to internal resources of the NewViewport:
-			FDisplayClusterViewport_RenderSettings&           InOutRenderSettings = NewViewport->GetRenderSettingsImpl();
-			FDisplayClusterViewport_RenderSettingsICVFX& InOutRenderSettingsICVFX = NewViewport->GetRenderSettingsICVFXImpl();
-
-			// Mark as internal resource
-			EnumAddFlags(InOutRenderSettingsICVFX.RuntimeFlags, EDisplayClusterViewportRuntimeICVFXFlags::InternalResource);
-
-			// Dont show ICVFX composing viewports on frame target
-			InOutRenderSettings.bVisible = false;
-
-			return NewViewport.Get();
-		}
+		return OutProjPolicy;
 	}
 
 	return nullptr;
 }
 
-FDisplayClusterViewport* FDisplayClusterViewportConfigurationHelpers_ICVFX::FindCameraViewport(FDisplayClusterViewportConfiguration& InConfiguration, UDisplayClusterICVFXCameraComponent& InCameraComponent)
-{
-	return FindViewportICVFX(InConfiguration, InCameraComponent.GetCameraUniqueId(), DisplayClusterViewportStrings::icvfx::camera, DisplayClusterViewportStrings::icvfx::suffix::none);
-}
-
 FDisplayClusterViewport* FDisplayClusterViewportConfigurationHelpers_ICVFX::GetOrCreateCameraViewport(FDisplayClusterViewportConfiguration& InConfiguration, UDisplayClusterICVFXCameraComponent& InCameraComponent, const FDisplayClusterConfigurationICVFX_CameraSettings& InCameraSettings)
 {
-	FDisplayClusterViewport* CameraViewport = FindCameraViewport(InConfiguration, InCameraComponent);
-	if (CameraViewport == nullptr)
-	{
-		// Create new camera viewport
-		TSharedPtr<IDisplayClusterProjectionPolicy, ESPMode::ThreadSafe> CameraProjectionPolicy;
-		if (!CreateProjectionPolicyCameraICVFX(InConfiguration, InCameraComponent, InCameraSettings, CameraProjectionPolicy))
-		{
-			return nullptr;
-		}
-
-		CameraViewport = CreateViewportICVFX(InConfiguration, InCameraComponent.GetCameraUniqueId(), DisplayClusterViewportStrings::icvfx::camera, DisplayClusterViewportStrings::icvfx::suffix::none, CameraProjectionPolicy);
-		if (CameraViewport == nullptr)
+	FDisplayClusterViewport* CameraViewport = InConfiguration.FindViewport(InCameraComponent.GetCameraUniqueId(), EDisplayClusterViewportFlags::ICVFX_InnerFrustum);
+	if (!CameraViewport)
+	{
+		// Create viewport instance
+		CameraViewport = InConfiguration.CreateViewport(
+			InCameraComponent.GetCameraUniqueId(),
+			CreateProjectionPolicyCameraICVFX(InConfiguration, InCameraComponent, InCameraSettings),
+			EDisplayClusterViewportFlags::ICVFX_InnerFrustum);
+		if(!CameraViewport)
 		{
 			return nullptr;
 		}
@@ -197,8 +170,6 @@ FDisplayClusterViewport* FDisplayClusterViewportConfigurationHelpers_ICVFX::GetO
 	// Mark viewport as used
 	EnumRemoveFlags(InOutRenderSettingsICVFX.RuntimeFlags, EDisplayClusterViewportRuntimeICVFXFlags::Unused);
 
-	// Add viewport ICVFX usage as Incamera
-	EnumAddFlags(InOutRenderSettingsICVFX.RuntimeFlags, EDisplayClusterViewportRuntimeICVFXFlags::InCamera);
 
 	return CameraViewport;
 }
@@ -259,24 +230,24 @@ TArray<TSharedPtr<FDisplayClusterViewport, ESPMode::ThreadSafe>> FDisplayCluster
 {
 	TArray<TSharedPtr<FDisplayClusterViewport, ESPMode::ThreadSafe>> OutViewports;
 
-	// Search for rendered camera viewport on other cluster nodes
-	const EDisplayClusterViewportRuntimeICVFXFlags InRuntimeFlagsMask = bGetChromakey ? EDisplayClusterViewportRuntimeICVFXFlags::Chromakey : EDisplayClusterViewportRuntimeICVFXFlags::InCamera;
-	const FString ViewportTypeId = bGetChromakey ? DisplayClusterViewportStrings::icvfx::chromakey : DisplayClusterViewportStrings::icvfx::camera;
+	// Filter viewports by usage flags
+	const EDisplayClusterViewportFlags ExpectedViewportFlags = bGetChromakey
+		? EDisplayClusterViewportFlags::ICVFX_ChromaKey
+		: EDisplayClusterViewportFlags::ICVFX_InnerFrustum;
 
 	if (FDisplayClusterViewportManager* ViewportManager = InConfiguration.GetViewportManagerImpl())
 	{
 		for (const TSharedPtr<FDisplayClusterViewport, ESPMode::ThreadSafe>& ViewportIt : ViewportManager->ImplGetEntireClusterViewports())
 		{
 			if (ViewportIt.IsValid()
-				&& EnumHasAnyFlags(ViewportIt->GetRenderSettingsICVFX().RuntimeFlags, InRuntimeFlagsMask)
+				&& EnumHasAnyFlags(ViewportIt->ViewportFlags, ExpectedViewportFlags)
 				&& ( ViewportIt->GetViewportResources(EDisplayClusterViewportResource::InputShaderResources).Num() > 0
 					&& ViewportIt->GetViewportResources(EDisplayClusterViewportResource::InputShaderResources)[0] != nullptr
 					&& ViewportIt->GetContexts().Num() > 0)
 				&& (!ViewportIt->GetRenderSettings().IsViewportOverridden()))
 			{
-				// this is incamera viewport. Check by name
-				const FString RequiredViewportId = ImplGetNameICVFX(ViewportIt->GetClusterNodeId(), InICVFXCameraId, ViewportTypeId, DisplayClusterViewportStrings::icvfx::suffix::none);
-				if (RequiredViewportId.Equals(ViewportIt->GetId()))
+				// This is InCamera based viewport. Check by the base name.
+				if (ViewportIt->ViewportBaseId == InICVFXCameraId)
 				{
 					OutViewports.Add(ViewportIt);
 				}
@@ -360,18 +331,13 @@ bool FDisplayClusterViewportConfigurationHelpers_ICVFX::PreviewReuseChromakeyVie
 FDisplayClusterViewport* FDisplayClusterViewportConfigurationHelpers_ICVFX::GetOrCreateChromakeyViewport(FDisplayClusterViewportConfiguration& InConfiguration, UDisplayClusterICVFXCameraComponent& InCameraComponent, const FDisplayClusterConfigurationICVFX_CameraSettings& InCameraSettings)
 {
 	const FString InICVFXCameraId = InCameraComponent.GetCameraUniqueId();
-	FDisplayClusterViewport* ChromakeyViewport = FindViewportICVFX(InConfiguration, InICVFXCameraId, DisplayClusterViewportStrings::icvfx::chromakey, DisplayClusterViewportStrings::icvfx::suffix::none);
-
-	// Create new chromakey viewport
+	FDisplayClusterViewport* ChromakeyViewport = InConfiguration.FindViewport(InICVFXCameraId, EDisplayClusterViewportFlags::ICVFX_ChromaKey);
 	if (ChromakeyViewport == nullptr)
 	{
-		TSharedPtr<IDisplayClusterProjectionPolicy, ESPMode::ThreadSafe> ChromakeyProjectionPolicy;
-		if (!CreateProjectionPolicyICVFX(InConfiguration, InICVFXCameraId, DisplayClusterViewportStrings::icvfx::chromakey, DisplayClusterViewportStrings::icvfx::suffix::none, false, ChromakeyProjectionPolicy))
-		{
-			return nullptr;
-		}
-
-		ChromakeyViewport = CreateViewportICVFX(InConfiguration, InICVFXCameraId, DisplayClusterViewportStrings::icvfx::chromakey, DisplayClusterViewportStrings::icvfx::suffix::none, ChromakeyProjectionPolicy);
+		ChromakeyViewport = InConfiguration.CreateViewport(
+			InICVFXCameraId,
+			InConfiguration.CreateProjectionPolicy_ICVFXViewport(InICVFXCameraId, EDisplayClusterViewportFlags::ICVFX_ChromaKey),
+			EDisplayClusterViewportFlags::ICVFX_ChromaKey);
 		if (ChromakeyViewport == nullptr)
 		{
 			return nullptr;
@@ -387,28 +353,24 @@ FDisplayClusterViewport* FDisplayClusterViewportConfigurationHelpers_ICVFX::GetO
 	// Mark viewport as used
 	EnumRemoveFlags(InOutRenderSettingsICVFX.RuntimeFlags, EDisplayClusterViewportRuntimeICVFXFlags::Unused);
 
-	// Add viewport ICVFX usage as Chromakey
-	EnumAddFlags(InOutRenderSettingsICVFX.RuntimeFlags, EDisplayClusterViewportRuntimeICVFXFlags::Chromakey);
-
 	return ChromakeyViewport;
 }
 
 FDisplayClusterViewport* FDisplayClusterViewportConfigurationHelpers_ICVFX::GetOrCreateLightcardViewport(FDisplayClusterViewport& BaseViewport, const bool bOverInFrustum)
 {
 	// Create new lightcard viewport
-	const FString ResourceId = DisplayClusterViewportStrings::icvfx::lightcard;
-	const FString ResourceSuffixId = bOverInFrustum ? DisplayClusterViewportStrings::icvfx::suffix::over : DisplayClusterViewportStrings::icvfx::suffix::under;
+	const EDisplayClusterViewportFlags ViewportFlags= bOverInFrustum
+		? EDisplayClusterViewportFlags::ICVFX_LightCard | EDisplayClusterViewportFlags::ICVFX_Modifier_RenderOverInnerFrustum
+		: EDisplayClusterViewportFlags::ICVFX_LightCard | EDisplayClusterViewportFlags::ICVFX_Modifier_RenderUnderInnerFrustum;
 
-	FDisplayClusterViewport* LightcardViewport = FindViewportICVFX(*BaseViewport.Configuration, BaseViewport.GetId(), ResourceId, ResourceSuffixId);
+	FDisplayClusterViewport* LightcardViewport = BaseViewport.Configuration->FindViewport(BaseViewport.GetId(), ViewportFlags);
 	if (LightcardViewport == nullptr)
 	{
-		TSharedPtr<IDisplayClusterProjectionPolicy, ESPMode::ThreadSafe> LightcardProjectionPolicy;
-		if (!CreateProjectionPolicyICVFX(*BaseViewport.Configuration, BaseViewport.GetId(), ResourceId, ResourceSuffixId, false, LightcardProjectionPolicy))
-		{
-			return nullptr;
-		}
+		LightcardViewport = BaseViewport.Configuration->CreateViewport(
+			BaseViewport.GetId(),
+			BaseViewport.Configuration->CreateProjectionPolicy_ICVFXViewport(BaseViewport.GetId(), ViewportFlags),
+			ViewportFlags);
 
-		LightcardViewport = CreateViewportICVFX(*BaseViewport.Configuration, BaseViewport.GetId(), ResourceId, ResourceSuffixId, LightcardProjectionPolicy);
 		if (LightcardViewport == nullptr)
 		{
 			return nullptr;
@@ -427,31 +389,24 @@ FDisplayClusterViewport* FDisplayClusterViewportConfigurationHelpers_ICVFX::GetO
 	// Mark viewport as used
 	EnumRemoveFlags(InOutRenderSettingsICVFX.RuntimeFlags, EDisplayClusterViewportRuntimeICVFXFlags::Unused);
 
-	// Add viewport ICVFX usage as Lightcard
-	EnumAddFlags(InOutRenderSettingsICVFX.RuntimeFlags, EDisplayClusterViewportRuntimeICVFXFlags::Lightcard);
-
-	// Add LC viewport position flag
-	EnumAddFlags(InOutRenderSettingsICVFX.RuntimeFlags, bOverInFrustum ? EDisplayClusterViewportRuntimeICVFXFlags::OverInFrustum : EDisplayClusterViewportRuntimeICVFXFlags::UnderInFrustum);
-
 	return LightcardViewport;
 }
 
 FDisplayClusterViewport* FDisplayClusterViewportConfigurationHelpers_ICVFX::GetOrCreateUVLightcardViewport(FDisplayClusterViewport& BaseViewport, const bool bOverInFrustum)
 {
-	// Create new lightcard viewport
-	const FString ResourceId = DisplayClusterViewportStrings::icvfx::uv_lightcard;
-	const FString ResourceSuffixId = bOverInFrustum ? DisplayClusterViewportStrings::icvfx::suffix::over : DisplayClusterViewportStrings::icvfx::suffix::under;
+	// Create new UV lightcard viewport
+	const EDisplayClusterViewportFlags ViewportFlags = bOverInFrustum
+		? EDisplayClusterViewportFlags::ICVFX_UVLightCard | EDisplayClusterViewportFlags::ICVFX_Modifier_RenderOverInnerFrustum
+		: EDisplayClusterViewportFlags::ICVFX_UVLightCard | EDisplayClusterViewportFlags::ICVFX_Modifier_RenderUnderInnerFrustum;
 
-	FDisplayClusterViewport* UVLightcardViewport = FindViewportICVFX(*BaseViewport.Configuration, BaseViewport.GetId(), ResourceId, ResourceSuffixId);
+	FDisplayClusterViewport* UVLightcardViewport = BaseViewport.Configuration->FindViewport(BaseViewport.GetId(), ViewportFlags);
 	if (UVLightcardViewport == nullptr)
 	{
-		TSharedPtr<IDisplayClusterProjectionPolicy, ESPMode::ThreadSafe> UVLightcardProjectionPolicy;
-		if (!CreateProjectionPolicyICVFX(*BaseViewport.Configuration, BaseViewport.GetId(), ResourceId, ResourceSuffixId, false, UVLightcardProjectionPolicy))
-		{
-			return nullptr;
-		}
-
-		UVLightcardViewport = CreateViewportICVFX(*BaseViewport.Configuration, BaseViewport.GetId(), ResourceId, ResourceSuffixId, UVLightcardProjectionPolicy);
+		UVLightcardViewport = BaseViewport.Configuration->CreateViewport(
+			BaseViewport.GetId(),
+			BaseViewport.Configuration->CreateProjectionPolicy_ICVFXViewport(BaseViewport.GetId(), ViewportFlags),
+			ViewportFlags);
+		
 		if (UVLightcardViewport == nullptr)
 		{
 			return nullptr;
@@ -467,12 +422,6 @@ FDisplayClusterViewport* FDisplayClusterViewportConfigurationHelpers_ICVFX::GetO
 	// Mark viewport as used
 	EnumRemoveFlags(InOutRenderSettingsICVFX.RuntimeFlags, EDisplayClusterViewportRuntimeICVFXFlags::Unused);
 
-	// Add viewport ICVFX usage as Lightcard
-	EnumAddFlags(InOutRenderSettingsICVFX.RuntimeFlags, EDisplayClusterViewportRuntimeICVFXFlags::UVLightcard);
-
-	// Add LC viewport position flag
-	EnumAddFlags(InOutRenderSettingsICVFX.RuntimeFlags, bOverInFrustum ? EDisplayClusterViewportRuntimeICVFXFlags::OverInFrustum : EDisplayClusterViewportRuntimeICVFXFlags::UnderInFrustum);
-
 	return UVLightcardViewport;
 }
 
diff --git a/Engine/Plugins/Runtime/nDisplay/Source/DisplayCluster/Private/Render/Viewport/Configuration/DisplayClusterViewportConfigurationHelpers_ICVFX.h b/Engine/Plugins/Runtime/nDisplay/Source/DisplayCluster/Private/Render/Viewport/Configuration/DisplayClusterViewportConfigurationHelpers_ICVFX.h
--- a/Engine/Plugins/Runtime/nDisplay/Source/DisplayCluster/Private/Render/Viewport/Configuration/DisplayClusterViewportConfigurationHelpers_ICVFX.h
+++ b/Engine/Plugins/Runtime/nDisplay/Source/DisplayCluster/Private/Render/Viewport/Configuration/DisplayClusterViewportConfigurationHelpers_ICVFX.h
@@ -23,9 +23,6 @@ struct FDisplayClusterViewport_CustomFrustumSettings;
 class FDisplayClusterViewportConfigurationHelpers_ICVFX
 {
 public:
-	/* Find existing InnerFrustum viewport. */
-	static FDisplayClusterViewport* FindCameraViewport(FDisplayClusterViewportConfiguration& InConfiguration, UDisplayClusterICVFXCameraComponent& InCameraComponent);
-
 	/* Return the new or existing InnerFrustum viewport. */
 	static FDisplayClusterViewport* GetOrCreateCameraViewport(FDisplayClusterViewportConfiguration& InConfiguration, UDisplayClusterICVFXCameraComponent& InCameraComponent, const FDisplayClusterConfigurationICVFX_CameraSettings& InCameraSettings);
 
@@ -57,8 +54,7 @@ public:
 
 	static bool IsCameraUsed(const FDisplayClusterConfigurationICVFX_CameraSettings& InCameraSettings);
 
-	static bool CreateProjectionPolicyICVFX(FDisplayClusterViewportConfiguration& InConfiguration, const FString& InViewportId, const FString& InResourceId, const FString& InSuffixId, bool bIsCameraProjection, TSharedPtr<IDisplayClusterProjectionPolicy, ESPMode::ThreadSafe>& OutProjPolicy);
-	static bool CreateProjectionPolicyCameraICVFX(FDisplayClusterViewportConfiguration& InConfiguration, UDisplayClusterICVFXCameraComponent& CameraComponent, const FDisplayClusterConfigurationICVFX_CameraSettings& InCameraSettings, TSharedPtr<IDisplayClusterProjectionPolicy, ESPMode::ThreadSafe>& OutProjPolicy);
+	static TSharedPtr<IDisplayClusterProjectionPolicy, ESPMode::ThreadSafe> CreateProjectionPolicyCameraICVFX(FDisplayClusterViewportConfiguration& InConfiguration, UDisplayClusterICVFXCameraComponent& CameraComponent, const FDisplayClusterConfigurationICVFX_CameraSettings& InCameraSettings);
 	static bool UpdateCameraProjectionSettingsICVFX(FDisplayClusterViewportConfiguration& InConfiguration, UDisplayClusterICVFXCameraComponent& CameraComponent, const FDisplayClusterConfigurationICVFX_CameraSettings& InCameraSettings, const TSharedPtr<IDisplayClusterProjectionPolicy, ESPMode::ThreadSafe>& InProjectionPolicy);
 
 	/** Configures the camera's chromakey render settings to match the specified configuration settings  */
@@ -73,15 +69,9 @@ public:
 
 	static void UpdateCameraCustomFrustum(const FDisplayClusterConfigurationICVFX_CameraCustomFrustum& InCameraCustomFrustumConfiguration, FDisplayClusterViewport_CustomFrustumSettings& OutCustomFrustumSettings);
 
-	static FDisplayClusterViewport* FindViewportICVFX(FDisplayClusterViewportConfiguration& InConfiguration, const FString& InViewportId, const FString& InResourceId, const FString& InResourceSuffixId);
-	static FDisplayClusterViewport* CreateViewportICVFX(FDisplayClusterViewportConfiguration& InConfiguration, const FString& InViewportId, const FString& InResourceId, const FString& InResourceSuffixId, TSharedPtr<IDisplayClusterProjectionPolicy, ESPMode::ThreadSafe>& InProjectionPolicy);
-
-	/** Constructs the full ICVFX camera viewport ID from a camera name.
-	 * 
-	 * @param InClusterNodeId - The cluster node ID
-	 * @param InCameraName - The camera component name
-	 * 
-	 * @return The constructed viewport ID string.
-	 */
-	static FString BuildCameraViewportId(const FString& InClusterNodeId, const FString& InCameraName);
+	/**
+	* Backward-compatibility bridge: derive legacy EDisplayClusterViewportRuntimeICVFXFlags
+	* from the new EDisplayClusterViewportFlags until the legacy runtime enum is removed.
+	*/ 
+	static void UpdateLegacyRuntimeICVFXFlagsFromViewportFlags(FDisplayClusterViewport& InOutViewport);
 };
diff --git a/Engine/Plugins/Runtime/nDisplay/Source/DisplayCluster/Private/Render/Viewport/Configuration/DisplayClusterViewportConfigurationHelpers_Tile.cpp b/Engine/Plugins/Runtime/nDisplay/Source/DisplayCluster/Private/Render/Viewport/Configuration/DisplayClusterViewportConfigurationHelpers_Tile.cpp
--- a/Engine/Plugins/Runtime/nDisplay/Source/DisplayCluster/Private/Render/Viewport/Configuration/DisplayClusterViewportConfigurationHelpers_Tile.cpp
+++ b/Engine/Plugins/Runtime/nDisplay/Source/DisplayCluster/Private/Render/Viewport/Configuration/DisplayClusterViewportConfigurationHelpers_Tile.cpp
@@ -8,7 +8,6 @@
 
 #include "Render/Viewport/DisplayClusterViewport.h"
 #include "Render/Viewport/DisplayClusterViewportManager.h"
-#include "Render/Viewport/DisplayClusterViewportStrings.h"
 #include "Render/Viewport/DisplayClusterViewportHelpers.h"
 #include "Render/Viewport/Configuration/DisplayClusterViewportConfiguration.h"
 
@@ -70,45 +69,6 @@ FIntRect FDisplayClusterViewportConfigurationHelpers_Tile::GetDestRect(const FDi
 	return FIntRect(DestPos, DestPos + DestSize);
 }
 
-FString FDisplayClusterViewportConfigurationHelpers_Tile::GetUniqueViewportNameForTile(const FString& InViewportId, const FIntPoint& TilePos)
-{
-	check(!InViewportId.IsEmpty());
-
-	return FString::Printf(TEXT("%s_%s_%d_%d"), *InViewportId, DisplayClusterViewportStrings::tile::prefix, TilePos.X, TilePos.Y);
-}
-
-bool FDisplayClusterViewportConfigurationHelpers_Tile::CreateProjectionPolicyForTileViewport(FDisplayClusterViewport& InSourceViewport, const FIntPoint& TilePos, TSharedPtr<IDisplayClusterProjectionPolicy, ESPMode::ThreadSafe>& OutProjPolicy)
-{
-	FDisplayClusterConfigurationProjection CameraProjectionPolicyConfig;
-
-	// Projection policy with type 'link' has support for tile rendering
-	CameraProjectionPolicyConfig.Type = DisplayClusterProjectionStrings::projection::Link;
-
-	// Create projection policy for viewport
-	OutProjPolicy = FDisplayClusterViewportManager::CreateProjectionPolicy(GetUniqueViewportNameForTile(InSourceViewport.GetId(), TilePos), &CameraProjectionPolicyConfig);
-
-	if (!OutProjPolicy.IsValid())
-	{
-		UE_LOG(LogDisplayClusterViewport, Error, TEXT("Tile Viewport '%s': projection policy for tile [%d-%d] not created for node '%s'."), *InSourceViewport.GetId(), TilePos.X, TilePos.Y, *InSourceViewport.GetClusterNodeId());
-
-		return false;
-	}
-
-	return true;
-}
-
-FDisplayClusterViewport* FDisplayClusterViewportConfigurationHelpers_Tile::FindTileViewport(FDisplayClusterViewport& InSourceViewport, const FIntPoint& TilePos)
-{
-	if (FDisplayClusterViewportManager* ViewportManager = InSourceViewport.Configuration->GetViewportManagerImpl())
-	{
-		TSharedPtr<FDisplayClusterViewport, ESPMode::ThreadSafe> Viewport = ViewportManager->ImplFindViewport(GetUniqueViewportNameForTile(InSourceViewport.GetId(), TilePos));
-
-		return Viewport.Get();
-	}
-
-	return nullptr;
-}
-
 FDisplayClusterViewport* FDisplayClusterViewportConfigurationHelpers_Tile::GetOrCreateTileViewport(FDisplayClusterViewport& InSourceViewport, const FIntPoint& InTilePos)
 {
 	// Note: At this point, the viewports should already be configured.
@@ -120,20 +80,18 @@ FDisplayClusterViewport* FDisplayClusterViewportConfigurationHelpers_Tile::GetOr
 	}
 
 	const FIntPoint& InTileSize = SourceRenderSettings.TileSettings.GetSize();
+	const EDisplayClusterViewportFlags TileViewportUsage = InSourceViewport.ViewportFlags | EDisplayClusterViewportFlags::TileViewport; // Combine flags on top;
+	const FString& TileViewportBaseId = InSourceViewport.ViewportBaseId;
 
-	FDisplayClusterViewport* TileViewport = FindTileViewport(InSourceViewport, InTilePos);
-	if (TileViewport == nullptr)
+	FDisplayClusterViewport* TileViewport = InSourceViewport.Configuration->FindViewport(
+		TileViewportBaseId, TileViewportUsage, &InTilePos);
+	if (!TileViewport)
 	{
-		if (FDisplayClusterViewportManager* ViewportManager = InSourceViewport.Configuration->GetViewportManagerImpl())
-		{
-			// Create new camera viewport
-			TSharedPtr<IDisplayClusterProjectionPolicy, ESPMode::ThreadSafe> TileProjectionPolicy;
-			if (CreateProjectionPolicyForTileViewport(InSourceViewport, InTilePos, TileProjectionPolicy))
-			{
-				// Create viewport for new projection policy
-				TileViewport = ViewportManager->ImplCreateViewport(GetUniqueViewportNameForTile(InSourceViewport.GetId(), InTilePos), TileProjectionPolicy).Get();
-			}
-		}
+		TileViewport = InSourceViewport.Configuration->CreateViewport(
+			TileViewportBaseId,
+			InSourceViewport.Configuration->CreateProjectionPolicy_TileViewport(TileViewportBaseId, TileViewportUsage, InTilePos),
+			TileViewportUsage,
+			&InTilePos);
 	}
 
 	// Update tile viewport settings from the source viewport
diff --git a/Engine/Plugins/Runtime/nDisplay/Source/DisplayCluster/Private/Render/Viewport/Configuration/DisplayClusterViewportConfigurationHelpers_Tile.h b/Engine/Plugins/Runtime/nDisplay/Source/DisplayCluster/Private/Render/Viewport/Configuration/DisplayClusterViewportConfigurationHelpers_Tile.h
--- a/Engine/Plugins/Runtime/nDisplay/Source/DisplayCluster/Private/Render/Viewport/Configuration/DisplayClusterViewportConfigurationHelpers_Tile.h
+++ b/Engine/Plugins/Runtime/nDisplay/Source/DisplayCluster/Private/Render/Viewport/Configuration/DisplayClusterViewportConfigurationHelpers_Tile.h
@@ -17,15 +17,9 @@ struct FDisplayClusterViewport_TileSettings;
 class FDisplayClusterViewportConfigurationHelpers_Tile
 {
 public:
-	/* Find existing InnerFrustum viewport. */
-	static FDisplayClusterViewport* FindTileViewport(FDisplayClusterViewport& InSourceViewport, const FIntPoint& TilePos);
-
 	/** Get or create a new tile viewport from source. */
 	static FDisplayClusterViewport* GetOrCreateTileViewport(FDisplayClusterViewport& InSourceViewport, const FIntPoint& TilePos);
 
-	/** Get unique viewport name for tile. */
-	static FString GetUniqueViewportNameForTile(const FString& InViewportId, const FIntPoint& TilePos);
-
 public:
 	/** Get tile settings for the camera viewport. */
 	static void UpdateICVFXCameraViewportTileSettings(FDisplayClusterViewport& InSourceViewport, const FDisplayClusterConfigurationMediaICVFX& InCameraMediaSettings);
@@ -35,8 +29,4 @@ public:
 
 	/** Get tile overscan settings. */
 	static FDisplayClusterViewport_OverscanSettings GetTileOverscanSettings(const FDisplayClusterConfigurationTile_Overscan& InTileOverscan);
-
-private:
-	/** Create projection policy for tile viewport. */
-	static bool CreateProjectionPolicyForTileViewport(FDisplayClusterViewport& InSourceViewport, const FIntPoint& TilePos, TSharedPtr<IDisplayClusterProjectionPolicy, ESPMode::ThreadSafe>& OutProjPolicy);
 };
diff --git a/Engine/Plugins/Runtime/nDisplay/Source/DisplayCluster/Private/Render/Viewport/Configuration/DisplayClusterViewportConfiguration_ICVFX.cpp b/Engine/Plugins/Runtime/nDisplay/Source/DisplayCluster/Private/Render/Viewport/Configuration/DisplayClusterViewportConfiguration_ICVFX.cpp
--- a/Engine/Plugins/Runtime/nDisplay/Source/DisplayCluster/Private/Render/Viewport/Configuration/DisplayClusterViewportConfiguration_ICVFX.cpp
+++ b/Engine/Plugins/Runtime/nDisplay/Source/DisplayCluster/Private/Render/Viewport/Configuration/DisplayClusterViewportConfiguration_ICVFX.cpp
@@ -24,7 +24,6 @@
 #include "Render/Viewport/Configuration/DisplayClusterViewportConfigurationHelpers_Visibility.h"
 #include "Render/Viewport/DisplayClusterViewport.h"
 #include "Render/Viewport/DisplayClusterViewportManager.h"
-#include "Render/Viewport/DisplayClusterViewportStrings.h"
 #include "Render/Viewport/LightCard/DisplayClusterViewportLightCardManager.h"
 #include "Render/Viewport/RenderFrame/DisplayClusterRenderFrameSettings.h"
 
@@ -170,15 +169,13 @@ void FDisplayClusterViewportConfiguration_ICVFX::Update()
 
 				if (!ViewportCfg->ICVFX.EnforcedCameraOrder.IsEmpty())
 				{
-					const FString& ClusterNodeId = Configuration.GetClusterNodeId();
-
 					// Build full viewport IDs from camera names using helper
 					for (const FString& CameraName : ViewportCfg->ICVFX.EnforcedCameraOrder)
 					{
 						if (!CameraName.IsEmpty())
 						{
-							const FString ViewportId = FDisplayClusterViewportConfigurationHelpers_ICVFX::BuildCameraViewportId(ClusterNodeId, CameraName);
-							InOutRenderSettingsICVFX.ICVFX.EnforcedCameraOrder.Add(ViewportId);
+							const FString InCameraViewportId = Configuration.BuildViewportName(CameraName, EDisplayClusterViewportFlags::ICVFX_InnerFrustum);
+							InOutRenderSettingsICVFX.ICVFX.EnforcedCameraOrder.Add(InCameraViewportId);
 						}
 					}
 				}
diff --git a/Engine/Plugins/Runtime/nDisplay/Source/DisplayCluster/Private/Render/Viewport/Configuration/DisplayClusterViewportConfiguration_ICVFXCamera.cpp b/Engine/Plugins/Runtime/nDisplay/Source/DisplayCluster/Private/Render/Viewport/Configuration/DisplayClusterViewportConfiguration_ICVFXCamera.cpp
--- a/Engine/Plugins/Runtime/nDisplay/Source/DisplayCluster/Private/Render/Viewport/Configuration/DisplayClusterViewportConfiguration_ICVFXCamera.cpp
+++ b/Engine/Plugins/Runtime/nDisplay/Source/DisplayCluster/Private/Render/Viewport/Configuration/DisplayClusterViewportConfiguration_ICVFXCamera.cpp
@@ -19,7 +19,6 @@
 
 #include "Render/Viewport/DisplayClusterViewport.h"
 #include "Render/Viewport/DisplayClusterViewportManager.h"
-#include "Render/Viewport/DisplayClusterViewportStrings.h"
 #include "Render/Viewport/RenderFrame/DisplayClusterRenderFrameSettings.h"
 #include "Render/Viewport/LightCard/DisplayClusterViewportLightCardManager.h"
 
@@ -170,8 +169,9 @@ void FDisplayClusterViewportConfiguration_ICVFXCamera::Update()
 bool FDisplayClusterViewportConfiguration_ICVFXCamera::Initialize()
 {
 	// Create new camera projection policy for camera viewport
-	TSharedPtr<IDisplayClusterProjectionPolicy, ESPMode::ThreadSafe> CameraProjectionPolicy;
-	if (!FDisplayClusterViewportConfigurationHelpers_ICVFX::CreateProjectionPolicyCameraICVFX(Configuration, CameraComponent, GetCameraSettings(), CameraProjectionPolicy))
+	TSharedPtr<IDisplayClusterProjectionPolicy, ESPMode::ThreadSafe> CameraProjectionPolicy =
+		FDisplayClusterViewportConfigurationHelpers_ICVFX::CreateProjectionPolicyCameraICVFX(Configuration, CameraComponent, GetCameraSettings());
+	if (!CameraProjectionPolicy.IsValid())
 	{
 		return false;
 	}
diff --git a/Engine/Plugins/Runtime/nDisplay/Source/DisplayCluster/Private/Render/Viewport/Configuration/DisplayClusterViewportConfiguration_Tile.cpp b/Engine/Plugins/Runtime/nDisplay/Source/DisplayCluster/Private/Render/Viewport/Configuration/DisplayClusterViewportConfiguration_Tile.cpp
--- a/Engine/Plugins/Runtime/nDisplay/Source/DisplayCluster/Private/Render/Viewport/Configuration/DisplayClusterViewportConfiguration_Tile.cpp
+++ b/Engine/Plugins/Runtime/nDisplay/Source/DisplayCluster/Private/Render/Viewport/Configuration/DisplayClusterViewportConfiguration_Tile.cpp
@@ -8,7 +8,6 @@
 
 #include "Render/Viewport/DisplayClusterViewport.h"
 #include "Render/Viewport/DisplayClusterViewportManager.h"
-#include "Render/Viewport/DisplayClusterViewportStrings.h"
 #include "Render/Viewport/RenderFrame/DisplayClusterRenderFrameSettings.h"
 
 #include "DisplayClusterRootActor.h"
diff --git a/Engine/Plugins/Runtime/nDisplay/Source/DisplayCluster/Private/Render/Viewport/DisplayClusterViewport.cpp b/Engine/Plugins/Runtime/nDisplay/Source/DisplayCluster/Private/Render/Viewport/DisplayClusterViewport.cpp
--- a/Engine/Plugins/Runtime/nDisplay/Source/DisplayCluster/Private/Render/Viewport/DisplayClusterViewport.cpp
+++ b/Engine/Plugins/Runtime/nDisplay/Source/DisplayCluster/Private/Render/Viewport/DisplayClusterViewport.cpp
@@ -12,6 +12,7 @@
 #include "Render/Viewport/Configuration/DisplayClusterViewportConfiguration.h"
 #include "Render/Viewport/Configuration/DisplayClusterViewportConfigurationProxy.h"
 #include "Render/Viewport/Configuration/DisplayClusterViewportConfigurationHelpers_Tile.h"
+#include "Render/Viewport/Configuration/DisplayClusterViewportConfigurationHelpers_ICVFX.h"
 
 #include "Render/Viewport/LightCard/DisplayClusterViewportLightCardManager.h"
 
@@ -28,7 +29,9 @@
 #include "Render/Viewport/Preview/DisplayClusterViewportPreview.h"
 
 #include "Render/DisplayDevice/Components/DisplayClusterDisplayDeviceBaseComponent.h"
-#include "Render/Upscaler/DisplayClusterUpscaler.h"
+
+#include "Render/Features/AutoExposure/DisplayClusterAutoExposure.h"
+#include "Render/Features/DisplayClusterUpscaler.h"
 
 #include "EngineUtils.h"
 #include "OpenColorIORendering.h"
@@ -67,18 +70,42 @@ namespace UE::DisplayCluster::Viewport
 ///////////////////////////////////////////////////////////////////////////////////////
 //          FDisplayClusterViewport
 ///////////////////////////////////////////////////////////////////////////////////////
-FDisplayClusterViewport::FDisplayClusterViewport(const TSharedRef<FDisplayClusterViewportConfiguration, ESPMode::ThreadSafe>& InConfiguration, const FString& InViewportId, const TSharedPtr<IDisplayClusterProjectionPolicy, ESPMode::ThreadSafe>& InProjectionPolicy)
+FDisplayClusterViewport::FDisplayClusterViewport(
+	const TSharedRef<FDisplayClusterViewportConfiguration, ESPMode::ThreadSafe>& InConfiguration,
+	const FString& InViewportBaseId,
+	const EDisplayClusterViewportFlags InViewportFlags,
+	const TSharedPtr<IDisplayClusterProjectionPolicy, ESPMode::ThreadSafe>& InProjectionPolicy,
+	const FIntPoint* InOptTilePos)
 	: Configuration(InConfiguration)
-	, ViewportPreview(MakeShared< FDisplayClusterViewportPreview, ESPMode::ThreadSafe>(InConfiguration, InViewportId))
-	, ViewportProxy(MakeShared<FDisplayClusterViewportProxy, ESPMode::ThreadSafe>(InConfiguration, InViewportId, InProjectionPolicy))
-	, ViewportId(InViewportId)
+	, ViewportBaseId(InViewportBaseId)
+	, ViewportFlags(InViewportFlags)
+	, ViewportId(InConfiguration->BuildViewportName(InViewportBaseId, InViewportFlags, InOptTilePos))
 	, ClusterNodeId(InConfiguration->GetClusterNodeId())
+	, ViewportProxy(MakeShared<FDisplayClusterViewportProxy, ESPMode::ThreadSafe>(
+		InConfiguration, ViewportBaseId, InViewportFlags, ViewportId, InProjectionPolicy))
+	, ViewportPreview(MakeShared< FDisplayClusterViewportPreview, ESPMode::ThreadSafe>(InConfiguration, ViewportId))
 	, UninitializedProjectionPolicy(InProjectionPolicy)
 {
 	check(!ClusterNodeId.IsEmpty());
-	check(!ViewportId.IsEmpty());
+	check(!ViewportBaseId.IsEmpty());
 	check(UninitializedProjectionPolicy.IsValid());
 
+	// Initialize ICVFX viewport
+	if (EnumHasAnyFlags(InViewportFlags, EDisplayClusterViewportFlags::ICVFX_AnyRole))
+	{
+		// ICVFX viewports are not directly visible on the final frame texture (back buffer).
+		// They are composed using ICVFX shaders.
+		RenderSettings.bVisible = false;
+
+		// All ICVFX viewports should have a flag to manage it properly
+		// Mark as internal resource
+		EnumAddFlags(RenderSettingsICVFX.RuntimeFlags, EDisplayClusterViewportRuntimeICVFXFlags::InternalResource);
+
+		// Backward-compatibility bridge: derive legacy EDisplayClusterViewportRuntimeICVFXFlags
+		// from the new EDisplayClusterViewportFlags until the legacy runtime enum is removed.
+		FDisplayClusterViewportConfigurationHelpers_ICVFX::UpdateLegacyRuntimeICVFXFlagsFromViewportFlags(*this);
+	}
+
 	if (FDisplayClusterViewportManagerProxy* ViewportManagerProxy = Configuration->Proxy->GetViewportManagerProxyImpl())
 	{
 		// Add viewport proxy on renderthread
@@ -272,9 +299,17 @@ void FDisplayClusterViewport::Initialize()
 	ViewportPreview->Initialize(*this);
 }
 
-void FDisplayClusterViewport::ReleaseTextures()
+void FDisplayClusterViewport::ReleaseAllResources()
 {
+	// Release viewport-owned render resources (textures, RTT).
 	Resources.ReleaseAllResources();
+
+	// Release material/view-state references on the game thread.
+	CleanupViewState();
+
+	// Drop game-thread view state references held by this viewport.
+	// The viewport proxy keeps its own copy of these references.
+	ViewStates.Reset();
 }
 
 void FDisplayClusterViewport::OnHandleStartScene()
@@ -388,6 +423,12 @@ void FDisplayClusterViewport::SetupSceneView(uint32 ContextNum, class UWorld* Wo
 	// Handle depth of field parameters
 	CameraDepthOfField.SetupSceneView(InOutView);
 
+	// Handle AutoExposure for this viewport
+	if (FDisplayClusterViewportManager* ViewportManager = Configuration->GetViewportManagerImpl())
+	{
+		ViewportManager->AutoExposure->ConfigureAutoExposureForSceneView(*this, ContextNum, InOutView);
+	}
+
 	// Handle DisplayDevice
 	if (UDisplayClusterDisplayDeviceBaseComponent* InDisplayDeviceComponent = GetDisplayDeviceComponent(Configuration->GetPreviewSettings().DisplayDeviceRootActorType))
 	{
@@ -755,6 +796,15 @@ bool FDisplayClusterViewport::UpdateFrameContexts(const uint32 InStereoViewIndex
 		Context.TileDestRect = TileDestRect;
 		Context.ContextSize = ContextSize;
 
+		if (OverscanRuntimeSettings.bIsEnabled)
+		{
+			Context.OverscanInnerRenderTargetRect = OverscanRuntimeSettings.OverscanPixels.GetInnerRect(RenderTargetRect);
+		}
+		else
+		{
+			Context.OverscanInnerRenderTargetRect = RenderTargetRect;
+		}
+
 		// r.ScreenPercentage
 		switch (RenderSettings.CaptureMode)
 		{
diff --git a/Engine/Plugins/Runtime/nDisplay/Source/DisplayCluster/Private/Render/Viewport/DisplayClusterViewport.h b/Engine/Plugins/Runtime/nDisplay/Source/DisplayCluster/Private/Render/Viewport/DisplayClusterViewport.h
--- a/Engine/Plugins/Runtime/nDisplay/Source/DisplayCluster/Private/Render/Viewport/DisplayClusterViewport.h
+++ b/Engine/Plugins/Runtime/nDisplay/Source/DisplayCluster/Private/Render/Viewport/DisplayClusterViewport.h
@@ -11,7 +11,6 @@
 #include "Render/Viewport/Containers/DisplayClusterViewport_OverscanRuntimeSettings.h"
 #include "Render/Viewport/Containers/DisplayClusterViewportRemap.h"
 #include "Render/Viewport/Containers/DisplayClusterViewport_Enums.h"
-#include "Render/Viewport/Containers/DisplayClusterViewport_InternalEnums.h"
 
 #include "Render/Viewport/DisplayClusterViewport_CustomPostProcessSettings.h"
 #include "Render/Viewport/DisplayClusterViewport_VisibilitySettings.h"
@@ -40,7 +39,23 @@ class FDisplayClusterViewport
 	, public TSharedFromThis<FDisplayClusterViewport, ESPMode::ThreadSafe>
 {
 public:
-	FDisplayClusterViewport(const TSharedRef<FDisplayClusterViewportConfiguration, ESPMode::ThreadSafe>& InConfiguration, const FString& ViewportId, const TSharedPtr<IDisplayClusterProjectionPolicy, ESPMode::ThreadSafe>& InProjectionPolicy);
+	/**
+	* Constructs a new nDisplay viewport.
+	*
+	* @param InConfiguration      Viewport configuration object (thread-safe).
+	* @param InViewportBaseId     The base viewport identifier (before ICVFX unique naming).
+	* @param InViewportFlags      Usage type of the viewport (e.g. InCamera, Chromakey, LightCard).
+	* @param InProjectionPolicy   Projection policy assigned to the viewport
+	*                             (defines how the scene is projected/rendered).
+	* @param InOptTilePos         Optional parameter for tile viewport position
+	*/
+	FDisplayClusterViewport(
+		const TSharedRef<FDisplayClusterViewportConfiguration, ESPMode::ThreadSafe>& InConfiguration,
+		const FString& InViewportBaseId,
+		const EDisplayClusterViewportFlags InViewportFlags,
+		const TSharedPtr<IDisplayClusterProjectionPolicy, ESPMode::ThreadSafe>& InProjectionPolicy,
+		const FIntPoint* InOptTilePos = nullptr);
+
 	virtual ~FDisplayClusterViewport();
 
 public:
@@ -89,6 +104,16 @@ public:
 	/** Get viewport preview API */
 	virtual IDisplayClusterViewportPreview& GetViewportPreview() const override;
 
+	virtual EDisplayClusterViewportFlags GetViewportFlags() const override
+	{
+		return ViewportFlags;
+	}
+
+	virtual FString GetBaseId() const override
+	{
+		return ViewportBaseId;
+	}
+
 	virtual FString GetId() const override
 	{ 
 		return ViewportId; 
@@ -182,7 +207,9 @@ public:
 	//////////////////////////////////////////////////////
 
 	void Initialize();
-	void ReleaseTextures();
+
+	/** Release all viewport resources: textures and ViewState. */
+	void ReleaseAllResources();
 
 
 	// Get from logic request for additional targetable resource
@@ -480,18 +507,26 @@ public:
 	// Configuration of the current cluster node
 	const TSharedRef<FDisplayClusterViewportConfiguration, ESPMode::ThreadSafe> Configuration;
 
-	// Viewport preview
-	const TSharedRef<FDisplayClusterViewportPreview, ESPMode::ThreadSafe> ViewportPreview;
+	// Base viewport identifier. 
+	// May not be unique across ICVFX viewports (multiple viewports can share the same base ID).
+	const FString ViewportBaseId;
 
-	// viewport proxy (render thread data)
-	const TSharedRef<FDisplayClusterViewportProxy, ESPMode::ThreadSafe> ViewportProxy;
+	// Usage flags of the viewport (e.g. InCamera, Chromakey, LightCard, Tile).
+	const EDisplayClusterViewportFlags ViewportFlags;
 
-	// Unique viewport name
+	// Final unique viewport identifier. 
+	// Guaranteed to be unique across the entire cluster (includes node ID and usage type if needed).
 	const FString ViewportId;
 
 	// Owner cluster node name
 	const FString ClusterNodeId;
 
+	// viewport proxy (render thread data)
+	const TSharedRef<FDisplayClusterViewportProxy, ESPMode::ThreadSafe> ViewportProxy;
+
+	// Viewport preview
+	const TSharedRef<FDisplayClusterViewportPreview, ESPMode::ThreadSafe> ViewportPreview;
+
 private:
 	// Unified repository of viewport resources
 	FDisplayClusterViewportResources Resources;
diff --git a/Engine/Plugins/Runtime/nDisplay/Source/DisplayCluster/Private/Render/Viewport/DisplayClusterViewportManager.cpp b/Engine/Plugins/Runtime/nDisplay/Source/DisplayCluster/Private/Render/Viewport/DisplayClusterViewportManager.cpp
--- a/Engine/Plugins/Runtime/nDisplay/Source/DisplayCluster/Private/Render/Viewport/DisplayClusterViewportManager.cpp
+++ b/Engine/Plugins/Runtime/nDisplay/Source/DisplayCluster/Private/Render/Viewport/DisplayClusterViewportManager.cpp
@@ -29,8 +29,8 @@
 #include "Render/Viewport/Configuration/DisplayClusterViewportConfigurationProxy.h"
 #include "Render/Viewport/Configuration/DisplayClusterViewportConfiguration_Viewport.h"
 
-#include "Render/Viewport/DisplayClusterViewportStrings.h"
-#include "Render/Upscaler/DisplayClusterUpscaler.h"
+#include "Render/Features/AutoExposure/DisplayClusterAutoExposure.h"
+#include "Render/Features/DisplayClusterUpscaler.h"
 
 #include "IDisplayClusterWarpBlend.h"
 
@@ -240,6 +240,7 @@ FDisplayClusterViewportManager::FDisplayClusterViewportManager()
 	, PostProcessManager(MakeShared<FDisplayClusterViewportPostProcessManager, ESPMode::ThreadSafe>(Configuration))
 	, LightCardManager(MakeShared<FDisplayClusterViewportLightCardManager, ESPMode::ThreadSafe>(Configuration))
 	, RenderFrameManager(MakeShared<FDisplayClusterRenderFrameManager>(Configuration))
+	, AutoExposure(MakeShared<FDisplayClusterAutoExposure>(Configuration))
 {
 	// Always reset RTT when root actor re-created
 	ResetSceneRenderTargetSize();
@@ -267,7 +268,7 @@ void FDisplayClusterViewportManager::Release()
 	HandleEndScene();
 
 	// Frees GPU memory as fast as possible.
-	ReleaseTextures();
+	ReleaseAllResources();
 
 	// This is a flag to check in the destructor that Release() has been called before.
 	bReleaseFunctionCalledBeforeDestructor = true;
@@ -314,28 +315,31 @@ FDisplayClusterViewportManager::~FDisplayClusterViewportManager()
 	}
 }
 
-void FDisplayClusterViewportManager::ReleaseTextures()
+void FDisplayClusterViewportManager::ReleaseAllResources()
 {
 	for (const TSharedPtr<FDisplayClusterViewport, ESPMode::ThreadSafe>& ViewportIt : ImplGetEntireClusterViewports())
 	{
 		if (ViewportIt.IsValid())
 		{
-			ViewportIt->ReleaseTextures();
+			ViewportIt->ReleaseAllResources();
 		}
 	}
 
 	if (ViewportManagerProxy.IsValid())
 	{
 		// Remove viewport manager proxy on render_thread
-		ENQUEUE_RENDER_COMMAND(DisplayClusterViewportManagerProxy_ReleaseTextures)(
+		ENQUEUE_RENDER_COMMAND(DisplayClusterViewportManagerProxy_ReleaseAllResources)(
 			[ViewportManagerProxy = ViewportManagerProxy](FRHICommandListImmediate& RHICmdList)
 			{
-				ViewportManagerProxy->ReleaseTextures_RenderThread();
+				ViewportManagerProxy->ReleaseAllResources_RenderThread();
 			});
 	}
 
 	// Release RTT manager caches
 	RenderTargetManager->Release();
+
+	// Release AutoExposure inner data
+	AutoExposure->Release();
 }
 
 IDisplayClusterViewportManagerPreview& FDisplayClusterViewportManager::GetViewportManagerPreview()
@@ -759,6 +763,9 @@ void FDisplayClusterViewportManager::InitializeNewFrame()
 
 	// Send updated postprocess data to rendering thread
 	PostProcessManager->FinalizeNewFrame();
+
+	// Update AutoExposure data.
+	AutoExposure->SyncAndComputeAutoExposure();
 }
 
 void FDisplayClusterViewportManager::FinalizeNewFrame()
@@ -777,6 +784,12 @@ void FDisplayClusterViewportManager::FinalizeNewFrame()
 	}
 }
 
+void FDisplayClusterViewportManager::OnEntireClusterPreviewRendered()
+{
+	// Forward the "entire cluster preview rendered" event so AutoExposure can refresh its data.
+	AutoExposure->OnEntireClusterPreviewRendered();
+}
+
 FSceneViewFamily::ConstructionValues FDisplayClusterViewportManager::CreateViewFamilyConstructionValues(
 	const FDisplayClusterRenderFrameTarget& InFrameTarget,
 	FSceneInterface* InScene,
@@ -913,31 +926,31 @@ void FDisplayClusterViewportManager::RenderFrame(FViewport* InViewport)
 	ViewportManagerProxy->ImplRenderFrame_GameThread(InViewport);
 }
 
-FDisplayClusterViewport* FDisplayClusterViewportManager::CreateViewport(const FString& InViewportId, const class UDisplayClusterConfigurationViewport& ConfigurationViewport)
+FDisplayClusterViewport* FDisplayClusterViewportManager::CreateViewport(const FString& InViewportBaseId, const class UDisplayClusterConfigurationViewport& ConfigurationViewport)
 {
 	check(IsInGameThread());
 
 	// Check viewport ID
-	if (InViewportId.IsEmpty())
+	if (InViewportBaseId.IsEmpty())
 	{
 		UE_LOG(LogDisplayClusterViewport, Warning, TEXT("Wrong viewport ID"));
 		return nullptr;
 	}
 
 	// ID must be unique
-	if (FindViewport(InViewportId) != nullptr)
+	if (FindViewport(InViewportBaseId) != nullptr)
 	{
-		UE_LOG(LogDisplayClusterViewport, Warning, TEXT("Viewport '%s' already exists"), *InViewportId);
+		UE_LOG(LogDisplayClusterViewport, Warning, TEXT("Viewport '%s' already exists"), *InViewportBaseId);
 
 		return nullptr;
 	}
 
 	// Create projection policy for viewport
-	const TSharedPtr<IDisplayClusterProjectionPolicy, ESPMode::ThreadSafe> NewProjectionPolicy = CreateProjectionPolicy(InViewportId, &ConfigurationViewport.ProjectionPolicy);
+	const TSharedPtr<IDisplayClusterProjectionPolicy, ESPMode::ThreadSafe> NewProjectionPolicy = CreateProjectionPolicy(InViewportBaseId, &ConfigurationViewport.ProjectionPolicy);
 	if (NewProjectionPolicy.IsValid())
 	{
 		// Create viewport for new projection policy
-		TSharedPtr<FDisplayClusterViewport, ESPMode::ThreadSafe> NewViewport = ImplCreateViewport(InViewportId, NewProjectionPolicy);
+		TSharedPtr<FDisplayClusterViewport, ESPMode::ThreadSafe> NewViewport = ImplCreateViewport(InViewportBaseId, NewProjectionPolicy);
 		if (NewViewport.IsValid())
 		{
 			FDisplayClusterViewportConfiguration_Viewport::UpdateViewportConfiguration(*NewViewport , ConfigurationViewport);
@@ -946,7 +959,7 @@ FDisplayClusterViewport* FDisplayClusterViewportManager::CreateViewport(const FS
 		}
 	}
 
-	UE_LOG(LogDisplayClusterViewport, Error, TEXT("Viewports '%s' not created."), *InViewportId);
+	UE_LOG(LogDisplayClusterViewport, Error, TEXT("Viewports '%s' not created."), *InViewportBaseId);
 
 	return nullptr;
 }
@@ -957,12 +970,19 @@ IDisplayClusterViewport* FDisplayClusterViewportManager::FindViewport(const FStr
 	return ExistViewport.IsValid() ? ExistViewport.Get() : nullptr;
 }
 
-TSharedPtr<FDisplayClusterViewport, ESPMode::ThreadSafe> FDisplayClusterViewportManager::ImplFindViewport(const FString& ViewportId) const
+TSharedPtr<FDisplayClusterViewport, ESPMode::ThreadSafe> FDisplayClusterViewportManager::ImplFindViewport(
+	const FString& InViewportBaseId,
+	const EDisplayClusterViewportFlags InViewportFlags,
+	const FIntPoint* InOptTilePos) const
 {
 	check(IsInGameThread());
 
+	// Generate unique viewport name
+	const FString ViewportId = Configuration->BuildViewportName(InViewportBaseId, InViewportFlags, InOptTilePos);
+
 	// Ok, we have a request for a particular viewport. Let's find it.
-	TSharedPtr<FDisplayClusterViewport, ESPMode::ThreadSafe> const* DesiredViewport = ImplGetEntireClusterViewports().FindByPredicate([ViewportId](const TSharedPtr<FDisplayClusterViewport, ESPMode::ThreadSafe>& ItemViewport)
+	TSharedPtr<FDisplayClusterViewport, ESPMode::ThreadSafe> const* DesiredViewport = ImplGetEntireClusterViewports().FindByPredicate(
+		[ViewportId](const TSharedPtr<FDisplayClusterViewport, ESPMode::ThreadSafe>& ItemViewport)
 	{
 		return ItemViewport.IsValid() && ViewportId.Equals(ItemViewport->ViewportId, ESearchCase::IgnoreCase);
 	});
@@ -970,18 +990,18 @@ TSharedPtr<FDisplayClusterViewport, ESPMode::ThreadSafe> FDisplayClusterViewport
 	return (DesiredViewport && DesiredViewport->IsValid()) ? *DesiredViewport : nullptr;
 }
 
-FDisplayClusterViewport* FDisplayClusterViewportManager::CreateViewport(const FString& ViewportId, const TSharedPtr<IDisplayClusterProjectionPolicy, ESPMode::ThreadSafe>& InProjectionPolicy)
+FDisplayClusterViewport* FDisplayClusterViewportManager::CreateViewport(const FString& InViewportBaseId, const TSharedPtr<IDisplayClusterProjectionPolicy, ESPMode::ThreadSafe>& InProjectionPolicy)
 {
 	check(IsInGameThread());
 
-	TSharedPtr<FDisplayClusterViewport, ESPMode::ThreadSafe> ExistViewport = ImplFindViewport(ViewportId);
+	TSharedPtr<FDisplayClusterViewport, ESPMode::ThreadSafe> ExistViewport = ImplFindViewport(InViewportBaseId);
 	if (ExistViewport.IsValid())
 	{
 		//add error log: Viewport with name '%s' already exist
 		return nullptr;
 	}
 
-	TSharedPtr<FDisplayClusterViewport, ESPMode::ThreadSafe> NewViewport = ImplCreateViewport(ViewportId, InProjectionPolicy);
+	TSharedPtr<FDisplayClusterViewport, ESPMode::ThreadSafe> NewViewport = ImplCreateViewport(InViewportBaseId, InProjectionPolicy);
 
 	return NewViewport.IsValid() ? NewViewport.Get() : nullptr;
 }
@@ -1001,13 +1021,18 @@ bool FDisplayClusterViewportManager::DeleteViewport(const FString& ViewportId)
 	return false;
 }
 
-TSharedPtr<FDisplayClusterViewport, ESPMode::ThreadSafe> FDisplayClusterViewportManager::ImplCreateViewport(const FString& ViewportId, const TSharedPtr<IDisplayClusterProjectionPolicy, ESPMode::ThreadSafe>& InProjectionPolicy)
+TSharedPtr<FDisplayClusterViewport, ESPMode::ThreadSafe> FDisplayClusterViewportManager::ImplCreateViewport(
+	const FString& InViewportBaseId,
+	const TSharedPtr<IDisplayClusterProjectionPolicy, ESPMode::ThreadSafe>& InProjectionPolicy,
+	const EDisplayClusterViewportFlags InViewportFlags,
+	const FIntPoint* InOptTilePos)
 {
 	check(IsInGameThread());
 	check(InProjectionPolicy.IsValid());
 
 	// Create a new viewport for cluster node used for rendering
-	TSharedPtr<FDisplayClusterViewport, ESPMode::ThreadSafe> NewViewport = MakeShared<FDisplayClusterViewport>(Configuration, ViewportId, InProjectionPolicy);
+	TSharedPtr<FDisplayClusterViewport, ESPMode::ThreadSafe> NewViewport = MakeShared<FDisplayClusterViewport>(
+		Configuration, InViewportBaseId, InViewportFlags, InProjectionPolicy, InOptTilePos);
 	if (NewViewport.IsValid())
 	{
 		NewViewport->Initialize();
@@ -1015,7 +1040,7 @@ TSharedPtr<FDisplayClusterViewport, ESPMode::ThreadSafe> FDisplayClusterViewport
 		// Add viewport on gamethread
 		EntireClusterViewports.Add(NewViewport);
 		
-		// Raises the stale flag, marking the current frame’s viewport list as needing an update.
+		// Raises the stale flag, marking the current frame's viewport list as needing an update.
 		Configuration->MarkCurrentRenderFrameViewportsOutOfDate();
 
 		if (Configuration->IsSceneOpened())
@@ -1043,7 +1068,7 @@ void FDisplayClusterViewportManager::ImplDeleteViewport(const TSharedPtr<FDispla
 		EntireClusterViewports.RemoveAt(ViewportIndex);
 	}
 
-	// Raises the stale flag, marking the current frame’s viewport list as needing an update.
+	// Raises the stale flag, marking the current frame's viewport list as needing an update.
 	Configuration->MarkCurrentRenderFrameViewportsOutOfDate();
 
 	// Reset RTT size after viewport delete
@@ -1077,7 +1102,7 @@ TSharedPtr<IDisplayClusterProjectionPolicy, ESPMode::ThreadSafe> FDisplayCluster
 	check(InConfigurationProjectionPolicy != nullptr);
 
 	// Generate unique projection policy id from viewport name
-	const FString ProjectionPolicyId = FString::Printf(TEXT("%s_%s"), DisplayClusterViewportStrings::prefix::projection, *InViewportId);
+	const FString ProjectionPolicyId = FString::Printf(TEXT("prj_%s"), *InViewportId);
 
 	IDisplayClusterRenderManager* const DCRenderManager = IDisplayCluster::Get().GetRenderMgr();
 	check(DCRenderManager);
diff --git a/Engine/Plugins/Runtime/nDisplay/Source/DisplayCluster/Private/Render/Viewport/DisplayClusterViewportManager.h b/Engine/Plugins/Runtime/nDisplay/Source/DisplayCluster/Private/Render/Viewport/DisplayClusterViewportManager.h
--- a/Engine/Plugins/Runtime/nDisplay/Source/DisplayCluster/Private/Render/Viewport/DisplayClusterViewportManager.h
+++ b/Engine/Plugins/Runtime/nDisplay/Source/DisplayCluster/Private/Render/Viewport/DisplayClusterViewportManager.h
@@ -12,6 +12,7 @@
 #include "Render/Viewport/IDisplayClusterViewportManager.h"
 #include "Render/Viewport/DisplayClusterViewport.h"
 #include "Render/Viewport/Containers/DisplayClusterViewport_Enums.h"
+#include "Render/Viewport/Containers/DisplayClusterViewport_InternalEnums.h"
 
 class FDisplayClusterViewportManagerProxy;
 class IDisplayClusterProjectionPolicy;
@@ -43,9 +44,12 @@ public:
 	*/
 	void Release();
 
-	/** Release all texture resources used for viewports.
-	* After that, the BeginNewFrame() function should be called for a new resource allocation.*/
-	void ReleaseTextures();
+	/** Release all resources.
+	* After that, the BeginNewFrame() function should be called for a new resource allocation.
+	*
+	* Note: This function has been renamed; its previous name was 'ReleaseTextures'.
+	*/
+	void ReleaseAllResources();
 
 	//~ Begin IDisplayClusterViewportManager
 	virtual TSharedPtr<IDisplayClusterViewportManager, ESPMode::ThreadSafe> ToSharedPtr() override
@@ -171,29 +175,48 @@ public:
 	}
 
 	// internal use only
+	FDisplayClusterViewport* CreateViewport(const FString& InViewportBaseId, const class UDisplayClusterConfigurationViewport& ConfigurationViewport);
+	FDisplayClusterViewport* CreateViewport(const FString& InViewportBaseId, const TSharedPtr<IDisplayClusterProjectionPolicy, ESPMode::ThreadSafe>& InProjectionPolicy);
+	bool                     DeleteViewport(const FString& ViewportId);
 
+	FDisplayClusterViewport* ImplFindViewport(const int32 ViewIndex, uint32* OutContextNum = nullptr) const;
 	/**
-	* Find viewport object by name [Game thread func]
+	* Creates and initializes a new nDisplay viewport instance.
 	*
-	* @param ViewportId - Viewport name
-	* @return - viewport object ref, or nullptr
+	* @param InViewportBaseId    The base viewport identifier (before ICVFX unique naming).
+	* @param InProjectionPolicy  Projection policy assigned to the viewport
+	*                            (defines how the scene is projected/rendered).
+	* @param InViewportFlags     Usage flags of the viewport (e.g. InCamera/LightCard/Tile, etc.).
+	* @param InOptTilePos        Optional parameter for tile viewport position
+	*
+	* @return A thread-safe shared pointer to the created FDisplayClusterViewport.
 	*/
-	inline FDisplayClusterViewport* FindViewportByName(const FString& InViewportId) const
-	{
-		return ImplFindViewport(InViewportId).Get();
-	}
-
-	FDisplayClusterViewport* CreateViewport(const FString& ViewportId, const class UDisplayClusterConfigurationViewport& ConfigurationViewport);
-	FDisplayClusterViewport* CreateViewport(const FString& ViewportId, const TSharedPtr<IDisplayClusterProjectionPolicy, ESPMode::ThreadSafe>& InProjectionPolicy);
-	bool                     DeleteViewport(const FString& ViewportId);
+	TSharedPtr<FDisplayClusterViewport, ESPMode::ThreadSafe> ImplCreateViewport(
+		const FString& InViewportBaseId,
+		const TSharedPtr<IDisplayClusterProjectionPolicy, ESPMode::ThreadSafe>& InProjectionPolicy,
+		const EDisplayClusterViewportFlags InViewportFlags = EDisplayClusterViewportFlags::None,
+		const FIntPoint* InOptTilePos = nullptr);
 
-	FDisplayClusterViewport* ImplFindViewport(const int32 ViewIndex, uint32* OutContextNum = nullptr) const;
+	/**
+	* Finds a viewport instance by identifier and usage type.
+	*
+	* @param InViewportBaseId The base viewport identifier (before ICVFX unique naming).
+	* @param InViewportFlags  Usage flags of the viewport (e.g. InCamera/LightCard/Tile, etc.).
+	* @param InOptTilePos     Optional parameter for tile viewport position
+	*
+	* @return A thread-safe shared pointer to the matching FDisplayClusterViewport,
+	*         or nullptr if no viewport matches the given criteria.
+	*/
+	TSharedPtr<FDisplayClusterViewport, ESPMode::ThreadSafe> ImplFindViewport(
+		const FString& InViewportBaseId,
+		const EDisplayClusterViewportFlags InViewportFlags = EDisplayClusterViewportFlags::None,
+		const FIntPoint* InOptTilePos = nullptr) const;
 
-	TSharedPtr<FDisplayClusterViewport, ESPMode::ThreadSafe> ImplCreateViewport(const FString& ViewportId, const TSharedPtr<IDisplayClusterProjectionPolicy, ESPMode::ThreadSafe>& InProjectionPolicy);
-	TSharedPtr<FDisplayClusterViewport, ESPMode::ThreadSafe> ImplFindViewport(const FString& InViewportId) const;
 	void ImplDeleteViewport(const TSharedPtr<FDisplayClusterViewport, ESPMode::ThreadSafe>& InExistViewport);
 
-	static TSharedPtr<IDisplayClusterProjectionPolicy, ESPMode::ThreadSafe> CreateProjectionPolicy(const FString& InViewportId, const FDisplayClusterConfigurationProjection* InConfigurationProjectionPolicy);
+	static TSharedPtr<IDisplayClusterProjectionPolicy, ESPMode::ThreadSafe> CreateProjectionPolicy(
+		const FString& InViewportId,
+		const FDisplayClusterConfigurationProjection* InConfigurationProjectionPolicy);
 
 	/** Get DC viewport manager proxy. */
 	TSharedPtr<FDisplayClusterViewportManagerProxy, ESPMode::ThreadSafe> GetViewportManagerProxy() const
@@ -233,6 +256,9 @@ public:
 
 	void ResetSceneRenderTargetSize();
 
+	/** Called every time previews for all cluster nodes are rendered. */
+	void OnEntireClusterPreviewRendered();
+
 private:
 	/** This function updates CurrentRenderFrameViewports according to the current rendering settings. */
 	void UpdateCurrentRenderFrameViewports() const;
@@ -265,6 +291,9 @@ public:
 	// Manager for creating the render frame stucture
 	const TSharedRef<class FDisplayClusterRenderFrameManager> RenderFrameManager;
 
+	// nDisplay multi-viewport AutoExposure feature.
+	const TSharedRef<class FDisplayClusterAutoExposure> AutoExposure;
+
 private:
 	// Render thread proxy manager. Deleted on render thread
 	TSharedPtr<FDisplayClusterViewportManagerProxy, ESPMode::ThreadSafe> ViewportManagerProxy;
diff --git a/Engine/Plugins/Runtime/nDisplay/Source/DisplayCluster/Private/Render/Viewport/DisplayClusterViewportManagerProxy.cpp b/Engine/Plugins/Runtime/nDisplay/Source/DisplayCluster/Private/Render/Viewport/DisplayClusterViewportManagerProxy.cpp
--- a/Engine/Plugins/Runtime/nDisplay/Source/DisplayCluster/Private/Render/Viewport/DisplayClusterViewportManagerProxy.cpp
+++ b/Engine/Plugins/Runtime/nDisplay/Source/DisplayCluster/Private/Render/Viewport/DisplayClusterViewportManagerProxy.cpp
@@ -12,25 +12,29 @@
 #include "IDisplayCluster.h"
 #include "IDisplayClusterCallbacks.h"
 #include "Render/IDisplayClusterRenderManager.h"
-#include "Render/Viewport/DisplayClusterViewportManager.h"
+
+#include "Render/Features/AutoExposure/DisplayClusterAutoExposure.h"
+
 #include "Render/Projection/IDisplayClusterProjectionPolicyFactory.h"
 #include "Render/Projection/IDisplayClusterProjectionPolicy.h"
 
+#include "Render/Viewport/DisplayClusterViewportManager.h"
+#include "Render/Viewport/DisplayClusterViewportManagerViewExtension.h"
+#include "Render/Viewport/DisplayClusterViewportProxy.h"
+
+#include "Render/Viewport/Containers/DisplayClusterViewportProxyData.h"
+
 #include "Render/Viewport/LightCard/DisplayClusterViewportLightCardManager.h"
 #include "Render/Viewport/LightCard/DisplayClusterViewportLightCardManagerProxy.h"
 
-#include "Render/Viewport/DisplayClusterViewportProxy.h"
-#include "Render/Viewport/DisplayClusterViewportManagerViewExtension.h"
+#include "Render/Viewport/Postprocess/DisplayClusterViewportPostProcessManager.h"
 
 #include "Render/Viewport/RenderTarget/DisplayClusterRenderTargetManager.h"
 #include "Render/Viewport/RenderTarget/DisplayClusterRenderTargetResource.h"
-#include "Render/Viewport/Postprocess/DisplayClusterViewportPostProcessManager.h"
-#include "Render/Viewport/Containers/DisplayClusterViewportProxyData.h"
 #include "Render/Viewport/DisplayClusterViewportStereoscopicPass.h"
 
 #include "Misc/DisplayClusterLog.h"
 
-#include "Render/Viewport/DisplayClusterViewportStrings.h"
 #include "RHIContext.h"
 
 // Enable/Disable ClearTexture for Frame RTT
@@ -52,6 +56,7 @@ FDisplayClusterViewportManagerProxy::FDisplayClusterViewportManagerProxy(const F
 	, RenderTargetManager(InViewportManager.RenderTargetManager)
 	, PostProcessManager(InViewportManager.PostProcessManager)
 	, LightCardManagerProxy(InViewportManager.LightCardManager->LightCardManagerProxy)
+	, AutoExposure(InViewportManager.AutoExposure)
 	, ViewportManagerViewExtension(InViewportManager.GetViewportManagerViewExtension())
 { }
 
@@ -201,6 +206,10 @@ void FDisplayClusterViewportManagerProxy::ImplRenderFrame_GameThread(FViewport*
 
 		// Handle render setup
 		ViewportManagerProxy->PostProcessManager->HandleRenderFrameSetup_RenderThread(RHICmdList, ViewportManagerProxy);
+
+		// Collects AutoExposure (Eye Adaptation) data from the viewports
+		ViewportManagerProxy->AutoExposure->GatherMeterViewportsAutoExposureDataAndSync_RenderThread(
+			ViewportManagerProxy->ImplGetCurrentRenderFrameViewportProxies_RenderThread());
 	});
 
 	ENQUEUE_RENDER_COMMAND(DisplayClusterRenderFrame_CrossGPUTransfer)(
@@ -699,13 +708,13 @@ bool FDisplayClusterViewportManagerProxy::ResolveFrameTargetToBackBuffer_RenderT
 	return false;
 }
 
-void FDisplayClusterViewportManagerProxy::ReleaseTextures_RenderThread()
+void FDisplayClusterViewportManagerProxy::ReleaseAllResources_RenderThread()
 {
 	for (const TSharedPtr<FDisplayClusterViewportProxy, ESPMode::ThreadSafe>& ViewportProxyIt : ImplGetEntireClusterViewportProxies_RenderThread())
 	{
 		if (ViewportProxyIt.IsValid())
 		{
-			ViewportProxyIt->ReleaseTextures_RenderThread();
+			ViewportProxyIt->ReleaseAllResources_RenderThread();
 		}
 	}
 }
diff --git a/Engine/Plugins/Runtime/nDisplay/Source/DisplayCluster/Private/Render/Viewport/DisplayClusterViewportManagerProxy.h b/Engine/Plugins/Runtime/nDisplay/Source/DisplayCluster/Private/Render/Viewport/DisplayClusterViewportManagerProxy.h
--- a/Engine/Plugins/Runtime/nDisplay/Source/DisplayCluster/Private/Render/Viewport/DisplayClusterViewportManagerProxy.h
+++ b/Engine/Plugins/Runtime/nDisplay/Source/DisplayCluster/Private/Render/Viewport/DisplayClusterViewportManagerProxy.h
@@ -20,6 +20,7 @@ class FDisplayClusterViewportManager;
 class FDisplayClusterViewportManagerViewExtension;
 class FDisplayClusterViewportResource;
 class FDisplayClusterViewportLightCardManagerProxy;
+class FDisplayClusterAutoExposure;
 class IDisplayClusterProjectionPolicy;
 class FViewport;
 
@@ -86,8 +87,11 @@ public:
 	*/
 	bool ImplResolveFrameTargetToBackBuffer_RenderThread(FRDGBuilder& GraphBuilder, const uint32 InContextNum, const int32 DestArrayIndex, FRDGTextureRef DestTexture, FVector2f WindowSize) const;
 
-	/** Release all textures on render thread. */
-	void ReleaseTextures_RenderThread();
+	/** Release all resources on render thread.
+	*
+	* Note: This function has been renamed; its previous name was 'ReleaseTextures_RenderThread()'
+	*/
+	void ReleaseAllResources_RenderThread();
 
 	/** Custom cross-GPU implementation for mGPU. Applyed for viewports with bOverrideCrossGPUTransfer=true. */
 	void DoCrossGPUTransfers_RenderThread(FRHICommandListImmediate& RHICmdList) const;
@@ -160,6 +164,9 @@ public:
 	// LightCard manager proxy object.
 	const TSharedRef<FDisplayClusterViewportLightCardManagerProxy, ESPMode::ThreadSafe> LightCardManagerProxy;
 
+	// nDisplay multi-viewport AutoExposure feature.
+	const TSharedRef<FDisplayClusterAutoExposure> AutoExposure;
+
 private:
 	// Viewport proxies of the entire cluster.
 	TArray<TSharedPtr<FDisplayClusterViewportProxy, ESPMode::ThreadSafe>> EntireClusterViewportProxies;
diff --git a/Engine/Plugins/Runtime/nDisplay/Source/DisplayCluster/Private/Render/Viewport/DisplayClusterViewportProxy.cpp b/Engine/Plugins/Runtime/nDisplay/Source/DisplayCluster/Private/Render/Viewport/DisplayClusterViewportProxy.cpp
--- a/Engine/Plugins/Runtime/nDisplay/Source/DisplayCluster/Private/Render/Viewport/DisplayClusterViewportProxy.cpp
+++ b/Engine/Plugins/Runtime/nDisplay/Source/DisplayCluster/Private/Render/Viewport/DisplayClusterViewportProxy.cpp
@@ -22,18 +22,49 @@
 #include "Render/Viewport/Containers/DisplayClusterViewportProxyData.h"
 #include "Render/Viewport/Resource/DisplayClusterViewportResource.h"
 
+#include "ScreenPass.h"
+#include "SceneTextures.h"
+#include "PostProcess/PostProcessMaterialInputs.h"
 
-///////////////////////////////////////////////////////////////////////////////////////
-IDisplayClusterShaders& FDisplayClusterViewportProxy::GetShadersAPI()
+namespace UE::DisplayCluster::ViewportProxy
+{
+	/**
+	 * Collects the current eye adaptation data from the given scene view.
+	 * 
+	 * @param InSceneView  The scene view to query.
+	 * @return A populated FDisplayClusterViewport_EyeAdaptationData struct with the latest auto-exposure data.
+	 */
+	static inline FDisplayClusterViewport_EyeAdaptationData GetEyeAdaptationData(const FSceneView& InSceneView)
 	{
-		static IDisplayClusterShaders& ShadersAPISingleton = IDisplayClusterShaders::Get();
+		// Get eye adaptation data from the FSceneView.
+		FDisplayClusterViewport_EyeAdaptationData OutEyeAdaptationData;
+		if (InSceneView.HasEyeAdaptationViewState())
+		{
+			OutEyeAdaptationData.LastAverageSceneLuminance = InSceneView.GetLastAverageSceneLuminance();
+		}
 
-		return ShadersAPISingleton;
+		return OutEyeAdaptationData;
 	}
+};
+
+///////////////////////////////////////////////////////////////////////////////////////
+IDisplayClusterShaders& FDisplayClusterViewportProxy::GetShadersAPI()
+{
+	static IDisplayClusterShaders& ShadersAPISingleton = IDisplayClusterShaders::Get();
+
+	return ShadersAPISingleton;
+}
 
 ///////////////////////////////////////////////////////////////////////////////////////
-FDisplayClusterViewportProxy::FDisplayClusterViewportProxy(const TSharedRef<FDisplayClusterViewportConfiguration, ESPMode::ThreadSafe>& InConfiguration, const FString& InViewportId, const TSharedPtr<IDisplayClusterProjectionPolicy, ESPMode::ThreadSafe>& InProjectionPolicy)
+FDisplayClusterViewportProxy::FDisplayClusterViewportProxy(
+	const TSharedRef<FDisplayClusterViewportConfiguration, ESPMode::ThreadSafe>& InConfiguration,
+	const FString& InViewportBaseId,
+	const EDisplayClusterViewportFlags InViewportFlags,
+	const FString& InViewportId,
+	const TSharedPtr<IDisplayClusterProjectionPolicy, ESPMode::ThreadSafe>& InProjectionPolicy)
 	: ConfigurationProxy(InConfiguration->Proxy)
+	, ViewportBaseId(InViewportBaseId)
+	, ViewportFlags(InViewportFlags)
 	, ViewportId(InViewportId)
 	, ClusterNodeId(InConfiguration->GetClusterNodeId())
 	, ProjectionPolicy(InProjectionPolicy)
@@ -42,8 +73,7 @@ FDisplayClusterViewportProxy::FDisplayClusterViewportProxy(const TSharedRef<FDis
 }
 
 FDisplayClusterViewportProxy::~FDisplayClusterViewportProxy()
-{
-}
+{ }
 
 void FDisplayClusterViewportProxy::UpdateViewportProxyData_RenderThread(const FDisplayClusterViewportProxyData& InViewportProxyData)
 {
@@ -378,20 +408,21 @@ void FDisplayClusterViewportProxy::OnPostRenderViewFamily_RenderThread(FRDGBuild
 {
 	const uint32 InContextNum = InProxyContext.ContextNum;
 
-#if WITH_MGPU
-	// Get the GPUIndex used to render this viewport
 	if (Contexts.IsValidIndex(InContextNum))
 	{
-
+		FDisplayClusterViewport_Context::FRenderThreadData& RenderThreadData = Contexts[InContextNum].RenderThreadData;
+#if WITH_MGPU
+		// Get the GPUIndex used to render this viewport
 		const uint32 GPUIndex = InSceneView.GPUMask.GetFirstIndex();
-		Contexts[InContextNum].RenderThreadData.GPUIndex = (GPUIndex < GNumExplicitGPUsForRendering) ? GPUIndex : -1;
-	}
+		RenderThreadData.GPUIndex = (GPUIndex < GNumExplicitGPUsForRendering) ? GPUIndex : -1;
 #endif
 
-	if (Contexts.IsValidIndex(InContextNum))
-	{
-		Contexts[InContextNum].RenderThreadData.EngineDisplayGamma = InSceneView.Family->RenderTarget->GetDisplayGamma();
-		Contexts[InContextNum].RenderThreadData.EngineShowFlags = InSceneView.Family->EngineShowFlags;
+		RenderThreadData.EngineDisplayGamma = InSceneView.Family->RenderTarget->GetDisplayGamma();
+		RenderThreadData.EngineShowFlags = InSceneView.Family->EngineShowFlags;
+		RenderThreadData.FrameNumber = InSceneView.Family->FrameNumber;
+
+		// Collect EyeAdaptation data from FSceneView.
+		RenderThreadData.EyeAdaptationData = UE::DisplayCluster::ViewportProxy::GetEyeAdaptationData(InSceneView);
 	}
 
 	// Allow the GUI controller to process output textures before PostRenderViewFamily is broadcasted. A dedicated
@@ -454,9 +485,13 @@ void FDisplayClusterViewportProxy::OnPostRenderViewFamily_RenderThread(FRDGBuild
 	}
 }
 
-void FDisplayClusterViewportProxy::ReleaseTextures_RenderThread()
+void FDisplayClusterViewportProxy::ReleaseAllResources_RenderThread()
 {
+	// Release textures
 	Resources.ReleaseAllResources();
+
+	// Release references to the ViewState data
+	ViewStates.Reset();
 }
 
 FDisplayClusterColorEncoding FDisplayClusterViewportProxy::GetResourceColorEncoding_RenderThread(const EDisplayClusterViewportResourceType InResourceType) const
diff --git a/Engine/Plugins/Runtime/nDisplay/Source/DisplayCluster/Private/Render/Viewport/DisplayClusterViewportProxy.h b/Engine/Plugins/Runtime/nDisplay/Source/DisplayCluster/Private/Render/Viewport/DisplayClusterViewportProxy.h
--- a/Engine/Plugins/Runtime/nDisplay/Source/DisplayCluster/Private/Render/Viewport/DisplayClusterViewportProxy.h
+++ b/Engine/Plugins/Runtime/nDisplay/Source/DisplayCluster/Private/Render/Viewport/DisplayClusterViewportProxy.h
@@ -50,7 +50,23 @@ class FDisplayClusterViewportProxy
 	, public TSharedFromThis<FDisplayClusterViewportProxy, ESPMode::ThreadSafe>
 {
 public:
-	FDisplayClusterViewportProxy(const TSharedRef<FDisplayClusterViewportConfiguration, ESPMode::ThreadSafe>& InConfiguration, const FString& InViewportId, const TSharedPtr<IDisplayClusterProjectionPolicy, ESPMode::ThreadSafe>& InProjectionPolicy);
+	/**
+	* Creates a viewport proxy instance.
+	* 
+	* @param InConfiguration    Viewport configuration object (thread-safe).
+	* @param InViewportBaseId   The base viewport identifier (before ICVFX unique naming).
+	* @param InViewportFlags    Usage type of the viewport (e.g. InCamera, Chromakey, LightCard).
+	* @param InViewportId       A viewport name string unique within the cluster.
+	* @param InProjectionPolicy Projection policy assigned to the viewport
+	*                           (defines how the scene is projected/rendered).
+	*/
+	FDisplayClusterViewportProxy(
+		const TSharedRef<FDisplayClusterViewportConfiguration, ESPMode::ThreadSafe>& InConfiguration,
+		const FString& InViewportBaseId,
+		const EDisplayClusterViewportFlags InViewportFlags,
+		const FString& InViewportId,
+		const TSharedPtr<IDisplayClusterProjectionPolicy, ESPMode::ThreadSafe>& InProjectionPolicy);
+
 	virtual ~FDisplayClusterViewportProxy();
 
 public:
@@ -70,6 +86,16 @@ public:
 		return *ConfigurationProxy;
 	}
 
+	virtual EDisplayClusterViewportFlags GetViewportFlags() const override
+	{
+		return ViewportFlags;
+	}
+
+	virtual FString GetBaseId() const override
+	{
+		return ViewportBaseId;
+	}
+
 	virtual FString GetId() const override
 	{
 		check(IsInRenderingThread());
@@ -251,8 +277,8 @@ public:
 	 */
 	void PostResolveViewport_RenderThread(FRHICommandListImmediate& RHICmdList) const;
 
-	/** Release all textures. */
-	void ReleaseTextures_RenderThread();
+	/** Release all resources on the rendering thread. */
+	void ReleaseAllResources_RenderThread();
 
 	inline bool FindContext_RenderThread(const int32 ViewIndex, uint32* OutContextNum)
 	{
@@ -332,7 +358,15 @@ public:
 	// Configuration for proxy
 	const TSharedRef<FDisplayClusterViewportConfigurationProxy, ESPMode::ThreadSafe> ConfigurationProxy;
 
-	/** Unique viewport name. */
+	// Base viewport identifier. 
+	// May not be unique across ICVFX viewports (multiple viewports can share the same base ID).
+	const FString ViewportBaseId;
+
+	// Usage type of the viewport (e.g. InCamera, Chromakey, LightCard).
+	const EDisplayClusterViewportFlags ViewportFlags;
+
+	// Final unique viewport identifier. 
+	// Guaranteed to be unique across the entire cluster (includes node ID and usage type if needed).
 	const FString ViewportId;
 
 	/** Cluster node name. */
diff --git a/Engine/Plugins/Runtime/nDisplay/Source/DisplayCluster/Private/Render/Viewport/DisplayClusterViewportStrings.h b/Engine/Plugins/Runtime/nDisplay/Source/DisplayCluster/Private/Render/Viewport/DisplayClusterViewportStrings.h
deleted file mode 100644
diff --git a/Engine/Plugins/Runtime/nDisplay/Source/DisplayCluster/Private/Render/Viewport/Preview/DisplayClusterViewportManagerPreview.cpp b/Engine/Plugins/Runtime/nDisplay/Source/DisplayCluster/Private/Render/Viewport/Preview/DisplayClusterViewportManagerPreview.cpp
--- a/Engine/Plugins/Runtime/nDisplay/Source/DisplayCluster/Private/Render/Viewport/Preview/DisplayClusterViewportManagerPreview.cpp
+++ b/Engine/Plugins/Runtime/nDisplay/Source/DisplayCluster/Private/Render/Viewport/Preview/DisplayClusterViewportManagerPreview.cpp
@@ -46,6 +46,12 @@ FString FDisplayClusterViewportManagerPreview::GetClusterNodeId() const
 			{
 				NodeIndex = 0;
 
+				// When rendering of all nodes in the cluster is complete, notify parent.
+				if (FDisplayClusterViewportManager* ViewportManager = Configuration->GetViewportManagerImpl())
+				{
+					ViewportManager->OnEntireClusterPreviewRendered();
+				}
+
 				// When rendering of all nodes in the cluster is complete, generate events for subscribers
 				OnEntireClusterPreviewGenerated.ExecuteIfBound();
 			}
diff --git a/Engine/Plugins/Runtime/nDisplay/Source/DisplayCluster/Private/Render/Viewport/RenderFrame/DisplayClusterRenderFrameManager.cpp b/Engine/Plugins/Runtime/nDisplay/Source/DisplayCluster/Private/Render/Viewport/RenderFrame/DisplayClusterRenderFrameManager.cpp
--- a/Engine/Plugins/Runtime/nDisplay/Source/DisplayCluster/Private/Render/Viewport/RenderFrame/DisplayClusterRenderFrameManager.cpp
+++ b/Engine/Plugins/Runtime/nDisplay/Source/DisplayCluster/Private/Render/Viewport/RenderFrame/DisplayClusterRenderFrameManager.cpp
@@ -6,7 +6,7 @@
 #include "Render/Viewport/Configuration/DisplayClusterViewportConfiguration.h"
 #include "Render/Viewport/DisplayClusterViewport.h"
 #include "Render/Viewport/DisplayClusterViewportManager.h"
-#include "Render/Upscaler/DisplayClusterUpscaler.h"
+#include "Render/Features/DisplayClusterUpscaler.h"
 
 #include "Misc/Parse.h"
 
diff --git a/Engine/Plugins/Runtime/nDisplay/Source/DisplayCluster/Private/Render/Viewport/RenderFrame/DisplayClusterRenderFrameSettings.cpp b/Engine/Plugins/Runtime/nDisplay/Source/DisplayCluster/Private/Render/Viewport/RenderFrame/DisplayClusterRenderFrameSettings.cpp
--- a/Engine/Plugins/Runtime/nDisplay/Source/DisplayCluster/Private/Render/Viewport/RenderFrame/DisplayClusterRenderFrameSettings.cpp
+++ b/Engine/Plugins/Runtime/nDisplay/Source/DisplayCluster/Private/Render/Viewport/RenderFrame/DisplayClusterRenderFrameSettings.cpp
@@ -164,6 +164,24 @@ bool FDisplayClusterRenderFrameSettings::IsPreviewRendering() const
 	return false;
 }
 
+/** Returns true if this DCRA is using cluster rendering. */
+bool FDisplayClusterRenderFrameSettings::IsClusterRendering() const
+{
+	switch (RenderMode)
+	{
+	case EDisplayClusterRenderFrameMode::Mono:
+	case EDisplayClusterRenderFrameMode::Stereo:
+	case EDisplayClusterRenderFrameMode::SideBySide:
+	case EDisplayClusterRenderFrameMode::TopBottom:
+		return true;
+
+	default:
+		break;
+	}
+
+	return false;
+}
+
 bool FDisplayClusterRenderFrameSettings::IsTechvisEnabled() const
 {
 	// Don't use Techvis to render ProxyHit
diff --git a/Engine/Plugins/Runtime/nDisplay/Source/DisplayCluster/Private/Render/Viewport/RenderFrame/DisplayClusterRenderFrameSettings.h b/Engine/Plugins/Runtime/nDisplay/Source/DisplayCluster/Private/Render/Viewport/RenderFrame/DisplayClusterRenderFrameSettings.h
--- a/Engine/Plugins/Runtime/nDisplay/Source/DisplayCluster/Private/Render/Viewport/RenderFrame/DisplayClusterRenderFrameSettings.h
+++ b/Engine/Plugins/Runtime/nDisplay/Source/DisplayCluster/Private/Render/Viewport/RenderFrame/DisplayClusterRenderFrameSettings.h
@@ -108,6 +108,9 @@ public:
 	/** Current frame is preview. */
 	bool IsPreviewRendering() const;
 
+	/** Returns true if this DCRA is using cluster rendering. */
+	bool IsClusterRendering() const;
+
 	/** Returns true, if Techvis is used. */
 	bool IsTechvisEnabled() const;
 
diff --git a/Engine/Plugins/Runtime/nDisplay/Source/DisplayCluster/Public/DisplayClusterRootActor.h b/Engine/Plugins/Runtime/nDisplay/Source/DisplayCluster/Public/DisplayClusterRootActor.h
--- a/Engine/Plugins/Runtime/nDisplay/Source/DisplayCluster/Public/DisplayClusterRootActor.h
+++ b/Engine/Plugins/Runtime/nDisplay/Source/DisplayCluster/Public/DisplayClusterRootActor.h
@@ -309,7 +309,7 @@ private:
 
 	UPROPERTY(EditInstanceOnly, Transient, Category = Viewports, meta = (DisplayName = "Viewport Overscan", PropertyPath = "CurrentConfigData.Cluster.Nodes.Viewports.RenderSettings.Overscan", ToolTip = "Render a larger frame than specified in the configuration to achieve continuity across displays when using post-processing effects."))
 	FDisplayClusterEditorPropertyReference ViewportOverscanRef;
-	
+
 	UPROPERTY(EditAnywhere, Transient, Category = Viewports, meta = (PropertyPath = "CurrentConfigData.StageSettings.bFreezeRenderOuterViewports"))
 	FDisplayClusterEditorPropertyReference FreezeRenderOuterViewportsRef;
 
@@ -349,6 +349,15 @@ private:
 	UPROPERTY(EditAnywhere, Transient, Category = OCIO, meta = (PropertyPath = "CurrentConfigData.StageSettings.ViewportOCIO.PerViewportOCIOProfiles", DisplayName = "Per-Viewport OCIO Overrides"))
 	FDisplayClusterEditorPropertyReference PerViewportOCIOProfilesRef;
 
+	UPROPERTY(EditAnywhere, Transient, Category = "Exposure", meta = (PropertyPath = "CurrentConfigData.StageSettings.AutoExposureSettings.bEnabled"))
+	FDisplayClusterEditorPropertyReference AutoExposureEnabledRef;
+
+	UPROPERTY(EditAnywhere, Transient, Category = "Exposure", meta = (PropertyPath = "CurrentConfigData.StageSettings.AutoExposureSettings.MeteringSettings"))
+	FDisplayClusterEditorPropertyReference AutoExposureMeteringSettingsRef;
+
+	UPROPERTY(EditAnywhere, Transient, Category = "Exposure", meta = (PropertyPath = "CurrentConfigData.StageSettings.AutoExposureSettings.ExcludeSettings"))
+	FDisplayClusterEditorPropertyReference AutoExposureExcludeSettingsRef;
+
 	UPROPERTY(EditAnywhere, Transient, Category = Chromakey, meta = (PropertyPath = "CurrentConfigData.StageSettings.GlobalChromakey.ChromakeyColor"))
 	FDisplayClusterEditorPropertyReference GlobalChromakeyColorRef;
 
diff --git a/Engine/Plugins/Runtime/nDisplay/Source/DisplayCluster/Public/Render/Viewport/Containers/DisplayClusterViewport_Context.h b/Engine/Plugins/Runtime/nDisplay/Source/DisplayCluster/Public/Render/Viewport/Containers/DisplayClusterViewport_Context.h
--- a/Engine/Plugins/Runtime/nDisplay/Source/DisplayCluster/Public/Render/Viewport/Containers/DisplayClusterViewport_Context.h
+++ b/Engine/Plugins/Runtime/nDisplay/Source/DisplayCluster/Public/Render/Viewport/Containers/DisplayClusterViewport_Context.h
@@ -7,6 +7,7 @@
 #include "ShowFlags.h"
 
 #include "Render/Viewport/Containers/DisplayClusterViewport_Enums.h"
+#include "Render/Viewport/Containers/DisplayClusterViewport_EyeAdaptation.h"
 
 /**
  * Viewport context with cahched data and states
@@ -149,6 +150,9 @@ public:
 	// GPU index for this context render target
 	int32 GPUIndex = INDEX_NONE;
 
+	// Inner rectangle inside RenderTargetRect
+	FIntRect OverscanInnerRenderTargetRect;
+
 	// Location and size on a render target texture
 	FIntRect RenderTargetRect;
 
@@ -184,6 +188,12 @@ public:
 
 		// Engine flags used to render this context
 		FEngineShowFlags EngineShowFlags;
+
+		// EyeAdaptation data
+		FDisplayClusterViewport_EyeAdaptationData EyeAdaptationData;
+
+		// The value from FSceneViewFamily::FrameNumber
+		uint32 FrameNumber = 0;
 	};
 
 	// This data updated only on rendering thread
diff --git a/Engine/Plugins/Runtime/nDisplay/Source/DisplayCluster/Public/Render/Viewport/Containers/DisplayClusterViewport_Enums.h b/Engine/Plugins/Runtime/nDisplay/Source/DisplayCluster/Public/Render/Viewport/Containers/DisplayClusterViewport_Enums.h
--- a/Engine/Plugins/Runtime/nDisplay/Source/DisplayCluster/Public/Render/Viewport/Containers/DisplayClusterViewport_Enums.h
+++ b/Engine/Plugins/Runtime/nDisplay/Source/DisplayCluster/Public/Render/Viewport/Containers/DisplayClusterViewport_Enums.h
@@ -235,3 +235,33 @@ enum class EDisplayClusterViewportRenderingFlags : uint8
 	StereoRendering = 1 << 0,
 };
 ENUM_CLASS_FLAGS(EDisplayClusterViewportRenderingFlags);
+
+/**
+ * Viewport usage type flags.
+ */
+enum class EDisplayClusterViewportFlags : uint8
+{
+	/** The standard viewport does not have any special flags. */
+	None = 0,
+
+	/** Tile viewport. */
+	TileViewport = 1 << 0,
+
+	/** ICVFX viewports. */
+	ICVFX_InnerFrustum = 1 << 1,
+	ICVFX_ChromaKey    = 1 << 2,
+	ICVFX_LightCard    = 1 << 3,
+	ICVFX_UVLightCard  = 1 << 4,
+
+	/** ICVFX render modifiers. */
+	ICVFX_Modifier_RenderOverInnerFrustum  = 1 << 5,
+	ICVFX_Modifier_RenderUnderInnerFrustum = 1 << 6,
+
+	/** Combined mask for any ICVFX viewport role (excludes modifiers). */
+	ICVFX_AnyRole =
+	  ICVFX_InnerFrustum
+	| ICVFX_ChromaKey
+	| ICVFX_LightCard
+	| ICVFX_UVLightCard,
+};
+ENUM_CLASS_FLAGS(EDisplayClusterViewportFlags);
diff --git a/Engine/Plugins/Runtime/nDisplay/Source/DisplayCluster/Public/Render/Viewport/Containers/DisplayClusterViewport_EyeAdaptation.h b/Engine/Plugins/Runtime/nDisplay/Source/DisplayCluster/Public/Render/Viewport/Containers/DisplayClusterViewport_EyeAdaptation.h
new file mode 100644
--- /dev/null
+++ b/Engine/Plugins/Runtime/nDisplay/Source/DisplayCluster/Public/Render/Viewport/Containers/DisplayClusterViewport_EyeAdaptation.h
@@ -0,0 +1,23 @@
+// Copyright Epic Games, Inc. All Rights Reserved.
+
+#pragma once
+
+#include "CoreMinimal.h"
+
+/**
+ * Eye Adaptation data for nDisplay viewport.
+ */
+struct FDisplayClusterViewport_EyeAdaptationData
+{
+	/** Return true if this data can be used. */
+	inline bool CanBeUsed() const
+	{
+		return FMath::IsFinite(LastAverageSceneLuminance)
+			&& LastAverageSceneLuminance > 0.0f;
+	}
+
+	// The last valid average scene luminance for eye adaptation (exposure compensation curve).
+	// Contains the value from FSceneView::GetLastAverageSceneLuminance().
+	// The value 0.0f means that there is no luminance data available.
+	float LastAverageSceneLuminance = 0.0f;
+};
diff --git a/Engine/Plugins/Runtime/nDisplay/Source/DisplayCluster/Public/Render/Viewport/IDisplayClusterViewport.h b/Engine/Plugins/Runtime/nDisplay/Source/DisplayCluster/Public/Render/Viewport/IDisplayClusterViewport.h
--- a/Engine/Plugins/Runtime/nDisplay/Source/DisplayCluster/Public/Render/Viewport/IDisplayClusterViewport.h
+++ b/Engine/Plugins/Runtime/nDisplay/Source/DisplayCluster/Public/Render/Viewport/IDisplayClusterViewport.h
@@ -39,7 +39,20 @@ public:
 	virtual TSharedRef<class FDisplayClusterViewport, ESPMode::ThreadSafe> ToSharedRef() = 0;
 	virtual TSharedRef<const class FDisplayClusterViewport, ESPMode::ThreadSafe> ToSharedRef() const = 0;
 
+	/** Returns the usage flags of the viewport. */
+	virtual EDisplayClusterViewportFlags GetViewportFlags() const = 0;
+
+	/** Returns the base viewport identifier (not guaranteed to be unique across ICVFX viewports).
+	* Note: For ICVFX viewports returns:
+	*       - The name of the ICVFX camera for InCamera/Chromakey.
+	*       - The name of the Outer viewport for Lightcards.
+	*/
+	virtual FString GetBaseId() const = 0;
+
+	/** Returns the unique viewport identifier, guaranteed to be unique across the entire cluster. */
 	virtual FString GetId() const = 0;
+
+	/** Returns the identifier of the cluster node that owns this viewport. */
 	virtual FString GetClusterNodeId() const = 0;
 
 	/** Get viewport manager configuration interface. */
diff --git a/Engine/Plugins/Runtime/nDisplay/Source/DisplayCluster/Public/Render/Viewport/IDisplayClusterViewportConfiguration.h b/Engine/Plugins/Runtime/nDisplay/Source/DisplayCluster/Public/Render/Viewport/IDisplayClusterViewportConfiguration.h
--- a/Engine/Plugins/Runtime/nDisplay/Source/DisplayCluster/Public/Render/Viewport/IDisplayClusterViewportConfiguration.h
+++ b/Engine/Plugins/Runtime/nDisplay/Source/DisplayCluster/Public/Render/Viewport/IDisplayClusterViewportConfiguration.h
@@ -136,6 +136,9 @@ public:
 	/** Returns true if preview rendering mode is used. */
 	virtual bool IsPreviewRendering() const = 0;
 
+	/** Returns true if this DCRA is using cluster rendering. */
+	virtual bool IsClusterRendering() const = 0;
+
 	/** Returns true, if Techvis is used. */
 	virtual bool IsTechvisEnabled() const = 0;
 
diff --git a/Engine/Plugins/Runtime/nDisplay/Source/DisplayCluster/Public/Render/Viewport/IDisplayClusterViewportProxy.h b/Engine/Plugins/Runtime/nDisplay/Source/DisplayCluster/Public/Render/Viewport/IDisplayClusterViewportProxy.h
--- a/Engine/Plugins/Runtime/nDisplay/Source/DisplayCluster/Public/Render/Viewport/IDisplayClusterViewportProxy.h
+++ b/Engine/Plugins/Runtime/nDisplay/Source/DisplayCluster/Public/Render/Viewport/IDisplayClusterViewportProxy.h
@@ -32,10 +32,20 @@ public:
 	virtual TSharedPtr<IDisplayClusterViewportProxy, ESPMode::ThreadSafe> ToSharedPtr() = 0;
 	virtual TSharedPtr<const IDisplayClusterViewportProxy, ESPMode::ThreadSafe> ToSharedPtr() const = 0;
 
-	/** Get the unique viewport name. */
+	/** Returns the usage flags of the viewport. */
+	virtual EDisplayClusterViewportFlags GetViewportFlags() const = 0;
+
+	/** Returns the base viewport identifier (not guaranteed to be unique across ICVFX viewports).
+	 * Note: For ICVFX viewports returns:
+	 *       - The name of the ICVFX camera for InCamera/Chromakey.
+	 *       - The name of the Outer viewport for Lightcards.
+	 */
+	virtual FString GetBaseId() const = 0;
+
+	/** Returns the unique viewport identifier, guaranteed to be unique across the entire cluster. */
 	virtual FString GetId() const = 0;
 
-	/** Get the name of the cluster node to which this viewport belongs. */
+	/** Returns the identifier of the cluster node that owns this viewport. */
 	virtual FString GetClusterNodeId() const = 0;
 
 	/** Get viewport manager configuration interface. */
diff --git a/Engine/Plugins/Runtime/nDisplay/Source/DisplayClusterConfiguration/Public/DisplayClusterConfigurationStrings.h b/Engine/Plugins/Runtime/nDisplay/Source/DisplayClusterConfiguration/Public/DisplayClusterConfigurationStrings.h
--- a/Engine/Plugins/Runtime/nDisplay/Source/DisplayClusterConfiguration/Public/DisplayClusterConfigurationStrings.h
+++ b/Engine/Plugins/Runtime/nDisplay/Source/DisplayClusterConfiguration/Public/DisplayClusterConfigurationStrings.h
@@ -29,6 +29,7 @@ namespace DisplayClusterConfigurationStrings
 		static constexpr const TCHAR* ChromaKeyCategory            = TEXT("Chromakey");
 		static constexpr const TCHAR* LightcardCategory            = TEXT("Light Cards");
 		static constexpr const TCHAR* OCIOCategory                 = TEXT("OCIO");
+		static constexpr const TCHAR* ExposureCategory             = TEXT("Exposure");
 		static constexpr const TCHAR* MediaCategory                = TEXT("Media");
 		static constexpr const TCHAR* TileCategory                 = TEXT("Tile Rendering");
 		static constexpr const TCHAR* OverrideCategory             = TEXT("Texture Replacement");
diff --git a/Engine/Plugins/Runtime/nDisplay/Source/DisplayClusterConfiguration/Public/DisplayClusterConfigurationTypes_AutoExposure.h b/Engine/Plugins/Runtime/nDisplay/Source/DisplayClusterConfiguration/Public/DisplayClusterConfigurationTypes_AutoExposure.h
new file mode 100644
--- /dev/null
+++ b/Engine/Plugins/Runtime/nDisplay/Source/DisplayClusterConfiguration/Public/DisplayClusterConfigurationTypes_AutoExposure.h
@@ -0,0 +1,82 @@
+// Copyright Epic Games, Inc. All Rights Reserved.
+
+#pragma once
+
+#include "CoreMinimal.h"
+
+#include "UObject/ObjectMacros.h"
+#include "UObject/Object.h"
+
+#include "DisplayClusterConfigurationTypes_Enums.h"
+
+#include "DisplayClusterConfigurationTypes_AutoExposure.generated.h"
+
+USTRUCT(Blueprintable)
+struct DISPLAYCLUSTERCONFIGURATION_API FDisplayClusterConfiguration_AutoExposureMeteringSettings
+{
+	GENERATED_BODY()
+
+public:
+	/**
+	 * List of allowed outer viewports for auto-exposure metering when the override is enabled.
+	 * If no viewports are specified, all outer viewports are used for metering.
+	 */
+	UPROPERTY(EditAnywhere, BlueprintReadWrite, Category = "Exposure", meta = (DisplayName = "Viewports"))
+	TArray<FString> Viewports;
+
+	/**
+	 * List of allowed ICVFX cameras for auto-exposure metering when the override is enabled.
+	 * If no ICVFX cameras are specified, all ICVFX cameras are used for metering.
+	 */
+	UPROPERTY(EditAnywhere, BlueprintReadWrite, Category = "Exposure", meta = (DisplayName = "ICVFX Cameras"))
+	TArray<FString> ICVFXCameras;
+
+	/**
+	 * List of allowed cluster nodes for auto-exposure metering when the override is enabled.
+	 * If no nodes are specified, all cluster nodes can be used for metering.
+	 */
+	UPROPERTY(EditAnywhere, BlueprintReadWrite, Category = "Exposure", meta = (DisplayName = "Nodes"))
+	TArray<FString> Nodes;
+};
+
+USTRUCT(Blueprintable)
+struct DISPLAYCLUSTERCONFIGURATION_API FDisplayClusterConfiguration_AutoExposureExcludeSettings
+{
+	GENERATED_BODY()
+
+public:
+	/** Outer viewports excluded from nDisplay auto-exposure (they use the default exposure behavior). */
+	UPROPERTY(EditAnywhere, BlueprintReadWrite, Category = "Exposure", meta = (DisplayName = "Viewports"))
+	TArray<FString> Viewports;
+
+	/** ICVFX cameras excluded from nDisplay auto-exposure (they use the default exposure behavior). */
+	UPROPERTY(EditAnywhere, BlueprintReadWrite, Category = "Exposure", meta = (DisplayName = "ICVFX Cameras"))
+	TArray<FString> ICVFXCameras;
+
+	/** Cluster nodes excluded from nDisplay auto-exposure (they use the default exposure behavior). */
+	UPROPERTY(EditAnywhere, BlueprintReadWrite, Category = "Exposure", meta = (DisplayName = "Nodes"))
+	TArray<FString> Nodes;
+};
+
+USTRUCT(Blueprintable)
+struct DISPLAYCLUSTERCONFIGURATION_API FDisplayClusterConfiguration_AutoExposureSettings
+{
+	GENERATED_BODY()
+
+public:
+	/** Enable/disable AutoExposure override. */
+	UPROPERTY(EditAnywhere, BlueprintReadWrite, Category = "Exposure", meta = (DisplayName = "Enable Auto Exposure Override"))
+	bool bEnabled = false;
+
+	/**
+	 * Metering settings for auto-exposure metering when the override is enabled.
+	 */
+	UPROPERTY(EditAnywhere, BlueprintReadWrite, Category = "Exposure", meta = (DisplayName = "Meter Auto Exposure on Items"))
+	FDisplayClusterConfiguration_AutoExposureMeteringSettings MeteringSettings;
+
+	/**
+	 * Exclude settings for auto-exposure.
+	 */
+	UPROPERTY(EditAnywhere, BlueprintReadWrite, Category = "Exposure", meta = (DisplayName = "Exclude Items from Auto Exposure"))
+	FDisplayClusterConfiguration_AutoExposureExcludeSettings ExcludeSettings;
+};
diff --git a/Engine/Plugins/Runtime/nDisplay/Source/DisplayClusterConfiguration/Public/DisplayClusterConfigurationTypes_ICVFX.h b/Engine/Plugins/Runtime/nDisplay/Source/DisplayClusterConfiguration/Public/DisplayClusterConfigurationTypes_ICVFX.h
--- a/Engine/Plugins/Runtime/nDisplay/Source/DisplayClusterConfiguration/Public/DisplayClusterConfigurationTypes_ICVFX.h
+++ b/Engine/Plugins/Runtime/nDisplay/Source/DisplayClusterConfiguration/Public/DisplayClusterConfigurationTypes_ICVFX.h
@@ -12,6 +12,7 @@
 
 #include "CineCameraActor.h"
 
+#include "DisplayClusterConfigurationTypes_AutoExposure.h"
 #include "DisplayClusterConfigurationTypes_Upscaler.h"
 #include "DisplayClusterConfigurationTypes_Media.h"
 #include "DisplayClusterConfigurationTypes_PostRender.h"
@@ -948,7 +949,6 @@ public:
 	UPROPERTY(EditAnywhere, BlueprintReadWrite, Category = Viewports, meta = (DisplayName = "Default Upscaler Settings"))
 	FDisplayClusterConfigurationUpscalerSettings OuterViewportUpscalerSettings;
 
-
 	// Hide list for all icvfx viewports (outer, inner, cameras, etc)
 	// (This allow to hide all actors from layers for icvfx render logic)
 	UPROPERTY(EditAnywhere, BlueprintReadWrite, Category = Viewports, meta = (DisplayName = "Content Hidden from Entire Cluster", ToolTip = "Content specified here will not appear anywhere in the nDisplay cluster", Substitutions = "LayersTooltip = Layers hidden from the entire nDisplay cluster, ActorsTooltip = Actors hidden from the entire nDisplay cluster"))
@@ -974,6 +974,10 @@ public:
 	UPROPERTY(EditAnywhere, BlueprintReadWrite, Category = "OCIO", meta = (DisplayName = "Outer Viewports OCIO"))
 	FDisplayClusterConfigurationICVFX_ViewportOCIO ViewportOCIO;
 
+	/** Entire cluster Auto Exposure settings. */
+	UPROPERTY(EditAnywhere, BlueprintReadWrite, Category = "Exposure", meta = (DisplayName = "Auto Exposure Settings"))
+	FDisplayClusterConfiguration_AutoExposureSettings AutoExposureSettings;
+
 	/** Global chromakey settings that all ICVFX camera components can opt to use */
 	UPROPERTY(EditAnywhere, BlueprintReadWrite, Category = "Chromakey")
 	FDisplayClusterConfigurationICVFX_GlobalChromakeySettings GlobalChromakey;
diff --git a/Engine/Plugins/Runtime/nDisplay/Source/DisplayClusterConfigurator/Private/DisplayClusterConfiguratorModule.cpp b/Engine/Plugins/Runtime/nDisplay/Source/DisplayClusterConfigurator/Private/DisplayClusterConfiguratorModule.cpp
--- a/Engine/Plugins/Runtime/nDisplay/Source/DisplayClusterConfigurator/Private/DisplayClusterConfiguratorModule.cpp
+++ b/Engine/Plugins/Runtime/nDisplay/Source/DisplayClusterConfigurator/Private/DisplayClusterConfiguratorModule.cpp
@@ -303,6 +303,11 @@ void FDisplayClusterConfiguratorModule::RegisterSectionMappings()
 			DisplayClusterConfigurationStrings::categories::OCIOCategory, LOCTEXT("OCIO", "OCIO"));
 		Section->AddCategory(DisplayClusterConfigurationStrings::categories::OCIOCategory);
 	}
+	{
+		const TSharedRef<FPropertySection> Section = PropertyModule.FindOrCreateSection(ADisplayClusterRootActor::StaticClass()->GetFName(),
+			DisplayClusterConfigurationStrings::categories::ExposureCategory, LOCTEXT("Exposure", "Exposure"));
+		Section->AddCategory(DisplayClusterConfigurationStrings::categories::ExposureCategory);
+	}
 	{
 		const TSharedRef<FPropertySection> Section = PropertyModule.FindOrCreateSection(ADisplayClusterRootActor::StaticClass()->GetFName(),
 			DisplayClusterConfigurationStrings::categories::LightcardCategory, LOCTEXT("Light Cards", "Light Cards"));
@@ -419,6 +424,7 @@ void FDisplayClusterConfiguratorModule::UnregisterSectionMappings()
 		PropertyModule.RemoveSection(ADisplayClusterRootActor::StaticClass()->GetFName(), DisplayClusterConfigurationStrings::categories::InCameraVFXCategory);
 		PropertyModule.RemoveSection(ADisplayClusterRootActor::StaticClass()->GetFName(), DisplayClusterConfigurationStrings::categories::ColorGradingCategory);
 		PropertyModule.RemoveSection(ADisplayClusterRootActor::StaticClass()->GetFName(), DisplayClusterConfigurationStrings::categories::OCIOCategory);
+		PropertyModule.RemoveSection(ADisplayClusterRootActor::StaticClass()->GetFName(), DisplayClusterConfigurationStrings::categories::ExposureCategory);
 		PropertyModule.RemoveSection(ADisplayClusterRootActor::StaticClass()->GetFName(), DisplayClusterConfigurationStrings::categories::LightcardCategory);
 		PropertyModule.RemoveSection(ADisplayClusterRootActor::StaticClass()->GetFName(), DisplayClusterConfigurationStrings::categories::PreviewCategory);
 
diff --git a/Engine/Plugins/Runtime/nDisplay/Source/DisplayClusterConfigurator/Private/Views/Details/DisplayClusterRootActorDetailsCustomization.cpp b/Engine/Plugins/Runtime/nDisplay/Source/DisplayClusterConfigurator/Private/Views/Details/DisplayClusterRootActorDetailsCustomization.cpp
--- a/Engine/Plugins/Runtime/nDisplay/Source/DisplayClusterConfigurator/Private/Views/Details/DisplayClusterRootActorDetailsCustomization.cpp
+++ b/Engine/Plugins/Runtime/nDisplay/Source/DisplayClusterConfigurator/Private/Views/Details/DisplayClusterRootActorDetailsCustomization.cpp
@@ -36,6 +36,7 @@ namespace DisplayClusterRootActorDetailsCustomizationUtils
 			DisplayClusterConfigurationStrings::categories::InCameraVFXCategory,
 			DisplayClusterConfigurationStrings::categories::ColorGradingCategory,
 			DisplayClusterConfigurationStrings::categories::OCIOCategory,
+			DisplayClusterConfigurationStrings::categories::ExposureCategory,
 			DisplayClusterConfigurationStrings::categories::ChromaKeyCategory,
 			DisplayClusterConfigurationStrings::categories::LightcardCategory,
 			DisplayClusterConfigurationStrings::categories::OverrideCategory,
