[nDisplay] Fixed multiple issues for Tile viewports

- Fixed a bug where Tile viewports were recreated every frame.

This bug led to:
1. ViewState was also recreated every frame, preventing temporary frame data from being preserved.
2. ViewportProxy missed required data in RenderSettings, which led to incorrect logic.
(see FDisplayClusterViewportConfiguration_ICVFX::ImplBeginReallocateViewports().)

- Fixed a bug where TileSource viewports were always split into Tiles, even in cases where splitting was not allowed by FDisplayClusterViewport::CanSplitIntoTiles().

This bug created unnecessary Tile viewports, which consumed additional resources and reduced performance.

Example: In preview rendering, the InCamera viewport is rendered only on a single cluster node.
On other nodes, InCamera viewports are not rendered. Instead, they use the image produced for that node, so they cannot be split.
Other scenarios follow the rules defined in CanSplitIntoTiles().

Reworked FDisplayClusterViewportConfiguration_Tile::ForEachTileViewport().

- Fixed a bug where FDisplayClusterViewportProxy::GetPriority_RenderThread() returned an incorrect value.
This caused TileSource viewports to be processed before the Tile viewport, resulting in a black texture.

Example timeline of the issue:
1. Initially, TileSource.RTT contains a black texture.
2. TileSource.UpdateDeferredResources() expects TileSource.RTT already contain the image produced by Tile viewports.
It then copies TileSource.RTT, directly or through OCIO, into TileSource.Input.
After this step, TileSource.RTT is no longer used in this code path.
3. Tile.UpdateDeferredResources() copies Tile.RTT into TileSource.RTT, so Tile must be processed before TileSource.
4. FDisplayClusterViewportManagerProxy::UpdateFrameResources_RenderThread() uses TileSource.Input for warp and ICVFX.

#Synced-CL 0

diff --git a/Engine/Plugins/Runtime/nDisplay/Source/DisplayCluster/Private/Render/Viewport/Configuration/DisplayClusterViewportConfiguration_ICVFX.cpp b/Engine/Plugins/Runtime/nDisplay/Source/DisplayCluster/Private/Render/Viewport/Configuration/DisplayClusterViewportConfiguration_ICVFX.cpp
--- a/Engine/Plugins/Runtime/nDisplay/Source/DisplayCluster/Private/Render/Viewport/Configuration/DisplayClusterViewportConfiguration_ICVFX.cpp
+++ b/Engine/Plugins/Runtime/nDisplay/Source/DisplayCluster/Private/Render/Viewport/Configuration/DisplayClusterViewportConfiguration_ICVFX.cpp
@@ -245,12 +245,17 @@ void FDisplayClusterViewportConfiguration_ICVFX::ImplBeginReallocateViewports()
 		for (const TSharedPtr<FDisplayClusterViewport, ESPMode::ThreadSafe>& ViewportIt : ViewportManager->ImplGetCurrentRenderFrameViewports())
 		{
 			// Runtime icvfx viewport support reallocate feature:
-			if (ViewportIt.IsValid() && EnumHasAllFlags(ViewportIt->GetRenderSettingsICVFX().RuntimeFlags, EDisplayClusterViewportRuntimeICVFXFlags::InternalResource))
+			if (ViewportIt.IsValid()
+				// Exclude tile viewports: they handled in FDisplayClusterViewportConfiguration_Tile::Update()
+				&& !EnumHasAnyFlags(ViewportIt->GetViewportFlags(), EDisplayClusterViewportFlags::TileViewport)
+				// Process only ICVFX viewports
+				&& EnumHasAnyFlags(ViewportIt->GetViewportFlags(), EDisplayClusterViewportFlags::ICVFX_AnyRole))
 			{
 				// Gain direct access to internal settings of the viewport:
 				FDisplayClusterViewport_RenderSettingsICVFX& InOutRenderSettingsICVFX = ViewportIt->GetRenderSettingsICVFXImpl();
 
-				// Mark all dynamic ICVFX viewports for delete
+				// Mark this ICVFX viewport as unused. This flag is cleared later if the viewport is actually used.
+				// Otherwise, this viewport will be removed in ImplFinishReallocateViewports().
 				EnumAddFlags(InOutRenderSettingsICVFX.RuntimeFlags, EDisplayClusterViewportRuntimeICVFXFlags::Unused);
 			}
 		}
diff --git a/Engine/Plugins/Runtime/nDisplay/Source/DisplayCluster/Private/Render/Viewport/Configuration/DisplayClusterViewportConfiguration_Tile.cpp b/Engine/Plugins/Runtime/nDisplay/Source/DisplayCluster/Private/Render/Viewport/Configuration/DisplayClusterViewportConfiguration_Tile.cpp
--- a/Engine/Plugins/Runtime/nDisplay/Source/DisplayCluster/Private/Render/Viewport/Configuration/DisplayClusterViewportConfiguration_Tile.cpp
+++ b/Engine/Plugins/Runtime/nDisplay/Source/DisplayCluster/Private/Render/Viewport/Configuration/DisplayClusterViewportConfiguration_Tile.cpp
@@ -16,34 +16,43 @@
 #include "Misc/DisplayClusterLog.h"
 #include "Misc/Parse.h"
 
-///////////////////////////////////////////////////////////////////
-// FDisplayClusterViewportConfiguration_Tile
-///////////////////////////////////////////////////////////////////
+/**
+ * Iterate over tile source viewports that are ready to split.
+ *
+ * @param Configuration Viewport configuration.
+ * @param Pred          Callable invoked as Pred(Viewport, TilePosition) for each matching viewport.
+ */
 template <typename Predicate>
-void FDisplayClusterViewportConfiguration_Tile::ForEachTileViewport(Predicate Pred) const
+void ForEachReadyToSplitTileSourceViewport(const FDisplayClusterViewportConfiguration& Configuration, Predicate Pred)
 {
 	if (FDisplayClusterViewportManager* ViewportManager = Configuration.GetViewportManagerImpl())
-{
-	for (const TSharedPtr<FDisplayClusterViewport, ESPMode::ThreadSafe>& ViewportIt : ViewportManager->ImplGetCurrentRenderFrameViewports())
 	{
-		if (ViewportIt.IsValid() && ViewportIt->GetRenderSettings().TileSettings.GetType() == EDisplayClusterViewportTileType::Source)
+		for (const TSharedPtr<FDisplayClusterViewport, ESPMode::ThreadSafe>& ViewportIt : ViewportManager->ImplGetCurrentRenderFrameViewports())
 		{
-			const FIntPoint& TileSize = ViewportIt->GetRenderSettings().TileSettings.GetSize();
+			if (ViewportIt.IsValid() && ViewportIt->GetRenderSettings().TileSettings.GetType() == EDisplayClusterViewportTileType::Source)
+			{
+				// Check if the viewport can be split into tiles.
+				if (!ViewportIt->CanSplitIntoTiles())
+				{
+					continue;
+				}
+
+				const FIntPoint& TileSize = ViewportIt->GetRenderSettings().TileSettings.GetSize();
 
 				// Iterate over all tile viewports
-			for (int32 PosX = 0; PosX < TileSize.X; PosX++)
-			{
-				for (int32 PosY = 0; PosY < TileSize.Y; PosY++)
+				for (int32 PosX = 0; PosX < TileSize.X; PosX++)
 				{
+					for (int32 PosY = 0; PosY < TileSize.Y; PosY++)
+					{
 						::Invoke(Pred, *ViewportIt, FIntPoint(PosX, PosY));
+					}
 				}
 			}
 		}
 	}
-	}
 	else
 	{
-		UE_LOG(LogDisplayClusterViewport, Warning, TEXT("Tile:ForEachTileViewport - ViewportManager not found. "));
+		UE_LOG(LogDisplayClusterViewport, Warning, TEXT("ForEachReadyToSplitTileSourceViewport() - ViewportManager not found. "));
 	}
 }
 
@@ -51,7 +60,7 @@ void FDisplayClusterViewportConfiguration_Tile::Update()
 {
 	ImplBeginReallocateViewports();
 
-	ForEachTileViewport([](FDisplayClusterViewport& InSourceViewport, const FIntPoint& InTilePos)
+	ForEachReadyToSplitTileSourceViewport(Configuration, [](FDisplayClusterViewport& InSourceViewport, const FIntPoint& InTilePos)
 	{
 		// Split the source viewports into multiple tiles.
 		FDisplayClusterViewportConfigurationHelpers_Tile::GetOrCreateTileViewport(InSourceViewport, InTilePos);
diff --git a/Engine/Plugins/Runtime/nDisplay/Source/DisplayCluster/Private/Render/Viewport/Configuration/DisplayClusterViewportConfiguration_Tile.h b/Engine/Plugins/Runtime/nDisplay/Source/DisplayCluster/Private/Render/Viewport/Configuration/DisplayClusterViewportConfiguration_Tile.h
--- a/Engine/Plugins/Runtime/nDisplay/Source/DisplayCluster/Private/Render/Viewport/Configuration/DisplayClusterViewportConfiguration_Tile.h
+++ b/Engine/Plugins/Runtime/nDisplay/Source/DisplayCluster/Private/Render/Viewport/Configuration/DisplayClusterViewportConfiguration_Tile.h
@@ -30,14 +30,6 @@ private:
 	/** Delete unused Tile viewports. */
 	void ImplFinishReallocateViewports() const;
 
-	/**
-	 * Iterate over all tile viewports
-	 *
-	 * @param Pred The predicate to match.
-	 */
-	template <typename Predicate>
-	void ForEachTileViewport(Predicate Pred) const;
-
 private:
 	// Viewport configuration API
 	const FDisplayClusterViewportConfiguration& Configuration;
diff --git a/Engine/Plugins/Runtime/nDisplay/Source/DisplayCluster/Private/Render/Viewport/DisplayClusterViewportProxy_Rules.cpp b/Engine/Plugins/Runtime/nDisplay/Source/DisplayCluster/Private/Render/Viewport/DisplayClusterViewportProxy_Rules.cpp
--- a/Engine/Plugins/Runtime/nDisplay/Source/DisplayCluster/Private/Render/Viewport/DisplayClusterViewportProxy_Rules.cpp
+++ b/Engine/Plugins/Runtime/nDisplay/Source/DisplayCluster/Private/Render/Viewport/DisplayClusterViewportProxy_Rules.cpp
@@ -334,6 +334,7 @@ uint8 FDisplayClusterViewportProxy::GetPriority_RenderThread() const
 	case EDisplayClusterViewportTileType::Source:
 		OutOrder += (uint8)EDisplayClusterViewportProxyPriority::TileSource;
 		break;
+
 	case EDisplayClusterViewportTileType::Tile:
 	case EDisplayClusterViewportTileType::UnusedTile:
 		OutOrder += (uint8)EDisplayClusterViewportProxyPriority::Tile;
@@ -347,6 +348,13 @@ uint8 FDisplayClusterViewportProxy::GetPriority_RenderThread() const
 		break;
 	}
 
+	// Internal tile viewports must be processed first.
+	// Additional rules does not apply to them.
+	if (RenderSettings.TileSettings.IsInternalViewport())
+	{
+		return OutOrder;
+	}
+
 	if (RenderSettings.IsViewportHasParent())
 	{
 		OutOrder += (uint8)EDisplayClusterViewportProxyPriority::Linked;
