[nDisplay] Fix preview-rendering crash caused by destroyed ViewState.

FDisplayClusterViewportPreview stores only a weak reference to its owning FDisplayClusterViewport.
Because of this, preview rendering can continue after the owning viewport has started or completed teardown.

In FDisplayClusterViewportPreview::CalcSceneView(), the preview retrieves ViewState from the viewport through InViewport->GetViewState(InContextNum).
That ViewState is owned by FDisplayClusterViewport, so it may be destroyed when the viewport releases its shared reference.
This could leave preview rendering with a dangling reference and cause a crash.

Fix:
Preserve shared ViewState references for the full preview-rendering lifetime.

Preview rendering runs outside DisplayClusterViewport and DisplayClusterViewportProxy lifecycle management, so referenced objects such as ViewState are not retained automatically through the normal ownership path.
This change keeps the required ViewState references alive in FDisplayClusterViewportPreview and preserves them in FDisplayClusterViewportManagerPreview::RenderViewportPreview() until render-thread processing is complete.

IDisplayClusterViewportPreview: Deprecated function GetViewport(). Use GetOwningViewport() instead.

Added missing DISPLAYCLUSTER_API macros to IDisplayClusterViewportManagerPreview and IDisplayClusterViewportPreview.

Added a new API to IDisplayClusterViewport:
TSharedPtr<FSceneViewStateReference, ESPMode::ThreadSafe> GetOrCreateViewState(uint32 InContextNum)
This function returns the ViewState for the specified viewport context, creating it if necessary.

#Synced-CL 0

diff --git a/Engine/Plugins/Runtime/nDisplay/Source/DisplayCluster/Private/Render/Device/DisplayClusterDeviceBase.cpp b/Engine/Plugins/Runtime/nDisplay/Source/DisplayCluster/Private/Render/Device/DisplayClusterDeviceBase.cpp
--- a/Engine/Plugins/Runtime/nDisplay/Source/DisplayCluster/Private/Render/Device/DisplayClusterDeviceBase.cpp
+++ b/Engine/Plugins/Runtime/nDisplay/Source/DisplayCluster/Private/Render/Device/DisplayClusterDeviceBase.cpp
@@ -318,7 +318,8 @@ FSceneViewStateInterface* FDisplayClusterDeviceBase::GetExternalViewState(const
 		uint32 ViewportContextNum = 0;
 		if (FDisplayClusterViewport* ViewportPtr = ViewportManager->ImplFindViewport(StereoViewIndex, &ViewportContextNum))
 		{
-			return ViewportPtr->GetViewState(ViewportContextNum);
+			TSharedPtr<FSceneViewStateReference, ESPMode::ThreadSafe> ViewState = ViewportPtr->GetOrCreateViewState(ViewportContextNum);
+			return ViewState.IsValid() ? ViewState->GetReference() : nullptr;
 		}
 
 		UE_LOG(LogDisplayClusterRender, Warning, TEXT("Viewport StereoViewIndex='%i' not found"), StereoViewIndex);
diff --git a/Engine/Plugins/Runtime/nDisplay/Source/DisplayCluster/Private/Render/Viewport/DisplayClusterViewport.h b/Engine/Plugins/Runtime/nDisplay/Source/DisplayCluster/Private/Render/Viewport/DisplayClusterViewport.h
--- a/Engine/Plugins/Runtime/nDisplay/Source/DisplayCluster/Private/Render/Viewport/DisplayClusterViewport.h
+++ b/Engine/Plugins/Runtime/nDisplay/Source/DisplayCluster/Private/Render/Viewport/DisplayClusterViewport.h
@@ -185,6 +185,8 @@ public:
 		return Contexts;
 	}
 
+	virtual TSharedPtr<FSceneViewStateReference, ESPMode::ThreadSafe> GetOrCreateViewState(uint32 InContextNum) override;
+
 	virtual const IDisplayClusterViewport_CustomPostProcessSettings& GetViewport_CustomPostProcessSettings() const override
 	{
 		check(IsInGameThread());
@@ -445,9 +447,6 @@ public:
 	/** Sets viewport configuration data. */
 	void SetViewportConfigurationData(const UDisplayClusterConfigurationViewport* InConfigurationData);
 
-	/** Support view states for preview. */
-	class FSceneViewStateInterface* GetViewState(uint32 ViewIndex);
-
 	/** Cleanup view states. */
 	void CleanupViewState();
 
diff --git a/Engine/Plugins/Runtime/nDisplay/Source/DisplayCluster/Private/Render/Viewport/DisplayClusterViewport_EditorPreview.cpp b/Engine/Plugins/Runtime/nDisplay/Source/DisplayCluster/Private/Render/Viewport/DisplayClusterViewport_EditorPreview.cpp
--- a/Engine/Plugins/Runtime/nDisplay/Source/DisplayCluster/Private/Render/Viewport/DisplayClusterViewport_EditorPreview.cpp
+++ b/Engine/Plugins/Runtime/nDisplay/Source/DisplayCluster/Private/Render/Viewport/DisplayClusterViewport_EditorPreview.cpp
@@ -60,7 +60,7 @@ void FDisplayClusterViewport::CleanupViewState()
 	}
 }
 
-FSceneViewStateInterface* FDisplayClusterViewport::GetViewState(uint32 ViewIndex)
+TSharedPtr<FSceneViewStateReference, ESPMode::ThreadSafe> FDisplayClusterViewport::GetOrCreateViewState(uint32 ViewIndex)
 {
 	if(!ShouldUseViewStates())
 	{
@@ -70,6 +70,13 @@ FSceneViewStateInterface* FDisplayClusterViewport::GetViewState(uint32 ViewIndex
 		return nullptr;
 	}
 
+	// Do not create or access ViewState for ViewIndex outside the configured range.
+	const uint32 MaxViewIndex = FMath::Max(0, Configuration->GetRenderFrameSettings().GetViewPerViewportAmount());
+	if (ViewIndex >= MaxViewIndex)
+	{
+		return nullptr;
+	}
+
 	int32 RequiredAmount = (int32)ViewIndex - ViewStates.Num() + 1;
 	if (RequiredAmount > 0)
 	{
@@ -89,5 +96,5 @@ FSceneViewStateInterface* FDisplayClusterViewport::GetViewState(uint32 ViewIndex
 		ViewStates[ViewIndex]->Allocate(FeatureLevel);
 	}
 
-	return ViewStates[ViewIndex]->GetReference();
+	return ViewStates[ViewIndex];
 }
\ No newline at end of file
diff --git a/Engine/Plugins/Runtime/nDisplay/Source/DisplayCluster/Private/Render/Viewport/Preview/DisplayClusterViewportManagerPreview_Render.cpp b/Engine/Plugins/Runtime/nDisplay/Source/DisplayCluster/Private/Render/Viewport/Preview/DisplayClusterViewportManagerPreview_Render.cpp
--- a/Engine/Plugins/Runtime/nDisplay/Source/DisplayCluster/Private/Render/Viewport/Preview/DisplayClusterViewportManagerPreview_Render.cpp
+++ b/Engine/Plugins/Runtime/nDisplay/Source/DisplayCluster/Private/Render/Viewport/Preview/DisplayClusterViewportManagerPreview_Render.cpp
@@ -200,6 +200,10 @@ bool FDisplayClusterViewportManagerPreview::RenderViewportPreview(FViewport* InV
 				}
 			}
 
+			// ViewState references that must stay alive until rendering is complete.
+			// They are passed to the render thread to prevent premature destruction.
+			TArray<TSharedPtr<FSceneViewStateReference, ESPMode::ThreadSafe>> ViewStatesToKeepAlive;
+
 			TArray<FSceneView*> Views;
 			for(const FDisplayClusterRenderFrameTargetView& ViewIt : ViewFamiliesIt.Views)
 			{
@@ -213,7 +217,8 @@ bool FDisplayClusterViewportManagerPreview::RenderViewportPreview(FViewport* InV
 
 				// Always call CalcSceneView() because this function also starts VE and after that calls FDisplayClusterViewport::SetupSceneView() -> OCIO
 				// uvLC viewport is not rendered, but late OCIO is used in FDisplayClusterViewportProxy::ApplyOCIO_RenderThread()
-				FSceneView* View = ViewportPtr->ViewportPreview->CalcSceneView(ViewFamily, ViewIt.ContextNum);
+				TSharedPtr<FSceneViewStateReference, ESPMode::ThreadSafe> ReferencedViewState;
+				FSceneView* View = ViewportPtr->ViewportPreview->CalcSceneView(ViewFamily, ViewIt.ContextNum, ReferencedViewState);
 
 				// Remove the view that is not used for rendering.
 				if (View && (!ViewIt.IsViewportContextCanBeRendered() || ViewFamily.RenderTarget == nullptr))
@@ -227,6 +232,12 @@ bool FDisplayClusterViewportManagerPreview::RenderViewportPreview(FViewport* InV
 				if (View)
 				{
 					Views.Add(View);
+
+					// Keep the ViewState alive until rendering is complete.
+					if (ReferencedViewState.IsValid())
+					{
+						ViewStatesToKeepAlive.Add(ReferencedViewState);
+					}
 				}
 			}
 
@@ -248,6 +259,10 @@ bool FDisplayClusterViewportManagerPreview::RenderViewportPreview(FViewport* InV
 					GetRendererModule().BeginRenderingViewFamily(&Canvas, &ViewFamily);
 				}
 
+				// Pass ViewState references to the render thread to keep them alive until rendering is complete.
+				ENQUEUE_RENDER_COMMAND(FDisplayClusterViewportManagerPreview_KeepViewStatesAlive)(
+					[ViewStatesToKeepAlive = MoveTemp(ViewStatesToKeepAlive)](FRHICommandListImmediate& RHICmdList) { });
+
 				if (GNumExplicitGPUsForRendering > 1)
 				{
 					const FRHIGPUMask SubmitGPUMask = ViewFamily.Views.Num() == 1 ? ViewFamily.Views[0]->GPUMask : FRHIGPUMask::All();
diff --git a/Engine/Plugins/Runtime/nDisplay/Source/DisplayCluster/Private/Render/Viewport/Preview/DisplayClusterViewportPreview.cpp b/Engine/Plugins/Runtime/nDisplay/Source/DisplayCluster/Private/Render/Viewport/Preview/DisplayClusterViewportPreview.cpp
--- a/Engine/Plugins/Runtime/nDisplay/Source/DisplayCluster/Private/Render/Viewport/Preview/DisplayClusterViewportPreview.cpp
+++ b/Engine/Plugins/Runtime/nDisplay/Source/DisplayCluster/Private/Render/Viewport/Preview/DisplayClusterViewportPreview.cpp
@@ -27,9 +27,14 @@ FDisplayClusterViewportPreview::~FDisplayClusterViewportPreview()
 	Release();
 }
 
+TSharedPtr<IDisplayClusterViewport, ESPMode::ThreadSafe> FDisplayClusterViewportPreview::GetOwningViewport() const
+{
+	return OwningViewportWeakPtr.Pin();
+}
+
 void FDisplayClusterViewportPreview::Initialize(FDisplayClusterViewport& InViewport)
 {
-	ViewportWeakPtr = InViewport.AsShared();
+	OwningViewportWeakPtr = InViewport.AsShared();
 }
 
 void FDisplayClusterViewportPreview::Release()
@@ -37,10 +42,12 @@ void FDisplayClusterViewportPreview::Release()
 	PreviewRTT.Reset();
 	RuntimeFlags = EDisplayClusterViewportPreviewFlags::None;
 
-	FDisplayClusterViewport* InViewport = GetViewportImpl();
-	PreviewMesh.Release(InViewport);
-	PreviewEditableMesh.Release(InViewport);
-	}
+	// Keeps a reference to the owning viewport so that it is not deleted within this logic.
+	TSharedPtr<FDisplayClusterViewport, ESPMode::ThreadSafe> OwningViewport = OwningViewportWeakPtr.Pin();
+
+	PreviewMesh.Release(OwningViewport.Get());
+	PreviewEditableMesh.Release(OwningViewport.Get());
+}
 
 void FDisplayClusterViewportPreview::Update(TSet<const UMeshComponent*>& MeshComponentsVisited)
 {
@@ -66,14 +73,16 @@ void FDisplayClusterViewportPreview::Update(TSet<const UMeshComponent*>& MeshCom
 		return;
 	}
 
-	FDisplayClusterViewport* Viewport = GetViewportImpl();	
+	// Keeps a reference to the owning viewport so that it is not deleted within this logic.
+	TSharedPtr<FDisplayClusterViewport, ESPMode::ThreadSafe> OwningViewport = OwningViewportWeakPtr.Pin();
+
 	UDisplayClusterCameraComponent* ViewPointComponent = nullptr;
 	UDisplayClusterDisplayDeviceBaseComponent* DisplayDeviceComponent = nullptr;
 
-	if (Viewport)
+	if (OwningViewport.IsValid())
 	{
-		ViewPointComponent = Viewport->GetViewPointCameraComponent(EDisplayClusterRootActorType::Configuration);
-		DisplayDeviceComponent = Viewport->GetDisplayDeviceComponent(Configuration->GetPreviewSettings().DisplayDeviceRootActorType);
+		ViewPointComponent = OwningViewport->GetViewPointCameraComponent(EDisplayClusterRootActorType::Configuration);
+		DisplayDeviceComponent = OwningViewport->GetDisplayDeviceComponent(Configuration->GetPreviewSettings().DisplayDeviceRootActorType);
 	}
 	
 	// Utility function to update both PreviewMesh and PreviewEditableMesh
@@ -89,7 +98,7 @@ void FDisplayClusterViewportPreview::Update(TSet<const UMeshComponent*>& MeshCom
 				MeshType = EDisplayClusterDisplayDeviceMeshType::DefaultMesh;
 			}
 
-			PreviewMeshObj.Update(Viewport, DisplayDeviceComponent, ViewPointComponent);
+			PreviewMeshObj.Update(OwningViewport.Get(), DisplayDeviceComponent, ViewPointComponent);
 
 			// Update Runtime Flags
 			if (PreviewMeshObj.HasAnyFlag(EDisplayClusterViewportPreviewMeshFlags::HasDeletedMaterialInstance | EDisplayClusterViewportPreviewMeshFlags::HasChangedMaterialInstance))
@@ -143,21 +152,13 @@ void FDisplayClusterViewportPreview::Update(TSet<const UMeshComponent*>& MeshCom
 	);
 }
 
-IDisplayClusterViewport* FDisplayClusterViewportPreview::GetViewport() const
-{
-	return GetViewportImpl();
-}
-
-FDisplayClusterViewport* FDisplayClusterViewportPreview::GetViewportImpl() const
-{
-	return ViewportWeakPtr.IsValid() ? ViewportWeakPtr.Pin().Get() : nullptr;
-}
-
 TSharedPtr<FDisplayClusterViewportResource, ESPMode::ThreadSafe> FDisplayClusterViewportPreview::GetOutputPreviewTargetableResources() const
 {
-	if (FDisplayClusterViewport* InViewport = GetViewportImpl())
+	// Keeps a reference to the owning viewport so that it is not deleted within this logic.
+	TSharedPtr<FDisplayClusterViewport, ESPMode::ThreadSafe> OwningViewport = OwningViewportWeakPtr.Pin();
+	if (OwningViewport.IsValid())
 	{
-		const TArray<TSharedPtr<FDisplayClusterViewportResource, ESPMode::ThreadSafe>>& PreviewResources = InViewport->GetViewportResources(EDisplayClusterViewportResource::OutputPreviewTargetableResources);
+		const TArray<TSharedPtr<FDisplayClusterViewportResource, ESPMode::ThreadSafe>>& PreviewResources = OwningViewport->GetViewportResources(EDisplayClusterViewportResource::OutputPreviewTargetableResources);
 		if (!PreviewResources.IsEmpty())
 		{
 			return PreviewResources[0];
diff --git a/Engine/Plugins/Runtime/nDisplay/Source/DisplayCluster/Private/Render/Viewport/Preview/DisplayClusterViewportPreview.h b/Engine/Plugins/Runtime/nDisplay/Source/DisplayCluster/Private/Render/Viewport/Preview/DisplayClusterViewportPreview.h
--- a/Engine/Plugins/Runtime/nDisplay/Source/DisplayCluster/Private/Render/Viewport/Preview/DisplayClusterViewportPreview.h
+++ b/Engine/Plugins/Runtime/nDisplay/Source/DisplayCluster/Private/Render/Viewport/Preview/DisplayClusterViewportPreview.h
@@ -55,7 +55,8 @@ public:
 		return ClusterNodeId;
 	}
 
-	virtual IDisplayClusterViewport* GetViewport() const override;
+	virtual TSharedPtr<IDisplayClusterViewport, ESPMode::ThreadSafe> GetOwningViewport() const override;
+
 	virtual UTextureRenderTarget2D* GetPreviewTextureRenderTarget2D() const override;
 
 	virtual UMeshComponent* GetPreviewMeshComponent() const override
@@ -85,14 +86,18 @@ public:
 	void Initialize(FDisplayClusterViewport& InViewport);
 
 	/**
-	* Create FSceneView for the viewport. This function is based on a similar function from LocalPlayer.
-	* 
-	* @param InOutViewFamily - Viewfamily that will render this viewport
-	* @param ContextNum - context index for rendering
-	* 
-	* @return - an instance of FSceneView or nullptr in case of failure.
-	*/
-	FSceneView* CalcSceneView(class FSceneViewFamilyContext& InOutViewFamily, uint32 ContextNum);
+	 * Creates an FSceneView for the viewport.
+	 * This function is based on similar LocalPlayer logic.
+	 *
+	 * @param InOutViewFamily View family used to render this viewport.
+	 * @param ContextNum      Rendering context index.
+	 * @param OutViewState    Receives the ViewState used by the SceneView.
+	 *                        This must be passed to the render thread to ensure it is not destroyed before rendering is complete.
+	 *
+	 * @return A pointer to the created FSceneView, or nullptr if creation fails.
+	 */	
+	FSceneView* CalcSceneView(class FSceneViewFamilyContext& InOutViewFamily, uint32 ContextNum,
+		TSharedPtr<FSceneViewStateReference, ESPMode::ThreadSafe>& OutViewState);
 
 	/** Returns true once if this type of log message can be displayed for the first time.*/
 	bool CanShowLogMsgOnce(const EDisplayClusterViewportPreviewShowLogMsgOnce& InLogState)
@@ -114,15 +119,12 @@ public:
 	}
 
 protected:
-	/** Get private viewport API. */
-	FDisplayClusterViewport* GetViewportImpl() const;
-
 	/** Return output preview resource. */
 	TSharedPtr<FDisplayClusterViewportResource, ESPMode::ThreadSafe> GetOutputPreviewTargetableResources() const;
 
 	/** Calculating stereoviewer offset for preview. */
 	bool CalculateStereoViewOffset(
-		FDisplayClusterViewport& InViewport,
+		const TSharedRef<FDisplayClusterViewport, ESPMode::ThreadSafe>& InViewport,
 		const uint32 InContextNum,
 		FRotator& ViewRotation,
 		FVector& ViewLocation,
@@ -142,8 +144,8 @@ public:
 	const FString ClusterNodeId;
 
 private:
-	// Pointer to the DC Viewport from the DCRA that is currently being used for rendering.
-	TWeakPtr<FDisplayClusterViewport, ESPMode::ThreadSafe> ViewportWeakPtr;
+	// Weak pointer to the owning DC viewport currently used for rendering.
+	TWeakPtr<FDisplayClusterViewport, ESPMode::ThreadSafe> OwningViewportWeakPtr;
 
 	// Preview RTT
 	TSharedPtr<FDisplayClusterViewportResource, ESPMode::ThreadSafe> PreviewRTT;
diff --git a/Engine/Plugins/Runtime/nDisplay/Source/DisplayCluster/Private/Render/Viewport/Preview/DisplayClusterViewportPreview_Render.cpp b/Engine/Plugins/Runtime/nDisplay/Source/DisplayCluster/Private/Render/Viewport/Preview/DisplayClusterViewportPreview_Render.cpp
--- a/Engine/Plugins/Runtime/nDisplay/Source/DisplayCluster/Private/Render/Viewport/Preview/DisplayClusterViewportPreview_Render.cpp
+++ b/Engine/Plugins/Runtime/nDisplay/Source/DisplayCluster/Private/Render/Viewport/Preview/DisplayClusterViewportPreview_Render.cpp
@@ -22,15 +22,20 @@
 ////////////////////////////////////////////////////////////////////////////////////////
 // FDisplayClusterViewportPreview
 ////////////////////////////////////////////////////////////////////////////////////////
-FSceneView* FDisplayClusterViewportPreview::CalcSceneView(FSceneViewFamilyContext& InOutViewFamily, uint32 InContextNum)
+FSceneView* FDisplayClusterViewportPreview::CalcSceneView(FSceneViewFamilyContext& InOutViewFamily, uint32 InContextNum,
+	TSharedPtr<FSceneViewStateReference, ESPMode::ThreadSafe>& OutViewState)
 {
-	FDisplayClusterViewport* InViewport = GetViewportImpl();
-	if (!InViewport || !InViewport->GetContexts().IsValidIndex(InContextNum))
+	// Reset the output ViewState before attempting to create a new scene view.
+	OutViewState.Reset();
+
+	// Keeps a reference to the owning viewport so that it is not deleted within this logic.
+	TSharedPtr<FDisplayClusterViewport, ESPMode::ThreadSafe> OwningViewport = OwningViewportWeakPtr.Pin();
+	if (!OwningViewport.IsValid() || !OwningViewport->GetContexts().IsValidIndex(InContextNum))
 	{
 		return nullptr;
 	}
 
-	const FDisplayClusterViewport_Context& InViewportContext = InViewport->GetContexts()[InContextNum];
+	const FDisplayClusterViewport_Context& InViewportContext = OwningViewport->GetContexts()[InContextNum];
 
 	const float MaxViewDistance = 1000000;
 	const bool bUseSceneColorTexture = false;
@@ -42,9 +47,9 @@ FSceneView* FDisplayClusterViewportPreview::CalcSceneView(FSceneViewFamilyContex
 	FVector StereoViewLocation;
 	FRotator ViewRotation;
 
-	if (CalculateStereoViewOffset(*InViewport, InContextNum, ViewRotation, ViewLocation, StereoViewLocation))
+	if (CalculateStereoViewOffset(OwningViewport.ToSharedRef(), InContextNum, ViewRotation, ViewLocation, StereoViewLocation))
 	{
-		FMatrix ProjectionMatrix = GetStereoProjectionMatrix(*InViewport, InContextNum);
+		FMatrix ProjectionMatrix = GetStereoProjectionMatrix(*OwningViewport, InContextNum);
 
 		FMatrix ViewRotationMatrix = FInverseRotationMatrix(ViewRotation);
 		ViewRotationMatrix = ViewRotationMatrix * FMatrix(
@@ -61,7 +66,10 @@ FSceneView* FDisplayClusterViewportPreview::CalcSceneView(FSceneViewFamilyContex
 		ViewInitOptions.SetViewRectangle(ViewRect);
 		ViewInitOptions.ViewFamily = &InOutViewFamily;
 
-		ViewInitOptions.SceneViewStateInterface = InViewport->GetViewState(InContextNum);
+		// Keeps a reference to the ViewState so that it is not destroyed before rendering is complete.
+		OutViewState = OwningViewport->GetOrCreateViewState(InContextNum);
+		ViewInitOptions.SceneViewStateInterface = OutViewState.IsValid() ? OutViewState->GetReference() : nullptr;
+
 		ViewInitOptions.ViewActor = ViewOwner;
 
 		ViewInitOptions.ViewOrigin = StereoViewLocation;
@@ -95,27 +103,27 @@ FSceneView* FDisplayClusterViewportPreview::CalcSceneView(FSceneViewFamilyContex
 		// The code below is based on the code from ULocalPlayer.
 		{
 			// ERenderPass::Start
-		View->StartFinalPostprocessSettings(ViewLocation);
-			InViewport->GetViewport_CustomPostProcessSettings().ApplyCustomPostProcess(InViewport, InContextNum, IDisplayClusterViewport_CustomPostProcessSettings::ERenderPass::Start, View->FinalPostProcessSettings);
+			View->StartFinalPostprocessSettings(ViewLocation);
+			OwningViewport->GetViewport_CustomPostProcessSettings().ApplyCustomPostProcess(OwningViewport.Get(), InContextNum, IDisplayClusterViewport_CustomPostProcessSettings::ERenderPass::Start, View->FinalPostProcessSettings);
 
 			// ERenderPass::Override
-		FPostProcessSettings OverridePostProcessingSettings;
-		float OverridePostProcessBlendWeight = 1.0f;
-			if (InViewport->GetViewport_CustomPostProcessSettings().ApplyCustomPostProcess(InViewport, InContextNum, IDisplayClusterViewport_CustomPostProcessSettings::ERenderPass::Override, OverridePostProcessingSettings, &OverridePostProcessBlendWeight))
-		{
-			View->OverridePostProcessSettings(OverridePostProcessingSettings, OverridePostProcessBlendWeight);
-		}
+			FPostProcessSettings OverridePostProcessingSettings;
+			float OverridePostProcessBlendWeight = 1.0f;
+			if (OwningViewport->GetViewport_CustomPostProcessSettings().ApplyCustomPostProcess(OwningViewport.Get(), InContextNum, IDisplayClusterViewport_CustomPostProcessSettings::ERenderPass::Override, OverridePostProcessingSettings, &OverridePostProcessBlendWeight))
+			{
+				View->OverridePostProcessSettings(OverridePostProcessingSettings, OverridePostProcessBlendWeight);
+			}
 
 			// ERenderPass::Final
-			InViewport->GetViewport_CustomPostProcessSettings().ApplyCustomPostProcess(InViewport, InContextNum, IDisplayClusterViewport_CustomPostProcessSettings::ERenderPass::Final, View->FinalPostProcessSettings);
-		View->EndFinalPostprocessSettings(ViewInitOptions);
+			OwningViewport->GetViewport_CustomPostProcessSettings().ApplyCustomPostProcess(OwningViewport.Get(), InContextNum, IDisplayClusterViewport_CustomPostProcessSettings::ERenderPass::Final, View->FinalPostProcessSettings);
+			View->EndFinalPostprocessSettings(ViewInitOptions);
 		}
 
 		FDisplayClusterViewportManager* ViewportManager = Configuration->GetViewportManagerImpl();
-		if (ViewportManager && InViewport->GetContexts().IsValidIndex(InContextNum))
+		if (ViewportManager && OwningViewport->GetContexts().IsValidIndex(InContextNum))
 		{
 			// Configure FDisplayClusterViewportManagerViewPointExtension for this viewport.
-			ViewportManager->SetCurrentStereoViewIndexForViewPointExtension(InViewport->GetContexts()[InContextNum].StereoViewIndex);
+			ViewportManager->SetCurrentStereoViewIndexForViewPointExtension(OwningViewport->GetContexts()[InContextNum].StereoViewIndex);
 		}
 
 		// Setup view extension for this view
@@ -155,7 +163,7 @@ FMatrix FDisplayClusterViewportPreview::GetStereoProjectionMatrix(FDisplayCluste
 }
 
 bool FDisplayClusterViewportPreview::CalculateStereoViewOffset(
-	FDisplayClusterViewport& InViewport,
+	const TSharedRef<FDisplayClusterViewport, ESPMode::ThreadSafe>& InViewport,
 	const uint32 InContextNum,
 	FRotator& ViewRotation,
 	FVector& ViewLocation,
@@ -165,7 +173,7 @@ bool FDisplayClusterViewportPreview::CalculateStereoViewOffset(
 
 	// Obtaining the internal viewpoint for a given viewport with stereo eye offset distance.
 	FMinimalViewInfo ViewInfo;
-	if (!InViewport.SetupViewPoint(InContextNum, ViewInfo))
+	if (!InViewport->SetupViewPoint(InContextNum, ViewInfo))
 	{
 		return false;
 	}
@@ -176,11 +184,11 @@ bool FDisplayClusterViewportPreview::CalculateStereoViewOffset(
 	// Obtaining the offset of the stereo eye and the values of the projection clipping plane for the given viewport was moved inside CalculateView().
 	// Perform view calculations on a policy side
 	StereoViewLocation = ViewInfo.Location;
-	if (!InViewport.CalculateView(InContextNum, StereoViewLocation, ViewRotation, Configuration->GetWorldToMeters()))
+	if (!InViewport->CalculateView(InContextNum, StereoViewLocation, ViewRotation, Configuration->GetWorldToMeters()))
 	{
 		if (CanShowLogMsgOnce(EDisplayClusterViewportPreviewShowLogMsgOnce::CalculateViewIsFailed))
 		{
-			UE_LOG(LogDisplayClusterViewport, Verbose, TEXT("Couldn't compute preview parameters for Viewport %s, ViewIdx: %d"), *InViewport.GetId(), InContextNum);
+			UE_LOG(LogDisplayClusterViewport, Verbose, TEXT("Couldn't compute preview parameters for Viewport %s, ViewIdx: %d"), *InViewport->GetId(), InContextNum);
 		}
 
 		return false;
diff --git a/Engine/Plugins/Runtime/nDisplay/Source/DisplayCluster/Public/Render/Viewport/IDisplayClusterViewport.h b/Engine/Plugins/Runtime/nDisplay/Source/DisplayCluster/Public/Render/Viewport/IDisplayClusterViewport.h
--- a/Engine/Plugins/Runtime/nDisplay/Source/DisplayCluster/Public/Render/Viewport/IDisplayClusterViewport.h
+++ b/Engine/Plugins/Runtime/nDisplay/Source/DisplayCluster/Public/Render/Viewport/IDisplayClusterViewport.h
@@ -18,6 +18,7 @@
 
 class UCameraComponent;
 class UWorld;
+class FSceneViewStateReference;
 struct FMinimalViewInfo;
 struct FOpenColorIOColorConversionSettings;
 
@@ -168,6 +169,15 @@ public:
 	/** Get viewport contexts */
 	virtual const TArray<FDisplayClusterViewport_Context>& GetContexts() const = 0;
 
+	/**
+	 * Returns the ViewState for the specified viewport context, creating it if necessary.
+	 *
+	 * @param InContextNum Eye context index for this viewport.
+	 *
+	 * @return Shared pointer to the corresponding ViewState, or nullptr if ViewState usage is not allowed for this viewport.
+	 */
+	virtual TSharedPtr<FSceneViewStateReference, ESPMode::ThreadSafe> GetOrCreateViewState(uint32 InContextNum) = 0;
+
 	/** Get clipping planes used by this viewport
 	 * 
 	 * @return Clipping planes as FVector2D(ZNear, ZFar)
diff --git a/Engine/Plugins/Runtime/nDisplay/Source/DisplayCluster/Public/Render/Viewport/IDisplayClusterViewportManagerPreview.h b/Engine/Plugins/Runtime/nDisplay/Source/DisplayCluster/Public/Render/Viewport/IDisplayClusterViewportManagerPreview.h
--- a/Engine/Plugins/Runtime/nDisplay/Source/DisplayCluster/Public/Render/Viewport/IDisplayClusterViewportManagerPreview.h
+++ b/Engine/Plugins/Runtime/nDisplay/Source/DisplayCluster/Public/Render/Viewport/IDisplayClusterViewportManagerPreview.h
@@ -11,7 +11,7 @@ class FCanvas;
 /**
 * Viewport manager preview rendering interface
 */
-class IDisplayClusterViewportManagerPreview
+class DISPLAYCLUSTER_API IDisplayClusterViewportManagerPreview
 {
 public:
 	virtual ~IDisplayClusterViewportManagerPreview() = default;
diff --git a/Engine/Plugins/Runtime/nDisplay/Source/DisplayCluster/Public/Render/Viewport/IDisplayClusterViewportPreview.h b/Engine/Plugins/Runtime/nDisplay/Source/DisplayCluster/Public/Render/Viewport/IDisplayClusterViewportPreview.h
--- a/Engine/Plugins/Runtime/nDisplay/Source/DisplayCluster/Public/Render/Viewport/IDisplayClusterViewportPreview.h
+++ b/Engine/Plugins/Runtime/nDisplay/Source/DisplayCluster/Public/Render/Viewport/IDisplayClusterViewportPreview.h
@@ -12,7 +12,7 @@ class UMaterialInstanceDynamic;
 /**
 * Provides access to the resources used to render the viewport preview.
 */
-class IDisplayClusterViewportPreview
+class DISPLAYCLUSTER_API IDisplayClusterViewportPreview
 {
 public:
 	virtual ~IDisplayClusterViewportPreview() = default;
@@ -32,8 +32,8 @@ public:
 	/** Gets the cluster node name. */
 	virtual FString GetClusterNodeId() const = 0;
 
-	/** Get owner viewport API */
-	virtual IDisplayClusterViewport* GetViewport() const = 0;
+	/** Returns the viewport that owns this object. */
+	virtual TSharedPtr<IDisplayClusterViewport, ESPMode::ThreadSafe> GetOwningViewport() const = 0;
 
 	/** Returns true if the preview has any of the requested flags. */
 	virtual bool HasAnyFlags(const EDisplayClusterViewportPreviewFlags InPreviewFlags) const = 0;
@@ -46,4 +46,9 @@ public:
 
 	/** Get viewport preview editable mesh.*/
 	virtual UMeshComponent* GetPreviewEditableMeshComponent() const = 0;
+
+	/** Deprecated. The base implementation returns nullptr. Use GetOwningViewport() instead. */
+	UE_DEPRECATED(5.8, "This function is deprecated. Use GetOwningViewport() instead.")
+	virtual IDisplayClusterViewport* GetViewport() const
+	{ return nullptr; };
 };
diff --git a/Engine/Plugins/Runtime/nDisplay/Source/DisplayClusterWarp/Private/Policy/InFrustumFit/DisplayClusterWarpInFrustumFitPolicy.cpp b/Engine/Plugins/Runtime/nDisplay/Source/DisplayClusterWarp/Private/Policy/InFrustumFit/DisplayClusterWarpInFrustumFitPolicy.cpp
--- a/Engine/Plugins/Runtime/nDisplay/Source/DisplayClusterWarp/Private/Policy/InFrustumFit/DisplayClusterWarpInFrustumFitPolicy.cpp
+++ b/Engine/Plugins/Runtime/nDisplay/Source/DisplayClusterWarp/Private/Policy/InFrustumFit/DisplayClusterWarpInFrustumFitPolicy.cpp
@@ -239,8 +239,8 @@ void FDisplayClusterWarpInFrustumFitPolicy::OnUpdateDisplayDeviceMeshAndMaterial
 
 	// Process only viewports with a projection policy based on the warpblend interface.
 	TSharedPtr<IDisplayClusterWarpBlend, ESPMode::ThreadSafe> WarpBlend;
-	IDisplayClusterViewport* InViewport = InViewportPreview.GetViewport();
-	if (!InViewport || !InViewport->GetProjectionPolicy().IsValid() || !InViewport->GetProjectionPolicy()->GetWarpBlendInterface(WarpBlend))
+	TSharedPtr<IDisplayClusterViewport, ESPMode::ThreadSafe> OwningViewport = InViewportPreview.GetOwningViewport();
+	if (!OwningViewport .IsValid()|| !OwningViewport->GetProjectionPolicy().IsValid() || !OwningViewport->GetProjectionPolicy()->GetWarpBlendInterface(WarpBlend))
 	{
 		return;
 	}
