Fix CI compilation error

#rnx
#Synced-CL 51735391

diff --git a/Engine/Plugins/Runtime/nDisplay/Source/DisplayClusterProjection/Private/Policy/Reference/DisplayClusterProjectionReferencePolicy.cpp b/Engine/Plugins/Runtime/nDisplay/Source/DisplayClusterProjection/Private/Policy/Reference/DisplayClusterProjectionReferencePolicy.cpp
--- a/Engine/Plugins/Runtime/nDisplay/Source/DisplayClusterProjection/Private/Policy/Reference/DisplayClusterProjectionReferencePolicy.cpp
+++ b/Engine/Plugins/Runtime/nDisplay/Source/DisplayClusterProjection/Private/Policy/Reference/DisplayClusterProjectionReferencePolicy.cpp
@@ -11,6 +11,9 @@
 #include "Render/Viewport/IDisplayClusterViewportManagerProxy.h"
 #include "Render/Viewport/IDisplayClusterViewport.h"
 #include "Render/Viewport/IDisplayClusterViewportProxy.h"
+#include "RenderingThread.h"
+#include "RHI.h"
+#include "RHICommandList.h"
 
 
 FDisplayClusterProjectionReferencePolicy::FDisplayClusterProjectionReferencePolicy(
diff --git a/Engine/Source/Runtime/InteractiveToolsFramework/Private/BaseGizmos/ViewBasedTransformAdjusters.cpp b/Engine/Source/Runtime/InteractiveToolsFramework/Private/BaseGizmos/ViewBasedTransformAdjusters.cpp
--- a/Engine/Source/Runtime/InteractiveToolsFramework/Private/BaseGizmos/ViewBasedTransformAdjusters.cpp
+++ b/Engine/Source/Runtime/InteractiveToolsFramework/Private/BaseGizmos/ViewBasedTransformAdjusters.cpp
@@ -6,6 +6,9 @@
 #include "BaseGizmos/ViewAdjustedStaticMeshGizmoComponent.h"
 #include "FrameTypes.h"
 #include "Misc/AxisDisplayInfo.h"
+#include "RenderingThread.h"
+#include "RHI.h"
+#include "RHICommandList.h"
 
 namespace ViewBasedTransformAdjustersLocals
 {
