[nDisplay] Fixed wrong logic in a checkSlow()

#rnx
#Synced-CL 0

diff --git a/Engine/Plugins/Runtime/nDisplay/Source/DisplayCluster/Private/Cluster/DisplayClusterClusterManager.cpp b/Engine/Plugins/Runtime/nDisplay/Source/DisplayCluster/Private/Cluster/DisplayClusterClusterManager.cpp
--- a/Engine/Plugins/Runtime/nDisplay/Source/DisplayCluster/Private/Cluster/DisplayClusterClusterManager.cpp
+++ b/Engine/Plugins/Runtime/nDisplay/Source/DisplayCluster/Private/Cluster/DisplayClusterClusterManager.cpp
@@ -1240,7 +1240,7 @@ void FDisplayClusterClusterManager::ProcessClusterStatesData(const TArray<uint8>
 	DataReader << StatesNum;
 	if (StatesNum <= 0)
 	{
-		checkSlow(StatesNum < 0);
+		checkSlow(StatesNum == 0);
 		return;
 	}
 
