[nDisplay] Fixed issues with nDisplay auto-configuration (.ini settings)
- Fixed nDisplay general and replication settings after in-memory settings cache implementation
- Use static type paths instead of string literals
- Added DPI parameters
- Minor code cleaning

#rnx
#Synced-CL 0

diff --git a/Engine/Plugins/Runtime/nDisplay/Source/DisplayClusterEditor/DisplayClusterEditor.Build.cs b/Engine/Plugins/Runtime/nDisplay/Source/DisplayClusterEditor/DisplayClusterEditor.Build.cs
--- a/Engine/Plugins/Runtime/nDisplay/Source/DisplayClusterEditor/DisplayClusterEditor.Build.cs
+++ b/Engine/Plugins/Runtime/nDisplay/Source/DisplayClusterEditor/DisplayClusterEditor.Build.cs
@@ -9,12 +9,13 @@ public class DisplayClusterEditor : ModuleRules
 	{
 		PrivateDependencyModuleNames.AddRange(
 			new string[] {
-				"DisplayCluster",
-				"DisplayClusterConfiguration",
-
 				"Core",
 				"CoreUObject",
+				"DisplayCluster",
+				"DisplayClusterConfiguration",
 				"Engine",
+				"EngineSettings",
+				"EnhancedInput",
 				"UnrealEd"
 			});
 	}
diff --git a/Engine/Plugins/Runtime/nDisplay/Source/DisplayClusterEditor/Private/DisplayClusterEditorModule.cpp b/Engine/Plugins/Runtime/nDisplay/Source/DisplayClusterEditor/Private/DisplayClusterEditorModule.cpp
--- a/Engine/Plugins/Runtime/nDisplay/Source/DisplayClusterEditor/Private/DisplayClusterEditorModule.cpp
+++ b/Engine/Plugins/Runtime/nDisplay/Source/DisplayClusterEditor/Private/DisplayClusterEditorModule.cpp
@@ -9,6 +9,14 @@
 
 #define LOCTEXT_NAMESPACE "DisplayClusterEditor"
 
+
+FDisplayClusterEditorModule::FDisplayClusterEditorModule()
+	: Container(TEXT("Project"))
+	, Category(TEXT("Plugins"))
+	, Section(TEXT("nDisplay"))
+{
+}
+
 void FDisplayClusterEditorModule::StartupModule()
 {
 	RegisterSettings();
@@ -28,9 +36,7 @@ void FDisplayClusterEditorModule::RegisterSettings()
 	if (ISettingsModule* SettingsModule = FModuleManager::GetModulePtr<ISettingsModule>("Settings"))
 	{
 		SettingsModule->RegisterSettings(
-			UDisplayClusterEditorSettings::Container,
-			UDisplayClusterEditorSettings::Category,
-			UDisplayClusterEditorSettings::Section,
+			Container, Category, Section,
 			LOCTEXT("RuntimeSettingsName", "nDisplay"),
 			LOCTEXT("RuntimeSettingsDescription", "Configure nDisplay"),
 			GetMutableDefault<UDisplayClusterEditorSettings>()
@@ -42,7 +48,7 @@ void FDisplayClusterEditorModule::UnregisterSettings()
 {
 	if (ISettingsModule* SettingsModule = FModuleManager::GetModulePtr<ISettingsModule>("Settings"))
 	{
-		SettingsModule->UnregisterSettings(UDisplayClusterEditorSettings::Container, UDisplayClusterEditorSettings::Category, UDisplayClusterEditorSettings::Section);
+		SettingsModule->UnregisterSettings(Container, Category, Section);
 	}
 }
 
diff --git a/Engine/Plugins/Runtime/nDisplay/Source/DisplayClusterEditor/Private/DisplayClusterEditorModule.h b/Engine/Plugins/Runtime/nDisplay/Source/DisplayClusterEditor/Private/DisplayClusterEditorModule.h
--- a/Engine/Plugins/Runtime/nDisplay/Source/DisplayClusterEditor/Private/DisplayClusterEditorModule.h
+++ b/Engine/Plugins/Runtime/nDisplay/Source/DisplayClusterEditor/Private/DisplayClusterEditorModule.h
@@ -11,6 +11,9 @@
 class FDisplayClusterEditorModule :
 	public IDisplayClusterEditor
 {
+public:
+	FDisplayClusterEditorModule();
+
 public:
 	//~ IModuleInterface interface
 	virtual void StartupModule() override;
@@ -19,4 +22,9 @@ public:
 protected:
 	void RegisterSettings();
 	void UnregisterSettings();
+
+private:
+	const FName Container;
+	const FName Category;
+	const FName Section;
 };
diff --git a/Engine/Plugins/Runtime/nDisplay/Source/DisplayClusterEditor/Private/Settings/DisplayClusterEditorSettings.cpp b/Engine/Plugins/Runtime/nDisplay/Source/DisplayClusterEditor/Private/Settings/DisplayClusterEditorSettings.cpp
--- a/Engine/Plugins/Runtime/nDisplay/Source/DisplayClusterEditor/Private/Settings/DisplayClusterEditorSettings.cpp
+++ b/Engine/Plugins/Runtime/nDisplay/Source/DisplayClusterEditor/Private/Settings/DisplayClusterEditorSettings.cpp
@@ -1,96 +1,102 @@
 // Copyright Epic Games, Inc. All Rights Reserved.
 
 #include "DisplayClusterEditorSettings.h"
+
 #include "DisplayClusterEditorEngine.h"
+#include "DisplayClusterGameEngine.h"
+#include "DisplayClusterPlayerInput.h"
+#include "DisplayClusterViewportClient.h"
+#include "EnhancedPlayerInput.h"
+#include "GeneralProjectSettings.h"
+
+#include "Editor/UnrealEdEngine.h"
+#include "Engine/GameEngine.h"
+#include "Engine/GameViewportClient.h"
+#include "GameFramework/InputSettings.h"
 #include "HAL/PlatformFileManager.h"
-#include "Serialization/JsonSerializerMacros.h"
 #include "Misc/ConfigCacheIni.h"
+#include "Serialization/JsonSerializerMacros.h"
 
 
-const FName UDisplayClusterEditorSettings::Container = TEXT("Project");
-const FName UDisplayClusterEditorSettings::Category  = TEXT("Plugins");
-const FName UDisplayClusterEditorSettings::Section   = TEXT("nDisplay");
-
-
-UDisplayClusterEditorSettings::UDisplayClusterEditorSettings(const FObjectInitializer& ObjectInitializer)
-	: Super(ObjectInitializer) 
-{
-	GET_MEMBER_NAME_CHECKED(UDisplayClusterEditorSettings, bEnabled);
-	GET_MEMBER_NAME_CHECKED(UDisplayClusterEditorSettings, bClusterReplicationEnabled);
-}
-
 void UDisplayClusterEditorSettings::PostEditChangeProperty(FPropertyChangedEvent& PropertyChangedEvent)
 {
 	Super::PostEditChangeProperty(PropertyChangedEvent);
+
+	SaveConfig();
+
 	if (PropertyChangedEvent.Property != nullptr)
 	{
-		// Since nDisplay is Windows only, save configs that depends on nDisplay plugin assets in Windows specific config file
-		const FString DefaultEnginePath = FString::Printf(TEXT("%sDefaultEngine.ini"), *FPaths::SourceConfigDir());
-		const FString DefaultGamePath   = FString::Printf(TEXT("%sDefaultGame.ini"),   *FPaths::SourceConfigDir());
-		const FString DefaultInputPath  = FString::Printf(TEXT("%sDefaultInput.ini"),  *FPaths::SourceConfigDir());
+		const FName PropertyName = PropertyChangedEvent.Property->GetFName();
 
-		FName PropertyName(PropertyChangedEvent.Property->GetFName());
+		const FString DefaultEnginePath = FPaths::Combine(FPaths::ProjectConfigDir(), TEXT("DefaultEngine.ini"));
+		const FString DefaultGamePath   = FPaths::Combine(FPaths::ProjectConfigDir(), TEXT("DefaultGame.ini"));
+		const FString DefaultInputPath  = FPaths::Combine(FPaths::ProjectConfigDir(), TEXT("DefaultInput.ini"));
 
-		if (PropertyName == GET_MEMBER_NAME_CHECKED(UDisplayClusterEditorSettings, bEnabled))
+		// Reset read-only flag from the config files
 		{
-			// Reset read-only flag from the config files
-			{
-				const FString* ConfigFiles[] = {
-					&DefaultEnginePath,
-					&DefaultGamePath,
-					&DefaultInputPath
-				};
+			const FString* ConfigFiles[] = { &DefaultEnginePath, &DefaultGamePath, &DefaultInputPath };
 
-				for (const FString* ConfigFile : ConfigFiles)
+			for (const FString* ConfigFile : ConfigFiles)
+			{
+				if (FPlatformFileManager::Get().GetPlatformFile().IsReadOnly(**ConfigFile))
 				{
-					if (FPlatformFileManager::Get().GetPlatformFile().IsReadOnly(**ConfigFile))
-					{
-						FPlatformFileManager::Get().GetPlatformFile().SetReadOnly(**ConfigFile, false);
-					}
+					FPlatformFileManager::Get().GetPlatformFile().SetReadOnly(**ConfigFile, false);
 				}
 			}
+		}
 
-			// Process nDisplay 'enable' command
-			if (bEnabled)
+		if (PropertyName == GET_MEMBER_NAME_CHECKED(UDisplayClusterEditorSettings, bEnabled))
+		{
+			// DefaultEngine.ini
 			{
-				// DefaultEngine.ini
-				GConfig->SetString(TEXT("/Script/Engine.Engine"), TEXT("GameEngine"), TEXT("/Script/DisplayCluster.DisplayClusterGameEngine"), DefaultEnginePath);
-				GConfig->SetString(TEXT("/Script/Engine.Engine"), TEXT("UnrealEdEngine"), TEXT("/Script/DisplayClusterEditor.DisplayClusterEditorEngine"), DefaultEnginePath);
-				GConfig->SetString(TEXT("/Script/Engine.Engine"), TEXT("GameViewportClientClassName"), TEXT("/Script/DisplayCluster.DisplayClusterViewportClient"), DefaultEnginePath);
+				// GameEngine
+				{
+					const UClass* const Class2Set = bEnabled ? UDisplayClusterGameEngine::StaticClass() : UGameEngine::StaticClass();
+					GConfig->SetString(TEXT("/Script/Engine.Engine"), TEXT("GameEngine"), *Class2Set->GetPathName(), DefaultEnginePath);
+				}
 
-				// DefaultGame.ini
-				GConfig->SetString(TEXT("/Script/EngineSettings.GeneralProjectSettings"), TEXT("bUseBorderlessWindow"), TEXT("True"), DefaultGamePath);
+				// UnrealEdEngine
+				{
+					const UClass* const Class2Set = bEnabled ? UDisplayClusterEditorEngine::StaticClass() : UUnrealEdEngine::StaticClass();
+					GConfig->SetString(TEXT("/Script/Engine.Engine"), TEXT("UnrealEdEngine"), *Class2Set->GetPathName(), DefaultEnginePath);
+				}
 
-				// DefaultInput.ini
-				GConfig->SetString(TEXT("/Script/Engine.InputSettings"), TEXT("DefaultPlayerInputClass"), TEXT("/Script/DisplayCluster.DisplayClusterPlayerInput"), DefaultInputPath);
-			}
-			// Process nDisplay 'disable' command
-			else
-			{
-				// DefaultEngine.ini
-				GConfig->SetString(TEXT("/Script/Engine.Engine"), TEXT("GameEngine"), TEXT("/Script/Engine.GameEngine"), DefaultEnginePath);
-				GConfig->SetString(TEXT("/Script/Engine.Engine"), TEXT("UnrealEdEngine"), TEXT("/Script/UnrealEd.UnrealEdEngine"), DefaultEnginePath);
-				GConfig->SetString(TEXT("/Script/Engine.Engine"), TEXT("GameViewportClientClassName"), TEXT("/Script/Engine.GameViewportClient"), DefaultEnginePath);
+				// GameViewportClientClassName
+				{
+					
+					const UClass* const Class2Set = bEnabled ? UDisplayClusterViewportClient::StaticClass() : UGameViewportClient::StaticClass();
+					GConfig->SetString(TEXT("/Script/Engine.Engine"), TEXT("GameViewportClientClassName"), *Class2Set->GetPathName(), DefaultEnginePath);
+				}
+
+				// Other engine settings
+				{
+					GConfig->SetBool(TEXT("/Script/Engine.UserInterfaceSettings"), TEXT("bAllowHighDPIInGameMode"),     true, DefaultEnginePath);
+					GConfig->SetBool(TEXT("/Script/Engine.UserInterfaceSettings"), TEXT("bAllowHighDpiWhenUnattended"), true, DefaultEnginePath);
+				}
 
-				// DefaultGame.ini
-				GConfig->SetString(TEXT("/Script/EngineSettings.GeneralProjectSettings"), TEXT("bUseBorderlessWindow"), TEXT("False"), DefaultGamePath);
+				GConfig->Flush(true, DefaultEnginePath);
+				GConfig->LoadFile(DefaultEnginePath);
+			}
 
-				// DefaultInput.ini
-				GConfig->SetString(TEXT("/Script/Engine.InputSettings"), TEXT("DefaultPlayerInputClass"), TEXT("/Script/EnhancedInput.EnhancedPlayerInput"), DefaultInputPath);
+			// DefaultGame.ini
+			if (UGeneralProjectSettings* GeneralProjectSettings = GetMutableDefault<UGeneralProjectSettings>())
+			{
+				GeneralProjectSettings->bUseBorderlessWindow = bEnabled;
+				GeneralProjectSettings->SaveConfig();
+				GeneralProjectSettings->TryUpdateDefaultConfigFile();
 			}
 
-			// Save changes to the files
-			GConfig->Flush(false, DefaultEnginePath);
-			GConfig->Flush(false, DefaultGamePath);
-			GConfig->Flush(false, DefaultInputPath);
+			// DefaultInput.ini
+			if (UInputSettings* InputSettings = GetMutableDefault<UInputSettings>())
+			{
+				UClass* const Class2Set = bEnabled ? UDisplayClusterPlayerInput::StaticClass() : UEnhancedPlayerInput::StaticClass();
+				InputSettings->SetDefaultPlayerInputClass(Class2Set);
+				InputSettings->SaveConfig();
+				InputSettings->TryUpdateDefaultConfigFile();
+			}
 		}
 		else if (PropertyName == GET_MEMBER_NAME_CHECKED(UDisplayClusterEditorSettings, bClusterReplicationEnabled))
 		{
-			if (FPlatformFileManager::Get().GetPlatformFile().IsReadOnly(*DefaultEnginePath))
-			{
-				FPlatformFileManager::Get().GetPlatformFile().SetReadOnly(*DefaultEnginePath, false);
-			}
-
 			if (bClusterReplicationEnabled)
 			{
 				FJsonSerializableArray NetDriverDefinitions;
@@ -106,7 +112,9 @@ void UDisplayClusterEditorSettings::PostEditChangeProperty(FPropertyChangedEvent
 				GConfig->RemoveKey(TEXT("/Script/Engine.GameEngine"), TEXT("+NetDriverDefinitions"), DefaultEnginePath);
 				GConfig->RemoveKey(TEXT("/Script/DisplayClusterReplication.DisplayClusterNetDriver"), TEXT("NetConnectionClassName"), DefaultEnginePath);
 			}
-			GConfig->Flush(false, DefaultEnginePath);
+
+			GConfig->Flush(true, DefaultEnginePath);
+			GConfig->LoadFile(DefaultEnginePath);
 		}
 	}
 }
diff --git a/Engine/Plugins/Runtime/nDisplay/Source/DisplayClusterEditor/Private/Settings/DisplayClusterEditorSettings.h b/Engine/Plugins/Runtime/nDisplay/Source/DisplayClusterEditor/Private/Settings/DisplayClusterEditorSettings.h
--- a/Engine/Plugins/Runtime/nDisplay/Source/DisplayClusterEditor/Private/Settings/DisplayClusterEditorSettings.h
+++ b/Engine/Plugins/Runtime/nDisplay/Source/DisplayClusterEditor/Private/Settings/DisplayClusterEditorSettings.h
@@ -15,19 +15,13 @@ class DISPLAYCLUSTEREDITOR_API UDisplayClusterEditorSettings : public UObject
 {
 	GENERATED_BODY()
 
-public:
-	UDisplayClusterEditorSettings(const FObjectInitializer& ObjectInitializer);
-
-public:
-	static const FName Container;
-	static const FName Category;
-	static const FName Section;
-
 protected:
 	virtual void PostEditChangeProperty(struct FPropertyChangedEvent& PropertyChangedEvent) override;
 
 protected:
-	UPROPERTY(config, EditAnywhere, Category = Main, meta = (ToolTip = "When enabled, replaces the original GameEngine to DisplayClusterGameEngine", ConfigRestartRequired = true))
+
+	/** Used to set or reset additional nDisplay settings */
+	UPROPERTY(config, EditAnywhere, Category = Main, meta = (ToolTip = "When enabled, automatically overrides some engine classes and settings", ConfigRestartRequired = true))
 	bool bEnabled;
 
 	/** When enabled, replaces the original NetDriver to DisplayClusterNetDriver, original NetConnection to DisplayClusterNetDriver */
