[nDisplay] MoviePipeline support and rendering improvements

Reworked the nDisplay/MoviePipeline integration to support simultaneous
rendering of the same ADisplayClusterRootActor from multiple independent
MovieGraph render pass nodes:
- ADisplayClusterRootActor:
* Removed the runtime transient property bMoviePipelineRenderPass.
* Replaced with TSet<const void*> ExclusiveLockOwners to support multiple
concurrent exclusive locks. The single boolean was insufficient because
* MovieGraph can render the same DCRA instance from multiple nodes
simultaneously.
- IDisplayClusterViewportManager:
* Added virtual Release(), which must be called before the destructor.
* Intended for external owners that manage their own IDisplayClusterViewportManager
instance (e.g. nDisplayMoviePipeline creates a unique ViewportManager per
cluster node instead of sharing the single instance owned by the DCRA).

Rework viewport rendering math:
- IDisplayClusterViewport:
* SetupViewPoint() reworked: calls CalculateView() early to compute the render view and projection, then encodes the nDisplay projection matrix into FMinimalViewInfo fields (FOV, AspectRatio, OffCenterProjectionOffset) so that CalculateProjectionMatrix() reproduces the exact nDisplay off-axis matrix. Also clears overscan, overrides near clip, and sets PreviousViewTransform for motion blur.
* GetPlayerCameraView(): added OutCustomNearClippingPlane output parameter to return the near clip override from an active CineCamera on the player's view target.
* Added GetPrevFrameContexts() returning contexts from the previous frame.
- FDisplayClusterViewport_Context::FCachedViewData: added RenderCameraLocation ??? camera head (pivot) position in world space, distinct from RenderLocation (per-eye render origin).
- Moved free math helpers into FDisplayClusterProjectionMath (Public/Math/DisplayClusterProjectionMath.h).

Implement frame number override for MoviePipeline rendering:
- IDisplayClusterViewportConfiguration
* UpdateConfigurationForClusterNode(): added InFrameNumberOverride ??? when set, used in place of GFrameNumber, which does not reflect the actual rendered frame during MoviePipeline rendering.
* Added IsMoviePipelineRendering() returning true when MoviePipeline drives rendering.
- FDisplayClusterRenderFrameSettings: added FrameNumberOverride to propagate the custom frame number through the render frame settings.

Add nDisplay render frame modes for MoviePipeline:
- EDisplayClusterRenderFrameMode: Added MRG_Mono/MRG_Stereo and MRQ_Stereo modes for MovieRenderGraph rendering paths.

#Synced-CL 0

diff --git a/Engine/Plugins/Runtime/nDisplay/Source/DisplayCluster/Private/Game/EngineClasses/Scene/DisplayClusterCameraComponent.cpp b/Engine/Plugins/Runtime/nDisplay/Source/DisplayCluster/Private/Game/EngineClasses/Scene/DisplayClusterCameraComponent.cpp
--- a/Engine/Plugins/Runtime/nDisplay/Source/DisplayCluster/Private/Game/EngineClasses/Scene/DisplayClusterCameraComponent.cpp
+++ b/Engine/Plugins/Runtime/nDisplay/Source/DisplayCluster/Private/Game/EngineClasses/Scene/DisplayClusterCameraComponent.cpp
@@ -185,7 +185,7 @@ bool UDisplayClusterCameraComponent::GetTargetCameraDesiredViewInternal(const ID
 	if (CfgCameraComponent.IsActiveEngineCameraBeingUsed())
 	{
 		// Get ViewInfo from the Game camera.
-		bViewInfoChanged = IDisplayClusterViewport::GetPlayerCameraView(InViewportConfiguration.GetCurrentWorld(), bUseCameraPostprocess, InOutViewInfo);
+		bViewInfoChanged = IDisplayClusterViewport::GetPlayerCameraView(InViewportConfiguration.GetCurrentWorld(), bUseCameraPostprocess, InOutViewInfo, OutCustomNearClippingPlane);
 	}
 	else if (UCameraComponent* SceneCameraComponent = GetTargetCameraComponent(InViewportConfiguration))
 	{
diff --git a/Engine/Plugins/Runtime/nDisplay/Source/DisplayCluster/Private/Render/Viewport/Configuration/DisplayClusterViewportConfiguration.cpp b/Engine/Plugins/Runtime/nDisplay/Source/DisplayCluster/Private/Render/Viewport/Configuration/DisplayClusterViewportConfiguration.cpp
--- a/Engine/Plugins/Runtime/nDisplay/Source/DisplayCluster/Private/Render/Viewport/Configuration/DisplayClusterViewportConfiguration.cpp
+++ b/Engine/Plugins/Runtime/nDisplay/Source/DisplayCluster/Private/Render/Viewport/Configuration/DisplayClusterViewportConfiguration.cpp
@@ -344,13 +344,23 @@ bool FDisplayClusterViewportConfiguration::IsExclusiveLocked() const
 	if (ADisplayClusterRootActor* SceneRootActor = GetRootActor(EDisplayClusterRootActorType::Scene))
 	{
 		// MRQ requires exclusivity: preview and proxies cannot use or modify the DCRA in scene.
-		if (SceneRootActor->bMoviePipelineRenderPass)
+		if (SceneRootActor->IsExclusiveLocked())
 		{
 			return true;
 		}
 	}
 #endif
 
+	// MoviePipeline instances can't use preview-in-scene and PIE.
+	switch (RenderFrameSettings.RenderMode)
+	{
+	case EDisplayClusterRenderFrameMode::MRG_Mono:
+	case EDisplayClusterRenderFrameMode::MRG_Stereo:
+	case EDisplayClusterRenderFrameMode::MRQ_Mono:
+	case EDisplayClusterRenderFrameMode::MRQ_Stereo:
+		return true;
+	}
+
 	return false;
 }
 
@@ -405,7 +415,12 @@ void FDisplayClusterViewportConfiguration::ReleaseConfiguration()
 	RenderFrameSettings = FDisplayClusterRenderFrameSettings();
 }
 
-bool FDisplayClusterViewportConfiguration::ImplUpdateConfiguration(EDisplayClusterRenderFrameMode InRenderMode, const UWorld* InWorld, const FString& InClusterNodeId, const TArray<FString>* InViewportNames)
+bool FDisplayClusterViewportConfiguration::ImplUpdateConfiguration(
+	EDisplayClusterRenderFrameMode InRenderMode,
+	const UWorld* InWorld,
+	const FString& InClusterNodeId,
+	const TArray<FString>* InViewportNames,
+	const uint32* InFrameNumberOverride)
 {
 	check(IsInGameThread());
 
@@ -434,6 +449,17 @@ bool FDisplayClusterViewportConfiguration::ImplUpdateConfiguration(EDisplayClust
 		return false;
 	}
 
+	// Override frame number
+	if (InFrameNumberOverride)
+	{
+		RenderFrameSettings.FrameNumberOverride = *InFrameNumberOverride;
+	}
+	else
+	{
+		// by default use GFrameNumber
+		RenderFrameSettings.FrameNumberOverride.Reset();
+	}
+
 	// The current world should always be initialized immediately before updating the configuration, since there is a dependency (e.g., visibility lists, etc.)
 	SetCurrentWorldImpl(InWorld);
 
diff --git a/Engine/Plugins/Runtime/nDisplay/Source/DisplayCluster/Private/Render/Viewport/Configuration/DisplayClusterViewportConfiguration.h b/Engine/Plugins/Runtime/nDisplay/Source/DisplayCluster/Private/Render/Viewport/Configuration/DisplayClusterViewportConfiguration.h
--- a/Engine/Plugins/Runtime/nDisplay/Source/DisplayCluster/Private/Render/Viewport/Configuration/DisplayClusterViewportConfiguration.h
+++ b/Engine/Plugins/Runtime/nDisplay/Source/DisplayCluster/Private/Render/Viewport/Configuration/DisplayClusterViewportConfiguration.h
@@ -47,14 +47,14 @@ public:
 		return RenderFrameSettings.PreviewSettings;
 	}
 
-	virtual bool UpdateConfigurationForClusterNode(EDisplayClusterRenderFrameMode InRenderMode, const UWorld* InWorld, const FString& InClusterNodeId) override
+	virtual bool UpdateConfigurationForClusterNode(EDisplayClusterRenderFrameMode InRenderMode, const UWorld* InWorld, const FString& InClusterNodeId, const uint32* InFrameNumberOverride = nullptr) override
 	{
-		return ImplUpdateConfiguration(InRenderMode, InWorld, InClusterNodeId, nullptr);
+		return ImplUpdateConfiguration(InRenderMode, InWorld, InClusterNodeId, nullptr, InFrameNumberOverride);
 	}
 
-	virtual bool UpdateConfigurationForViewportsList(EDisplayClusterRenderFrameMode InRenderMode, const UWorld* InWorld, const TArray<FString>& InViewportNames) override
+	virtual bool UpdateConfigurationForViewportsList(EDisplayClusterRenderFrameMode InRenderMode, const UWorld* InWorld, const TArray<FString>& InViewportNames, const uint32* InFrameNumberOverride = nullptr) override
 	{
-		return ImplUpdateConfiguration(InRenderMode, InWorld, TEXT(""), &InViewportNames);
+		return ImplUpdateConfiguration(InRenderMode, InWorld, TEXT(""), &InViewportNames, InFrameNumberOverride);
 	}
 
 	virtual void ReleaseConfiguration() override;
@@ -99,6 +99,23 @@ public:
 		return RenderFrameSettings.IsClusterRendering();
 	}
 
+	virtual bool IsMoviePipelineRendering() const override
+	{
+		switch (RenderFrameSettings.RenderMode)
+		{
+		case EDisplayClusterRenderFrameMode::MRQ_Mono:
+		case EDisplayClusterRenderFrameMode::MRG_Mono:
+		case EDisplayClusterRenderFrameMode::MRQ_Stereo:
+		case EDisplayClusterRenderFrameMode::MRG_Stereo:
+			return true;
+
+		default:
+			break;
+		}
+
+		return false;
+	}
+
 	virtual bool IsTechvisEnabled() const override
 	{
 		return RenderFrameSettings.IsTechvisEnabled();
@@ -269,12 +286,18 @@ public:
 private:
 	/** Update configuration implementation.
 	* 
-	* @param InRenderMode    - rendering mode.
-	* @param InWorld         - ptr to the world to be rendered
-	* @param InClusterNodeId - (opt) Configuring rendering for a cluster node.
-	* @param InViewportNames - (opt) Configuring rendering for a list of viewports.
+	* @param InRenderMode          - rendering mode.
+	* @param InWorld               - ptr to the world to be rendered
+	* @param InClusterNodeId       - (opt) Configuring rendering for a cluster node.
+	* @param InViewportNames       - (opt) Configuring rendering for a list of viewports.
+	* @param InFrameNumberOverride - (opt) If set, use this value instead of GFrameNumber.
 	*/
-	bool ImplUpdateConfiguration(EDisplayClusterRenderFrameMode InRenderMode, const UWorld* InWorld, const FString& InClusterNodeId, const TArray<FString>* InViewportNames);
+	bool ImplUpdateConfiguration(
+		EDisplayClusterRenderFrameMode InRenderMode,
+		const UWorld* InWorld,
+		const FString& InClusterNodeId,
+		const TArray<FString>* InViewportNames,
+		const uint32* InFrameNumberOverride);
 
 	/** Hide DCRA components for nDisplay rendering.*/
 	void ImplUpdateConfigurationVisibility();
diff --git a/Engine/Plugins/Runtime/nDisplay/Source/DisplayCluster/Private/Render/Viewport/DisplayClusterViewport.cpp b/Engine/Plugins/Runtime/nDisplay/Source/DisplayCluster/Private/Render/Viewport/DisplayClusterViewport.cpp
--- a/Engine/Plugins/Runtime/nDisplay/Source/DisplayCluster/Private/Render/Viewport/DisplayClusterViewport.cpp
+++ b/Engine/Plugins/Runtime/nDisplay/Source/DisplayCluster/Private/Render/Viewport/DisplayClusterViewport.cpp
@@ -518,6 +518,10 @@ float FDisplayClusterViewport::GetCustomBufferRatio(const FDisplayClusterRenderF
 void FDisplayClusterViewport::ResetFrameContexts()
 {
 	Resources.ReleaseAllResources();
+	Contexts.Empty();
+
+	PrevFrameContextsFrameNumber = 0;
+	PrevFrameContexts.Empty();
 }
 
 bool FDisplayClusterViewport::UpdateFrameContexts(const uint32 InStereoViewIndex)
@@ -533,6 +537,18 @@ bool FDisplayClusterViewport::UpdateFrameContexts(const uint32 InStereoViewIndex
 		return false;
 	}
 
+	// Snapshot the current contexts as the previous frame's contexts once per frame.
+	// The frame number guard prevents redundant copies when this function is called multiple times per frame.
+	if (!Contexts.IsEmpty())
+	{
+		uint32 CurrentFrameNumber = Configuration->GetRenderFrameSettings().GetFrameNumber();
+		if (PrevFrameContextsFrameNumber != CurrentFrameNumber)
+		{
+			PrevFrameContextsFrameNumber = CurrentFrameNumber;
+			PrevFrameContexts = Contexts;
+		}
+	}
+
 	FIntRect DesiredFrameTargetRect = RenderSettings.Rect;
 
 	// Apply desired frame mult
diff --git a/Engine/Plugins/Runtime/nDisplay/Source/DisplayCluster/Private/Render/Viewport/DisplayClusterViewport.h b/Engine/Plugins/Runtime/nDisplay/Source/DisplayCluster/Private/Render/Viewport/DisplayClusterViewport.h
--- a/Engine/Plugins/Runtime/nDisplay/Source/DisplayCluster/Private/Render/Viewport/DisplayClusterViewport.h
+++ b/Engine/Plugins/Runtime/nDisplay/Source/DisplayCluster/Private/Render/Viewport/DisplayClusterViewport.h
@@ -185,6 +185,12 @@ public:
 		return Contexts;
 	}
 
+	virtual const TArray<FDisplayClusterViewport_Context>& GetPrevFrameContexts() const override
+	{
+		check(IsInGameThread());
+		return PrevFrameContexts;
+	}
+
 	virtual TSharedPtr<FSceneViewStateReference, ESPMode::ThreadSafe> GetOrCreateViewState(uint32 InContextNum) override;
 
 	virtual const IDisplayClusterViewport_CustomPostProcessSettings& GetViewport_CustomPostProcessSettings() const override
@@ -566,6 +572,12 @@ private:
 	// Viewport contexts (left/center/right eyes)
 	TArray<FDisplayClusterViewport_Context> Contexts;
 
+	// Viewport contexts captured during the previous frame.
+	TArray<FDisplayClusterViewport_Context> PrevFrameContexts;
+
+	// Frame number for the previous-frame viewport contexts (PrevFrameContexts).
+	uint32 PrevFrameContextsFrameNumber = 0;
+
 	// View states (preview only)
 	TArray<TSharedPtr<FSceneViewStateReference, ESPMode::ThreadSafe>> ViewStates;
 
diff --git a/Engine/Plugins/Runtime/nDisplay/Source/DisplayCluster/Private/Render/Viewport/DisplayClusterViewportManager.h b/Engine/Plugins/Runtime/nDisplay/Source/DisplayCluster/Private/Render/Viewport/DisplayClusterViewportManager.h
--- a/Engine/Plugins/Runtime/nDisplay/Source/DisplayCluster/Private/Render/Viewport/DisplayClusterViewportManager.h
+++ b/Engine/Plugins/Runtime/nDisplay/Source/DisplayCluster/Private/Render/Viewport/DisplayClusterViewportManager.h
@@ -38,12 +38,6 @@ public:
 
 	void Initialize();
 
-	/** 
-	* Must be called before the destructor.
-	* Because all subsystems connected to this class via WeakPtr will not work when we are in destructor.
-	*/
-	void Release();
-
 	/** Release all resources.
 	* After that, the BeginNewFrame() function should be called for a new resource allocation.
 	*
@@ -52,6 +46,9 @@ public:
 	void ReleaseAllResources();
 
 	//~ Begin IDisplayClusterViewportManager
+
+	virtual void Release() override;
+
 	virtual TSharedPtr<IDisplayClusterViewportManager, ESPMode::ThreadSafe> ToSharedPtr() override
 	{
 		return AsShared();
diff --git a/Engine/Plugins/Runtime/nDisplay/Source/DisplayCluster/Private/Render/Viewport/DisplayClusterViewport_CustomPostProcessSettings.cpp b/Engine/Plugins/Runtime/nDisplay/Source/DisplayCluster/Private/Render/Viewport/DisplayClusterViewport_CustomPostProcessSettings.cpp
--- a/Engine/Plugins/Runtime/nDisplay/Source/DisplayCluster/Private/Render/Viewport/DisplayClusterViewport_CustomPostProcessSettings.cpp
+++ b/Engine/Plugins/Runtime/nDisplay/Source/DisplayCluster/Private/Render/Viewport/DisplayClusterViewport_CustomPostProcessSettings.cpp
@@ -63,10 +63,7 @@ namespace UE::DisplayClusterViewport::CustomPostProcess
 			return false;
 		}
 
-		const FMatrix& ProjectionMatrix =
-			ViewportContext.ProjectionData.bUseOverscan ?
-			ViewportContext.OverscanProjectionMatrix :
-			ViewportContext.ProjectionMatrix;
+		const FMatrix& ProjectionMatrix = ViewportContext.GetRenderProjectionMatrix();
 
 		// Ignore if the projection matrix is invalid.
 		if (ProjectionMatrix.M[0][0] == 0.0f || ProjectionMatrix.M[1][1] == 0.0f)
diff --git a/Engine/Plugins/Runtime/nDisplay/Source/DisplayCluster/Private/Render/Viewport/DisplayClusterViewport_Math.cpp b/Engine/Plugins/Runtime/nDisplay/Source/DisplayCluster/Private/Render/Viewport/DisplayClusterViewport_Math.cpp
--- a/Engine/Plugins/Runtime/nDisplay/Source/DisplayCluster/Private/Render/Viewport/DisplayClusterViewport_Math.cpp
+++ b/Engine/Plugins/Runtime/nDisplay/Source/DisplayCluster/Private/Render/Viewport/DisplayClusterViewport_Math.cpp
@@ -8,6 +8,7 @@
 #include "Render/Viewport/Containers/DisplayClusterViewport_PostRenderSettings.h"
 
 #include "Render/Projection/IDisplayClusterProjectionPolicy.h"
+#include "Math/DisplayClusterProjectionMath.h"
 
 #include "Render/Viewport/RenderFrame/DisplayClusterRenderFrame.h"
 #include "Render/Viewport/RenderFrame/DisplayClusterRenderFrameSettings.h"
@@ -21,37 +22,6 @@
 #include "EngineUtils.h"
 #include "SceneView.h"
 
-namespace UE::DisplayCluster::Viewport::Math
-{
-	/**
-	 * Calculates the ViewOffset for the eye from the view location
-	 *
-	 * @param InViewRotation - (in) rotation of the view (from this rotator we get the direction to the eye)
-	 * @param PassOffsetSwap - (in) Distance to the eye from the midpoint between the eyes.
-	 *
-	 * @return - the distance to the eye from the original ViewLocation
-	 */
-	static inline FVector GetStereoEyeOffset(const FRotator& InViewRotation, const float StereoEyeOffsetDistance)
-	{
-		return  InViewRotation.Quaternion().RotateVector(FVector(0.0f, StereoEyeOffsetDistance, 0.0f));
-	}
-
-	/** check frustum. */
-	static inline void GetNonZeroFrustumRange(double& InOutValue0, double& InOutValue1, double n)
-	{
-		static const double MinHalfFOVRangeRad = FMath::DegreesToRadians(0.5f);
-		static const double MinRangeBase = FMath::Tan(MinHalfFOVRangeRad * 2);;
-
-		const double MinRangeValue = n * MinRangeBase;
-		if ((InOutValue1 - InOutValue0) < MinRangeValue)
-		{
-			// Get minimal values from center of range
-			const double CenterRad = (FMath::Atan(InOutValue0 / n) + (FMath::Atan(InOutValue1 / n))) * 0.5f;
-			InOutValue0 = double(n * FMath::Tan(CenterRad - MinHalfFOVRangeRad));
-			InOutValue1 = double(n * FMath::Tan(CenterRad + MinHalfFOVRangeRad));
-		}
-	}
-};
 
 ///////////////////////////////////////////////////////////////////////////////////////
 //          FDisplayClusterViewport
@@ -76,8 +46,6 @@ FVector2D FDisplayClusterViewport::GetClippingPlanes() const
 
 bool FDisplayClusterViewport::GetViewPointCameraEye(const uint32 InContextNum, FVector& OutViewLocation, FRotator& OutViewRotation, FVector& OutViewOffset)
 {
-	using namespace UE::DisplayCluster::Viewport::Math;
-
 	// Here we use the ViewPoint component as the eye position
 	if (UDisplayClusterCameraComponent* SceneCameraComponent = GetViewPointCameraComponent(EDisplayClusterRootActorType::Scene))
 	{
@@ -85,7 +53,7 @@ bool FDisplayClusterViewport::GetViewPointCameraEye(const uint32 InContextNum, F
 		SceneCameraComponent->GetEyePosition(GetConfiguration(), OutViewLocation, OutViewRotation);
 
 		// Calculate stereo ViewOffset:
-		OutViewOffset = GetStereoEyeOffset(OutViewRotation, GetStereoEyeOffsetDistance(InContextNum));
+		OutViewOffset = FDisplayClusterProjectionMath::GetStereoEyeOffset(OutViewRotation, GetStereoEyeOffsetDistance(InContextNum));
 
 		// Add stereo eye offset to the output view location.
 		OutViewLocation += OutViewOffset;
@@ -125,9 +93,10 @@ bool FDisplayClusterViewport::CalculateView(const uint32 InContextNum, FVector&
 	return false;
 }
 
-bool FDisplayClusterViewport::CalculateViewData(const uint32 InContextNum, const FVector& InViewLocation, const FRotator& InViewRotation, const float WorldToMeters)
+bool FDisplayClusterViewport::CalculateViewData(const uint32 InContextNum, const FVector& InViewLocation, const FRotator& InViewRotation, const float InWorldToMeters)
 {
-	using namespace UE::DisplayCluster::Viewport::Math;
+	// Use the provided WorldToMeters value when valid; otherwise fall back to the configuration value.
+	const float WorldToMeters = InWorldToMeters > 0 ? InWorldToMeters : Configuration->GetWorldToMeters();
 
 	if (!PreCalculateViewData(InContextNum, InViewLocation, InViewRotation, WorldToMeters))
 	{
@@ -164,7 +133,7 @@ bool FDisplayClusterViewport::CalculateViewData(const uint32 InContextNum, const
 
 	// Stereo offset distance
 	const float StereoEyeOffsetDistance = GetStereoEyeOffsetDistance(InContextNum);
-	FVector ViewOffset = GetStereoEyeOffset(ViewRotation, StereoEyeOffsetDistance);
+	FVector ViewOffset = FDisplayClusterProjectionMath::GetStereoEyeOffset(ViewRotation, StereoEyeOffsetDistance);
 
 	// Projection policy expects view location with stereo offset.
 	ViewLocation += ViewOffset;
@@ -234,6 +203,11 @@ bool FDisplayClusterViewport::CalculateViewData(const uint32 InContextNum, const
 	OutContext.ViewData.RenderLocation = ViewLocation;
 	OutContext.ViewData.RenderRotation = ViewRotation;
 
+	// calc final origin location
+	const FVector RenderViewOffset = FDisplayClusterProjectionMath::GetStereoEyeOffset(ViewRotation, StereoEyeOffsetDistance);
+	ViewLocation -= RenderViewOffset;
+	OutContext.ViewData.RenderCameraLocation = ViewLocation;
+
 	// Calculate Projection matrix
 	FMatrix ProjectionMatrix;
 	if (!ProjectionPolicy->GetProjectionMatrix(this, InContextNum, ProjectionMatrix))
@@ -257,10 +231,7 @@ bool FDisplayClusterViewport::GetProjectionMatrix(const uint32 InContextNum, FMa
 		if (Contexts[InContextNum].bValidData)
 		{
 			// Use calculated values
-			OutPrjMatrix = Contexts[InContextNum].ProjectionData.bUseOverscan
-				? Contexts[InContextNum].OverscanProjectionMatrix
-				: Contexts[InContextNum].ProjectionMatrix;
-
+			OutPrjMatrix = Contexts[InContextNum].GetRenderProjectionMatrix();
 			return true;
 		}
 	}
@@ -270,19 +241,18 @@ bool FDisplayClusterViewport::GetProjectionMatrix(const uint32 InContextNum, FMa
 
 void FDisplayClusterViewport::CalculateProjectionMatrix(const uint32 InContextNum, float Left, float Right, float Top, float Bottom, float ZNear, float ZFar, bool bIsAnglesInput)
 {
-	using namespace UE::DisplayCluster::Viewport::Math;
-
-	// limit max frustum to 89
-	static const double MaxFrustumAngle = FMath::Tan(FMath::DegreesToRadians(89));
-	const double MaxValue = ZNear * MaxFrustumAngle;
-
 	const double n = ZNear;
 	const double f = ZFar;
 
-	double t = bIsAnglesInput ? (ZNear * FMath::Tan(FMath::DegreesToRadians(Top)))    : Top;
-	double b = bIsAnglesInput ? (ZNear * FMath::Tan(FMath::DegreesToRadians(Bottom))) : Bottom;
-	double l = bIsAnglesInput ? (ZNear * FMath::Tan(FMath::DegreesToRadians(Left)))   : Left;
-	double r = bIsAnglesInput ? (ZNear * FMath::Tan(FMath::DegreesToRadians(Right)))  : Right;
+	// limit max frustum to 89 degree
+	const double MaxAngleDegree = 89;
+	const double MaxValue = FDisplayClusterProjectionMath::DegreesToNearPlaneExtent(MaxAngleDegree, n);
+
+
+	double t = bIsAnglesInput ? FDisplayClusterProjectionMath::DegreesToNearPlaneExtent(static_cast<double>(Top),    n) : Top;
+	double b = bIsAnglesInput ? FDisplayClusterProjectionMath::DegreesToNearPlaneExtent(static_cast<double>(Bottom), n) : Bottom;
+	double l = bIsAnglesInput ? FDisplayClusterProjectionMath::DegreesToNearPlaneExtent(static_cast<double>(Left),   n) : Left;
+	double r = bIsAnglesInput ? FDisplayClusterProjectionMath::DegreesToNearPlaneExtent(static_cast<double>(Right),  n) : Right;
 
 	// Protect PrjMatrix from bad input values, and fix\clamp FOV to limits
 	{
@@ -327,10 +297,10 @@ void FDisplayClusterViewport::CalculateProjectionMatrix(const uint32 InContextNu
 		}
 	}
 
-	GetNonZeroFrustumRange(l, r, n);
-	GetNonZeroFrustumRange(b, t, n);
+	FDisplayClusterProjectionMath::GetNonZeroFrustumRange(l, r, n);
+	FDisplayClusterProjectionMath::GetNonZeroFrustumRange(b, t, n);
 
-	Contexts[InContextNum].ProjectionMatrix = IDisplayClusterViewport::MakeProjectionMatrix(l, r, t, b, n, f);
+	Contexts[InContextNum].ProjectionMatrix = FDisplayClusterProjectionMath::MakeAsymmetricProjectionMatrix(l, r, t, b, n, f);
 
 	// Update cached projection data:
 	FDisplayClusterViewport_Context::FCachedProjectionData& CachedProjectionData = Contexts[InContextNum].ProjectionData;
@@ -343,7 +313,6 @@ void FDisplayClusterViewport::CalculateProjectionMatrix(const uint32 InContextNu
 	CachedProjectionData.ZNear = n;
 	CachedProjectionData.ZFar = f;
 
-
 	if (FDisplayClusterViewport_OverscanRuntimeSettings::UpdateProjectionAngles(OverscanRuntimeSettings, l, r, t, b))
 	{
 		if (FMath::IsWithin(l, -MaxValue, MaxValue) &&
@@ -353,7 +322,7 @@ void FDisplayClusterViewport::CalculateProjectionMatrix(const uint32 InContextNu
 			)
 		{
 			// Use overscan projection matrix
-			Contexts[InContextNum].OverscanProjectionMatrix = IDisplayClusterViewport::MakeProjectionMatrix(l, r, t, b, n, f);
+			Contexts[InContextNum].OverscanProjectionMatrix = FDisplayClusterProjectionMath::MakeAsymmetricProjectionMatrix(l, r, t, b, n, f);
 
 			// Cache projection data for overscan
 			CachedProjectionData.bUseOverscan = true;
@@ -372,32 +341,8 @@ void FDisplayClusterViewport::CalculateProjectionMatrix(const uint32 InContextNu
 ///////////////////////////////////////////////////////////////////////////////////////
 FMatrix IDisplayClusterViewport::MakeProjectionMatrix(float l, float r, float t, float b, float n, float f)
 {
-	const float mx = 2.f * n / (r - l);
-	const float my = 2.f * n / (t - b);
-	const float ma = -(r + l) / (r - l);
-	const float mb = -(t + b) / (t - b);
-
-	// Support unlimited far plane (f==n)
-	const float mc = (f == n) ? (1.0f - Z_PRECISION) : (f / (f - n));
-	const float md = (f == n) ? (-n * (1.0f - Z_PRECISION)) : (-(f * n) / (f - n));
-
-	const float me = 1.f;
-
-	// Normal LHS
-	const FMatrix ProjectionMatrix = FMatrix(
-		FPlane(mx, 0, 0, 0),
-		FPlane(0, my, 0, 0),
-		FPlane(ma, mb, mc, me),
-		FPlane(0, 0, md, 0));
-
-	// Invert Z-axis (UE uses Z-inverted LHS)
-	static const FMatrix flipZ = FMatrix(
-		FPlane(1, 0, 0, 0),
-		FPlane(0, 1, 0, 0),
-		FPlane(0, 0, -1, 0),
-		FPlane(0, 0, 1, 1));
-
-	const FMatrix ResultMatrix(ProjectionMatrix * flipZ);
-
-	return ResultMatrix;
+	return FDisplayClusterProjectionMath::MakeAsymmetricProjectionMatrix(
+		static_cast<double>(l), static_cast<double>(r),
+		static_cast<double>(t), static_cast<double>(b),
+		static_cast<double>(n), static_cast<double>(f));
 }
diff --git a/Engine/Plugins/Runtime/nDisplay/Source/DisplayCluster/Private/Render/Viewport/DisplayClusterViewport_Rules.cpp b/Engine/Plugins/Runtime/nDisplay/Source/DisplayCluster/Private/Render/Viewport/DisplayClusterViewport_Rules.cpp
--- a/Engine/Plugins/Runtime/nDisplay/Source/DisplayCluster/Private/Render/Viewport/DisplayClusterViewport_Rules.cpp
+++ b/Engine/Plugins/Runtime/nDisplay/Source/DisplayCluster/Private/Render/Viewport/DisplayClusterViewport_Rules.cpp
@@ -555,6 +555,12 @@ bool FDisplayClusterViewport::ShouldUseFullSizeFrameTargetableResource() const
 
 bool FDisplayClusterViewport::ShouldFreezeRender() const
 {
+	// Never freeze when the DCRA is exclusively locked.
+	if (Configuration->IsExclusiveLocked())
+	{
+		return false;
+	}
+
 	// Freeze preview rendering for all viewport
 	if (Configuration->GetRenderFrameSettings().IsPreviewFreezeRender())
 	{
@@ -592,6 +598,12 @@ bool FDisplayClusterViewport::ShouldUseViewStates() const
 
 	// MRQ has its own ViewState management.
 	case EDisplayClusterRenderFrameMode::MRQ_Mono:
+	case EDisplayClusterRenderFrameMode::MRQ_Stereo:
+		return false;
+
+	// MRG has its own ViewState management.
+	case EDisplayClusterRenderFrameMode::MRG_Mono:
+	case EDisplayClusterRenderFrameMode::MRG_Stereo:
 		return false;
 
 	// Preview-In-Scene
diff --git a/Engine/Plugins/Runtime/nDisplay/Source/DisplayCluster/Private/Render/Viewport/DisplayClusterViewport_ViewPoint.cpp b/Engine/Plugins/Runtime/nDisplay/Source/DisplayCluster/Private/Render/Viewport/DisplayClusterViewport_ViewPoint.cpp
--- a/Engine/Plugins/Runtime/nDisplay/Source/DisplayCluster/Private/Render/Viewport/DisplayClusterViewport_ViewPoint.cpp
+++ b/Engine/Plugins/Runtime/nDisplay/Source/DisplayCluster/Private/Render/Viewport/DisplayClusterViewport_ViewPoint.cpp
@@ -4,6 +4,7 @@
 #include "Render/Viewport/DisplayClusterViewportManager.h"
 
 #include "Render/Projection/IDisplayClusterProjectionPolicy.h"
+#include "Math/DisplayClusterProjectionMath.h"
 
 #include "DisplayClusterRootActor.h"
 #include "Components/DisplayClusterCameraComponent.h"
@@ -72,15 +73,35 @@ bool IDisplayClusterViewport::GetCameraComponentView(UCameraComponent* InCameraC
 	return true;
 }
 
-bool IDisplayClusterViewport::GetPlayerCameraView(UWorld* InWorld, const bool bUseCameraPostprocess, FMinimalViewInfo& InOutViewInfo)
+bool IDisplayClusterViewport::GetPlayerCameraView(UWorld* InWorld, const bool bUseCameraPostprocess, FMinimalViewInfo& InOutViewInfo, float* OutCustomNearClippingPlane)
 {
 	if (InWorld)
 	{
-		if (APlayerController* const CurPlayerController = InWorld->GetFirstPlayerController())
+		if (APlayerController* const LocalPlayerController = InWorld->GetFirstPlayerController())
 		{
-			if (APlayerCameraManager* const CurPlayerCameraManager = CurPlayerController->PlayerCameraManager)
+			if (APlayerCameraManager* const LocalPlayerCameraManager = LocalPlayerController->PlayerCameraManager)
 			{
-				InOutViewInfo = CurPlayerCameraManager->GetCameraCacheView(); // Get desired view with postprocess from player camera
+				if (OutCustomNearClippingPlane)
+				{
+					*OutCustomNearClippingPlane = -1;
+
+					if (AActor* ViewActor = LocalPlayerController->GetViewTarget())
+					{
+						TArray<UCineCameraComponent*> CineCameraComponents;
+						ViewActor->GetComponents<UCineCameraComponent>(CineCameraComponents, true);
+
+						// Get settings from this cinecamera component
+						if (UCineCameraComponent* CineCameraComponent = !CineCameraComponents.IsEmpty() ? CineCameraComponents[0] : nullptr)
+						{
+							if (CineCameraComponent->bOverride_CustomNearClippingPlane)
+							{
+								*OutCustomNearClippingPlane = CineCameraComponent->CustomNearClippingPlane;
+							}
+						}
+					}
+				}
+
+				InOutViewInfo = LocalPlayerCameraManager->GetCameraCacheView(); // Get desired view with postprocess from player camera
 
 				if (!bUseCameraPostprocess)
 				{
@@ -88,8 +109,10 @@ bool IDisplayClusterViewport::GetPlayerCameraView(UWorld* InWorld, const bool bU
 					InOutViewInfo.PostProcessBlendWeight = 0.0f;
 				}
 
-				InOutViewInfo.FOV = CurPlayerCameraManager->GetFOVAngle();
-				CurPlayerCameraManager->GetCameraViewPoint(/*out*/ InOutViewInfo.Location, /*out*/ InOutViewInfo.Rotation);
+				const float ActualFOV = LocalPlayerCameraManager->GetFOVAngle();
+				InOutViewInfo.FOV = ActualFOV;
+
+				LocalPlayerCameraManager->GetCameraViewPoint(/*out*/ InOutViewInfo.Location, /*out*/ InOutViewInfo.Rotation);
 
 				if (!bUseCameraPostprocess)
 				{
@@ -167,45 +190,104 @@ UDisplayClusterCameraComponent* FDisplayClusterViewport::GetViewPointCameraCompo
 
 bool FDisplayClusterViewport::SetupViewPoint(const uint32 InContextNum, FMinimalViewInfo& InOutViewInfo)
 {
-	if (UDisplayClusterCameraComponent* SceneCameraComponent = GetViewPointCameraComponent(EDisplayClusterRootActorType::Scene))
+	if (!Contexts.IsValidIndex(InContextNum))
 	{
-		// Get ViewPoint from DCRA component
-		SceneCameraComponent->GetDesiredView(*Configuration, InOutViewInfo, &CustomNearClippingPlane);
+		return false;
+	}
 
-		// The projection policy can override these ViewPoint data.
+	const FDisplayClusterViewport_Context& SrcContext = Contexts[InContextNum];
+
+	// Get DCRA ViewPoint component (required).
+	UDisplayClusterCameraComponent* ViewPointComponent = GetViewPointCameraComponent(EDisplayClusterRootActorType::Scene);
+	if (!ViewPointComponent)
+	{
+		return false;
+	}
+
+	// Get ViewInfo
+	{
+		// Get ViewInfo from the ViewPoint component
+		ViewPointComponent->GetDesiredView(*Configuration, InOutViewInfo, &CustomNearClippingPlane);
+
+		// Apply projection policy overrides (policy type 'camera', etc)
 		if (ProjectionPolicy.IsValid())
 		{
 			ProjectionPolicy->SetupProjectionViewPoint(this, Configuration->GetRootActorWorldDeltaSeconds(), InOutViewInfo, &CustomNearClippingPlane);
 		}
 
-		// Save additional data at this point:
-		if (Contexts.IsValidIndex(InContextNum))
+		// Calculate View data for the current viewport context.
+		if (!CalculateView(InContextNum, InOutViewInfo.Location, InOutViewInfo.Rotation, 0.f))
 		{
-			FDisplayClusterViewport_Context& DestContext = Contexts[InContextNum];
+			return false;
+		}
 
-			// DoF FocalLength
-			if (InOutViewInfo.PostProcessSettings.DepthOfFieldFocalDistance > 0.0f)
-			{
-				// Convert FOV to focal length,
-				// 
-				// fov = 2 * atan(d/(2*f))
-				// where,
-				//   d = sensor dimension (APS-C 24.576 mm)
-				//   f = focal length
-				// 
-				// f = 0.5 * d * (1/tan(fov/2))
-				const FMatrix ProjectionMatrix = InOutViewInfo.CalculateProjectionMatrix();
-				DestContext.DepthOfField.SensorFocalLength = 0.5f * InOutViewInfo.PostProcessSettings.DepthOfFieldSensorWidth * ProjectionMatrix.M[0][0];
-
-				// Save actual squeeze factor value
-				DestContext.DepthOfField.SqueezeFactor = FMath::Clamp(InOutViewInfo.PostProcessSettings.DepthOfFieldSqueezeFactor, 1.0f, 2.0f);
-			}
+		// SetupViewPoint() should return Origin for ULocalPlayer::GetProjectionData()
+		InOutViewInfo.Location = SrcContext.ViewData.RenderCameraLocation;
+	}
+
+	// Override FMinimalViewInfo fields with nDisplay-specific values.
+	{
+		// nDisplay always uses its own overscan - ignore camera overscan settings.
+		InOutViewInfo.ClearOverscan();
+
+		// Override near clipping plane from nDisplay viewport.
+		InOutViewInfo.PerspectiveNearClipPlane = GetClippingPlanes().X;
+
+		// Encode the nDisplay projection matrix into FMinimalViewInfo so that
+		// CalculateProjectionMatrix() reproduces it exactly.
+		// Note: FMinimalViewInfo carries no far-plane distance,
+		//       so the recovered matrix always uses an infinite far plane.
+		{
+			const FMatrix& P = SrcContext.GetRenderProjectionMatrix();
+			const float mx = P.M[0][0];   // horizontal scale
+			const float my = P.M[1][1];   // vertical scale
+
+			// Recover horizontal FOV and aspect ratio from the projection scale factors.
+			InOutViewInfo.FOV         = 2.f * FMath::RadiansToDegrees(FMath::Atan(1.f / mx));
+			InOutViewInfo.AspectRatio = my / mx;
+
+			// Force CalculateProjectionMatrix() which uses FOV/AspectRatio directly,
+			// bypassing the viewport pixel dimensions used by the unconstrained path.
+			InOutViewInfo.bConstrainAspectRatio = true;
+			InOutViewInfo.AspectRatioAxisConstraint.Reset();
+
+			// Recover the off-center shift from the projection matrix.
+			InOutViewInfo.OffCenterProjectionOffset.X = -P.M[2][0];
+			InOutViewInfo.OffCenterProjectionOffset.Y = -P.M[2][1];
 		}
 
-		return true;
+		// Optionally set the previous-frame view transform for motion blur.
+		if (PrevFrameContexts.IsValidIndex(InContextNum))
+		{
+			const FDisplayClusterViewport_Context& PrevFrameContext = PrevFrameContexts[InContextNum];
+			InOutViewInfo.PreviousViewTransform = FTransform(
+				PrevFrameContext.ViewData.RenderRotation,
+				PrevFrameContext.ViewData.RenderCameraLocation);
+		}
 	}
 
-	return false;
+	// Save additional data at this point: DoF FocalLength
+	if (InOutViewInfo.PostProcessSettings.DepthOfFieldFocalDistance > 0.0f)
+	{
+		// Convert FOV to focal length,
+		// 
+		// fov = 2 * atan(d/(2*f))
+		// where,
+		//   d = sensor dimension (APS-C 24.576 mm)
+		//   f = focal length
+		// 
+		// f = 0.5 * d * (1/tan(fov/2))
+		const FMatrix ProjectionMatrix = InOutViewInfo.CalculateProjectionMatrix();
+
+		FDisplayClusterViewport_Context& DestContext = Contexts[InContextNum];
+
+		DestContext.DepthOfField.SensorFocalLength = 0.5f * InOutViewInfo.PostProcessSettings.DepthOfFieldSensorWidth * ProjectionMatrix.M[0][0];
+
+		// Save actual squeeze factor value
+		DestContext.DepthOfField.SqueezeFactor = FMath::Clamp(InOutViewInfo.PostProcessSettings.DepthOfFieldSqueezeFactor, 1.0f, 2.0f);
+	}
+
+	return true;
 }
 
 bool FDisplayClusterViewport::PreCalculateViewData(const uint32 InContextNum, const FVector& InViewLocation, const FRotator& InViewRotation, const float WorldToMeters)
diff --git a/Engine/Plugins/Runtime/nDisplay/Source/DisplayCluster/Private/Render/Viewport/RenderFrame/DisplayClusterRenderFrameSettings.cpp b/Engine/Plugins/Runtime/nDisplay/Source/DisplayCluster/Private/Render/Viewport/RenderFrame/DisplayClusterRenderFrameSettings.cpp
--- a/Engine/Plugins/Runtime/nDisplay/Source/DisplayCluster/Private/Render/Viewport/RenderFrame/DisplayClusterRenderFrameSettings.cpp
+++ b/Engine/Plugins/Runtime/nDisplay/Source/DisplayCluster/Private/Render/Viewport/RenderFrame/DisplayClusterRenderFrameSettings.cpp
@@ -219,13 +219,22 @@ int32 FDisplayClusterRenderFrameSettings::GetViewPerViewportAmount() const
 {
 	switch (RenderMode)
 	{
+	// Quadbuffer stereo
 	case EDisplayClusterRenderFrameMode::Stereo:
+		return 2;
+
+	// Render stereo on monoscopic display
 	case EDisplayClusterRenderFrameMode::SideBySide:
 	case EDisplayClusterRenderFrameMode::TopBottom:
 	case EDisplayClusterRenderFrameMode::PIE_SideBySide:
 	case EDisplayClusterRenderFrameMode::PIE_TopBottom:
 		return 2;
 
+	// MRQ/MRG stereo
+	case EDisplayClusterRenderFrameMode::MRQ_Stereo:
+	case EDisplayClusterRenderFrameMode::MRG_Stereo:
+		return 2;
+
 	default:
 		break;
 	}
diff --git a/Engine/Plugins/Runtime/nDisplay/Source/DisplayCluster/Private/Render/Viewport/RenderFrame/DisplayClusterRenderFrameSettings.h b/Engine/Plugins/Runtime/nDisplay/Source/DisplayCluster/Private/Render/Viewport/RenderFrame/DisplayClusterRenderFrameSettings.h
--- a/Engine/Plugins/Runtime/nDisplay/Source/DisplayCluster/Private/Render/Viewport/RenderFrame/DisplayClusterRenderFrameSettings.h
+++ b/Engine/Plugins/Runtime/nDisplay/Source/DisplayCluster/Private/Render/Viewport/RenderFrame/DisplayClusterRenderFrameSettings.h
@@ -11,6 +11,9 @@
 // Settings for render frame builder
 struct FDisplayClusterRenderFrameSettings
 {
+	// If set, use this value instead of GFrameNumber.
+	TOptional<uint32> FrameNumberOverride;
+
 	// customize mono\stereo render modes
 	EDisplayClusterRenderFrameMode RenderMode = EDisplayClusterRenderFrameMode::Unknown;
 
@@ -168,4 +171,15 @@ public:
 
 	/** Return number of contexts per viewport. */
 	int32 GetViewPerViewportAmount() const;
+
+	/** Return current frame number. */
+	inline uint32 GetFrameNumber() const
+	{
+		if (FrameNumberOverride.IsSet())
+		{
+			return *FrameNumberOverride;
+		}
+
+		return GFrameNumber;
+	}
 };
diff --git a/Engine/Plugins/Runtime/nDisplay/Source/DisplayCluster/Public/DisplayClusterRootActor.h b/Engine/Plugins/Runtime/nDisplay/Source/DisplayCluster/Public/DisplayClusterRootActor.h
--- a/Engine/Plugins/Runtime/nDisplay/Source/DisplayCluster/Public/DisplayClusterRootActor.h
+++ b/Engine/Plugins/Runtime/nDisplay/Source/DisplayCluster/Public/DisplayClusterRootActor.h
@@ -557,10 +557,6 @@ public:
 	bool bPreviewInGameRenderFrustum = true;
 
 #if WITH_EDITORONLY_DATA
-	/** When the MRQ is rendered, this flag is raised. */
-	UPROPERTY(Transient, NonTransactional)
-	bool bMoviePipelineRenderPass = false;
-
 	/** Selectively preview a specific viewport or show all/none. */
 	UPROPERTY(EditAnywhere, BlueprintReadWrite, Category = "Preview", meta = (DisplayName = "Preview Node"))
 	FString PreviewNodeId = DisplayClusterConfigurationStrings::gui::preview::PreviewNodeNone;
@@ -767,4 +763,51 @@ private:
 	// UE_DEPRECATED 5.4
 	FOnPreviewUpdated DeprecatedPreviewDelegate;
 #endif
+
+public:
+	/**
+	 * Registers an owner acquiring an exclusive lock on this DCRA.
+	 * While exclusively locked:
+	 *  - No proxy actors may use this DCRA simultaneously (ICVFX panel, StageApp, etc).
+	 *  - Default scene preview rendering is disabled.
+	 *  - PIE rendering is disabled.
+	 *
+	 * Intended for scenarios where a system (e.g. Movie Pipeline) requires
+	 * sole ownership of the DCRA without interference.
+	 */
+	void AddExclusiveLockOwner(const void* InOwner)
+	{
+		FScopeLock Lock(&ExclusiveLockOwnersMutex);
+		ExclusiveLockOwners.Add(InOwner);
+	}
+
+	/**
+	 * Unregisters an owner previously registered via AddExclusiveLockOwner().
+	 */
+	void RemoveExclusiveLockOwner(const void* InOwner)
+	{
+		FScopeLock Lock(&ExclusiveLockOwnersMutex);
+		ExclusiveLockOwners.Remove(InOwner);
+	}
+
+	/**
+	 * Returns true if this DCRA is under an exclusive lock.
+	 * 
+	 * While exclusively locked:
+	 *  - No proxy actors may use this DCRA simultaneously (ICVFX panel, StageApp, etc).
+	 *  - Default scene preview rendering is disabled.
+	 *  - PIE rendering is disabled.
+	 */
+	bool IsExclusiveLocked() const
+	{
+		FScopeLock Lock(&ExclusiveLockOwnersMutex);
+		return !ExclusiveLockOwners.IsEmpty();
+	}
+
+private:
+	/** Owners holding an exclusive lock on this DCRA. See AddExclusiveLockOwner(). */
+	TSet<const void*> ExclusiveLockOwners;
+
+	/** Mutex for ExclusiveLockOwners. */
+	mutable FCriticalSection ExclusiveLockOwnersMutex;
 };
diff --git a/Engine/Plugins/Runtime/nDisplay/Source/DisplayCluster/Public/Math/DisplayClusterProjectionMath.h b/Engine/Plugins/Runtime/nDisplay/Source/DisplayCluster/Public/Math/DisplayClusterProjectionMath.h
new file mode 100644
--- /dev/null
+++ b/Engine/Plugins/Runtime/nDisplay/Source/DisplayCluster/Public/Math/DisplayClusterProjectionMath.h
@@ -0,0 +1,125 @@
+// Copyright Epic Games, Inc. All Rights Reserved.
+
+#pragma once
+
+#include "CoreMinimal.h"
+
+/** Math utilities for nDisplay projection and frustum calculations. */
+struct FDisplayClusterProjectionMath
+{
+	/**
+	 * Calculates the ViewOffset for the eye from the view location.
+	 * @param InViewRotation           Rotation of the view; used to derive the eye direction.
+	 * @param StereoEyeOffsetDistance  Distance to the eye from the midpoint between the eyes.
+	 * @return  Offset vector from the original ViewLocation to the eye.
+	 */
+	static inline FVector GetStereoEyeOffset(const FRotator& InViewRotation, const float StereoEyeOffsetDistance)
+	{
+		return InViewRotation.Quaternion().RotateVector(FVector(0.0f, StereoEyeOffsetDistance, 0.0f));
+	}
+
+	/** Ensures the frustum range [InOutValue0, InOutValue1] is at least 1 degree wide at the given near distance. */
+	static inline void GetNonZeroFrustumRange(double& InOutValue0, double& InOutValue1, double n)
+	{
+		static const double MinHalfFOVRangeRad = FMath::DegreesToRadians(0.5f);
+		static const double MinRangeBase = FMath::Tan(MinHalfFOVRangeRad * 2);
+
+		const double MinRangeValue = n * MinRangeBase;
+		if ((InOutValue1 - InOutValue0) < MinRangeValue)
+		{
+			// Get minimal values from center of range
+			const double CenterRad = (FMath::Atan(InOutValue0 / n) + FMath::Atan(InOutValue1 / n)) * 0.5;
+			InOutValue0 = n * FMath::Tan(CenterRad - MinHalfFOVRangeRad);
+			InOutValue1 = n * FMath::Tan(CenterRad + MinHalfFOVRangeRad);
+		}
+	}
+
+	/** Converts a frustum angle in degrees to a signed extent on the near plane. */
+	template <typename T>
+	static T DegreesToNearPlaneExtent(const T AngleDeg, const T ZNear)
+	{
+		return ZNear * FMath::Tan(FMath::DegreesToRadians(AngleDeg));
+	}
+
+	/** Converts a signed extent on the near plane to a frustum angle in degrees. */
+	template <typename T>
+	static T NearPlaneExtentToDegrees(const T NearPlaneExtent, const T ZNear)
+	{
+		return ZNear != T(0)
+			? FMath::RadiansToDegrees(FMath::Atan(NearPlaneExtent / ZNear))
+			: T(0);
+	}
+
+	/** Returns true if matrices A and B are nearly equal within the given Tolerance. */
+	template <typename T>
+	static bool MatricesNearlyEqual(const UE::Math::TMatrix<T>& A, const UE::Math::TMatrix<T>& B, float Tolerance = KINDA_SMALL_NUMBER)
+	{
+		for (int32 R = 0; R < 4; ++R)
+		{
+			for (int32 C = 0; C < 4; ++C)
+			{
+				if (!FMath::IsNearlyEqual(A.M[R][C], B.M[R][C], Tolerance))
+				{
+					return false;
+				}
+			}
+		}
+
+		return true;
+	}
+
+	/**
+	 * Builds a perspective projection matrix from asymmetric near-plane extents (off-axis frustum).
+	 * Inputs are signed distances from the eye to each frustum edge, measured at the near plane.
+	 * The result uses Unreal's Z-inverted left-handed convention.
+	 *
+	 * @param InLeft    Left edge extent at the near plane (typically negative).
+	 * @param InRight   Right edge extent at the near plane (typically positive).
+	 * @param InTop     Top edge extent at the near plane (typically positive).
+	 * @param InBottom  Bottom edge extent at the near plane (typically negative).
+	 * @param InZNear   Near clip plane distance (must be > 0).
+	 * @param InZFar    Far clip plane distance. Pass InZNear to use an unlimited far plane.
+	 */
+	template <typename T>
+	static UE::Math::TMatrix<T> MakeAsymmetricProjectionMatrix(const T InLeft, const T InRight, const T InTop, const T InBottom, const T InZNear, const T InZFar)
+	{
+		const T mx = static_cast<T>(2) * InZNear / (InRight - InLeft);
+		const T my = static_cast<T>(2) * InZNear / (InTop - InBottom);
+
+		const T ma = -(InRight + InLeft) / (InRight - InLeft);
+		const T mb = -(InTop + InBottom) / (InTop - InBottom);
+
+		// Support unlimited far plane (InZFar == InZNear)
+		const T mc = (InZFar == InZNear) ? (static_cast<T>(1) - Z_PRECISION) : (InZFar / (InZFar - InZNear));
+		const T md = (InZFar == InZNear) ? (-InZNear * (static_cast<T>(1) - Z_PRECISION)) : (-(InZFar * InZNear) / (InZFar - InZNear));
+
+		const T me = static_cast<T>(1);
+
+		// Normal LHS
+		const UE::Math::TMatrix<T> ProjectionMatrix = UE::Math::TMatrix<T>(
+			FPlane(mx, 0,  0,  0),
+			FPlane(0,  my, 0,  0),
+			FPlane(ma, mb, mc, me),
+			FPlane(0,  0,  md, 0));
+
+		// Invert Z-axis (UE uses Z-inverted LHS)
+		static const UE::Math::TMatrix<T> FlipZ = UE::Math::TMatrix<T>(
+			FPlane(1, 0,  0, 0),
+			FPlane(0, 1,  0, 0),
+			FPlane(0, 0, -1, 0),
+			FPlane(0, 0,  1, 1));
+
+		return ProjectionMatrix * FlipZ;
+	}
+
+	/**
+	 * Builds a perspective projection matrix from asymmetric near-plane extents (off-axis frustum).
+	 * @param InFrustumExtents  Near-plane extents as (Left, Right, Top, Bottom).
+	 * @param InClipPlanes      Clip distances as (ZNear, ZFar). Pass ZNear in both components for an unlimited far plane.
+	 */
+	template <typename T>
+	static UE::Math::TMatrix<T> MakeAsymmetricProjectionMatrix(const UE::Math::TVector4<T>& InFrustumExtents, const UE::Math::TVector2<T>& InClipPlanes)
+	{
+		return MakeAsymmetricProjectionMatrix(InFrustumExtents.X, InFrustumExtents.Y, InFrustumExtents.Z, InFrustumExtents.W, InClipPlanes.X, InClipPlanes.Y);
+	}
+};
diff --git a/Engine/Plugins/Runtime/nDisplay/Source/DisplayCluster/Public/Render/Viewport/Containers/DisplayClusterViewport_Context.h b/Engine/Plugins/Runtime/nDisplay/Source/DisplayCluster/Public/Render/Viewport/Containers/DisplayClusterViewport_Context.h
--- a/Engine/Plugins/Runtime/nDisplay/Source/DisplayCluster/Public/Render/Viewport/Containers/DisplayClusterViewport_Context.h
+++ b/Engine/Plugins/Runtime/nDisplay/Source/DisplayCluster/Public/Render/Viewport/Containers/DisplayClusterViewport_Context.h
@@ -21,6 +21,22 @@ public:
 		, StereoViewIndex(InStereoViewIndex)
 	{ }
 
+	/** Returns the effective projection matrix for rendering. */
+	inline const FMatrix& GetRenderProjectionMatrix() const
+	{
+		return ProjectionData.bUseOverscan ?
+			OverscanProjectionMatrix :
+			ProjectionMatrix;
+	}
+
+	/** Returns the effective near-plane frustum extents for rendering (X=Left, Y=Right, Z=Top, W=Bottom), in world units. */
+	inline const FVector4& GetRenderNearPlaneExtents() const
+	{
+		return ProjectionData.bUseOverscan
+			? ProjectionData.OverscanProjectionAngles
+			: ProjectionData.ProjectionAngles;
+	}
+
 public:
 	/**
 	* Declare all bitfields together to avoid alignment and padding.
@@ -78,10 +94,15 @@ public:
 		// This value is used as an argument to IDisplayClusterProjectionPolicy::CalculateView()
 		FRotator EyeRotation = FRotator::ZeroRotator;
 
-		// Render view location used in world space
-		FVector  RenderLocation = FVector::ZeroVector;
+		// Camera head (pivot) position in world space.
+		// Equals RenderLocation minus the stereo eye offset; identical to RenderLocation for mono rendering.
+		FVector RenderCameraLocation = FVector::ZeroVector;
+
+		// Render eye position in world space; the actual view origin passed to the scene renderer.
+		// Differs from RenderCameraLocation by the stereo eye offset in stereo contexts.
+		FVector RenderLocation = FVector::ZeroVector;
 
-		// Render view rotation used in world space
+		// View rotation in world space; shared by both the camera head and the render eye.
 		FRotator RenderRotation = FRotator::ZeroRotator;
 	};
 
@@ -96,10 +117,12 @@ public:
 		// Is overscan used
 		bool bUseOverscan = false;
 
-		// Projection angles [Left, Right, Top, Bottom]
+		// Near-plane frustum extents for rendering in world units.
+		// ( X=Left, Y=Right, Z=Top, W=Bottom )
 		FVector4 ProjectionAngles;
 
-		// Projection angles for Overscan [Left, Right, Top, Bottom]
+		// Near-plane frustum extents for rendering with overscan in world units.
+		// ( X=Left, Y=Right, Z=Top, W=Bottom )
 		FVector4 OverscanProjectionAngles;
 
 		// Clipping planes
diff --git a/Engine/Plugins/Runtime/nDisplay/Source/DisplayCluster/Public/Render/Viewport/IDisplayClusterViewport.h b/Engine/Plugins/Runtime/nDisplay/Source/DisplayCluster/Public/Render/Viewport/IDisplayClusterViewport.h
--- a/Engine/Plugins/Runtime/nDisplay/Source/DisplayCluster/Public/Render/Viewport/IDisplayClusterViewport.h
+++ b/Engine/Plugins/Runtime/nDisplay/Source/DisplayCluster/Public/Render/Viewport/IDisplayClusterViewport.h
@@ -169,6 +169,9 @@ public:
 	/** Get viewport contexts */
 	virtual const TArray<FDisplayClusterViewport_Context>& GetContexts() const = 0;
 
+	/** Get viewport contexts captured during the previous frame. */
+	virtual const TArray<FDisplayClusterViewport_Context>& GetPrevFrameContexts() const = 0;
+
 	/**
 	 * Returns the ViewState for the specified viewport context, creating it if necessary.
 	 *
@@ -235,13 +238,14 @@ public:
 
 	/** Get View from current player camera
 	*
-	* @param InWorld - In this world, an active game camera will be searched for.
-	* @param bUseCameraPostprocess - if false, then InOutViewInfo.PostProcessBlendWeight is set to 0.
-	* @param InOutViewInfo         - (in,out) ViewPoint data
+	* @param InWorld                    - (in)       In this world, an active game camera will be searched for.
+	* @param bUseCameraPostprocess      - (in)       if false, then InOutViewInfo.PostProcessBlendWeight is set to 0.
+	* @param InOutViewInfo              - (in,out)   ViewPoint data
+	* @param OutCustomNearClippingPlane - (opt, out) Custom NCP, or a value less than zero if not defined.
 	*
 	* @return false if the viewpoint data cannot be retrieved.
 	*/
-	static bool GetPlayerCameraView(UWorld* InWorld, const bool bUseCameraPostprocess, FMinimalViewInfo& InOutViewInfo);
+	static bool GetPlayerCameraView(UWorld* InWorld, const bool bUseCameraPostprocess, FMinimalViewInfo& InOutViewInfo, float* OutCustomNearClippingPlane = nullptr);
 
 	///////////////// UE_DEPRECATED 5.3 ///////////////////
 
diff --git a/Engine/Plugins/Runtime/nDisplay/Source/DisplayCluster/Public/Render/Viewport/IDisplayClusterViewportConfiguration.h b/Engine/Plugins/Runtime/nDisplay/Source/DisplayCluster/Public/Render/Viewport/IDisplayClusterViewportConfiguration.h
--- a/Engine/Plugins/Runtime/nDisplay/Source/DisplayCluster/Public/Render/Viewport/IDisplayClusterViewportConfiguration.h
+++ b/Engine/Plugins/Runtime/nDisplay/Source/DisplayCluster/Public/Render/Viewport/IDisplayClusterViewportConfiguration.h
@@ -44,21 +44,31 @@ public:
 	* Update\Create\Delete local node viewports
 	* Updating the configuration to render a ClusterNode in the specified mode
 	*
-	* @param InRenderMode     - Render mode
-	* @param InWorld          - ptr to the world to be rendered
-	* @param InClusterNodeId  - cluster node for rendering
+	* @param InRenderMode          - Render mode
+	* @param InWorld               - ptr to the world to be rendered
+	* @param InClusterNodeId       - cluster node for rendering
+	* @param InFrameNumberOverride - (opt) If set, use this value instead of GFrameNumber.
 	*/
-	virtual bool UpdateConfigurationForClusterNode(EDisplayClusterRenderFrameMode InRenderMode, const UWorld* InWorld, const FString& InClusterNodeId) = 0;
+	virtual bool UpdateConfigurationForClusterNode(
+		EDisplayClusterRenderFrameMode InRenderMode,
+		const UWorld* InWorld,
+		const FString& InClusterNodeId,
+		const uint32* InFrameNumberOverride = nullptr) = 0;
 
 	/**
 	* Update\Create\Delete local node viewports
 	* Updating the configuration to render a list of viewports in a given mode
 	*
-	* @param InRenderMode    - Render mode
-	* @param InWorld         - ptr to the world to be rendered
-	* @param InViewportNames - Viewports names for next frame
+	* @param InRenderMode          - Render mode
+	* @param InWorld               - ptr to the world to be rendered
+	* @param InViewportNames       - Viewports names for next frame
+	* @param InFrameNumberOverride - (opt) If set, use this value instead of GFrameNumber.
 	*/
-	virtual bool UpdateConfigurationForViewportsList(EDisplayClusterRenderFrameMode InRenderMode, const UWorld* InWorld, const TArray<FString>& InViewportNames) = 0;
+	virtual bool UpdateConfigurationForViewportsList(
+		EDisplayClusterRenderFrameMode InRenderMode,
+		const UWorld* InWorld,
+		const TArray<FString>& InViewportNames,
+		const uint32* InFrameNumberOverride = nullptr) = 0;
 
 	/** Release the current configuration and free resources. */
 	virtual void ReleaseConfiguration() = 0;
@@ -139,6 +149,9 @@ public:
 	/** Returns true if this DCRA is using cluster rendering. */
 	virtual bool IsClusterRendering() const = 0;
 
+	/** Returns true if MoviePipeline is used for rendering. */
+	virtual bool IsMoviePipelineRendering() const = 0;
+
 	/** Returns true, if Techvis is used. */
 	virtual bool IsTechvisEnabled() const = 0;
 
diff --git a/Engine/Plugins/Runtime/nDisplay/Source/DisplayCluster/Public/Render/Viewport/IDisplayClusterViewportManager.h b/Engine/Plugins/Runtime/nDisplay/Source/DisplayCluster/Public/Render/Viewport/IDisplayClusterViewportManager.h
--- a/Engine/Plugins/Runtime/nDisplay/Source/DisplayCluster/Public/Render/Viewport/IDisplayClusterViewportManager.h
+++ b/Engine/Plugins/Runtime/nDisplay/Source/DisplayCluster/Public/Render/Viewport/IDisplayClusterViewportManager.h
@@ -26,6 +26,12 @@ struct FDisplayClusterRenderFrameTargetViewFamily;
 class DISPLAYCLUSTER_API IDisplayClusterViewportManager
 {
 public:
+	/**
+	* Must be called before the destructor.
+	* Because all subsystems connected to this class via WeakPtr will not work when we are in destructor.
+	*/
+	virtual void Release() = 0;
+
 	virtual ~IDisplayClusterViewportManager() = default;
 
 public:
diff --git a/Engine/Plugins/Runtime/nDisplay/Source/DisplayCluster/Public/Render/Viewport/RenderFrame/DisplayClusterRenderFrameEnums.h b/Engine/Plugins/Runtime/nDisplay/Source/DisplayCluster/Public/Render/Viewport/RenderFrame/DisplayClusterRenderFrameEnums.h
--- a/Engine/Plugins/Runtime/nDisplay/Source/DisplayCluster/Public/Render/Viewport/RenderFrame/DisplayClusterRenderFrameEnums.h
+++ b/Engine/Plugins/Runtime/nDisplay/Source/DisplayCluster/Public/Render/Viewport/RenderFrame/DisplayClusterRenderFrameEnums.h
@@ -5,42 +5,53 @@
 #include "CoreMinimal.h"
 
 /**
-* Rendering Mode.
-*/
-enum class EDisplayClusterRenderFrameMode: uint8
+ * Defines the rendering output mode.
+ */
+enum class EDisplayClusterRenderFrameMode : uint8
 {
 	Unknown = 0,
 
-	// Render to single backbuffer texture
+	// Monoscopic rendering to a backbuffer texture.
 	Mono,
 
-	// Render to two separated backbuffer textures
+	// Stereoscopic rendering to a layered backbuffer texture (quad-buffer).
 	Stereo,
 
-	// Special render: render to two separated backbuffer textures, with half size X
+	// Stereoscopic rendering packed into one texture:
+	// left eye in the left half, right eye in the right half.
 	SideBySide,
 
-	// Special render: render to two separated backbuffer textures, with half size Y
+	// Stereoscopic rendering packed into one texture:
+	// left eye in the top half, right eye in the bottom half.
 	TopBottom,
 
-
-	// Monoscopic rendering in PIE
+	// Monoscopic rendering in Play In Editor (PIE).
 	PIE_Mono,
 
-	// Stereoscopic side-by-side rendering in PIE
+	// Stereoscopic rendering in PIE packed into one texture:
+	// left eye in the left half, right eye in the right half.
 	PIE_SideBySide,
 
-	// Stereoscopic top-bottom rendering in PIE
+	// Stereoscopic rendering in PIE packed into one texture:
+	// left eye in the top half, right eye in the bottom half.
 	PIE_TopBottom,
 
-	// Monoscopic rendering for MRQ
+	// Monoscopic rendering for Movie Render Queue (MRQ).
 	MRQ_Mono,
 
+	// Stereoscopic rendering for Movie Render Queue (MRQ).
+	MRQ_Stereo,
+
+	// Monoscopic rendering for Movie Render Graph (MRG).
+	MRG_Mono,
+
+	// Stereoscopic rendering for Movie Render Graph (MRG).
+	MRG_Stereo,
 
-	// Special render for preview in scene
+	// Rendering mode used for in-scene preview.
 	PreviewInScene,
 
-	// Special render of ProxyHit for preview in scene
-	// *** Not implemented.
+	// ProxyHit rendering mode used for in-scene preview.
+	// Not implemented.
 	PreviewProxyHitInScene,
-};
+};
\ No newline at end of file
diff --git a/Engine/Plugins/Runtime/nDisplay/Source/DisplayClusterMoviePipeline/Private/DisplayClusterMoviePipelineViewportPass.cpp b/Engine/Plugins/Runtime/nDisplay/Source/DisplayClusterMoviePipeline/Private/DisplayClusterMoviePipelineViewportPass.cpp
--- a/Engine/Plugins/Runtime/nDisplay/Source/DisplayClusterMoviePipeline/Private/DisplayClusterMoviePipelineViewportPass.cpp
+++ b/Engine/Plugins/Runtime/nDisplay/Source/DisplayClusterMoviePipeline/Private/DisplayClusterMoviePipelineViewportPass.cpp
@@ -438,10 +438,8 @@ void UDisplayClusterMoviePipelineViewportPassBase::ReleaseDisplayCluster()
 		// Delete all used rendering resources
 		DCRootActor->RemoveViewportManager();
 
-#if WITH_EDITORONLY_DATA
-		// Exit from the exclusive mode for MRQ rendering.
-		DCRootActor->bMoviePipelineRenderPass = false;
-#endif
+		// Unlock DCRA when MRQ rendering done.
+		DCRootActor->RemoveExclusiveLockOwner(this);
 	}
 
 	// Reset all runtime values
@@ -493,10 +491,8 @@ bool UDisplayClusterMoviePipelineViewportPassBase::InitializeDisplayCluster()
 				if (ViewportManager->GetConfiguration().UpdateConfigurationForViewportsList(EDisplayClusterRenderFrameMode::MRQ_Mono, GetWorld(), DisplayClusterViewports))
 				{
 
-#if WITH_EDITORONLY_DATA
-					// Enter to the exclusive mode for MRQ rendering.
-					DCRootActor->bMoviePipelineRenderPass = true;
-#endif
+					// Disable preview/PIE rendering for this DCRA.
+					DCRootActor->AddExclusiveLockOwner(this);
 
 					// Finally set this DCRA for rendering
 					RootActorWeakPtr = DCRootActor;
diff --git a/Engine/Plugins/Runtime/nDisplay/Source/DisplayClusterProjection/Private/Policy/Camera/DisplayClusterProjectionCameraPolicy.cpp b/Engine/Plugins/Runtime/nDisplay/Source/DisplayClusterProjection/Private/Policy/Camera/DisplayClusterProjectionCameraPolicy.cpp
--- a/Engine/Plugins/Runtime/nDisplay/Source/DisplayClusterProjection/Private/Policy/Camera/DisplayClusterProjectionCameraPolicy.cpp
+++ b/Engine/Plugins/Runtime/nDisplay/Source/DisplayClusterProjection/Private/Policy/Camera/DisplayClusterProjectionCameraPolicy.cpp
@@ -170,7 +170,7 @@ bool FDisplayClusterProjectionCameraPolicy::ImplSetupProjectionViewPoint(IDispla
 	else if (UWorld* CurrentWorld = InViewport ? InViewport->GetConfiguration().GetCurrentWorld() : nullptr)
 	{
 		// Get active player camera
-		bResult = IDisplayClusterViewport::GetPlayerCameraView(CurrentWorld, CameraSettings.bUseCameraPostprocess, InOutViewInfo);
+		bResult = IDisplayClusterViewport::GetPlayerCameraView(CurrentWorld, CameraSettings.bUseCameraPostprocess, InOutViewInfo, OutCustomNearClippingPlane);
 	}
 
 	// Fix camera lens deffects (prototype)
