[nDisplay] Fixed CIS issues
- Use UE_GET_LOG_VERBOSITY macro instead of direct GetVerbosity() calls
- Forward declarations
- Include filename capitalization (Linux)

#rnx
#Synced-CL 0

diff --git a/Engine/Plugins/Runtime/nDisplay/Source/DisplayClusterMonitor/Private/DisplayClusterMonitorMessenger.cpp b/Engine/Plugins/Runtime/nDisplay/Source/DisplayClusterMonitor/Private/DisplayClusterMonitorMessenger.cpp
--- a/Engine/Plugins/Runtime/nDisplay/Source/DisplayClusterMonitor/Private/DisplayClusterMonitorMessenger.cpp
+++ b/Engine/Plugins/Runtime/nDisplay/Source/DisplayClusterMonitor/Private/DisplayClusterMonitorMessenger.cpp
@@ -169,12 +169,11 @@ namespace UE::nDisplay::Monitor
 		// Validate input
 		if (!MessageEndpoint || !InPayload || !IsValid(InType) || InRecipients.IsEmpty())
 		{
-			delete InPayload;
 			return;
 		}
 
 		// Conditional detailed logs
-		if(LogDisplayClusterMonitorMsg.GetVerbosity() >= ELogVerbosity::VeryVerbose)
+		if(UE_GET_LOG_VERBOSITY(LogDisplayClusterMonitorMsg) >= ELogVerbosity::VeryVerbose)
 		{
 			for (const FMessageAddress& Recipient : InRecipients)
 			{
@@ -192,7 +191,6 @@ namespace UE::nDisplay::Monitor
 		// Validate input
 		if (!MessageEndpoint || !InPayload || !IsValid(InType) || InEndpointTypes.IsEmpty())
 		{
-			delete InPayload;
 			return;
 		}
 
@@ -206,7 +204,7 @@ namespace UE::nDisplay::Monitor
 			[&InEndpointTypes](const FDCEndpoint& EP) { return EP.Address; });
 
 		// Conditional detailed logs
-		if (LogDisplayClusterMonitorMsg.GetVerbosity() >= ELogVerbosity::VeryVerbose)
+		if (UE_GET_LOG_VERBOSITY(LogDisplayClusterMonitorMsg) >= ELogVerbosity::VeryVerbose)
 		{
 			for (const FMessageAddress& Recipient : Recipients)
 			{
@@ -223,7 +221,6 @@ namespace UE::nDisplay::Monitor
 	{
 		if (!MessageEndpoint || !InPayload || !IsValid(InType))
 		{
-			delete InPayload;
 			return;
 		}
 
diff --git a/Engine/Plugins/Runtime/nDisplay/Source/DisplayClusterMonitor/Private/DisplayClusterMonitorProviderMedia.cpp b/Engine/Plugins/Runtime/nDisplay/Source/DisplayClusterMonitor/Private/DisplayClusterMonitorProviderMedia.cpp
--- a/Engine/Plugins/Runtime/nDisplay/Source/DisplayClusterMonitor/Private/DisplayClusterMonitorProviderMedia.cpp
+++ b/Engine/Plugins/Runtime/nDisplay/Source/DisplayClusterMonitor/Private/DisplayClusterMonitorProviderMedia.cpp
@@ -9,7 +9,7 @@
 #include "MediaObservables/MediaObservableUI.h"
 #include "Misc/DisplayClusterConsoleExec.h"
 #include "Misc/Guid.h"
-#include "Render/GUILayer/IDisplayClusterGuiLayerController.h"
+#include "Render/GUILayer/IDisplayClusterGUILayerController.h"
 #include "Render/IDisplayClusterRenderManager.h"
 #include "Render/Viewport/Containers/DisplayClusterViewport_Enums.h"
 #include "Render/Viewport/IDisplayClusterViewport.h"
@@ -294,7 +294,7 @@ namespace UE::nDisplay::Monitor
 	void FDisplayClusterMonitorProviderMedia::OnEndFrame()
 	{
 		// Log information about active sessions
-		if (LogDisplayClusterMonitorProviderMedia.GetVerbosity() >= ELogVerbosity::VeryVerbose)
+		if (UE_GET_LOG_VERBOSITY(LogDisplayClusterMonitorProviderMedia) >= ELogVerbosity::VeryVerbose)
 		{
 			int32 SessionsNum = 0;
 			for (const auto& [Monitor, SessionMap] : ActiveSessions)
diff --git a/Engine/Plugins/Runtime/nDisplay/Source/DisplayClusterMonitorEditor/Private/Widgets/SClusterTreeView.h b/Engine/Plugins/Runtime/nDisplay/Source/DisplayClusterMonitorEditor/Private/Widgets/SClusterTreeView.h
--- a/Engine/Plugins/Runtime/nDisplay/Source/DisplayClusterMonitorEditor/Private/Widgets/SClusterTreeView.h
+++ b/Engine/Plugins/Runtime/nDisplay/Source/DisplayClusterMonitorEditor/Private/Widgets/SClusterTreeView.h
@@ -12,6 +12,7 @@
 class IClusterMonitorController;
 class IClusterObservable;
 class IClusterResidence;
+class SClusterTreeToolbar;
 
 
 /**
diff --git a/Engine/Plugins/Runtime/nDisplay/Source/DisplayClusterMonitorEditor/Private/Widgets/SSessionViewport.h b/Engine/Plugins/Runtime/nDisplay/Source/DisplayClusterMonitorEditor/Private/Widgets/SSessionViewport.h
--- a/Engine/Plugins/Runtime/nDisplay/Source/DisplayClusterMonitorEditor/Private/Widgets/SSessionViewport.h
+++ b/Engine/Plugins/Runtime/nDisplay/Source/DisplayClusterMonitorEditor/Private/Widgets/SSessionViewport.h
@@ -11,6 +11,7 @@ class IClusterObservable;
 class SClusterSessionsView;
 class SObservableMediaImage;
 class SSessionViewport;
+class SWidget;
 struct FToolMenuEntry;
 
 
